#UE4 - Fix copy & paste in comment

diff --git a/Engine/Source/Runtime/Online/OnlineSubsystem/Public/Interfaces/OnlineSessionInterface.h b/Engine/Source/Runtime/Online/OnlineSubsystem/Public/Interfaces/OnlineSessionInterface.h
--- a/Engine/Source/Runtime/Online/OnlineSubsystem/Public/Interfaces/OnlineSessionInterface.h
+++ b/Engine/Source/Runtime/Online/OnlineSubsystem/Public/Interfaces/OnlineSessionInterface.h
@@ -475,7 +475,7 @@ public:
 	 * Delegate fired when the session registration process has completed
 	 *
 	 * @param SessionName the name of the session the player joined or not
-	 * @param PlayerId the player that was unregistered from the online service
+	 * @param PlayerId the player that was registered in the online service
 	 * @param bWasSuccessful true if the async action completed without error, false if there was an error
 	 */
 	DEFINE_ONLINE_DELEGATE_THREE_PARAM(OnRegisterPlayersComplete, FName, const TArray< TSharedRef<class FUniqueNetId> >&, bool);
