build script, remove libcurl from rocket

diff --git a/Engine/Source/Runtime/Online/HTTP/HTTP.Build.cs b/Engine/Source/Runtime/Online/HTTP/HTTP.Build.cs
--- a/Engine/Source/Runtime/Online/HTTP/HTTP.Build.cs
+++ b/Engine/Source/Runtime/Online/HTTP/HTTP.Build.cs
@@ -20,7 +20,10 @@ public class HTTP : ModuleRules
 			Target.Platform == UnrealTargetPlatform.Win64)
 		{
 			AddThirdPartyPrivateStaticDependencies(Target, "WinInet");
-			AddThirdPartyPrivateStaticDependencies(Target, "libcurl");
+            if (!UnrealBuildTool.UnrealBuildTool.BuildingRocket() && !UnrealBuildTool.UnrealBuildTool.RunningRocket())
+            {
+                AddThirdPartyPrivateStaticDependencies(Target, "libcurl");
+            }
 		}
 
 		if (Target.Platform == UnrealTargetPlatform.Linux)
