CL for MS1 from 22nd Toys. Integrates from our branch into Main.

 - Friends Interface
 - Leaderboards Interface
 - User Interface
 - User Cloud Interface
 - Title File Interface

Fixes / additions to some Test classes.

diff --git a/Engine/Source/Runtime/Online/OnlineSubsystem/Public/Interfaces/OnlineFriendsInterface.h b/Engine/Source/Runtime/Online/OnlineSubsystem/Public/Interfaces/OnlineFriendsInterface.h
--- a/Engine/Source/Runtime/Online/OnlineSubsystem/Public/Interfaces/OnlineFriendsInterface.h
+++ b/Engine/Source/Runtime/Online/OnlineSubsystem/Public/Interfaces/OnlineFriendsInterface.h
@@ -12,7 +12,11 @@ namespace EFriendsLists
 		/** default friends list */
 		Default,
 		/** recent players friends list */
-		RecentPlayers
+		RecentPlayers,
+		/** online players friends list */
+		OnlinePlayers,
+		/** list of players running the same title/game */
+		InGamePlayers,
 	};
 
 	/** @return the stringified version of the enum passed in */
@@ -24,6 +28,10 @@ namespace EFriendsLists
 				return TEXT("default");
 			case RecentPlayers:
 				return TEXT("recentPlayers");
+			case OnlinePlayers:
+				return TEXT("onlinePlayers");
+			case InGamePlayers:
+				return TEXT("inGamePlayers");
 		}
 		return TEXT("");
 	}
diff --git a/Engine/Source/Runtime/Online/OnlineSubsystemUtils/Private/Tests/TestCloudInterface.cpp b/Engine/Source/Runtime/Online/OnlineSubsystemUtils/Private/Tests/TestCloudInterface.cpp
--- a/Engine/Source/Runtime/Online/OnlineSubsystemUtils/Private/Tests/TestCloudInterface.cpp
+++ b/Engine/Source/Runtime/Online/OnlineSubsystemUtils/Private/Tests/TestCloudInterface.cpp
@@ -73,6 +73,7 @@ bool FTestCloudInterface::Tick( float DeltaTime )
 			break;
 		case 1:
 			// Write out N files for the user of various sizes/names and enumerate
+			// @todo: make the file count configurable on a per-platform basis, since different platforms have different file count limits
 			WriteNUserCloudFiles(*UserId, FString(TEXT("UserCloud.bin")), 15, OnWriteUserCloudFileCompleteDelegate);
 			break;
 		case 2:
@@ -186,7 +187,9 @@ void FTestCloudInterface::WriteNUserCloudFiles(const FUniqueNetId& InUserId, con
 	TArray<uint8> DummyData;
 	for (int32 FileIdx=0; FileIdx<FileCount; FileIdx++)
 	{
-		WriteRandomFile(DummyData, FMath::Trunc(FMath::FRandRange(1024, 100*1024)));
+		// @todo: make the dummy data size configurable on a per-platform basis, since different platforms (e.g. PS4) have different size limits
+		//WriteRandomFile(DummyData, FMath::Trunc(FMath::FRandRange(1024, 100*1024)));
+		WriteRandomFile(DummyData, FMath::Trunc(FMath::FRandRange(256, 1024)));
 		UserCloud->WriteUserFile(InUserId, FString::Printf(TEXT("%s%d.%s"), *FPaths::GetBaseFilename(FileName), FileIdx, *FPaths::GetExtension(FileName)), DummyData);
 	}
 }
@@ -254,11 +257,13 @@ void FTestCloudInterface::ReadEnumeratedUserFiles(FOnReadUserFileCompleteDelegat
 void FTestCloudInterface::OnReadEnumeratedUserFilesComplete(bool bWasSuccessful, const FUniqueNetId& InUserId, const FString& FileName)
 {
 	int32 FileSize = 0;
+	char *buffer = NULL;
 	if (bWasSuccessful)
 	{
 		TArray<uint8> FileContents;
 		bWasSuccessful = UserCloud->GetFileContents(InUserId, FileName, FileContents);
 		FileSize = FileContents.Num();
+		buffer = (char *) FileContents.GetData();
 	}
 
 	UE_LOG(LogOnline, Log, TEXT("Read user file complete Success:%d UserId:%s FileName:%s Size:%d"), bWasSuccessful, *InUserId.ToDebugString(), *FileName, FileSize);
diff --git a/Engine/Source/Runtime/Online/OnlineSubsystemUtils/Private/Tests/TestFriendsInterface.cpp b/Engine/Source/Runtime/Online/OnlineSubsystemUtils/Private/Tests/TestFriendsInterface.cpp
--- a/Engine/Source/Runtime/Online/OnlineSubsystemUtils/Private/Tests/TestFriendsInterface.cpp
+++ b/Engine/Source/Runtime/Online/OnlineSubsystemUtils/Private/Tests/TestFriendsInterface.cpp
@@ -195,9 +195,9 @@ void FTestFriendsInterface::OnDeleteFriendComplete(int32 LocalPlayer, bool bWasS
 	UE_LOG(LogOnline, Log, TEXT("DeleteFriend() for player (%d) to friend (%s) was success=%d. %s"), 
 		LocalPlayer, *FriendId.ToDebugString(), bWasSuccessful, *ErrorStr);
 
-	// done with this part of the test if no more invites to send
+	// done with this part of the test if no more friends to delete
 	FriendsToDelete.RemoveAt(0);
-	if (FriendsToDelete.Num() == 0)
+	if (bWasSuccessful && FriendsToDelete.Num() == 0)
 	{
 		bDeleteFriends = false;
 		bReadFriendsList = true;
diff --git a/Engine/Source/Runtime/Online/OnlineSubsystemUtils/Private/Tests/TestLeaderboardInterface.cpp b/Engine/Source/Runtime/Online/OnlineSubsystemUtils/Private/Tests/TestLeaderboardInterface.cpp
--- a/Engine/Source/Runtime/Online/OnlineSubsystemUtils/Private/Tests/TestLeaderboardInterface.cpp
+++ b/Engine/Source/Runtime/Online/OnlineSubsystemUtils/Private/Tests/TestLeaderboardInterface.cpp
@@ -65,7 +65,7 @@ void FTestLeaderboardInterface::WriteLeaderboards()
 	
 	// Set some data
 	WriteObject.SetIntStat("TestIntStat1", 50);
-	WriteObject.SetFloatStat("TestFloatStat1", 99.0f);
+	WriteObject.SetFloatStat("TestFloatStat1", 99.5f);
 
 	// Write it to the buffers
 	Leaderboards->WriteLeaderboards(TEXT("TEST"), *UserId, WriteObject);
@@ -92,6 +92,17 @@ void FTestLeaderboardInterface::OnLeaderboardReadComplete(bool bWasSuccessful)
 	UE_LOG(LogOnline, Verbose, TEXT("OnLeaderboardReadComplete bWasSuccessful: %d"), bWasSuccessful);
 	bOverallSuccess = bOverallSuccess && bWasSuccessful;
 
+	for (int32 RowIdx = 0; RowIdx < ReadObject->Rows.Num(); ++RowIdx)
+	{
+		const FOnlineStatsRow& StatsRow = ReadObject->Rows[RowIdx];
+		UE_LOG(LogOnline, Log, TEXT("Leaderboard stats for: Nickname = %s, Rank = %d"), *StatsRow.NickName, StatsRow.Rank);
+
+		for (FStatsColumnArray::TConstIterator It(StatsRow.Columns); It; ++It)
+		{
+			UE_LOG(LogOnline, Log, TEXT("  %s = %s"), *It.Key().ToString(), *It.Value().ToString());
+		}
+	}
+
 	Leaderboards->ClearOnLeaderboardReadCompleteDelegate(LeaderboardReadCompleteDelegate);
 	TestPhase++;
 }
diff --git a/Engine/Source/Runtime/Online/OnlineSubsystemUtils/Private/Tests/TestUserInterface.cpp b/Engine/Source/Runtime/Online/OnlineSubsystemUtils/Private/Tests/TestUserInterface.cpp
--- a/Engine/Source/Runtime/Online/OnlineSubsystemUtils/Private/Tests/TestUserInterface.cpp
+++ b/Engine/Source/Runtime/Online/OnlineSubsystemUtils/Private/Tests/TestUserInterface.cpp
@@ -45,6 +45,15 @@ void FTestUserInterface::Test(UWorld* InWorld, const TArray<FString>& InUserIds)
 			}
 		}
 
+		// Include our own user so we have at least one UserId to query for
+		{
+			TSharedPtr<FUniqueNetId> UserId = OnlineSub->GetIdentityInterface()->GetUniquePlayerId(0);
+			if (QueryUserIds.Contains(UserId.ToSharedRef()) == false)
+			{
+				QueryUserIds.Add(UserId.ToSharedRef());
+			}
+		}
+
 		// kick off next test
 		StartNextTest();
 	}
@@ -85,24 +94,31 @@ void FTestUserInterface::OnQueryUserInfoComplete(int32 LocalPlayer, bool bWasSuc
 	UE_LOG(LogOnline, Log,
 		TEXT("GetUserInterface() for player (%d) was success=%d"), LocalPlayer, bWasSuccessful);
 
-	for (int32 UserIdx=0; UserIdx < UserIds.Num(); UserIdx++)
+	if (bWasSuccessful)
 	{
-		TSharedPtr<FOnlineUser> User = OnlineSub->GetUserInterface()->GetUserInfo(LocalPlayer, *UserIds[UserIdx]);
-		if (User.IsValid())
-		{
-			UE_LOG(LogOnline, Log,
-				TEXT("PlayerId=%s found"), *UserIds[UserIdx]->ToDebugString());
-			UE_LOG(LogOnline, Log,
-				TEXT("	DisplayName=%s"), *User->GetDisplayName());
-			UE_LOG(LogOnline, Log,
-				TEXT("	RealName=%s"), *User->GetRealName());
-		}
-		else
+		for (int32 UserIdx=0; UserIdx < UserIds.Num(); UserIdx++)
 		{
-			UE_LOG(LogOnline, Log,
-				TEXT("PlayerId=%s not found"), *UserIds[UserIdx]->ToDebugString());
+			TSharedPtr<FOnlineUser> User = OnlineSub->GetUserInterface()->GetUserInfo(LocalPlayer, *UserIds[UserIdx]);
+			if (User.IsValid())
+			{
+				UE_LOG(LogOnline, Log,
+					TEXT("PlayerId=%s found"), *UserIds[UserIdx]->ToDebugString());
+				UE_LOG(LogOnline, Log,
+					TEXT("	DisplayName=%s"), *User->GetDisplayName());
+				UE_LOG(LogOnline, Log,
+					TEXT("	RealName=%s"), *User->GetRealName());
+			}
+			else
+			{
+				UE_LOG(LogOnline, Log,
+					TEXT("PlayerId=%s not found"), *UserIds[UserIdx]->ToDebugString());
+			}
 		}
 	}
+	else
+	{
+		UE_LOG(LogOnline, Error, TEXT("GetUserInterface() failure. Error = %s"), *ErrorStr);
+	}
 
 	// done
 	bQueryUserInfo = false;
