#ttp 330561 - PORTAL: Patching hangs at 81% when user is patching the portal
#ttp 328690 - PORTAL: The install process of the Editor will fluctuate drastically and periodically stop at certain points.
#summary Adding HTTP progress info to Mac, and allowing Windows HTTP to have smaller read buffers.

diff --git a/Engine/Source/Runtime/Online/HTTP/Private/HttpModule.cpp b/Engine/Source/Runtime/Online/HTTP/Private/HttpModule.cpp
--- a/Engine/Source/Runtime/Online/HTTP/Private/HttpModule.cpp
+++ b/Engine/Source/Runtime/Online/HTTP/Private/HttpModule.cpp
@@ -16,6 +16,7 @@ FCriticalSection FHttpManager::RequestLock;
 void FHttpModule::StartupModule()
 {	
 	Singleton = this;
+	MaxReadBufferSize = 256 * 1024;
 
 	FPlatformHttp::Init();
 
diff --git a/Engine/Source/Runtime/Online/HTTP/Private/Mac/MacHTTP.cpp b/Engine/Source/Runtime/Online/HTTP/Private/Mac/MacHTTP.cpp
--- a/Engine/Source/Runtime/Online/HTTP/Private/Mac/MacHTTP.cpp
+++ b/Engine/Source/Runtime/Online/HTTP/Private/Mac/MacHTTP.cpp
@@ -14,6 +14,7 @@
 FMacHttpRequest::FMacHttpRequest()
 :	Connection(NULL)
 ,	CompletionStatus(EHttpRequestStatus::NotStarted)
+,	ProgressBytesSent(0)
 {
 	UE_LOG(LogHttp, Verbose, TEXT("FMacHttpRequest::FMacHttpRequest()"));
 	Request = [[NSMutableURLRequest alloc] init];
@@ -315,6 +316,15 @@ void FMacHttpRequest::Tick(float DeltaSeconds)
 {
 	if( CompletionStatus == EHttpRequestStatus::Processing || Response->HadError() )
 	{
+		if( OnRequestProgress().IsBound() )
+		{
+			const int32 NumBytesReceived = Response->GetNumBytesReceived();
+			if( ProgressBytesSent < NumBytesReceived )
+			{
+				ProgressBytesSent = NumBytesReceived;
+				OnRequestProgress().Execute(SharedThis(this), NumBytesReceived);
+			}
+		}
 		if( Response->IsReady() )
 		{
 			FinishedRequest();
@@ -546,3 +556,8 @@ bool FMacHttpResponse::HadError()
 	
 	return bHadError;
 }
+
+const int32 FMacHttpResponse::GetNumBytesReceived() const
+{
+	return [[ResponseWrapper Payload] length];
+}
diff --git a/Engine/Source/Runtime/Online/HTTP/Private/Mac/MacHTTP.h b/Engine/Source/Runtime/Online/HTTP/Private/Mac/MacHTTP.h
--- a/Engine/Source/Runtime/Online/HTTP/Private/Mac/MacHTTP.h
+++ b/Engine/Source/Runtime/Online/HTTP/Private/Mac/MacHTTP.h
@@ -94,6 +94,9 @@ private:
 
 	/** Current status of request being processed */
 	EHttpRequestStatus::Type CompletionStatus;
+
+	/** Number of bytes sent to progress update */
+	int32 ProgressBytesSent;
 };
 
 
@@ -169,6 +172,11 @@ public:
 	 */
 	bool HadError();
 
+	/**
+	 * Get the number of bytes received so far
+	 */
+	const int32 GetNumBytesReceived() const;
+
 	/**
 	 * Constructor
 	 *
diff --git a/Engine/Source/Runtime/Online/HTTP/Private/Windows/HttpWinInet.cpp b/Engine/Source/Runtime/Online/HTTP/Private/Windows/HttpWinInet.cpp
--- a/Engine/Source/Runtime/Online/HTTP/Private/Windows/HttpWinInet.cpp
+++ b/Engine/Source/Runtime/Online/HTTP/Private/Windows/HttpWinInet.cpp
@@ -751,6 +751,7 @@ FHttpResponseWinInet::FHttpResponseWinInet(const FHttpRequestWinInet& InRequest)
 ,	ContentLength(0)
 ,	bIsReady(0)
 ,	bResponseSucceeded(0)
+,	MaxReadBufferSize(FHttpModule::Get().GetMaxReadBufferSize())
 {
 
 }
@@ -845,8 +846,6 @@ void FHttpResponseWinInet::ProcessResponse()
 	// We might be calling back into this from another asynchronous read, so continue where we left off.
 	// if there is no content length, we're probably receiving chunked data.
 	ContentLength = QueryContentLength();
-	// Size of the buffer to read when calling InternetReadFile. Payload grows by this amount as necessary
-	const int32 MAX_READ_BUFFER_SIZE = 256 * 1024;
 	// Set buffer size based on content length or hard code if chunked response
 	int32 BufferSize = 0;
 	// For non-chunked responses, allocate one extra uint8 to check if we are sent extra content
@@ -862,7 +861,7 @@ void FHttpResponseWinInet::ProcessResponse()
 	else
 	{
 		// For chunked responses, add data using a fixed size buffer at a time.
-		BufferSize = TotalBytesRead + MAX_READ_BUFFER_SIZE;
+		BufferSize = TotalBytesRead + MaxReadBufferSize;
 		// Size read buffer
 		ResponsePayload.SetNum(BufferSize);
 	}
@@ -871,11 +870,12 @@ void FHttpResponseWinInet::ProcessResponse()
 	int32 LoopCount = 0;
 	do
 	{
+		const int32 NumBytesToRead = FMath::Min<int32>(ResponsePayload.Num() - TotalBytesRead, MaxReadBufferSize);
 		// Read directly into the response payload
 		const BOOL bReadFile = InternetReadFile(
 			Request.RequestHandle, 
 			&ResponsePayload[TotalBytesRead], 
-			ResponsePayload.Num() - TotalBytesRead, 
+			NumBytesToRead, 
 			(LPDWORD)&AsyncBytesRead);
 
 		UE_LOG(LogHttp, VeryVerbose, TEXT("InternetReadFile result=%d (%u bytes read) (%u bytes total read). LoopCount=%d %p"), 
@@ -914,7 +914,7 @@ void FHttpResponseWinInet::ProcessResponse()
 					UE_LOG(LogHttp, Log, TEXT("Response payload (%d bytes read so far) is larger than the content-length (%d). Resizing buffer to accommodate. %p"), 
 						TotalBytesRead, ContentLength, &Request);
 				}
-				ResponsePayload.AddZeroed(MAX_READ_BUFFER_SIZE);
+				ResponsePayload.AddZeroed(MaxReadBufferSize);
 			}
 		}
 		LoopCount++;
diff --git a/Engine/Source/Runtime/Online/HTTP/Private/Windows/HttpWinInet.h b/Engine/Source/Runtime/Online/HTTP/Private/Windows/HttpWinInet.h
--- a/Engine/Source/Runtime/Online/HTTP/Private/Windows/HttpWinInet.h
+++ b/Engine/Source/Runtime/Online/HTTP/Private/Windows/HttpWinInet.h
@@ -270,6 +270,8 @@ private:
 	int32 volatile bResponseSucceeded;
 	/** Threadsafe counter for exchanging bytes read so far with main thread tick */
 	FThreadSafeCounter ProgressBytesRead;
+	/** Max buffer size for individual http reads */
+	int32 MaxReadBufferSize;
 };
 
 /**
diff --git a/Engine/Source/Runtime/Online/HTTP/Public/HttpModule.h b/Engine/Source/Runtime/Online/HTTP/Public/HttpModule.h
--- a/Engine/Source/Runtime/Online/HTTP/Public/HttpModule.h
+++ b/Engine/Source/Runtime/Online/HTTP/Public/HttpModule.h
@@ -109,6 +109,23 @@ public:
 		return HttpSendTimeout;
 	}
 
+	/**
+	 * @return max read buffer size for http requests
+	 */
+	inline int32 GetMaxReadBufferSize() const
+	{
+		return MaxReadBufferSize;
+	}
+
+	/**
+	 * Sets timeout in seconds for the entire http request to complete
+	 * @param SizeInBytes	The maximum number of bytes to use for the read buffer
+	 */
+	inline void SetMaxReadBufferSize(int32 SizeInBytes)
+	{
+		MaxReadBufferSize = SizeInBytes;
+	}
+
 	/**
 	 * @return true if http requests are enabled
 	 */
@@ -146,6 +163,8 @@ private:
 	float HttpSendTimeout;
 	/** Max number of simultaneous connections to a specific server */
 	int32 HttpMaxConnectionsPerServer;
+	/** Max buffer size for individual http reads */
+	int32 MaxReadBufferSize;
 	/** toggles http requests */
 	bool bEnableHttp;
 	/** singleton for the module while loaded and available */
