#UE4 - Fix shadowed variable bug

diff --git a/Engine/Source/Runtime/Online/OnlineSubsystemUtils/Private/IpNetDriver.cpp b/Engine/Source/Runtime/Online/OnlineSubsystemUtils/Private/IpNetDriver.cpp
--- a/Engine/Source/Runtime/Online/OnlineSubsystemUtils/Private/IpNetDriver.cpp
+++ b/Engine/Source/Runtime/Online/OnlineSubsystemUtils/Private/IpNetDriver.cpp
@@ -108,7 +108,7 @@ bool UIpNetDriver::InitBase( bool bInitAsClient, FNetworkNotify* InNotify, const
 		// If initializing as a client, and the specific port failed, retry with port 0 (any port)
 		LocalAddr->SetPort(0);
 		AttemptPort = LocalAddr->GetPort();
-		int32 BoundPort = SocketSubsystem->BindNextPort( Socket, *LocalAddr, 1, 1 );
+		BoundPort = SocketSubsystem->BindNextPort( Socket, *LocalAddr, 1, 1 );
 	}
 
 	if (BoundPort == 0)
