#UE4 - Temp fix for PIE: When you try to connect to localhost before the server starts, you will consume your own packet, and the netconnection code doesn't handle this well

diff --git a/Engine/Source/Runtime/Online/OnlineSubsystemUtils/Private/IpNetDriver.cpp b/Engine/Source/Runtime/Online/OnlineSubsystemUtils/Private/IpNetDriver.cpp
--- a/Engine/Source/Runtime/Online/OnlineSubsystemUtils/Private/IpNetDriver.cpp
+++ b/Engine/Source/Runtime/Online/OnlineSubsystemUtils/Private/IpNetDriver.cpp
@@ -99,7 +99,7 @@ bool UIpNetDriver::InitBase( bool bInitAsClient, FNetworkNotify* InNotify, const
 	LocalAddr = SocketSubsystem->GetLocalBindAddr(*GLog);
 	
 	// Try the specific port first, even for clients. Required for some platforms.
-	LocalAddr->SetPort(URL.Port);
+	LocalAddr->SetPort(bInitAsClient ? 0 : URL.Port);
 	
 	int32 AttemptPort = LocalAddr->GetPort();
 	int32 BoundPort = SocketSubsystem->BindNextPort( Socket, *LocalAddr, MaxPortCountToTry + 1, 1 );
