#UE4 - fixed invalid UE_LOG statement

diff --git a/Engine/Source/Runtime/Online/OnlineSubsystemUtils/Private/Tests/TestSessionInterface.cpp b/Engine/Source/Runtime/Online/OnlineSubsystemUtils/Private/Tests/TestSessionInterface.cpp
--- a/Engine/Source/Runtime/Online/OnlineSubsystemUtils/Private/Tests/TestSessionInterface.cpp
+++ b/Engine/Source/Runtime/Online/OnlineSubsystemUtils/Private/Tests/TestSessionInterface.cpp
@@ -176,7 +176,7 @@ void FTestSessionInterface::WorldDestroyed( UWorld* InWorld )
 
 void FTestSessionInterface::OnReadFriendsListComplete(int32 LocalUserNum, bool bWasSuccessful, const FString& ListName, const FString& ErrorStr)
 {
-	UE_LOG(LogOnline, Verbose, TEXT("OnReadFriendsListComplete LocalUserNum: %d bSuccess: %d. "), LocalUserNum, bWasSuccessful, *ErrorStr);
+	UE_LOG(LogOnline, Verbose, TEXT("OnReadFriendsListComplete LocalUserNum: %d bSuccess: %d %s"), LocalUserNum, bWasSuccessful, *ErrorStr);
 	Friends->ClearOnReadFriendsListCompleteDelegate(0, OnReadFriendsListCompleteDelegate);
 	if (bWasSuccessful)
 	{
