#UE4 - fixed bug in new Online::GetSubsystem function

diff --git a/Engine/Source/Runtime/Online/OnlineSubsystemUtils/Public/OnlineSubsystemUtils.h b/Engine/Source/Runtime/Online/OnlineSubsystemUtils/Public/OnlineSubsystemUtils.h
--- a/Engine/Source/Runtime/Online/OnlineSubsystemUtils/Public/OnlineSubsystemUtils.h
+++ b/Engine/Source/Runtime/Online/OnlineSubsystemUtils/Public/OnlineSubsystemUtils.h
@@ -33,7 +33,7 @@ namespace Online
 			FWorldContext& CurrentContext = GEngine->GetWorldContextFromWorldChecked(World);
 			if (CurrentContext.PIEInstance != INDEX_NONE)
 			{ 
-				Identifier = FName(*FString::Printf(TEXT("%s:%d"), *SubsystemName.ToString(), CurrentContext.PIEInstance)); 
+				Identifier = FName(*FString::Printf(TEXT("%s:%d"), SubsystemName != NAME_None ? *SubsystemName.ToString() : TEXT(""), CurrentContext.PIEInstance));
 			} 
 		}
 
