#UE4 - added DestroySubsystem call to OSS module, allowing the destruction of a single instance of a OSS

diff --git a/Engine/Source/Runtime/Online/OnlineSubsystem/Private/OnlineSubsystemModule.cpp b/Engine/Source/Runtime/Online/OnlineSubsystem/Private/OnlineSubsystemModule.cpp
--- a/Engine/Source/Runtime/Online/OnlineSubsystem/Private/OnlineSubsystemModule.cpp
+++ b/Engine/Source/Runtime/Online/OnlineSubsystem/Private/OnlineSubsystemModule.cpp
@@ -214,6 +214,29 @@ IOnlineSubsystem* FOnlineSubsystemModule::GetOnlineSubsystem(const FName InSubsy
 	return (OnlineSubsystem == NULL) ? NULL : (*OnlineSubsystem).Get();
 }
 
+void FOnlineSubsystemModule::DestroyOnlineSubsystem(const FName InSubsystemName)
+{
+	FName SubsystemName, InstanceName;
+	ParseOnlineSubsystemName(InSubsystemName, SubsystemName, InstanceName);
+
+	IOnlineSubsystemPtr* OnlineSubsystem = NULL;
+	if (SubsystemName != NAME_None)
+	{
+		FName KeyName = FName(*FString::Printf(TEXT("%s:%s"), *SubsystemName.ToString(), *InstanceName.ToString()));
+
+		IOnlineSubsystemPtr OnlineSubsystem;
+		OnlineSubsystems.RemoveAndCopyValue(KeyName, OnlineSubsystem);
+		if (OnlineSubsystem.IsValid())
+		{
+			OnlineSubsystem->Shutdown();
+		}
+		else
+		{
+			UE_LOG(LogOnline, Warning, TEXT("OnlineSubsystem instance %s not found, unable to destroy."), *KeyName.ToString());
+		}
+	}
+}
+
 bool FOnlineSubsystemModule::IsOnlineSubsystemLoaded(const FName InSubsystemName) const
 {
 	bool bIsLoaded = false;
diff --git a/Engine/Source/Runtime/Online/OnlineSubsystem/Public/OnlineSubsystemModule.h b/Engine/Source/Runtime/Online/OnlineSubsystem/Public/OnlineSubsystemModule.h
--- a/Engine/Source/Runtime/Online/OnlineSubsystem/Public/OnlineSubsystemModule.h
+++ b/Engine/Source/Runtime/Online/OnlineSubsystem/Public/OnlineSubsystemModule.h
@@ -29,6 +29,21 @@ private:
 	/** Mapping of all currently loaded platform service subsystems to their name */
 	TMap<FName, IOnlineSubsystemPtr> OnlineSubsystems;
 
+	/**
+	 * Transform an online subsystem identifier into its Subsystem and Instance constituents
+	 *
+	 * accepts the following forms:
+	 * <subsystem name>:<instance name> -> subsystem name / instance name
+	 * :<instance name>					-> default subsystem / instance name
+	 * <subsystem name>:				-> subsystem name / default instance
+	 * <subsystem name>					-> subsystem name / default instance
+	 * <nothing>						-> default subsystem / default instance
+	 *
+	 * @param FullName full name of the subsystem and instance that is being referenced
+	 * @param SubsystemName parsed out value or default subsystem name
+	 * @param InstanceName parsed out value or default instance name
+	 *
+	 */
 	void ParseOnlineSubsystemName(const FName& FullName, FName& SubsystemName, FName& InstanceName);
 
 	/**
@@ -48,11 +63,21 @@ public:
 	 * Will load the appropriate module if the subsystem isn't currently loaded
 	 * It's possible that the subsystem doesn't exist and therefore can return NULL
 	 *
-	 * @param SubsystemName - name of subsystem as referenced by consumers
+	 * @param InSubsystemName - name of subsystem as referenced by consumers
+	 *
 	 * @return Requested online subsystem, or NULL if that subsystem was unable to load or doesn't exist
 	 */
 	virtual class IOnlineSubsystem* GetOnlineSubsystem(const FName InSubsystemName = NAME_None);
 
+	/**
+	 * Destroys an online subsystem created internally via access with GetOnlineSubsystem
+	 * Typically destruction of the subsystem is handled at application exit, but
+	 * there may be rare instances where the subsystem is destroyed by request
+	 *
+	 * @param InSubsystemName - name of subsystem as referenced by consumers
+	 */
+	virtual void DestroyOnlineSubsystem(const FName InSubsystemName);
+
 	/** 
 	 * Determine if a subsystem is loaded by the OSS module
 	 *
