#UE4 - log fixup

diff --git a/Engine/Source/Runtime/Online/OnlineSubsystemUtils/Private/OnlineBeaconClient.cpp b/Engine/Source/Runtime/Online/OnlineSubsystemUtils/Private/OnlineBeaconClient.cpp
--- a/Engine/Source/Runtime/Online/OnlineSubsystemUtils/Private/OnlineBeaconClient.cpp
+++ b/Engine/Source/Runtime/Online/OnlineSubsystemUtils/Private/OnlineBeaconClient.cpp
@@ -104,7 +104,7 @@ void AOnlineBeaconClient::NotifyControlMessage(UNetConnection* Connection, uint8
 
 		// We are the client
 #if !(UE_BUILD_SHIPPING && WITH_EDITOR)
-		UE_LOG(LogNet, Log, TEXT("Beacon client received: %s"), FNetControlMessageInfo::GetName(MessageType));
+		UE_LOG(LogNet, Log, TEXT("Beacon: Client received: %s"), FNetControlMessageInfo::GetName(MessageType));
 #endif
 		switch (MessageType)
 		{
