Allows players to create guest accounts from the in-game picker, and also prevents guests from playing without their sponsors.

diff --git a/Engine/Source/Runtime/Online/OnlineSubsystem/Public/Interfaces/OnlineIdentityInterface.h b/Engine/Source/Runtime/Online/OnlineSubsystem/Public/Interfaces/OnlineIdentityInterface.h
--- a/Engine/Source/Runtime/Online/OnlineSubsystem/Public/Interfaces/OnlineIdentityInterface.h
+++ b/Engine/Source/Runtime/Online/OnlineSubsystem/Public/Interfaces/OnlineIdentityInterface.h
@@ -191,6 +191,15 @@ public:
 	 */
 	virtual TSharedPtr<FUniqueNetId> GetUniquePlayerId(int32 LocalUserNum) const = 0;
 
+	/**
+	 * Gets the platform specific unique id for the sponsor of the specified player
+	 *
+	 * @param LocalUserNum the controller number of the associated user
+	 *
+	 * @return Valid player id object if the sponsor exists, NULL otherwise
+	 */
+	virtual TSharedPtr<FUniqueNetId> GetSponsorUniquePlayerId(int32 LocalUserNum) const { return nullptr; }
+
 	/**
 	 * Create a unique id from binary data (used during replication)
 	 * 
