#ue4
- Removed extra logging from ULeaderboardBlueprintLibrary

diff --git a/Engine/Source/Runtime/Online/OnlineSubsystemUtils/Private/LeaderboardBlueprintLibrary.cpp b/Engine/Source/Runtime/Online/OnlineSubsystemUtils/Private/LeaderboardBlueprintLibrary.cpp
--- a/Engine/Source/Runtime/Online/OnlineSubsystemUtils/Private/LeaderboardBlueprintLibrary.cpp
+++ b/Engine/Source/Runtime/Online/OnlineSubsystemUtils/Private/LeaderboardBlueprintLibrary.cpp
@@ -14,36 +14,22 @@ ULeaderboardBlueprintLibrary::ULeaderboardBlueprintLibrary(const FPostConstructI
 
 bool ULeaderboardBlueprintLibrary::WriteLeaderboardObject(APlayerController* PlayerController, class FOnlineLeaderboardWrite& WriteObject)
 {
-	UE_LOG(LogOnline, Display, TEXT("WriteLeaderboardObject site 1"));
-
 	if (APlayerState* PlayerState = (PlayerController != NULL) ? PlayerController->PlayerState : NULL)
 	{
-		UE_LOG(LogOnline, Display, TEXT("WriteLeaderboardObject site 2"));
-
 		TSharedPtr<FUniqueNetId> UserId = PlayerState->UniqueId.GetUniqueNetId();
 		if (UserId.IsValid())
 		{
-			UE_LOG(LogOnline, Display, TEXT("WriteLeaderboardObject site 3"));
-
 			if (IOnlineSubsystem* const OnlineSub = IOnlineSubsystem::Get())
 			{
-				UE_LOG(LogOnline, Display, TEXT("WriteLeaderboardObject site 4"));
-
 				IOnlineLeaderboardsPtr Leaderboards = OnlineSub->GetLeaderboardsInterface();
 				if (Leaderboards.IsValid())
 				{
-					UE_LOG(LogOnline, Display, TEXT("WriteLeaderboardObject site 5"));
-
 					// the call will copy the user id and write object to its own memory
 					bool bResult = Leaderboards->WriteLeaderboards(PlayerState->SessionName, *UserId, WriteObject);
 
-					UE_LOG(LogOnline, Display, TEXT("WriteLeaderboardObject site 6"));
-
 					// Flush the leaderboard immediately for now
 					bool bFlushResult = Leaderboards->FlushLeaderboards(PlayerState->SessionName);
 
-					UE_LOG(LogOnline, Display, TEXT("WriteLeaderboardObject site 7"));
-
 					return bResult && bFlushResult;
 				}
 				else
@@ -66,7 +52,6 @@ bool ULeaderboardBlueprintLibrary::WriteLeaderboardObject(APlayerController* Pla
 		FFrame::KismetExecutionMessage(TEXT("WriteLeaderboardObject - Invalid player state"), ELogVerbosity::Warning);
 	}
 
-	UE_LOG(LogOnline, Display, TEXT("WriteLeaderboardObject site 8"));
 	return false;
 }
 
