- Fixed compile issues related to OnlineSubsystemUtils

diff --git a/Engine/Source/Runtime/Online/OnlineSubsystemUtils/Classes/IpConnection.h b/Engine/Source/Runtime/Online/OnlineSubsystemUtils/Classes/IpConnection.h
--- a/Engine/Source/Runtime/Online/OnlineSubsystemUtils/Classes/IpConnection.h
+++ b/Engine/Source/Runtime/Online/OnlineSubsystemUtils/Classes/IpConnection.h
@@ -13,8 +13,8 @@ class ONLINESUBSYSTEMUTILS_API UIpConnection : public UNetConnection
     GENERATED_UCLASS_BODY()
 	// Variables.
 	TSharedPtr<FInternetAddr>	RemoteAddr;
-	FSocket*			Socket;
-	FResolveInfo*		ResolveInfo;
+	class FSocket*				Socket;
+	class FResolveInfo*			ResolveInfo;
 
 	// Begin NetConnection Interface
 	virtual void InitBase(UNetDriver* InDriver, class FSocket* InSocket, const FURL& InURL, EConnectionState InState, int32 InMaxPacket = 0, int32 InPacketOverhead = 0) OVERRIDE;
@@ -23,19 +23,7 @@ class ONLINESUBSYSTEMUTILS_API UIpConnection : public UNetConnection
 	virtual void LowLevelSend(void* Data,int32 Count) OVERRIDE;
 	FString LowLevelGetRemoteAddress(bool bAppendPort=false) OVERRIDE;
 	FString LowLevelDescribe() OVERRIDE;
-	virtual int32 GetAddrAsInt(void) OVERRIDE
-	{
-		uint32 OutAddr = 0;
-		// Get the host byte order ip addr
-		RemoteAddr->GetIp(OutAddr);
-		return (int32)OutAddr;
-	}
-	virtual int32 GetAddrPort(void) OVERRIDE
-	{
-		int32 OutPort = 0;
-		// Get the host byte order ip port
-		RemoteAddr->GetPort(OutPort);
-		return OutPort;
-	}
+	virtual int32 GetAddrAsInt(void) OVERRIDE;
+	virtual int32 GetAddrPort(void) OVERRIDE;
 	// End NetConnection Interface
 };
diff --git a/Engine/Source/Runtime/Online/OnlineSubsystemUtils/Private/IpConnection.cpp b/Engine/Source/Runtime/Online/OnlineSubsystemUtils/Private/IpConnection.cpp
--- a/Engine/Source/Runtime/Online/OnlineSubsystemUtils/Private/IpConnection.cpp
+++ b/Engine/Source/Runtime/Online/OnlineSubsystemUtils/Private/IpConnection.cpp
@@ -157,3 +157,19 @@ FString UIpConnection::LowLevelDescribe()
 		:								TEXT("Invalid")
 	);
 }
+
+int32 UIpConnection::GetAddrAsInt(void)
+{
+	uint32 OutAddr = 0;
+	// Get the host byte order ip addr
+	RemoteAddr->GetIp(OutAddr);
+	return (int32)OutAddr;
+}
+
+int32 UIpConnection::GetAddrPort(void)
+{
+	int32 OutPort = 0;
+	// Get the host byte order ip port
+	RemoteAddr->GetPort(OutPort);
+	return OutPort;
+}
