#UE4 - support for destroying a single OSS instance

diff --git a/Engine/Source/Runtime/Online/OnlineSubsystem/Public/OnlineSubsystem.h b/Engine/Source/Runtime/Online/OnlineSubsystem/Public/OnlineSubsystem.h
--- a/Engine/Source/Runtime/Online/OnlineSubsystem/Public/OnlineSubsystem.h
+++ b/Engine/Source/Runtime/Online/OnlineSubsystem/Public/OnlineSubsystem.h
@@ -83,6 +83,17 @@ public:
 		return OSSModule.GetOnlineSubsystem(SubsystemName); 
   	}
 
+	/** 
+	 * Destroy a single online subsystem instance
+	 * @param SubsystemName - Name of the online service to destroy
+	 */
+	static void Destroy(FName SubsystemName)
+	{
+		static const FName OnlineSubsystemModuleName = TEXT("OnlineSubsystem");
+		FOnlineSubsystemModule& OSSModule = FModuleManager::GetModuleChecked<FOnlineSubsystemModule>(OnlineSubsystemModuleName);
+		return OSSModule.DestroyOnlineSubsystem(SubsystemName);
+	}
+
 	/** 
 	 * Determine if the subsystem for a given interface is already loaded
 	 * @param SubsystemName - Name of the requested online service
