#UE4 - added DEFAULT_INSTANCE variable for OSS instance names

diff --git a/Engine/Source/Runtime/Online/OnlineSubsystem/Private/OnlineSubsystemModule.cpp b/Engine/Source/Runtime/Online/OnlineSubsystem/Private/OnlineSubsystemModule.cpp
--- a/Engine/Source/Runtime/Online/OnlineSubsystem/Private/OnlineSubsystemModule.cpp
+++ b/Engine/Source/Runtime/Online/OnlineSubsystem/Private/OnlineSubsystemModule.cpp
@@ -143,7 +143,7 @@ void FOnlineSubsystemModule::UnregisterPlatformService(const FName FactoryName)
 void FOnlineSubsystemModule::ParseOnlineSubsystemName(const FName& FullName, FName& SubsystemName, FName& InstanceName)
 {
 	SubsystemName = DefaultPlatformService;
-	InstanceName = TEXT("DefaultInstance");
+	InstanceName = DEFAULT_INSTANCE;
 
 	if (FullName != NAME_None)
 	{
diff --git a/Engine/Source/Runtime/Online/OnlineSubsystem/Public/OnlineSubsystemTypes.h b/Engine/Source/Runtime/Online/OnlineSubsystem/Public/OnlineSubsystemTypes.h
--- a/Engine/Source/Runtime/Online/OnlineSubsystem/Public/OnlineSubsystemTypes.h
+++ b/Engine/Source/Runtime/Online/OnlineSubsystem/Public/OnlineSubsystemTypes.h
@@ -3,6 +3,9 @@
 #pragma once
 #include "OnlineSubsystemPackage.h"
 
+/** Name given to default OSS instances (disambiguates for PIE) */
+#define DEFAULT_INSTANCE FName(TEXT("DefaultInstance"))
+
 /** Maximum players supported on a given platform */
 #if PLATFORM_XBOXONE
 #define MAX_LOCAL_PLAYERS 4
