#UE4 - added support for and InstanceName on various OSS

diff --git a/Engine/Source/Runtime/Online/OnlineSubsystemNull/Private/OnlineSubsystemModuleNull.cpp b/Engine/Source/Runtime/Online/OnlineSubsystemNull/Private/OnlineSubsystemModuleNull.cpp
--- a/Engine/Source/Runtime/Online/OnlineSubsystemNull/Private/OnlineSubsystemModuleNull.cpp
+++ b/Engine/Source/Runtime/Online/OnlineSubsystemNull/Private/OnlineSubsystemModuleNull.cpp
@@ -18,7 +18,7 @@ public:
 
 	virtual IOnlineSubsystemPtr CreateSubsystem(FName InstanceName)
 	{
-		FOnlineSubsystemNullPtr OnlineSub = MakeShareable(new FOnlineSubsystemNull());
+		FOnlineSubsystemNullPtr OnlineSub = MakeShareable(new FOnlineSubsystemNull(InstanceName));
 		if (OnlineSub->IsEnabled())
 		{
 			if(!OnlineSub->Init())
diff --git a/Engine/Source/Runtime/Online/OnlineSubsystemNull/Public/OnlineSubsystemNull.h b/Engine/Source/Runtime/Online/OnlineSubsystemNull/Public/OnlineSubsystemNull.h
--- a/Engine/Source/Runtime/Online/OnlineSubsystemNull/Public/OnlineSubsystemNull.h
+++ b/Engine/Source/Runtime/Online/OnlineSubsystemNull/Public/OnlineSubsystemNull.h
@@ -61,6 +61,11 @@ public:
 
 	// FOnlineSubsystemNull
 
+	/**
+	 * @return the name of the online subsystem instance
+	 */
+	FName GetInstanceName() const { return InstanceName; }
+
 	/**
 	 * Is the Null API available for use
 	 * @return true if Null functionality is available, false otherwise
@@ -70,17 +75,31 @@ public:
 PACKAGE_SCOPE:
 
 	/** Only the factory makes instances */
-	FOnlineSubsystemNull()
-		: SessionInterface(NULL)
-		, LeaderboardsInterface(NULL)
-		, IdentityInterface(NULL)
-		, AchievementsInterface(NULL)
-		, OnlineAsyncTaskThreadRunnable(NULL)
-		, OnlineAsyncTaskThread(NULL)
+	FOnlineSubsystemNull(FName InInstanceName) :
+		InstanceName(InInstanceName),
+		SessionInterface(NULL),
+		LeaderboardsInterface(NULL),
+		IdentityInterface(NULL),
+		AchievementsInterface(NULL),
+		OnlineAsyncTaskThreadRunnable(NULL),
+		OnlineAsyncTaskThread(NULL)
+	{}
+
+	FOnlineSubsystemNull() :
+		InstanceName(NAME_None),
+		SessionInterface(NULL),
+		LeaderboardsInterface(NULL),
+		IdentityInterface(NULL),
+		AchievementsInterface(NULL),
+		OnlineAsyncTaskThreadRunnable(NULL),
+		OnlineAsyncTaskThread(NULL)
 	{}
 
 private:
 
+	/** Instance name (disambiguates PIE instances for example) */
+	FName InstanceName;
+
 	/** Interface to the session services */
 	FOnlineSessionNullPtr SessionInterface;
 
diff --git a/Engine/Source/Runtime/Online/OnlineSubsystemSteam/Private/OnlineSubsystemModuleSteam.cpp b/Engine/Source/Runtime/Online/OnlineSubsystemSteam/Private/OnlineSubsystemModuleSteam.cpp
--- a/Engine/Source/Runtime/Online/OnlineSubsystemSteam/Private/OnlineSubsystemModuleSteam.cpp
+++ b/Engine/Source/Runtime/Online/OnlineSubsystemSteam/Private/OnlineSubsystemModuleSteam.cpp
@@ -41,7 +41,7 @@ public:
 	{
 		if (!SteamSingleton.IsValid())
 		{
-			SteamSingleton = MakeShareable(new FOnlineSubsystemSteam());
+			SteamSingleton = MakeShareable(new FOnlineSubsystemSteam(InstanceName));
 			if (SteamSingleton->IsEnabled())
 			{
 				if(!SteamSingleton->Init())
diff --git a/Engine/Source/Runtime/Online/OnlineSubsystemSteam/Public/OnlineSubsystemSteam.h b/Engine/Source/Runtime/Online/OnlineSubsystemSteam/Public/OnlineSubsystemSteam.h
--- a/Engine/Source/Runtime/Online/OnlineSubsystemSteam/Public/OnlineSubsystemSteam.h
+++ b/Engine/Source/Runtime/Online/OnlineSubsystemSteam/Public/OnlineSubsystemSteam.h
@@ -25,6 +25,9 @@ class ONLINESUBSYSTEMSTEAM_API FOnlineSubsystemSteam :
 {
 protected:
 
+	/** Instance name (disambiguates PIE instances for example) */
+	FName InstanceName;
+
 	/** Has the STEAM client APIs been initialized */
 	bool bSteamworksClientInitialized;
 
@@ -82,7 +85,28 @@ protected:
 PACKAGE_SCOPE:
 
 	/** Only the factory makes instances */
+	FOnlineSubsystemSteam(FName InInstanceName) :
+		InstanceName(InInstanceName),
+		bSteamworksClientInitialized(false),
+		bSteamworksGameServerInitialized(false),
+		SteamAppID(0),
+		GameServerSteamPort(0),
+		GameServerGamePort(0),
+		GameServerQueryPort(0),
+		SessionInterface(NULL),
+		IdentityInterface(NULL),
+		FriendInterface(NULL),
+		SharedCloudInterface(NULL),
+		UserCloudInterface(NULL),
+		LeaderboardsInterface(NULL),
+		VoiceInterface(NULL),
+		ExternalUIInterface(NULL),
+		OnlineAsyncTaskThreadRunnable(NULL),
+		OnlineAsyncTaskThread(NULL)
+	{}
+
 	FOnlineSubsystemSteam() : 
+		InstanceName(NAME_None),
 		bSteamworksClientInitialized(false),
 		bSteamworksGameServerInitialized(false),
 		SteamAppID(0),
@@ -224,6 +248,11 @@ public:
 
 	// FOnlineSubsystemSteam
 
+	/**
+	 * @return the name of the online subsystem instance
+	 */
+	FName GetInstanceName() const { return InstanceName; }
+
 	/**
 	 * Whether or not the Steam Client interfaces are available; these interfaces are only available, if the Steam Client program is running
 	 * NOTE: These interfaces are made unavailable, when running a dedicated server
