#UE4 - Fix IOS achievement crash. Make copy of delegate to force it as a copy inside the async task, otherwise, it remains a reference which ends up referring to a destroyed delegate.

diff --git a/Engine/Source/Runtime/Online/IOS/OnlineSubsystemIOS/Private/OnlineAchievementsInterfaceIOS.cpp b/Engine/Source/Runtime/Online/IOS/OnlineSubsystemIOS/Private/OnlineAchievementsInterfaceIOS.cpp
--- a/Engine/Source/Runtime/Online/IOS/OnlineSubsystemIOS/Private/OnlineAchievementsInterfaceIOS.cpp
+++ b/Engine/Source/Runtime/Online/IOS/OnlineSubsystemIOS/Private/OnlineAchievementsInterfaceIOS.cpp
@@ -16,6 +16,9 @@ FOnlineAchievementsIOS::FOnlineAchievementsIOS( class FOnlineSubsystemIOS* InSub
 
 void FOnlineAchievementsIOS::QueryAchievements(const FUniqueNetId & PlayerId, const FOnQueryAchievementsCompleteDelegate& Delegate)
 {
+	// Make a copy of the delegate so it won't be a reference inside the FIOSAsyncTask
+	FOnQueryAchievementsCompleteDelegate CopiedDelegate = Delegate;
+
 	dispatch_async(dispatch_get_main_queue(), ^
 	{
 	[GKAchievement loadAchievementsWithCompletionHandler: ^(NSArray *loadedAchievements, NSError *error)
@@ -47,7 +50,7 @@ void FOnlineAchievementsIOS::QueryAchievements(const FUniqueNetId & PlayerId, co
 			// Report back to the game thread whether this succeeded.
 			[FIOSAsyncTask CreateTaskWithBlock : ^ bool(void)
 			{
-				Delegate.ExecuteIfBound(PlayerId, bSuccess);
+				CopiedDelegate.ExecuteIfBound(PlayerId, bSuccess);
 				return true;
 			}];
 		}
@@ -57,6 +60,8 @@ void FOnlineAchievementsIOS::QueryAchievements(const FUniqueNetId & PlayerId, co
 
 void FOnlineAchievementsIOS::QueryAchievementDescriptions(const FUniqueNetId& PlayerId, const FOnQueryAchievementsCompleteDelegate& Delegate)
 {
+	FOnQueryAchievementsCompleteDelegate CopiedDelegate = Delegate;
+
 	dispatch_async(dispatch_get_main_queue(), ^
 	{
 		[GKAchievementDescription loadAchievementDescriptionsWithCompletionHandler : ^ (NSArray *descriptions, NSError *error)
@@ -97,7 +102,7 @@ void FOnlineAchievementsIOS::QueryAchievementDescriptions(const FUniqueNetId& Pl
 			// Report back to the game thread whether this succeeded.
 			[FIOSAsyncTask CreateTaskWithBlock : ^ bool(void)
 			{
-				Delegate.ExecuteIfBound(PlayerId, bSuccess);
+				CopiedDelegate.ExecuteIfBound(PlayerId, bSuccess);
 				return true;
 			}];
 		}];
