#UE4 - added FOnlineSubsystemImpl as a common implementation to shim between concrete/leaf implementations
- added named interfaces feature for basically "add to root" any UObject that online subsystem cares about across map transitions
- added ForceDedicated mode for PIE so that OSS platforms behave as dedicated while actually running as editor (possibly needs better solution)

[REDACTED_HANDLE] - derive all subsystems from this code and remove duplicate/common functionality (namedsessions, islocalplayer, isplayerinsession, other 'C' style functions)

diff --git a/Engine/Source/Runtime/Online/OnlineSubsystem/Private/NamedInterfaces.cpp b/Engine/Source/Runtime/Online/OnlineSubsystem/Private/NamedInterfaces.cpp
new file mode 100644
--- /dev/null
+++ b/Engine/Source/Runtime/Online/OnlineSubsystem/Private/NamedInterfaces.cpp
@@ -0,0 +1,78 @@
+// Copyright 1998-2014 Epic Games, Inc. All Rights Reserved.
+
+#include "OnlineSubsystemPrivatePCH.h"
+#include "NamedInterfaces.h"
+
+UNamedInterfaces::UNamedInterfaces(const class FPostConstructInitializeProperties& PCIP) :
+	Super(PCIP)
+{
+}
+
+void UNamedInterfaces::Initialize()
+{
+	// Iterate through each configured named interface load it and create an instance
+	for (int32 InterfaceIndex = 0; InterfaceIndex < NamedInterfaceDefs.Num(); InterfaceIndex++)
+	{
+		const FNamedInterfaceDef& Def = NamedInterfaceDefs[InterfaceIndex];
+		// Load the specified interface class name
+		UClass* Class = LoadClass<UObject>(NULL, *Def.InterfaceClassName, NULL, LOAD_None, NULL);
+		if (Class)
+		{
+			int32 AddIndex = NamedInterfaces.AddZeroed();
+			FNamedInterface& Interface = NamedInterfaces[AddIndex];
+			// Set the object and interface names
+			Interface.InterfaceName = Def.InterfaceName;
+			Interface.InterfaceObject = ConstructObject<UObject>(Class);
+			UE_LOG(LogOnline, Display,
+				TEXT("Created named interface (%s) of type (%s)"),
+				*Def.InterfaceName.ToString(),
+				*Def.InterfaceClassName);
+		}
+		else
+		{
+			UE_LOG(LogOnline, Warning,
+				TEXT("Failed to load class (%s) for named interface (%s)"),
+				*Def.InterfaceClassName,
+				*Def.InterfaceName.ToString());
+		}
+	}
+}
+
+UObject* UNamedInterfaces::GetNamedInterface(FName InterfaceName) const
+{
+	for (const FNamedInterface& Interface : NamedInterfaces)
+	{
+		if (Interface.InterfaceName == InterfaceName)
+		{
+			return Interface.InterfaceObject;
+		}
+	}
+
+	return NULL;
+}
+
+void UNamedInterfaces::SetNamedInterface(FName InterfaceName, UObject* NewInterface)
+{
+	int32 InterfaceIdx = 0;
+	for (; InterfaceIdx < NamedInterfaces.Num(); InterfaceIdx++)
+	{
+		const FNamedInterface& Interface = NamedInterfaces[InterfaceIdx];
+		if (Interface.InterfaceName == InterfaceName)
+		{
+			break;
+		}
+	}
+
+	if (InterfaceIdx >= NamedInterfaces.Num())
+	{
+		int32 AddIndex = NamedInterfaces.AddZeroed();
+		FNamedInterface& Interface = NamedInterfaces[AddIndex];
+		// Set the object and interface names
+		Interface.InterfaceName = InterfaceName;
+		Interface.InterfaceObject = NewInterface;
+	}
+	else
+	{
+		NamedInterfaces[InterfaceIdx].InterfaceObject = NewInterface;
+	}
+}
\ No newline at end of file
diff --git a/Engine/Source/Runtime/Online/OnlineSubsystem/Private/OnlineSubsystem.cpp b/Engine/Source/Runtime/Online/OnlineSubsystem/Private/OnlineSubsystem.cpp
--- a/Engine/Source/Runtime/Online/OnlineSubsystem/Private/OnlineSubsystem.cpp
+++ b/Engine/Source/Runtime/Online/OnlineSubsystem/Private/OnlineSubsystem.cpp
@@ -4,6 +4,8 @@
 #include "OnlineSessionInterface.h"
 #include "NboSerializer.h"
 
+#include "OnlineSubsystem.generated.inl"
+
 DEFINE_LOG_CATEGORY(LogOnline);
 DEFINE_LOG_CATEGORY(LogOnlineGame);
 
diff --git a/Engine/Source/Runtime/Online/OnlineSubsystem/Private/OnlineSubsystemImpl.cpp b/Engine/Source/Runtime/Online/OnlineSubsystem/Private/OnlineSubsystemImpl.cpp
new file mode 100644
--- /dev/null
+++ b/Engine/Source/Runtime/Online/OnlineSubsystem/Private/OnlineSubsystemImpl.cpp
@@ -0,0 +1,67 @@
+// Copyright 1998-2014 Epic Games, Inc. All Rights Reserved.
+
+#include "OnlineSubsystemPrivatePCH.h"
+#include "OnlineSubsystemImpl.h"
+#include "NamedInterfaces.h"
+
+FOnlineSubsystemImpl::FOnlineSubsystemImpl() :
+	InstanceName(NAME_None),
+	NamedInterfaces(NULL),
+	bForceDedicated(false)
+{
+}
+
+FOnlineSubsystemImpl::FOnlineSubsystemImpl(FName InInstanceName) :
+	InstanceName(InInstanceName),
+	NamedInterfaces(NULL),
+	bForceDedicated(false)
+{
+}
+
+FOnlineSubsystemImpl::~FOnlineSubsystemImpl()
+{
+	if (NamedInterfaces)
+	{
+		NamedInterfaces->RemoveFromRoot();
+		NamedInterfaces = NULL;
+	}
+}
+
+void FOnlineSubsystemImpl::InitNamedInterfaces()
+{
+	NamedInterfaces = ConstructObject<UNamedInterfaces>(UNamedInterfaces::StaticClass());
+	if (NamedInterfaces)
+	{
+		NamedInterfaces->Initialize();
+		NamedInterfaces->AddToRoot();
+	}
+}
+
+UObject* FOnlineSubsystemImpl::GetNamedInterface(FName InterfaceName)
+{
+	if (!NamedInterfaces)
+	{
+		InitNamedInterfaces();
+	}
+
+	if (NamedInterfaces)
+	{
+		return NamedInterfaces->GetNamedInterface(InterfaceName);
+	}
+
+	return NULL;
+}
+
+void FOnlineSubsystemImpl::SetNamedInterface(FName InterfaceName, UObject* NewInterface)
+{
+	if (!NamedInterfaces)
+	{
+		InitNamedInterfaces();
+	}
+
+	if (NamedInterfaces)
+	{
+		return NamedInterfaces->SetNamedInterface(InterfaceName, NewInterface);
+	}
+}
+
diff --git a/Engine/Source/Runtime/Online/OnlineSubsystem/Private/OnlineSubsystemPrivatePCH.h b/Engine/Source/Runtime/Online/OnlineSubsystem/Private/OnlineSubsystemPrivatePCH.h
--- a/Engine/Source/Runtime/Online/OnlineSubsystem/Private/OnlineSubsystemPrivatePCH.h
+++ b/Engine/Source/Runtime/Online/OnlineSubsystem/Private/OnlineSubsystemPrivatePCH.h
@@ -3,6 +3,7 @@
 #pragma once
 
 #include "Core.h"
+#include "CoreUObject.h"
 #include "OnlineSubsystem.h"
 #include "OnlineSubsystemModule.h"
 #include "SocketSubsystem.h"
diff --git a/Engine/Source/Runtime/Online/OnlineSubsystem/Public/NamedInterfaces.h b/Engine/Source/Runtime/Online/OnlineSubsystem/Public/NamedInterfaces.h
new file mode 100644
--- /dev/null
+++ b/Engine/Source/Runtime/Online/OnlineSubsystem/Public/NamedInterfaces.h
@@ -0,0 +1,71 @@
+// Copyright 1998-2014 Epic Games, Inc. All Rights Reserved.
+
+//
+// Named interfaces for Online Subsystem
+//
+
+#pragma once
+#include "NamedInterfaces.generated.h"
+
+/** Holds a named object interface for dynamically bound interfaces */
+USTRUCT()
+struct FNamedInterface
+{
+	GENERATED_USTRUCT_BODY()
+
+	FNamedInterface() :
+		InterfaceName(NAME_None),
+		InterfaceObject(NULL)
+	{
+	}
+
+	/** The name to bind this object to */
+	UPROPERTY()
+	FName InterfaceName;
+	/** The object to store at this location */
+	UPROPERTY()
+	UObject* InterfaceObject;
+};
+
+/** Holds a name to class name mapping for adding the named interfaces automatically */
+USTRUCT()
+struct FNamedInterfaceDef
+{
+	GENERATED_USTRUCT_BODY()
+
+	FNamedInterfaceDef() : 
+		InterfaceName(NAME_None)
+	{
+	}
+
+	/** The name to bind this object to */
+	UPROPERTY()
+	FName InterfaceName;
+	/** The class to load and create for the named interface */
+	UPROPERTY()
+	FString InterfaceClassName;
+};
+
+/**
+ * Named interfaces are a registry of UObjects accessible by an FName key that will persist for the lifetime of the process
+ */
+UCLASS(transient, config=Engine)
+class UNamedInterfaces: public UObject
+{
+	GENERATED_UCLASS_BODY()
+
+	void Initialize();
+	class UObject* GetNamedInterface(FName InterfaceName) const;
+	void SetNamedInterface(FName InterfaceName, class UObject* NewInterface);
+
+private:
+
+	/** Holds the set of registered named interfaces */
+	UPROPERTY()
+	TArray<FNamedInterface> NamedInterfaces;
+
+	/** The list of named interfaces to automatically create and store */
+	UPROPERTY(config)
+	TArray<FNamedInterfaceDef> NamedInterfaceDefs;
+};
+
diff --git a/Engine/Source/Runtime/Online/OnlineSubsystem/Public/OnlineSubsystem.h b/Engine/Source/Runtime/Online/OnlineSubsystem/Public/OnlineSubsystem.h
--- a/Engine/Source/Runtime/Online/OnlineSubsystem/Public/OnlineSubsystem.h
+++ b/Engine/Source/Runtime/Online/OnlineSubsystem/Public/OnlineSubsystem.h
@@ -217,6 +217,33 @@ public:
 	 */
 	virtual IOnlinePresencePtr GetPresenceInterface() const = 0;
 
+	/**
+	 * Get custom UObject data preserved by the online subsystem
+	 *
+	 * @param InterfaceName key to the custom data
+	 */
+	virtual class UObject* GetNamedInterface(FName InterfaceName) { return NULL; }
+
+	/**
+	 * Set a custom UObject to be preserved by the online subsystem
+	 *
+	 * @param InterfaceName key to the custom data
+	 * @param NewInterface object to preserve
+	 */
+	virtual void SetNamedInterface(FName InterfaceName, class UObject* NewInterface) {};
+
+	/**
+	 * Force the online subsystem to behave as if it's associated with running a dedicated server
+	 *
+	 * @param bForce force dedicated mode if true
+	 */
+	virtual void SetForceDedicated(bool bForce) {}
+
+	/**
+	 * @return true if the online subsystem is in forced dedicated server mode, false otherwise
+	 */
+	virtual bool IsDedicated() const { return false; }
+
 	/** 
 	 * Initialize the underlying subsystem APIs
 	 * @return true if the subsystem was successfully initialized, false otherwise
diff --git a/Engine/Source/Runtime/Online/OnlineSubsystem/Public/OnlineSubsystemImpl.h b/Engine/Source/Runtime/Online/OnlineSubsystem/Public/OnlineSubsystemImpl.h
new file mode 100644
--- /dev/null
+++ b/Engine/Source/Runtime/Online/OnlineSubsystem/Public/OnlineSubsystemImpl.h
@@ -0,0 +1,49 @@
+// Copyright 1998-2014 Epic Games, Inc. All Rights Reserved.
+
+#pragma once
+
+#include "OnlineSubsystem.h"
+#include "OnlineSubsystemPackage.h"
+
+/**
+ *	FOnlineSubsystemImpl - common functionality to share across online platforms, not intended for direct use
+ */
+class ONLINESUBSYSTEM_API FOnlineSubsystemImpl : public IOnlineSubsystem
+{
+protected:
+
+	/** Hidden on purpose */
+	FOnlineSubsystemImpl();
+	FOnlineSubsystemImpl(FName InInstanceName);
+
+	/** Instance name (disambiguates PIE instances for example) */
+	FName InstanceName;
+
+	/** Whether or not the online subsystem is in forced dedicated server mode */
+	bool bForceDedicated;
+
+	/** Holds all currently named interfaces */
+	mutable class UNamedInterfaces* NamedInterfaces;
+
+	/** Load in any named interfaces specified by the ini configuration */
+	void InitNamedInterfaces();
+
+public:
+	
+	virtual ~FOnlineSubsystemImpl();
+
+	// IOnlineSubsystem
+	virtual void SetForceDedicated(bool bForce) OVERRIDE { bForceDedicated = bForce; }
+	virtual bool IsDedicated() const OVERRIDE{ return bForceDedicated || IsRunningDedicatedServer(); }
+
+	virtual class UObject* GetNamedInterface(FName InterfaceName) OVERRIDE;
+	virtual void SetNamedInterface(FName InterfaceName, class UObject* NewInterface) OVERRIDE;
+
+	// FOnlineSubsystemImpl
+
+	/**
+	 * @return the name of the online subsystem instance
+	 */
+	FName GetInstanceName() const { return InstanceName; }
+};
+
