#UE4 - Online::GetSubsystem now stores the OSS identifier for PIE worlds using ContextHandle, not PIEInstance for easy retrieval of the WorldContext later

diff --git a/Engine/Source/Runtime/Online/OnlineSubsystemUtils/Public/OnlineSubsystemUtils.h b/Engine/Source/Runtime/Online/OnlineSubsystemUtils/Public/OnlineSubsystemUtils.h
--- a/Engine/Source/Runtime/Online/OnlineSubsystemUtils/Public/OnlineSubsystemUtils.h
+++ b/Engine/Source/Runtime/Online/OnlineSubsystemUtils/Public/OnlineSubsystemUtils.h
@@ -12,7 +12,6 @@
 ONLINESUBSYSTEMUTILS_API class UAudioComponent* CreateVoiceAudioComponent(uint32 SampleRate);
 #endif
 
-
 /** Macro to handle the boilerplate of accessing the proper online subsystem and getting the requested interface (UWorld version) */
 #define IMPLEMENT_GET_INTERFACE(InterfaceType) \
 static IOnline##InterfaceType##Ptr Get##InterfaceType##Interface(class UWorld* World, const FName SubsystemName = NAME_None) \
@@ -31,9 +30,9 @@ namespace Online
 		if (GIsPlayInEditorWorld)
 		{
 			FWorldContext& CurrentContext = GEngine->GetWorldContextFromWorldChecked(World);
-			if (CurrentContext.PIEInstance != INDEX_NONE)
+			if (CurrentContext.WorldType == EWorldType::PIE)
 			{ 
-				Identifier = FName(*FString::Printf(TEXT("%s:%d"), SubsystemName != NAME_None ? *SubsystemName.ToString() : TEXT(""), CurrentContext.PIEInstance));
+				Identifier = FName(*FString::Printf(TEXT("%s:%d"), SubsystemName != NAME_None ? *SubsystemName.ToString() : TEXT(""), CurrentContext.ContextHandle));
 			} 
 		}
 
