#UE4 - Override net driver on PS4 to allow the creation of special UDPP2P socket types

diff --git a/Engine/Source/Runtime/Online/OnlineSubsystemUtils/Classes/IpNetDriver.h b/Engine/Source/Runtime/Online/OnlineSubsystemUtils/Classes/IpNetDriver.h
--- a/Engine/Source/Runtime/Online/OnlineSubsystemUtils/Classes/IpNetDriver.h
+++ b/Engine/Source/Runtime/Online/OnlineSubsystemUtils/Classes/IpNetDriver.h
@@ -46,6 +46,10 @@ class ONLINESUBSYSTEMUTILS_API UIpNetDriver : public UNetDriver
 	}
 	// End UNetDriver Interface
 
+	// Begin UIpNetDriver interface.
+	virtual FSocket * CreateSocket();
+	// End UIpNetDriver interface.
+
 	// Begin FExec Interface
 	virtual bool Exec( UWorld* InWorld, const TCHAR* Cmd, FOutputDevice& Ar=*GLog ) OVERRIDE;
 	// End FExec Interface
diff --git a/Engine/Source/Runtime/Online/OnlineSubsystemUtils/Private/IpNetDriver.cpp b/Engine/Source/Runtime/Online/OnlineSubsystemUtils/Private/IpNetDriver.cpp
--- a/Engine/Source/Runtime/Online/OnlineSubsystemUtils/Private/IpNetDriver.cpp
+++ b/Engine/Source/Runtime/Online/OnlineSubsystemUtils/Private/IpNetDriver.cpp
@@ -35,6 +35,20 @@ ISocketSubsystem* UIpNetDriver::GetSocketSubsystem()
 	return ISocketSubsystem::Get();
 }
 
+FSocket * UIpNetDriver::CreateSocket()
+{
+	// Create UDP socket and enable broadcasting.
+	ISocketSubsystem* SocketSubsystem = GetSocketSubsystem();
+
+	if (SocketSubsystem == NULL)
+	{
+		UE_LOG(LogNet, Warning, TEXT("UIpNetDriver::CreateSocket: Unable to find socket subsystem"));
+		return NULL;
+	}
+
+	return SocketSubsystem->CreateSocket( NAME_DGram, TEXT( "Unreal" ) );
+}
+
 bool UIpNetDriver::InitBase( bool bInitAsClient, FNetworkNotify* InNotify, const FURL& URL, bool bReuseAddressAndPort, FString& Error )
 {
 	if (!Super::InitBase(bInitAsClient, InNotify, URL, bReuseAddressAndPort, Error))
@@ -51,18 +65,8 @@ bool UIpNetDriver::InitBase( bool bInitAsClient, FNetworkNotify* InNotify, const
 
 	// Derived types may have already allocated a socket
 
-	// First, try and create a socket specific to the session subsystem
-	if ( Socket == NULL )
-	{
-		// Create UDP socket and enable broadcasting.
-		Socket = SocketSubsystem->CreateSocket( FName( "SessionSocket" ), TEXT( "Unreal" ) );
-	}
-
-	if (Socket == NULL)
-	{
-		// Create UDP socket and enable broadcasting.
-		Socket = SocketSubsystem->CreateSocket(NAME_DGram, TEXT("Unreal"));
-	}
+	// Create the socket that we will use to communicate with
+	Socket = CreateSocket();
 
 	if( Socket == NULL )
 	{
