#UE4 - fixed initialization order issue

diff --git a/Engine/Source/Runtime/Online/OnlineSubsystem/Private/OnlineSubsystemImpl.cpp b/Engine/Source/Runtime/Online/OnlineSubsystem/Private/OnlineSubsystemImpl.cpp
--- a/Engine/Source/Runtime/Online/OnlineSubsystem/Private/OnlineSubsystemImpl.cpp
+++ b/Engine/Source/Runtime/Online/OnlineSubsystem/Private/OnlineSubsystemImpl.cpp
@@ -6,15 +6,15 @@
 
 FOnlineSubsystemImpl::FOnlineSubsystemImpl() :
 	InstanceName(NAME_None),
-	NamedInterfaces(NULL),
 	bForceDedicated(false)
+	NamedInterfaces(NULL)
 {
 }
 
 FOnlineSubsystemImpl::FOnlineSubsystemImpl(FName InInstanceName) :
 	InstanceName(InInstanceName),
-	NamedInterfaces(NULL),
-	bForceDedicated(false)
+	bForceDedicated(false),
+	NamedInterfaces(NULL)
 {
 }
 
