#UE4 - moved Steam/NULL OSS onto FOnlineSubsystemImpl base class

diff --git a/Engine/Source/Runtime/Online/OnlineSubsystemNull/Public/OnlineSubsystemNull.h b/Engine/Source/Runtime/Online/OnlineSubsystemNull/Public/OnlineSubsystemNull.h
--- a/Engine/Source/Runtime/Online/OnlineSubsystemNull/Public/OnlineSubsystemNull.h
+++ b/Engine/Source/Runtime/Online/OnlineSubsystemNull/Public/OnlineSubsystemNull.h
@@ -3,6 +3,8 @@
 #pragma once
 
 #include "OnlineSubsystem.h"
+#include "OnlineSubsystemImpl.h"
+#include "OnlineSubsystemNullPackage.h"
 
 /** Forward declarations of all interface classes */
 typedef TSharedPtr<class FOnlineSessionNull, ESPMode::ThreadSafe> FOnlineSessionNullPtr;
@@ -19,7 +21,7 @@ typedef TSharedPtr<class FOnlineAchievementsNull, ESPMode::ThreadSafe> FOnlineAc
  *	OnlineSubsystemNull - Implementation of the online subsystem for Null services
  */
 class ONLINESUBSYSTEMNULL_API FOnlineSubsystemNull : 
-	public IOnlineSubsystem,
+	public FOnlineSubsystemImpl,
 	public FTickerObjectBase
 {
 
@@ -61,11 +63,6 @@ public:
 
 	// FOnlineSubsystemNull
 
-	/**
-	 * @return the name of the online subsystem instance
-	 */
-	FName GetInstanceName() const { return InstanceName; }
-
 	/**
 	 * Is the Null API available for use
 	 * @return true if Null functionality is available, false otherwise
@@ -76,7 +73,7 @@ PACKAGE_SCOPE:
 
 	/** Only the factory makes instances */
 	FOnlineSubsystemNull(FName InInstanceName) :
-		InstanceName(InInstanceName),
+		FOnlineSubsystemImpl(InInstanceName),
 		SessionInterface(NULL),
 		LeaderboardsInterface(NULL),
 		IdentityInterface(NULL),
@@ -86,7 +83,6 @@ PACKAGE_SCOPE:
 	{}
 
 	FOnlineSubsystemNull() :
-		InstanceName(NAME_None),
 		SessionInterface(NULL),
 		LeaderboardsInterface(NULL),
 		IdentityInterface(NULL),
@@ -97,9 +93,6 @@ PACKAGE_SCOPE:
 
 private:
 
-	/** Instance name (disambiguates PIE instances for example) */
-	FName InstanceName;
-
 	/** Interface to the session services */
 	FOnlineSessionNullPtr SessionInterface;
 
diff --git a/Engine/Source/Runtime/Online/OnlineSubsystemSteam/Public/OnlineSubsystemSteam.h b/Engine/Source/Runtime/Online/OnlineSubsystemSteam/Public/OnlineSubsystemSteam.h
--- a/Engine/Source/Runtime/Online/OnlineSubsystemSteam/Public/OnlineSubsystemSteam.h
+++ b/Engine/Source/Runtime/Online/OnlineSubsystemSteam/Public/OnlineSubsystemSteam.h
@@ -3,6 +3,7 @@
 #pragma once
 
 #include "OnlineSubsystem.h"
+#include "OnlineSubsystemImpl.h"
 #include "OnlineSubsystemSteamPackage.h"
 
 /** Forward declarations of all interface classes */
@@ -20,14 +21,11 @@ typedef TSharedPtr<class FOnlineAchievementsSteam, ESPMode::ThreadSafe> FOnlineA
  *	OnlineSubsystemSteam - Implementation of the online subsystem for STEAM services
  */
 class ONLINESUBSYSTEMSTEAM_API FOnlineSubsystemSteam : 
-	public IOnlineSubsystem, 
+	public FOnlineSubsystemImpl, 
 	public FTickerObjectBase
 {
 protected:
 
-	/** Instance name (disambiguates PIE instances for example) */
-	FName InstanceName;
-
 	/** Has the STEAM client APIs been initialized */
 	bool bSteamworksClientInitialized;
 
@@ -86,7 +84,7 @@ PACKAGE_SCOPE:
 
 	/** Only the factory makes instances */
 	FOnlineSubsystemSteam(FName InInstanceName) :
-		InstanceName(InInstanceName),
+		FOnlineSubsystemImpl(InInstanceName),
 		bSteamworksClientInitialized(false),
 		bSteamworksGameServerInitialized(false),
 		SteamAppID(0),
@@ -106,7 +104,6 @@ PACKAGE_SCOPE:
 	{}
 
 	FOnlineSubsystemSteam() : 
-		InstanceName(NAME_None),
 		bSteamworksClientInitialized(false),
 		bSteamworksGameServerInitialized(false),
 		SteamAppID(0),
@@ -248,11 +245,6 @@ public:
 
 	// FOnlineSubsystemSteam
 
-	/**
-	 * @return the name of the online subsystem instance
-	 */
-	FName GetInstanceName() const { return InstanceName; }
-
 	/**
 	 * Whether or not the Steam Client interfaces are available; these interfaces are only available, if the Steam Client program is running
 	 * NOTE: These interfaces are made unavailable, when running a dedicated server
