#ue4
- Fixed incorrect grp. prefixing in FOnlineAchievementsIOS
- Added basic implementation of FOnlineAchievementsIOS::ResetAchievements

diff --git a/Engine/Source/Runtime/Online/IOS/OnlineSubsystemIOS/Private/OnlineAchievementsInterfaceIOS.cpp b/Engine/Source/Runtime/Online/IOS/OnlineSubsystemIOS/Private/OnlineAchievementsInterfaceIOS.cpp
--- a/Engine/Source/Runtime/Online/IOS/OnlineSubsystemIOS/Private/OnlineAchievementsInterfaceIOS.cpp
+++ b/Engine/Source/Runtime/Online/IOS/OnlineSubsystemIOS/Private/OnlineAchievementsInterfaceIOS.cpp
@@ -8,7 +8,7 @@
 #include <GameKit/GKAchievementDescription.h>
 
 
-FOnlineAchievementsIOS::FOnlineAchievementsIOS( class FOnlineSubsystemIOS* InSubsystem )
+FOnlineAchievementsIOS::FOnlineAchievementsIOS(class FOnlineSubsystemIOS* InSubsystem)
 {
 	UE_LOG(LogOnline, Display, TEXT("FOnlineSubsystemIOS::FOnlineAchievementsIOS()"));
 	IOSSubsystem = InSubsystem;
@@ -126,7 +126,8 @@ bool FOnlineAchievementsIOS::WriteAchievements(const FUniqueNetId& PlayerId, FOn
 		const FVariantData& Stat = It.Value();
 
 		// Create an achievement object which should be reported to the server.
-		NSString* AchievementID = [NSString stringWithFString:NormalizeAchievementName(It.Key().ToString())];
+		const FString AchievementName(It.Key().ToString());
+		NSString* AchievementID = [NSString stringWithFString:AchievementName];
 		
 		GKAchievement* Achievement = [[[GKAchievement alloc] initWithIdentifier:AchievementID] autorelease];
 
@@ -160,7 +161,7 @@ bool FOnlineAchievementsIOS::WriteAchievements(const FUniqueNetId& PlayerId, FOn
 	}
 
 	// flush the achievements to the server
-	if( [UnreportedAchievements count] > 0 )
+	if ([UnreportedAchievements count] > 0)
 	{
 		dispatch_async(dispatch_get_main_queue(), ^
 		{
@@ -171,7 +172,7 @@ bool FOnlineAchievementsIOS::WriteAchievements(const FUniqueNetId& PlayerId, FOn
 					for (GKAchievement* achievement in UnreportedAchievements)
 					{
 						FString AchievementId(achievement.identifier);
-						UE_LOG(LogOnline, Display, TEXT("Successfully reported achievement: %s, isComplete:%d"), *AchievementId, [achievement isCompleted] ? 1 : 0);
+						UE_LOG(LogOnline, Display, TEXT("Successfully reported achievement: %s, isCompleted:%d"), *AchievementId, [achievement isCompleted] ? 1 : 0);
 
 						if ([achievement isCompleted])
 						{
@@ -193,7 +194,7 @@ bool FOnlineAchievementsIOS::WriteAchievements(const FUniqueNetId& PlayerId, FOn
 				[UnreportedAchievements release];
 
 				// Report whether this succeeded or not back to whoever is listening.
-				WriteObject->WriteState = Error == nil ? EOnlineAsyncTaskState::Done : EOnlineAsyncTaskState::Failed;
+				WriteObject->WriteState = (Error == nil) ? EOnlineAsyncTaskState::Done : EOnlineAsyncTaskState::Failed;
 				// Report back to the game thread whether this succeeded.
 				
 				[FIOSAsyncTask CreateTaskWithBlock : ^ bool(void)
@@ -218,9 +219,9 @@ bool FOnlineAchievementsIOS::WriteAchievements(const FUniqueNetId& PlayerId, FOn
 
 EOnlineCachedResult::Type FOnlineAchievementsIOS::GetCachedAchievement(const FUniqueNetId& PlayerId, const FString& AchievementId, FOnlineAchievement& OutAchievement)
 {
-	for( int32 AchievementIdx = 0; AchievementIdx < Achievements.Num(); AchievementIdx++ )
+	for (int32 AchievementIdx = 0; AchievementIdx < Achievements.Num(); AchievementIdx++)
 	{
-		if( Achievements[AchievementIdx].Id == AchievementId )
+		if (Achievements[AchievementIdx].Id == AchievementId)
 		{
 			OutAchievement = Achievements[AchievementIdx];
 			return EOnlineCachedResult::Success;
@@ -237,35 +238,49 @@ EOnlineCachedResult::Type FOnlineAchievementsIOS::GetCachedAchievements(const FU
 	OutAchievements = Achievements;
 
 	// did we have them cached?
-	return OutAchievements.Num() > 0 ? EOnlineCachedResult::Success : EOnlineCachedResult::NotFound;
+	return (OutAchievements.Num() > 0) ? EOnlineCachedResult::Success : EOnlineCachedResult::NotFound;
 }
 
 
 EOnlineCachedResult::Type FOnlineAchievementsIOS::GetCachedAchievementDescription(const FString& AchievementId, FOnlineAchievementDesc& OutAchievementDesc)
 {
-	FOnlineAchievementDesc* FoundDesc = AchievementDescriptions.Find( AchievementId );
-	if( FoundDesc != NULL )
+	if (FOnlineAchievementDesc* FoundDesc = AchievementDescriptions.Find(AchievementId))
 	{
 		OutAchievementDesc = *FoundDesc;
 	}
-	return ( FoundDesc != NULL ) ? EOnlineCachedResult::Success : EOnlineCachedResult::NotFound;
+	return (FoundDesc != NULL) ? EOnlineCachedResult::Success : EOnlineCachedResult::NotFound;
 }
 
 
-FString FOnlineAchievementsIOS::NormalizeAchievementName( const FString& InAchievementId )
+FString FOnlineAchievementsIOS::NormalizeAchievementName(const FString& InAchievementId)
 {
-	FString Id = InAchievementId;
-	if( Id.Find( TEXT(".") ) == INDEX_NONE )
-	{
-		Id = FString::Printf( TEXT("grp.%s"), *InAchievementId );
-	}
-	return Id;
+	return InAchievementId;
 }
 
 #if !UE_BUILD_SHIPPING
 bool FOnlineAchievementsIOS::ResetAchievements( const FUniqueNetId& PlayerId )
 {
-	check(!TEXT("ResetAchievements has not been implemented for IOS"));
-	return false;
+	dispatch_async(dispatch_get_main_queue(), ^
+	{
+		[GKAchievement resetAchievementsWithCompletionHandler : ^(NSError *error)
+		{
+			bool bSuccess = error == NULL;
+			UE_LOG(LogOnline, Display, TEXT("FOnlineAchievementsIOS::ResetAchievements - %d"), bSuccess ? 1 : 0);
+
+			if (bSuccess)
+			{
+				// Wipe out the achievement descriptions back on the game thread.
+				[FIOSAsyncTask CreateTaskWithBlock : ^ bool(void)
+				{
+					Achievements.Empty();
+					AchievementDescriptions.Empty();
+					return true;
+				}];
+			}
+		}
+		];
+	});
+
+	return true;
 };
 #endif // !UE_BUILD_SHIPPING
diff --git a/Engine/Source/Runtime/Online/IOS/OnlineSubsystemIOS/Private/OnlineAchievementsInterfaceIOS.h b/Engine/Source/Runtime/Online/IOS/OnlineSubsystemIOS/Private/OnlineAchievementsInterfaceIOS.h
--- a/Engine/Source/Runtime/Online/IOS/OnlineSubsystemIOS/Private/OnlineAchievementsInterfaceIOS.h
+++ b/Engine/Source/Runtime/Online/IOS/OnlineSubsystemIOS/Private/OnlineAchievementsInterfaceIOS.h
@@ -50,19 +50,7 @@ public:
 	/**
 	 * Default destructor
 	 */
-	virtual ~FOnlineAchievementsIOS(){}
-
-
-private:
-
-	/**
-	 * Convert the name of the achievement to one which is valid with game center. I.e. grp.InAchievementName
-	 *
-	 * @param InAchievementName - The name we will check matches the game center requirements
-	 *
-	 * @return The name in a game center compatible format.
-	 */
-	FString NormalizeAchievementName( const FString& InAchievementName );
+	virtual ~FOnlineAchievementsIOS() {}
 };
 
 typedef TSharedPtr<FOnlineAchievementsIOS, ESPMode::ThreadSafe> FOnlineAchievementsIOSPtr;
\ No newline at end of file
