#ue4
- Fixed some iOS compile errors

diff --git a/Engine/Source/Runtime/Online/IOS/OnlineSubsystemIOS/Private/OnlineAchievementsInterfaceIOS.cpp b/Engine/Source/Runtime/Online/IOS/OnlineSubsystemIOS/Private/OnlineAchievementsInterfaceIOS.cpp
--- a/Engine/Source/Runtime/Online/IOS/OnlineSubsystemIOS/Private/OnlineAchievementsInterfaceIOS.cpp
+++ b/Engine/Source/Runtime/Online/IOS/OnlineSubsystemIOS/Private/OnlineAchievementsInterfaceIOS.cpp
@@ -241,20 +241,17 @@ EOnlineCachedResult::Type FOnlineAchievementsIOS::GetCachedAchievements(const FU
 	return (OutAchievements.Num() > 0) ? EOnlineCachedResult::Success : EOnlineCachedResult::NotFound;
 }
 
-
 EOnlineCachedResult::Type FOnlineAchievementsIOS::GetCachedAchievementDescription(const FString& AchievementId, FOnlineAchievementDesc& OutAchievementDesc)
 {
 	if (FOnlineAchievementDesc* FoundDesc = AchievementDescriptions.Find(AchievementId))
 	{
 		OutAchievementDesc = *FoundDesc;
+		return EOnlineCachedResult::Success;
+	}
+	else
+	{
+		return EOnlineCachedResult::NotFound;
 	}
-	return (FoundDesc != NULL) ? EOnlineCachedResult::Success : EOnlineCachedResult::NotFound;
-}
-
-
-FString FOnlineAchievementsIOS::NormalizeAchievementName(const FString& InAchievementId)
-{
-	return InAchievementId;
 }
 
 #if !UE_BUILD_SHIPPING
