#ue4
- Added a helper class FOnlineSubsystemBPCallHelper to reduce some of the boilerplate in BP async nodes

diff --git a/Engine/Source/Runtime/Online/OnlineSubsystemUtils/Private/OnlineSubsystemUtils.cpp b/Engine/Source/Runtime/Online/OnlineSubsystemUtils/Private/OnlineSubsystemUtils.cpp
--- a/Engine/Source/Runtime/Online/OnlineSubsystemUtils/Private/OnlineSubsystemUtils.cpp
+++ b/Engine/Source/Runtime/Online/OnlineSubsystemUtils/Private/OnlineSubsystemUtils.cpp
@@ -262,3 +262,33 @@ static bool OnlineExec( UWorld* InWorld, const TCHAR* Cmd, FOutputDevice& Ar )
 /** Our entry point for all online exec routing */
 FStaticSelfRegisteringExec OnlineExecRegistration(OnlineExec);
 
+//////////////////////////////////////////////////////////////////////////
+// FOnlineSubsystemBPCallHelper
+
+FOnlineSubsystemBPCallHelper::FOnlineSubsystemBPCallHelper(const TCHAR* CallFunctionContext, FName SystemName)
+	: OnlineSub(IOnlineSubsystem::Get(SystemName))
+	, FunctionContext(CallFunctionContext)
+{
+	if (OnlineSub == nullptr)
+	{
+		FFrame::KismetExecutionMessage(*FString::Printf(TEXT("%s - Invalid or uninitialized OnlineSubsystem"), FunctionContext), ELogVerbosity::Warning);
+	}
+}
+
+void FOnlineSubsystemBPCallHelper::QueryIDFromPlayerController(APlayerController* PlayerController)
+{
+	UserID.Reset();
+
+	if (APlayerState* PlayerState = (PlayerController != NULL) ? PlayerController->PlayerState : NULL)
+	{
+		UserID = PlayerState->UniqueId.GetUniqueNetId();
+		if (!UserID.IsValid())
+		{
+			FFrame::KismetExecutionMessage(*FString::Printf(TEXT("%s - Cannot map local player to unique net ID"), FunctionContext), ELogVerbosity::Warning);
+		}
+	}
+	else
+	{
+		FFrame::KismetExecutionMessage(*FString::Printf(TEXT("%s - Invalid player state"), FunctionContext), ELogVerbosity::Warning);
+	}
+}
diff --git a/Engine/Source/Runtime/Online/OnlineSubsystemUtils/Private/OnlineSubsystemUtilsPrivatePCH.h b/Engine/Source/Runtime/Online/OnlineSubsystemUtils/Private/OnlineSubsystemUtilsPrivatePCH.h
--- a/Engine/Source/Runtime/Online/OnlineSubsystemUtils/Private/OnlineSubsystemUtilsPrivatePCH.h
+++ b/Engine/Source/Runtime/Online/OnlineSubsystemUtils/Private/OnlineSubsystemUtilsPrivatePCH.h
@@ -12,6 +12,25 @@
 #include "ModuleManager.h"
 #include "OnlineSubsystemUtilsClasses.h"
 
+// Helper class for various methods to reduce the call hierarchy
+struct FOnlineSubsystemBPCallHelper
+{
+public:
+	FOnlineSubsystemBPCallHelper(const TCHAR* CallFunctionContext, FName SystemName = NAME_None);
+
+	void QueryIDFromPlayerController(APlayerController* PlayerController);
+
+	bool IsValid() const
+	{
+		return UserID.IsValid() && (OnlineSub != nullptr);
+	}
+
+public:
+	TSharedPtr<class FUniqueNetId> UserID;
+	IOnlineSubsystem* const OnlineSub;
+	const TCHAR* FunctionContext;
+};
+
 #define INVALID_INDEX -1
 
 
