#ue4
- Added support for reporting achievement progress from Blueprints with a WriteAchievementProgress async node

diff --git a/Engine/Source/Runtime/Online/OnlineSubsystemUtils/Classes/AchievementWriteCallbackProxy.h b/Engine/Source/Runtime/Online/OnlineSubsystemUtils/Classes/AchievementWriteCallbackProxy.h
new file mode 100644
--- /dev/null
+++ b/Engine/Source/Runtime/Online/OnlineSubsystemUtils/Classes/AchievementWriteCallbackProxy.h
@@ -0,0 +1,42 @@
+// Copyright 1998-2014 Epic Games, Inc. All Rights Reserved.
+#pragma once
+
+#include "Runtime/Online/OnlineSubsystem/Public/Interfaces/OnlineAchievementsInterface.h"
+#include "AchievementWriteCallbackProxy.generated.h"
+
+UCLASS(MinimalAPI)
+class UAchievementWriteCallbackProxy : public UOnlineBlueprintCallProxyBase
+{
+	GENERATED_UCLASS_BODY()
+
+	// Called when there is a successful achievement write
+	UPROPERTY(BlueprintAssignable)
+	FEmptyOnlineDelegate OnSuccess;
+
+	// Called when there is an unsuccessful achievement write
+	UPROPERTY(BlueprintAssignable)
+	FEmptyOnlineDelegate OnFailure;
+
+	// Writes progress about an achievement to the default online subsystem
+	UFUNCTION(BlueprintCallable, meta=(BlueprintInternalUseOnly = "true"), Category="Online|Achievements")
+	static UAchievementWriteCallbackProxy* WriteAchievementProgress(class APlayerController* PlayerController, FName AchievementName, float Progress = 100.0f);
+
+	// UOnlineBlueprintCallProxyBase interface
+	virtual void Activate() OVERRIDE;
+	// End of UOnlineBlueprintCallProxyBase interface
+
+	// UObject interface
+	virtual void BeginDestroy() OVERRIDE;
+	// End of UObject interface
+
+private:
+	// Internal callback when the achievement is written, calls out to the public success/failure callbacks
+	void OnAchievementWritten(const FUniqueNetId& UserID, bool bSuccess);
+
+private:
+	/** The player controller triggering things */
+	TWeakObjectPtr<APlayerController> PlayerControllerWeakPtr;
+
+	/** The achievements write object */
+	FOnlineAchievementsWritePtr WriteObject;
+};
diff --git a/Engine/Source/Runtime/Online/OnlineSubsystemUtils/Private/AchievementWriteCallbackProxy.cpp b/Engine/Source/Runtime/Online/OnlineSubsystemUtils/Private/AchievementWriteCallbackProxy.cpp
new file mode 100644
--- /dev/null
+++ b/Engine/Source/Runtime/Online/OnlineSubsystemUtils/Private/AchievementWriteCallbackProxy.cpp
@@ -0,0 +1,72 @@
+// Copyright 1998-2014 Epic Games, Inc. All Rights Reserved.
+
+#include "OnlineSubsystemUtilsPrivatePCH.h"
+
+//////////////////////////////////////////////////////////////////////////
+// ULeaderboardQueryCallbackProxy
+
+UAchievementWriteCallbackProxy::UAchievementWriteCallbackProxy(const class FPostConstructInitializeProperties& PCIP)
+	: Super(PCIP)
+{
+}
+
+UAchievementWriteCallbackProxy* UAchievementWriteCallbackProxy::WriteAchievementProgress(class APlayerController* PlayerController, FName AchievementName, float Progress)
+{
+	UAchievementWriteCallbackProxy* Proxy = NewObject<UAchievementWriteCallbackProxy>();
+
+	Proxy->WriteObject = MakeShareable(new FOnlineAchievementsWrite);
+	Proxy->WriteObject->SetFloatStat(AchievementName, Progress);
+	Proxy->PlayerControllerWeakPtr = PlayerController;
+
+	return Proxy;
+}
+
+void UAchievementWriteCallbackProxy::Activate()
+{
+	FOnlineSubsystemBPCallHelper Helper(TEXT("WriteAchievementObject"));
+	Helper.QueryIDFromPlayerController(PlayerControllerWeakPtr.Get());
+
+	if (Helper.IsValid())
+	{
+		IOnlineAchievementsPtr Achievements = Helper.OnlineSub->GetAchievementsInterface();
+		if (Achievements.IsValid())
+		{
+			//@TODO: Bind OnAchievementWritten in
+			FOnlineAchievementsWriteRef WriteObjectRef = WriteObject.ToSharedRef();
+			Achievements->WriteAchievements(*Helper.UserID, WriteObjectRef);
+
+			// OnAchievementWritten will get called, nothing more to do
+			//@TODO: Remove this call
+			OnSuccess.Broadcast();
+			return;
+		}
+		else
+		{
+			FFrame::KismetExecutionMessage(TEXT("WriteAchievementObject - Achievements not supported by Online Subsystem"), ELogVerbosity::Warning);
+		}
+	}
+
+	// Fail immediately
+	OnFailure.Broadcast();
+}
+
+void UAchievementWriteCallbackProxy::OnAchievementWritten(const FUniqueNetId& UserID, bool bSuccess)
+{
+	if (bSuccess)
+	{
+		OnSuccess.Broadcast();
+	}
+	else
+	{
+		OnFailure.Broadcast();
+	}
+
+	WriteObject.Reset();
+}
+
+void UAchievementWriteCallbackProxy::BeginDestroy()
+{
+	WriteObject.Reset();
+
+	Super::BeginDestroy();
+}
