#ue4
- Updated WriteAchievements to take in a completion/failure delegate instead of calling a stateful OnAchievementsWritten
- Fixed various locations that were not currently calling the delegate

diff --git a/Engine/Source/Runtime/Online/IOS/OnlineSubsystemIOS/Private/OnlineAchievementsInterfaceIOS.cpp b/Engine/Source/Runtime/Online/IOS/OnlineSubsystemIOS/Private/OnlineAchievementsInterfaceIOS.cpp
--- a/Engine/Source/Runtime/Online/IOS/OnlineSubsystemIOS/Private/OnlineAchievementsInterfaceIOS.cpp
+++ b/Engine/Source/Runtime/Online/IOS/OnlineSubsystemIOS/Private/OnlineAchievementsInterfaceIOS.cpp
@@ -110,8 +110,11 @@ void FOnlineAchievementsIOS::QueryAchievementDescriptions(const FUniqueNetId& Pl
 }
 
 
-bool FOnlineAchievementsIOS::WriteAchievements(const FUniqueNetId& PlayerId, FOnlineAchievementsWriteRef& InWriteObject)
+void FOnlineAchievementsIOS::WriteAchievements(const FUniqueNetId& PlayerId, FOnlineAchievementsWriteRef& InWriteObject, const FOnAchievementsWrittenDelegate& Delegate)
 {
+	// Make a copy of the delegate so it won't be a reference inside the FIOSAsyncTask
+	FOnAchievementsWrittenDelegate CopiedDelegate = Delegate;
+
 	// Hold a reference to the write object for the completion handler to write to.
 	FOnlineAchievementsWriteRef WriteObject = InWriteObject;
 	WriteObject->WriteState = EOnlineAsyncTaskState::InProgress;
@@ -174,9 +177,9 @@ bool FOnlineAchievementsIOS::WriteAchievements(const FUniqueNetId& PlayerId, FOn
 						FString AchievementId(achievement.identifier);
 						UE_LOG(LogOnline, Display, TEXT("Successfully reported achievement: %s, isCompleted:%d"), *AchievementId, [achievement isCompleted] ? 1 : 0);
 
+						// Report any completed achievements to the game thread
 						if ([achievement isCompleted])
 						{
-							// Report back to the game thread whether this succeeded.
 							[FIOSAsyncTask CreateTaskWithBlock : ^ bool(void)
 							{
 								TriggerOnAchievementUnlockedDelegates(PlayerId, AchievementId);
@@ -195,11 +198,11 @@ bool FOnlineAchievementsIOS::WriteAchievements(const FUniqueNetId& PlayerId, FOn
 
 				// Report whether this succeeded or not back to whoever is listening.
 				WriteObject->WriteState = (Error == nil) ? EOnlineAsyncTaskState::Done : EOnlineAsyncTaskState::Failed;
+
 				// Report back to the game thread whether this succeeded.
-				
 				[FIOSAsyncTask CreateTaskWithBlock : ^ bool(void)
 				{
-					TriggerOnAchievementsWrittenDelegates(PlayerId, Error == nil);
+					CopiedDelegate.ExecuteIfBound(PlayerId, Error == nil);
 					return true;
 				}];
 			}
@@ -209,11 +212,9 @@ bool FOnlineAchievementsIOS::WriteAchievements(const FUniqueNetId& PlayerId, FOn
 	else
 	{
 		UE_LOG(LogOnline, Warning, TEXT("No achievements were written to be flushed"));
-		TriggerOnAchievementsWrittenDelegates( PlayerId, false );
 		WriteObject->WriteState = EOnlineAsyncTaskState::Failed;
+		CopiedDelegate.ExecuteIfBound(PlayerId, false);
 	}
-
-	return WriteObject->WriteState != EOnlineAsyncTaskState::Failed;
 }
 
 
diff --git a/Engine/Source/Runtime/Online/IOS/OnlineSubsystemIOS/Private/OnlineAchievementsInterfaceIOS.h b/Engine/Source/Runtime/Online/IOS/OnlineSubsystemIOS/Private/OnlineAchievementsInterfaceIOS.h
--- a/Engine/Source/Runtime/Online/IOS/OnlineSubsystemIOS/Private/OnlineAchievementsInterfaceIOS.h
+++ b/Engine/Source/Runtime/Online/IOS/OnlineSubsystemIOS/Private/OnlineAchievementsInterfaceIOS.h
@@ -27,7 +27,7 @@ private:
 public:
 
 	// Begin IOnlineAchievements interface
-	virtual bool WriteAchievements( const FUniqueNetId& PlayerId, FOnlineAchievementsWriteRef& WriteObject ) OVERRIDE;
+	virtual void WriteAchievements(const FUniqueNetId& PlayerId, FOnlineAchievementsWriteRef& WriteObject, const FOnAchievementsWrittenDelegate& Delegate = FOnAchievementsWrittenDelegate()) OVERRIDE;
 	virtual void QueryAchievements(const FUniqueNetId & PlayerId, const FOnQueryAchievementsCompleteDelegate & Delegate = FOnQueryAchievementsCompleteDelegate()) OVERRIDE;
 	virtual void QueryAchievementDescriptions( const FUniqueNetId& PlayerId, const FOnQueryAchievementsCompleteDelegate & Delegate = FOnQueryAchievementsCompleteDelegate() ) OVERRIDE;
 	virtual EOnlineCachedResult::Type GetCachedAchievement(const FUniqueNetId& PlayerId, const FString& AchievementId, FOnlineAchievement& OutAchievement) OVERRIDE;
diff --git a/Engine/Source/Runtime/Online/OnlineSubsystem/Public/Interfaces/OnlineAchievementsInterface.h b/Engine/Source/Runtime/Online/OnlineSubsystem/Public/Interfaces/OnlineAchievementsInterface.h
--- a/Engine/Source/Runtime/Online/OnlineSubsystem/Public/Interfaces/OnlineAchievementsInterface.h
+++ b/Engine/Source/Runtime/Online/OnlineSubsystem/Public/Interfaces/OnlineAchievementsInterface.h
@@ -95,10 +95,9 @@ public:
 	 *
 	 * @param PlayerId - The uid of the player we are writing achievements for
 	 * @param WriteObject - The stats holder containing the achievements we are writing.
-	 *
-	 * @return Whether we have kicked off a write attempt
+	 * @param Delegate - The delegate to call when the write has completed or failed.
 	 */
-	virtual bool WriteAchievements(const FUniqueNetId& PlayerId, FOnlineAchievementsWriteRef& WriteObject) = 0;
+	virtual void WriteAchievements(const FUniqueNetId& PlayerId, FOnlineAchievementsWriteRef& WriteObject, const FOnAchievementsWrittenDelegate& Delegate = FOnAchievementsWrittenDelegate()) = 0;
 
 	
 	/**
@@ -165,12 +164,6 @@ public:
 	virtual bool ResetAchievements( const FUniqueNetId& PlayerId ) = 0;
 #endif // !UE_BUILD_SHIPPING
 
-	/**
-	 * Delegate fired when achievements were written to the server
-	 */
-	DEFINE_ONLINE_DELEGATE_TWO_PARAM(OnAchievementsWritten, const FUniqueNetId&, bool);
-
-
 	/**
 	 * Delegate fired when an achievement on the server was unlocked
 	 */
diff --git a/Engine/Source/Runtime/Online/OnlineSubsystemNull/Private/OnlineAchievementsInterfaceNull.cpp b/Engine/Source/Runtime/Online/OnlineSubsystemNull/Private/OnlineAchievementsInterfaceNull.cpp
--- a/Engine/Source/Runtime/Online/OnlineSubsystemNull/Private/OnlineAchievementsInterfaceNull.cpp
+++ b/Engine/Source/Runtime/Online/OnlineSubsystemNull/Private/OnlineAchievementsInterfaceNull.cpp
@@ -20,14 +20,14 @@ bool FOnlineAchievementsNull::ReadAchievementsFromConfig()
 	return Config.ReadAchievements(Achievements);
 }
 
-bool FOnlineAchievementsNull::WriteAchievements(const FUniqueNetId& PlayerId, FOnlineAchievementsWriteRef& WriteObject)
+void FOnlineAchievementsNull::WriteAchievements(const FUniqueNetId& PlayerId, FOnlineAchievementsWriteRef& WriteObject, const FOnAchievementsWrittenDelegate& Delegate)
 {
 	if (!ReadAchievementsFromConfig())
 	{
 		// we don't have achievements
-		TriggerOnAchievementsWrittenDelegates(PlayerId, false);
 		WriteObject->WriteState = EOnlineAsyncTaskState::Failed;
-		return false;
+		Delegate.ExecuteIfBound(PlayerId, false);
+		return;
 	}
 
 	FUniqueNetIdString NullId(PlayerId);
@@ -35,9 +35,9 @@ bool FOnlineAchievementsNull::WriteAchievements(const FUniqueNetId& PlayerId, FO
 	if (NULL == PlayerAch)
 	{
 		// achievements haven't been read for a player
-		TriggerOnAchievementsWrittenDelegates(PlayerId, false);
 		WriteObject->WriteState = EOnlineAsyncTaskState::Failed;
-		return false;
+		Delegate.ExecuteIfBound(PlayerId, false);
+		return;
 	}
 
 	// treat each achievement as unlocked
@@ -56,9 +56,7 @@ bool FOnlineAchievementsNull::WriteAchievements(const FUniqueNetId& PlayerId, FO
 	}
 
 	WriteObject->WriteState = EOnlineAsyncTaskState::Done;
-	TriggerOnAchievementsWrittenDelegates(PlayerId, true);
-
-	return true;
+	Delegate.ExecuteIfBound(PlayerId, true);
 };
 
 void FOnlineAchievementsNull::QueryAchievements( const FUniqueNetId& PlayerId, const FOnQueryAchievementsCompleteDelegate & Delegate )
diff --git a/Engine/Source/Runtime/Online/OnlineSubsystemNull/Private/OnlineAchievementsInterfaceNull.h b/Engine/Source/Runtime/Online/OnlineSubsystemNull/Private/OnlineAchievementsInterfaceNull.h
--- a/Engine/Source/Runtime/Online/OnlineSubsystemNull/Private/OnlineAchievementsInterfaceNull.h
+++ b/Engine/Source/Runtime/Online/OnlineSubsystemNull/Private/OnlineAchievementsInterfaceNull.h
@@ -112,7 +112,7 @@ public:
 	virtual ~FOnlineAchievementsNull(){}
 
 	// Begin IOnlineAchievements interface
-	virtual bool WriteAchievements(const FUniqueNetId& PlayerId, FOnlineAchievementsWriteRef& WriteObject) OVERRIDE;
+	virtual void WriteAchievements(const FUniqueNetId& PlayerId, FOnlineAchievementsWriteRef& WriteObject, const FOnAchievementsWrittenDelegate& Delegate = FOnAchievementsWrittenDelegate()) OVERRIDE;
 	virtual void QueryAchievements(const FUniqueNetId & PlayerId, const FOnQueryAchievementsCompleteDelegate & Delegate = FOnQueryAchievementsCompleteDelegate()) OVERRIDE;
 	virtual void QueryAchievementDescriptions( const FUniqueNetId& PlayerId, const FOnQueryAchievementsCompleteDelegate & Delegate = FOnQueryAchievementsCompleteDelegate() ) OVERRIDE;
 	virtual EOnlineCachedResult::Type GetCachedAchievement(const FUniqueNetId& PlayerId, const FString& AchievementId, FOnlineAchievement& OutAchievement) OVERRIDE;
diff --git a/Engine/Source/Runtime/Online/OnlineSubsystemSteam/Private/OnlineAchievementsInterfaceSteam.cpp b/Engine/Source/Runtime/Online/OnlineSubsystemSteam/Private/OnlineAchievementsInterfaceSteam.cpp
--- a/Engine/Source/Runtime/Online/OnlineSubsystemSteam/Private/OnlineAchievementsInterfaceSteam.cpp
+++ b/Engine/Source/Runtime/Online/OnlineSubsystemSteam/Private/OnlineAchievementsInterfaceSteam.cpp
@@ -29,16 +29,16 @@ bool FOnlineAchievementsSteam::ReadAchievementsFromConfig()
 	return Config.ReadAchievements(Achievements);
 }
 
-bool FOnlineAchievementsSteam::WriteAchievements(const FUniqueNetId& PlayerId, FOnlineAchievementsWriteRef& WriteObject)
+void FOnlineAchievementsSteam::WriteAchievements(const FUniqueNetId& PlayerId, FOnlineAchievementsWriteRef& WriteObject, const FOnAchievementsWrittenDelegate& Delegate)
 {
 	if (!bHaveConfiguredAchievements)
 	{
 		// we don't have achievements
 		UE_LOG_ONLINE(Warning, TEXT("Steam achievements have not been configured in .ini"));
 
-		TriggerOnAchievementsWrittenDelegates(PlayerId, false);
 		WriteObject->WriteState = EOnlineAsyncTaskState::Failed;
-		return false;
+		Delegate.ExecuteIfBound(PlayerId, false);
+		return;
 	}
 
 	FUniqueNetIdSteam SteamId(PlayerId);
@@ -48,9 +48,9 @@ bool FOnlineAchievementsSteam::WriteAchievements(const FUniqueNetId& PlayerId, F
 		// we don't have achievements
 		UE_LOG_ONLINE(Warning, TEXT("Cannot report Steam achievements for non-local player %s"), *PlayerId.ToString());
 
-		TriggerOnAchievementsWrittenDelegates(PlayerId, false);
 		WriteObject->WriteState = EOnlineAsyncTaskState::Failed;
-		return false;
+		Delegate.ExecuteIfBound(PlayerId, false);
+		return;
 	}
 
 	const TArray<FOnlineAchievement> * PlayerAch = PlayerAchievements.Find(SteamId);
@@ -59,9 +59,9 @@ bool FOnlineAchievementsSteam::WriteAchievements(const FUniqueNetId& PlayerId, F
 		// achievements haven't been read for a player
 		UE_LOG_ONLINE(Warning, TEXT("Steam achievements have not been read for player %s"), *PlayerId.ToString());
 
-		TriggerOnAchievementsWrittenDelegates(PlayerId, false);
 		WriteObject->WriteState = EOnlineAsyncTaskState::Failed;
-		return false;
+		Delegate.ExecuteIfBound(PlayerId, false);
+		return;
 	}
 
 	const int32 AchNum = PlayerAch->Num();
@@ -98,11 +98,10 @@ bool FOnlineAchievementsSteam::WriteAchievements(const FUniqueNetId& PlayerId, F
 		}
 	}
 
-	StatsInt->WriteAchievements(SteamId, WriteObject);
-	return true;
+	StatsInt->WriteAchievementsInternal(SteamId, WriteObject, Delegate);
 };
 
-void FOnlineAchievementsSteam::OnWriteAchievementsComplete(const FUniqueNetIdSteam& PlayerId, bool bWasSuccessful, FOnlineAchievementsWritePtr & WriteObject)
+void FOnlineAchievementsSteam::OnWriteAchievementsComplete(const FUniqueNetIdSteam& PlayerId, bool bWasSuccessful, FOnlineAchievementsWritePtr& WriteObject, const FOnAchievementsWrittenDelegate& Delegate)
 {
 	// write object should be valid
 	check(WriteObject.IsValid());
@@ -133,7 +132,7 @@ void FOnlineAchievementsSteam::OnWriteAchievementsComplete(const FUniqueNetIdSte
 		}
 	}
 
-	TriggerOnAchievementsWrittenDelegates(PlayerId, bWasSuccessful);
+	Delegate.ExecuteIfBound(PlayerId, bWasSuccessful);
 }
 
 void FOnlineAchievementsSteam::QueryAchievements(const FUniqueNetId& PlayerId, const FOnQueryAchievementsCompleteDelegate & Delegate)
@@ -148,7 +147,7 @@ void FOnlineAchievementsSteam::QueryAchievements(const FUniqueNetId& PlayerId, c
 	}
 
 	// schedule a read (this will trigger the OnAchievementsRead delegates)
-	StatsInt->QueryAchievements( FUniqueNetIdSteam(PlayerId), Delegate );
+	StatsInt->QueryAchievementsInternal( FUniqueNetIdSteam(PlayerId), Delegate );
 }
 
 void FOnlineAchievementsSteam::QueryAchievementDescriptions( const FUniqueNetId& PlayerId, const FOnQueryAchievementsCompleteDelegate & Delegate )
@@ -163,7 +162,7 @@ void FOnlineAchievementsSteam::QueryAchievementDescriptions( const FUniqueNetId&
 	}
 
 	// schedule a read (this will trigger the OnAchievementDescriptionsRead delegates)
-	StatsInt->QueryAchievements( FUniqueNetIdSteam(PlayerId), Delegate );
+	StatsInt->QueryAchievementsInternal( FUniqueNetIdSteam(PlayerId), Delegate );
 }
 
 
diff --git a/Engine/Source/Runtime/Online/OnlineSubsystemSteam/Private/OnlineAchievementsInterfaceSteam.h b/Engine/Source/Runtime/Online/OnlineSubsystemSteam/Private/OnlineAchievementsInterfaceSteam.h
--- a/Engine/Source/Runtime/Online/OnlineSubsystemSteam/Private/OnlineAchievementsInterfaceSteam.h
+++ b/Engine/Source/Runtime/Online/OnlineSubsystemSteam/Private/OnlineAchievementsInterfaceSteam.h
@@ -129,13 +129,14 @@ PACKAGE_SCOPE:
 	 * @param PlayerId - id of a player we were making read for
 	 * @param bWasSuccessful - whether the read completed successfully
 	 * @param WriteObject - achievements we have been writing (unlocking)
+	 * @param Delegate - delegate back to the game code indicating the write is completed
 	 */
-	void OnWriteAchievementsComplete(const FUniqueNetIdSteam& PlayerId, bool bWasSuccessful, FOnlineAchievementsWritePtr & WriteObject);
+	void OnWriteAchievementsComplete(const FUniqueNetIdSteam& PlayerId, bool bWasSuccessful, FOnlineAchievementsWritePtr & WriteObject, const FOnAchievementsWrittenDelegate& Delegate);
 
 public:
 
 	// Begin IOnlineAchievements interface
-	virtual bool WriteAchievements(const FUniqueNetId& PlayerId, FOnlineAchievementsWriteRef& WriteObject) OVERRIDE;
+	virtual void WriteAchievements(const FUniqueNetId& PlayerId, FOnlineAchievementsWriteRef& WriteObject, const FOnAchievementsWrittenDelegate& Delegate = FOnAchievementsWrittenDelegate()) OVERRIDE;
 	virtual void QueryAchievements(const FUniqueNetId& PlayerId, const FOnQueryAchievementsCompleteDelegate& Delegate = FOnQueryAchievementsCompleteDelegate()) OVERRIDE;
 	virtual void QueryAchievementDescriptions( const FUniqueNetId& PlayerId, const FOnQueryAchievementsCompleteDelegate & Delegate = FOnQueryAchievementsCompleteDelegate() ) OVERRIDE;
 	virtual EOnlineCachedResult::Type GetCachedAchievement(const FUniqueNetId& PlayerId, const FString& AchievementId, FOnlineAchievement& OutAchievement) OVERRIDE;
diff --git a/Engine/Source/Runtime/Online/OnlineSubsystemSteam/Private/OnlineLeaderboardInterfaceSteam.cpp b/Engine/Source/Runtime/Online/OnlineSubsystemSteam/Private/OnlineLeaderboardInterfaceSteam.cpp
--- a/Engine/Source/Runtime/Online/OnlineSubsystemSteam/Private/OnlineLeaderboardInterfaceSteam.cpp
+++ b/Engine/Source/Runtime/Online/OnlineSubsystemSteam/Private/OnlineLeaderboardInterfaceSteam.cpp
@@ -1218,6 +1218,9 @@ private:
 	/** Reference to write object for state tracking */
 	FOnlineAchievementsWritePtr WriteObject;
 
+	/** Delegate to call when the write finishes */
+	FOnAchievementsWrittenDelegate OnWriteFinishedDelegate;
+
 	/** Hidden on purpose */
 	FOnlineAsyncTaskSteamWriteAchievements() :
 		FOnlineAsyncTaskSteamStoreStats(),
@@ -1247,9 +1250,10 @@ private:
 	}
 
 public:
-	FOnlineAsyncTaskSteamWriteAchievements(FOnlineSubsystemSteam* InSteamSubsystem, const FUniqueNetIdSteam& InUserId, FOnlineAchievementsWriteRef& InWriteObject) :
-		FOnlineAsyncTaskSteamStoreStats(InSteamSubsystem, TEXT("Unused"), InUserId),
-		WriteObject(InWriteObject)
+	FOnlineAsyncTaskSteamWriteAchievements(FOnlineSubsystemSteam* InSteamSubsystem, const FUniqueNetIdSteam& InUserId, FOnlineAchievementsWriteRef& InWriteObject, const FOnAchievementsWrittenDelegate& InOnWriteFinishedDelegate)
+		: FOnlineAsyncTaskSteamStoreStats(InSteamSubsystem, TEXT("Unused"), InUserId)
+		, WriteObject(InWriteObject)
+		, OnWriteFinishedDelegate(InOnWriteFinishedDelegate)
 	{
 	}
 
@@ -1261,7 +1265,7 @@ public:
 		FOnlineAsyncTaskSteam::TriggerDelegates();
 		
 		FOnlineAchievementsSteamPtr Achievements = StaticCastSharedPtr<FOnlineAchievementsSteam>(Subsystem->GetAchievementsInterface());
-		Achievements->OnWriteAchievementsComplete(UserId, bWasSuccessful, WriteObject);
+		Achievements->OnWriteAchievementsComplete(UserId, bWasSuccessful, WriteObject, OnWriteFinishedDelegate);
 	}
 };
 
@@ -1292,7 +1296,7 @@ bool FOnlineLeaderboardsSteam::ReadLeaderboards(const TArray< TSharedRef<FUnique
 	return true;
 }
 
-void FOnlineLeaderboardsSteam::QueryAchievements( const FUniqueNetIdSteam &	UserId, const FOnQueryAchievementsCompleteDelegate & AchievementDelegate )
+void FOnlineLeaderboardsSteam::QueryAchievementsInternal(const FUniqueNetIdSteam& UserId, const FOnQueryAchievementsCompleteDelegate& AchievementDelegate)
 {
 	FOnlineAsyncTaskSteamGetAchievements* NewStatsTask = new FOnlineAsyncTaskSteamGetAchievements( SteamSubsystem, UserId, AchievementDelegate );
 
@@ -1366,12 +1370,11 @@ bool FOnlineLeaderboardsSteam::WriteLeaderboards(const FName& SessionName, const
 	return bWasSuccessful;
 }
 
-bool FOnlineLeaderboardsSteam::WriteAchievements(const FUniqueNetIdSteam& UserId, FOnlineAchievementsWriteRef& WriteObject)
+void FOnlineLeaderboardsSteam::WriteAchievementsInternal(const FUniqueNetIdSteam& UserId, FOnlineAchievementsWriteRef& WriteObject, const FOnAchievementsWrittenDelegate& OnWriteFinishedDelegate)
 {
 	SetUserStatsStoreState(UserId, EOnlineAsyncTaskState::InProgress);
-	FOnlineAsyncTaskSteamWriteAchievements* NewTask = new FOnlineAsyncTaskSteamWriteAchievements(SteamSubsystem, UserId, WriteObject);
+	FOnlineAsyncTaskSteamWriteAchievements* NewTask = new FOnlineAsyncTaskSteamWriteAchievements(SteamSubsystem, UserId, WriteObject, OnWriteFinishedDelegate);
 	SteamSubsystem->QueueAsyncTask(NewTask);
-	return true;
 }
 
 bool FOnlineLeaderboardsSteam::FlushLeaderboards(const FName& SessionName)
diff --git a/Engine/Source/Runtime/Online/OnlineSubsystemSteam/Private/OnlineLeaderboardInterfaceSteam.h b/Engine/Source/Runtime/Online/OnlineSubsystemSteam/Private/OnlineLeaderboardInterfaceSteam.h
--- a/Engine/Source/Runtime/Online/OnlineSubsystemSteam/Private/OnlineLeaderboardInterfaceSteam.h
+++ b/Engine/Source/Runtime/Online/OnlineSubsystemSteam/Private/OnlineLeaderboardInterfaceSteam.h
@@ -130,22 +130,20 @@ PACKAGE_SCOPE:
 	void SetUserStatsStoreState(const FUniqueNetIdSteam& UserId, EOnlineAsyncTaskState::Type NewState);
 
 	/**
-	 * Commits any changes in the online stats cache to the permanent storage
+	 * Commits any changes in the online stats cache to the permanent storage (internal helper for FOnlineAchievementsSteam::WriteAchievements)
 	 *
 	 * @param UserId user to set state of stats stored
 	 * @param WriteObject object to track the state of stats
-	 * @return true if the call is successful, false otherwise
 	 */
-	bool WriteAchievements(const FUniqueNetIdSteam& UserId, FOnlineAchievementsWriteRef& WriteObject);
+	void WriteAchievementsInternal(const FUniqueNetIdSteam& UserId, FOnlineAchievementsWriteRef& WriteObject, const FOnAchievementsWrittenDelegate& OnWriteFinishedDelegate);
 
 	/**
-	 * Reads achievements for user
+	 * Reads achievements for user (internal helper for FOnlineAchievementsSteam::QueryAchievements)
 	 *
 	 * @param UserId user to set state of stats stored
 	 * @param bTriggerDescriptionsDelegatesInstead - async task will normally trigger OnAchievementsRead delegates, this bool tells it to trigger OnAchievementDescriptionsRead instead
-	 * @return true if the call is successful, false otherwise
 	 */
-	void QueryAchievements(const FUniqueNetIdSteam& UserId, const FOnQueryAchievementsCompleteDelegate& AchievementDelegate);
+	void QueryAchievementsInternal(const FUniqueNetIdSteam& UserId, const FOnQueryAchievementsCompleteDelegate& AchievementDelegate);
 
 public:
 
diff --git a/Engine/Source/Runtime/Online/OnlineSubsystemUtils/Private/Tests/TestAchievementsInterface.cpp b/Engine/Source/Runtime/Online/OnlineSubsystemUtils/Private/Tests/TestAchievementsInterface.cpp
--- a/Engine/Source/Runtime/Online/OnlineSubsystemUtils/Private/Tests/TestAchievementsInterface.cpp
+++ b/Engine/Source/Runtime/Online/OnlineSubsystemUtils/Private/Tests/TestAchievementsInterface.cpp
@@ -159,9 +159,8 @@ void FTestAchievementsInterface::WriteAchievements()
 	}
 
 	// Add our delegate and write some achievements to the server
-	OnlineAchievements->AddOnAchievementsWrittenDelegate( OnAchievementsWrittenDelegate );
 	OnlineAchievements->AddOnAchievementUnlockedDelegate( OnAchievementUnlockedDelegate );
-	OnlineAchievements->WriteAchievements( *UserId.Get(), WriteObjectRef );
+	OnlineAchievements->WriteAchievements( *UserId.Get(), WriteObjectRef, OnAchievementsWrittenDelegate );
 }
 
 
@@ -170,7 +169,6 @@ void FTestAchievementsInterface::OnAchievementsWritten(const FUniqueNetId& Playe
 	UE_LOG(LogOnline, Display, TEXT("FTestAchievementsInterface::OnAchievementsWritten( bWasSuccessful = %s )"), bWasSuccessful ? TEXT("true") : TEXT("false"));
 
 	// Clear our delegate and delete this test.
-	OnlineAchievements->ClearOnAchievementsWrittenDelegate(OnAchievementsWrittenDelegate);
 	OnlineAchievements->ClearOnAchievementUnlockedDelegate(OnAchievementUnlockedDelegate);
 
 	if (bWasSuccessful && WriteObject->WriteState == EOnlineAsyncTaskState::Done)
