#ue4
- Fixed some threading issues with querying achievements (unshevled from CL# 2048640 and modified)

diff --git a/Engine/Source/Runtime/Online/IOS/OnlineSubsystemIOS/Private/OnlineAchievementsInterfaceIOS.cpp b/Engine/Source/Runtime/Online/IOS/OnlineSubsystemIOS/Private/OnlineAchievementsInterfaceIOS.cpp
--- a/Engine/Source/Runtime/Online/IOS/OnlineSubsystemIOS/Private/OnlineAchievementsInterfaceIOS.cpp
+++ b/Engine/Source/Runtime/Online/IOS/OnlineSubsystemIOS/Private/OnlineAchievementsInterfaceIOS.cpp
@@ -23,36 +23,43 @@ void FOnlineAchievementsIOS::QueryAchievements(const FUniqueNetId & PlayerId, co
 	{
 	[GKAchievement loadAchievementsWithCompletionHandler: ^(NSArray *loadedAchievements, NSError *error)
 		{
-			bool bSuccess = error == NULL;
+			bool bSuccess = (error == NULL);
 			UE_LOG(LogOnline, Display, TEXT("FOnlineSubsystemIOS::loadAchievementsWithCompletionHandler - %d"), bSuccess);
 
-			if( bSuccess )
+			if (bSuccess)
 			{
-				UE_LOG(LogOnline, Display, TEXT("Loaded %d achievements"), [loadedAchievements count]);
-				Achievements.Empty();
-
-				for( int32 AchievementIdx = 0; AchievementIdx < [loadedAchievements count]; AchievementIdx++ )
+				// Write the achievements back to our Achievements array
+				[FIOSAsyncTask CreateTaskWithBlock : ^ bool(void)
 				{
-					GKAchievement* LoadedAchievement = [loadedAchievements objectAtIndex:AchievementIdx];
-					UE_LOG(LogOnline, Display, TEXT("Loaded achievement: %s"), *FString(LoadedAchievement.identifier));
-					
-					Achievements.AddZeroed();
-					FOnlineAchievement& OnlineAchievement = Achievements[AchievementIdx];
-					OnlineAchievement.Id = LoadedAchievement.identifier;
-					OnlineAchievement.Progress = LoadedAchievement.percentComplete;
-				}
+					UE_LOG(LogOnline, Display, TEXT("Loaded %d achievements"), [loadedAchievements count]);
+					Achievements.Empty();
+
+					for( int32 AchievementIdx = 0; AchievementIdx < [loadedAchievements count]; AchievementIdx++ )
+					{
+						GKAchievement* LoadedAchievement = [loadedAchievements objectAtIndex : AchievementIdx];
+						UE_LOG(LogOnline, Display, TEXT("Loaded achievement: %s"), *FString(LoadedAchievement.identifier));
+
+						Achievements.AddZeroed();
+						FOnlineAchievement& OnlineAchievement = Achievements[AchievementIdx];
+						OnlineAchievement.Id = LoadedAchievement.identifier;
+						OnlineAchievement.Progress = LoadedAchievement.percentComplete;
+					}
+
+					CopiedDelegate.ExecuteIfBound(PlayerId, true);
+					return true;
+				}];
 			}
 			else
 			{
-				UE_LOG(LogOnline, Display, TEXT("Failed to load achievements with error [%d]"), [error code]);
-			}
+				UE_LOG(LogOnline, Warning, TEXT("Failed to load achievements with error [%d]"), [error code]);
 
-			// Report back to the game thread whether this succeeded.
-			[FIOSAsyncTask CreateTaskWithBlock : ^ bool(void)
-			{
-				CopiedDelegate.ExecuteIfBound(PlayerId, bSuccess);
-				return true;
-			}];
+				// Report back to the game thread this failed.
+				[FIOSAsyncTask CreateTaskWithBlock : ^ bool(void)
+				{
+					CopiedDelegate.ExecuteIfBound(PlayerId, false);
+					return true;
+				}];
+			}
 		}
 	];
 	});
@@ -66,45 +73,52 @@ void FOnlineAchievementsIOS::QueryAchievementDescriptions(const FUniqueNetId& Pl
 	{
 		[GKAchievementDescription loadAchievementDescriptionsWithCompletionHandler : ^ (NSArray *descriptions, NSError *error)
 		{
-			bool bSuccess = error == NULL;
+			bool bSuccess = (error == NULL);
 			UE_LOG(LogOnline, Display, TEXT("FOnlineSubsystemIOS::loadAchievementDescriptionsWithCompletionHandler - %d"), bSuccess);
 
 			if (bSuccess)
-			{
-				AchievementDescriptions.Empty();
-
-				for (GKAchievementDescription* desc in descriptions)
+			{				
+				// Report back to the game thread whether this succeeded.
+				[FIOSAsyncTask CreateTaskWithBlock : ^ bool(void)
 				{
-					FString Id(desc.identifier);
-
-					FOnlineAchievementDesc OnlineAchievementDesc;
-					OnlineAchievementDesc.Title = FText::FromString(desc.title);
-					OnlineAchievementDesc.LockedDesc = FText::FromString(desc.unachievedDescription);
-					OnlineAchievementDesc.UnlockedDesc = FText::FromString(desc.achievedDescription);
-					OnlineAchievementDesc.bIsHidden = desc.hidden;
-
-					UE_LOG(LogOnline, Display, TEXT("============================================"));
-					UE_LOG(LogOnline, Display, TEXT("Loaded achievement id: %s"), *Id);
-					UE_LOG(LogOnline, Display, TEXT("Loaded achievement title: %s"), *OnlineAchievementDesc.Title.ToString());
-					UE_LOG(LogOnline, Display, TEXT("Loaded achievement locked desc: %s"), *OnlineAchievementDesc.LockedDesc.ToString());
-					UE_LOG(LogOnline, Display, TEXT("Loaded achievement unlocked desc: %s"), *OnlineAchievementDesc.UnlockedDesc.ToString());
-					UE_LOG(LogOnline, Display, TEXT("Loaded achievement hidden: %d"), OnlineAchievementDesc.bIsHidden);
-					UE_LOG(LogOnline, Display, TEXT("============================================"));
-
-					AchievementDescriptions.Add(Id, OnlineAchievementDesc);
-				}
+					AchievementDescriptions.Empty();
+
+					for (GKAchievementDescription* desc in descriptions)
+					{
+						FString Id(desc.identifier);
+
+						FOnlineAchievementDesc OnlineAchievementDesc;
+						OnlineAchievementDesc.Title = FText::FromString(desc.title);
+						OnlineAchievementDesc.LockedDesc = FText::FromString(desc.unachievedDescription);
+						OnlineAchievementDesc.UnlockedDesc = FText::FromString(desc.achievedDescription);
+						OnlineAchievementDesc.bIsHidden = desc.hidden;
+
+						UE_LOG(LogOnline, Display, TEXT("============================================"));
+						UE_LOG(LogOnline, Display, TEXT("Loaded achievement id: %s"), *Id);
+						UE_LOG(LogOnline, Display, TEXT("Loaded achievement title: %s"), *OnlineAchievementDesc.Title.ToString());
+						UE_LOG(LogOnline, Display, TEXT("Loaded achievement locked desc: %s"), *OnlineAchievementDesc.LockedDesc.ToString());
+						UE_LOG(LogOnline, Display, TEXT("Loaded achievement unlocked desc: %s"), *OnlineAchievementDesc.UnlockedDesc.ToString());
+						UE_LOG(LogOnline, Display, TEXT("Loaded achievement hidden: %d"), OnlineAchievementDesc.bIsHidden);
+						UE_LOG(LogOnline, Display, TEXT("============================================"));
+
+						AchievementDescriptions.Add(Id, OnlineAchievementDesc);
+					}
+
+					CopiedDelegate.ExecuteIfBound(PlayerId, true);
+					return true;
+				}];
 			}
 			else
 			{
-				UE_LOG(LogOnline, Display, TEXT("Failed to load achievement descriptionss with error [%d]"), [error code]);
-			}
+				UE_LOG(LogOnline, Warning, TEXT("Failed to load achievement descriptions with error [%d]"), [error code]);
 
-			// Report back to the game thread whether this succeeded.
-			[FIOSAsyncTask CreateTaskWithBlock : ^ bool(void)
-			{
-				CopiedDelegate.ExecuteIfBound(PlayerId, bSuccess);
-				return true;
-			}];
+				// Report the failure back to the game thread
+				[FIOSAsyncTask CreateTaskWithBlock : ^ bool(void)
+				{
+					CopiedDelegate.ExecuteIfBound(PlayerId, false);
+					return true;
+				}];
+			}
 		}];
 	});
 }
@@ -190,7 +204,7 @@ void FOnlineAchievementsIOS::WriteAchievements(const FUniqueNetId& PlayerId, FOn
 				}
 				else
 				{
-					UE_LOG(LogOnline, Display, TEXT("Failed to report achievements with error[%d]"), [Error code]);
+					UE_LOG(LogOnline, Warning, TEXT("Failed to report achievements with error[%d]"), [Error code]);
 				}
 
 				// Release any handle of the achievements back to the IOS gc
