#ue4
- Updated the WriteAchievements proxy to use the delegate passing model

diff --git a/Engine/Source/Runtime/Online/OnlineSubsystemUtils/Private/AchievementWriteCallbackProxy.cpp b/Engine/Source/Runtime/Online/OnlineSubsystemUtils/Private/AchievementWriteCallbackProxy.cpp
--- a/Engine/Source/Runtime/Online/OnlineSubsystemUtils/Private/AchievementWriteCallbackProxy.cpp
+++ b/Engine/Source/Runtime/Online/OnlineSubsystemUtils/Private/AchievementWriteCallbackProxy.cpp
@@ -31,13 +31,11 @@ void UAchievementWriteCallbackProxy::Activate()
 		IOnlineAchievementsPtr Achievements = Helper.OnlineSub->GetAchievementsInterface();
 		if (Achievements.IsValid())
 		{
-			//@TODO: Bind OnAchievementWritten in
 			FOnlineAchievementsWriteRef WriteObjectRef = WriteObject.ToSharedRef();
-			Achievements->WriteAchievements(*Helper.UserID, WriteObjectRef);
+			FOnAchievementsWrittenDelegate WriteFinishedDelegate = FOnAchievementsWrittenDelegate::CreateUObject(this, &ThisClass::OnAchievementWritten);
+			Achievements->WriteAchievements(*Helper.UserID, WriteObjectRef, WriteFinishedDelegate);
 
 			// OnAchievementWritten will get called, nothing more to do
-			//@TODO: Remove this call
-			OnSuccess.Broadcast();
 			return;
 		}
 		else
@@ -48,6 +46,7 @@ void UAchievementWriteCallbackProxy::Activate()
 
 	// Fail immediately
 	OnFailure.Broadcast();
+	WriteObject.Reset();
 }
 
 void UAchievementWriteCallbackProxy::OnAchievementWritten(const FUniqueNetId& UserID, bool bSuccess)
