#ue4
- Exposed achievement queries to Blueprints: CacheAchievements, GetCachedAchievementProgress, CacheAchievementDescriptions, and GetCachedAchievementDescription

diff --git a/Engine/Source/Runtime/Online/OnlineSubsystemUtils/Classes/AchievementBlueprintLibrary.h b/Engine/Source/Runtime/Online/OnlineSubsystemUtils/Classes/AchievementBlueprintLibrary.h
new file mode 100644
--- /dev/null
+++ b/Engine/Source/Runtime/Online/OnlineSubsystemUtils/Classes/AchievementBlueprintLibrary.h
@@ -0,0 +1,35 @@
+// Copyright 1998-2014 Epic Games, Inc. All Rights Reserved.
+
+#pragma once
+
+#include "AchievementBlueprintLibrary.generated.h"
+
+// Library of synchronous achievement calls
+UCLASS()
+class ONLINESUBSYSTEMUTILS_API UAchievementBlueprintLibrary : public UBlueprintFunctionLibrary
+{
+	GENERATED_UCLASS_BODY()
+
+	/**
+	 * Gets the status of an achievement ID (you must call CacheAchievements first to cache them)
+	 *
+	 * @param AchievementID - The id of the achievement we are looking up
+	 * @param bFoundID - If the ID was found in the cache (if not, none of the other values are meaningful)
+	 * @param Progress - The progress amount of the achievement
+	 */
+	UFUNCTION(BlueprintCallable, Category = "Online|Achievements")
+	static void GetCachedAchievementProgress(APlayerController* PlayerController, FName AchievementID, /*out*/ bool& bFoundID, /*out*/ float& Progress);
+
+	/**
+	 * Get the description for an achievement ID (you must call CacheAchievementDescriptions first to cache them)
+	 *
+	 * @param AchievementID - The id of the achievement we are searching for data of
+	 * @param bFoundID - If the ID was found in the cache (if not, none of the other values are meaningful)
+	 * @param Title - The localized title of the achievement
+	 * @param LockedDescription - The localized locked description of the achievement
+	 * @param UnlockedDescription - The localized unlocked description of the achievement
+	 * @param bHidden - Whether the achievement is hidden
+	 */
+	UFUNCTION(BlueprintCallable, Category="Online|Achievements")
+	static void GetCachedAchievementDescription(APlayerController* PlayerController, FName AchievementID, /*out*/ bool& bFoundID, /*out*/ FText& Title, /*out*/ FText& LockedDescription, /*out*/ FText& UnlockedDescription, /*out*/ bool& bHidden);
+};
diff --git a/Engine/Source/Runtime/Online/OnlineSubsystemUtils/Classes/AchievementQueryCallbackProxy.h b/Engine/Source/Runtime/Online/OnlineSubsystemUtils/Classes/AchievementQueryCallbackProxy.h
new file mode 100644
--- /dev/null
+++ b/Engine/Source/Runtime/Online/OnlineSubsystemUtils/Classes/AchievementQueryCallbackProxy.h
@@ -0,0 +1,42 @@
+// Copyright 1998-2014 Epic Games, Inc. All Rights Reserved.
+#pragma once
+
+#include "Runtime/Online/OnlineSubsystem/Public/Interfaces/OnlineAchievementsInterface.h"
+#include "AchievementQueryCallbackProxy.generated.h"
+
+UCLASS(MinimalAPI)
+class UAchievementQueryCallbackProxy : public UOnlineBlueprintCallProxyBase
+{
+	GENERATED_UCLASS_BODY()
+
+	// Called when there is a successful query
+	UPROPERTY(BlueprintAssignable)
+	FEmptyOnlineDelegate OnSuccess;
+
+	// Called when there is an unsuccessful query
+	UPROPERTY(BlueprintAssignable)
+	FEmptyOnlineDelegate OnFailure;
+
+	// Fetches and caches achievement progress from the default online subsystem
+	UFUNCTION(BlueprintCallable, meta=(BlueprintInternalUseOnly = "true"), Category = "Online|Achievements")
+	static UAchievementQueryCallbackProxy* CacheAchievements(class APlayerController* PlayerController);
+
+	// Fetches and caches achievement descriptions from the default online subsystem
+	UFUNCTION(BlueprintCallable, meta=(BlueprintInternalUseOnly = "true"), Category = "Online|Achievements")
+	static UAchievementQueryCallbackProxy* CacheAchievementDescriptions(class APlayerController* PlayerController);
+
+	// UOnlineBlueprintCallProxyBase interface
+	virtual void Activate() OVERRIDE;
+	// End of UOnlineBlueprintCallProxyBase interface
+
+private:
+	// Internal callback when the achievement query completes, calls out to the public success/failure callbacks
+	void OnQueryCompleted(const FUniqueNetId& UserID, bool bSuccess);
+
+private:
+	// The player controller triggering things
+	TWeakObjectPtr<APlayerController> PlayerControllerWeakPtr;
+
+	// Are we querying achievement progress or achievement descriptions?
+	bool bFetchDescriptions;
+};
diff --git a/Engine/Source/Runtime/Online/OnlineSubsystemUtils/Private/AchievementBlueprintLibrary.cpp b/Engine/Source/Runtime/Online/OnlineSubsystemUtils/Private/AchievementBlueprintLibrary.cpp
new file mode 100644
--- /dev/null
+++ b/Engine/Source/Runtime/Online/OnlineSubsystemUtils/Private/AchievementBlueprintLibrary.cpp
@@ -0,0 +1,71 @@
+// Copyright 1998-2014 Epic Games, Inc. All Rights Reserved.
+
+#include "OnlineSubsystemUtilsPrivatePCH.h"
+#include "AchievementBlueprintLibrary.h"
+#include "Runtime/Online/OnlineSubsystem/Public/Interfaces/OnlineAchievementsInterface.h"
+
+UAchievementBlueprintLibrary::UAchievementBlueprintLibrary(const FPostConstructInitializeProperties& PCIP)
+	: Super(PCIP)
+{
+}
+
+void UAchievementBlueprintLibrary::GetCachedAchievementProgress(APlayerController* PlayerController, FName AchievementID, /*out*/ bool& bFoundID, /*out*/ float& Progress)
+{
+	bFoundID = false;
+	Progress = 0.0f;
+
+	FOnlineSubsystemBPCallHelper Helper(TEXT("GetCachedAchievementProgress"));
+	Helper.QueryIDFromPlayerController(PlayerController);
+
+	if (Helper.IsValid())
+	{
+		IOnlineAchievementsPtr Achievements = Helper.OnlineSub->GetAchievementsInterface();
+		if (Achievements.IsValid())
+		{
+			FOnlineAchievement AchievementStatus;
+			if (Achievements->GetCachedAchievement(*Helper.UserID, AchievementID.ToString(), AchievementStatus) == EOnlineCachedResult::Success)
+			{
+				bFoundID = true;
+				Progress = AchievementStatus.Progress;
+			}
+		}
+		else
+		{
+			FFrame::KismetExecutionMessage(TEXT("Achievements not supported by Online Subsystem"), ELogVerbosity::Warning);
+		}
+	}
+}
+
+void UAchievementBlueprintLibrary::GetCachedAchievementDescription(APlayerController* PlayerController, FName AchievementID, /*out*/ bool& bFoundID, /*out*/ FText& Title, /*out*/ FText& LockedDescription, /*out*/ FText& UnlockedDescription, /*out*/ bool& bHidden)
+{
+	bFoundID = false;
+	Title = FText::GetEmpty();
+	LockedDescription = FText::GetEmpty();
+	UnlockedDescription = FText::GetEmpty();
+	bHidden = false;
+
+	FOnlineSubsystemBPCallHelper Helper(TEXT("GetCachedAchievementDescription"));
+	Helper.QueryIDFromPlayerController(PlayerController);
+
+	if (Helper.IsValid())
+	{
+		IOnlineAchievementsPtr Achievements = Helper.OnlineSub->GetAchievementsInterface();
+		if (Achievements.IsValid())
+		{
+			FOnlineAchievementDesc AchievementDescription;
+			if (Achievements->GetCachedAchievementDescription(AchievementID.ToString(), AchievementDescription) == EOnlineCachedResult::Success)
+			{
+				bFoundID = true;
+				Title = AchievementDescription.Title;
+				LockedDescription = AchievementDescription.LockedDesc;
+				UnlockedDescription = AchievementDescription.UnlockedDesc;
+				bHidden = AchievementDescription.bIsHidden;
+				//@TODO: UnlockTime - FDateTime is not exposed to Blueprints right now, see TTP 315420
+			}
+		}
+		else
+		{
+			FFrame::KismetExecutionMessage(TEXT("Achievements not supported by Online Subsystem"), ELogVerbosity::Warning);
+		}
+	}
+}
diff --git a/Engine/Source/Runtime/Online/OnlineSubsystemUtils/Private/AchievementQueryCallbackProxy.cpp b/Engine/Source/Runtime/Online/OnlineSubsystemUtils/Private/AchievementQueryCallbackProxy.cpp
new file mode 100644
--- /dev/null
+++ b/Engine/Source/Runtime/Online/OnlineSubsystemUtils/Private/AchievementQueryCallbackProxy.cpp
@@ -0,0 +1,73 @@
+// Copyright 1998-2014 Epic Games, Inc. All Rights Reserved.
+
+#include "OnlineSubsystemUtilsPrivatePCH.h"
+
+//////////////////////////////////////////////////////////////////////////
+// UAchievementQueryCallbackProxy
+
+UAchievementQueryCallbackProxy::UAchievementQueryCallbackProxy(const class FPostConstructInitializeProperties& PCIP)
+	: Super(PCIP)
+{
+}
+
+UAchievementQueryCallbackProxy* UAchievementQueryCallbackProxy::CacheAchievements(class APlayerController* PlayerController)
+{
+	UAchievementQueryCallbackProxy* Proxy = NewObject<UAchievementQueryCallbackProxy>();
+	Proxy->PlayerControllerWeakPtr = PlayerController;
+	Proxy->bFetchDescriptions = false;
+	return Proxy;
+}
+
+UAchievementQueryCallbackProxy* UAchievementQueryCallbackProxy::CacheAchievementDescriptions(class APlayerController* PlayerController)
+{
+	UAchievementQueryCallbackProxy* Proxy = NewObject<UAchievementQueryCallbackProxy>();
+	Proxy->PlayerControllerWeakPtr = PlayerController;
+	Proxy->bFetchDescriptions = true;
+	return Proxy;
+}
+
+void UAchievementQueryCallbackProxy::Activate()
+{
+	FOnlineSubsystemBPCallHelper Helper(TEXT("CacheAchievements or CacheAchievementDescriptions"));
+	Helper.QueryIDFromPlayerController(PlayerControllerWeakPtr.Get());
+
+	if (Helper.IsValid())
+	{
+		IOnlineAchievementsPtr Achievements = Helper.OnlineSub->GetAchievementsInterface();
+		if (Achievements.IsValid())
+		{
+			FOnQueryAchievementsCompleteDelegate QueryFinishedDelegate = FOnQueryAchievementsCompleteDelegate::CreateUObject(this, &ThisClass::OnQueryCompleted);
+			
+			if (bFetchDescriptions)
+			{
+				Achievements->QueryAchievementDescriptions(*Helper.UserID, QueryFinishedDelegate);
+			}
+			else
+			{
+				Achievements->QueryAchievements(*Helper.UserID, QueryFinishedDelegate);
+			}
+
+			// OnQueryCompleted will get called, nothing more to do now
+			return;
+		}
+		else
+		{
+			FFrame::KismetExecutionMessage(TEXT("Achievements not supported by Online Subsystem"), ELogVerbosity::Warning);
+		}
+	}
+
+	// Fail immediately
+	OnFailure.Broadcast();
+}
+
+void UAchievementQueryCallbackProxy::OnQueryCompleted(const FUniqueNetId& UserID, bool bSuccess)
+{
+	if (bSuccess)
+	{
+		OnSuccess.Broadcast();
+	}
+	else
+	{
+		OnFailure.Broadcast();
+	}
+}
