#ue4
- Removed some of the code in the existing leaderboard query node since the parent class can take care of most of that work

diff --git a/Engine/Source/Editor/OnlineBlueprintSupport/Classes/K2Node_LeaderboardQuery.h b/Engine/Source/Editor/OnlineBlueprintSupport/Classes/K2Node_LeaderboardQuery.h
--- a/Engine/Source/Editor/OnlineBlueprintSupport/Classes/K2Node_LeaderboardQuery.h
+++ b/Engine/Source/Editor/OnlineBlueprintSupport/Classes/K2Node_LeaderboardQuery.h
@@ -8,13 +8,4 @@ UCLASS()
 class ONLINEBLUEPRINTSUPPORT_API UK2Node_LeaderboardQuery : public UK2Node_BaseAsyncTask
 {
 	GENERATED_UCLASS_BODY()
-
-	// UEdGraphNode interface
-	virtual FString GetTooltip() const OVERRIDE;
-	virtual FText GetNodeTitle(ENodeTitleType::Type TitleType) const OVERRIDE;
-	// End of UEdGraphNode interface
-
-	// UK2Node_BaseAsyncTask interface
-	virtual FString GetCategoryName() OVERRIDE;
-	// End of UK2Node_BaseAsyncTask interface
 };
diff --git a/Engine/Source/Editor/OnlineBlueprintSupport/Private/K2Node_LeaderboardQuery.cpp b/Engine/Source/Editor/OnlineBlueprintSupport/Private/K2Node_LeaderboardQuery.cpp
--- a/Engine/Source/Editor/OnlineBlueprintSupport/Private/K2Node_LeaderboardQuery.cpp
+++ b/Engine/Source/Editor/OnlineBlueprintSupport/Private/K2Node_LeaderboardQuery.cpp
@@ -15,19 +15,4 @@ UK2Node_LeaderboardQuery::UK2Node_LeaderboardQuery(const FPostConstructInitializ
 	ProxyClass = ULeaderboardQueryCallbackProxy::StaticClass();
 }
 
-FString UK2Node_LeaderboardQuery::GetTooltip() const
-{
-	return LOCTEXT("K2Node_LeaderboardQuery_Tooltip", "Queries a leaderboard for an integer value").ToString();
-}
-
-FText UK2Node_LeaderboardQuery::GetNodeTitle(ENodeTitleType::Type TitleType) const
-{
-	return LOCTEXT("ReadLeaderboardInteger", "Read Leaderboard Integer");
-}
-
-FString UK2Node_LeaderboardQuery::GetCategoryName()
-{
-	return TEXT("Online|Leaderboard");
-}
-
 #undef LOCTEXT_NAMESPACE
diff --git a/Engine/Source/Runtime/Online/OnlineSubsystemUtils/Classes/LeaderboardQueryCallbackProxy.h b/Engine/Source/Runtime/Online/OnlineSubsystemUtils/Classes/LeaderboardQueryCallbackProxy.h
--- a/Engine/Source/Runtime/Online/OnlineSubsystemUtils/Classes/LeaderboardQueryCallbackProxy.h
+++ b/Engine/Source/Runtime/Online/OnlineSubsystemUtils/Classes/LeaderboardQueryCallbackProxy.h
@@ -19,8 +19,8 @@ class ULeaderboardQueryCallbackProxy : public UObject
 	UPROPERTY(BlueprintAssignable)
 	FLeaderboardQueryResult OnFailure;
 
-	// Called to perform the query internally
-	UFUNCTION(BlueprintCallable, meta = (BlueprintInternalUseOnly = "true"))
+	// Queries a leaderboard for an integer value
+	UFUNCTION(BlueprintCallable, meta = (BlueprintInternalUseOnly = "true", FriendlyName="Read Leaderboard Integer"), Category="Online|Leaderboard")
 	static ULeaderboardQueryCallbackProxy* CreateProxyObjectForIntQuery(class APlayerController* PlayerController, FName StatName);
 
 public:
