#online,mcp
- updated title file interface to specify start,count of files to enumerate
- only implemented for Mcp

TTP #333057 EMS: add ability to specify start/count for EMS file enumeration

diff --git a/Engine/Plugins/Editor/EpicSurvey/Source/EpicSurvey/Private/SurveyTitleCdnStorage.cpp b/Engine/Plugins/Editor/EpicSurvey/Source/EpicSurvey/Private/SurveyTitleCdnStorage.cpp
--- a/Engine/Plugins/Editor/EpicSurvey/Source/EpicSurvey/Private/SurveyTitleCdnStorage.cpp
+++ b/Engine/Plugins/Editor/EpicSurvey/Source/EpicSurvey/Private/SurveyTitleCdnStorage.cpp
@@ -96,6 +96,11 @@ bool FSurveyTitleCdnStorage::EnumerateFiles()
 	
 }	
 
+bool FSurveyTitleCdnStorage::EnumerateFiles(int32 Start, int32 Count)
+{
+	EnumerateFiles();
+}
+
 void FSurveyTitleCdnStorage::GetFileList(TArray<FCloudFileHeader>& OutFiles) 
 {
 	TArray<FCloudFileHeader>* FilesPtr = &FileHeaders;
diff --git a/Engine/Plugins/Editor/EpicSurvey/Source/EpicSurvey/Private/SurveyTitleCdnStorage.h b/Engine/Plugins/Editor/EpicSurvey/Source/EpicSurvey/Private/SurveyTitleCdnStorage.h
--- a/Engine/Plugins/Editor/EpicSurvey/Source/EpicSurvey/Private/SurveyTitleCdnStorage.h
+++ b/Engine/Plugins/Editor/EpicSurvey/Source/EpicSurvey/Private/SurveyTitleCdnStorage.h
@@ -22,6 +22,8 @@ public:
 
 	virtual bool EnumerateFiles() OVERRIDE;
 
+	virtual bool EnumerateFiles(int32 Start, int32 Count) OVERRIDE;
+
 	virtual void GetFileList(TArray<FCloudFileHeader>& InFileHeaders) OVERRIDE;
 
 	virtual bool ReadFile(const FString& DLName) OVERRIDE;
diff --git a/Engine/Plugins/Editor/EpicSurvey/Source/EpicSurvey/Private/SurveyTitleLocalStorage.cpp b/Engine/Plugins/Editor/EpicSurvey/Source/EpicSurvey/Private/SurveyTitleLocalStorage.cpp
--- a/Engine/Plugins/Editor/EpicSurvey/Source/EpicSurvey/Private/SurveyTitleLocalStorage.cpp
+++ b/Engine/Plugins/Editor/EpicSurvey/Source/EpicSurvey/Private/SurveyTitleLocalStorage.cpp
@@ -85,6 +85,11 @@ bool FSurveyTitleLocalStorage::EnumerateFiles()
 	return true;
 }
 
+bool FSurveyTitleLocalStorage::EnumerateFiles(int32 Start, int32 Count)
+{
+	EnumerateFiles();
+}
+
 void FSurveyTitleLocalStorage::GetFileList(TArray<FCloudFileHeader>& InFileHeaders)
 {
 	InFileHeaders.Append( FileHeaders );
diff --git a/Engine/Plugins/Editor/EpicSurvey/Source/EpicSurvey/Private/SurveyTitleLocalStorage.h b/Engine/Plugins/Editor/EpicSurvey/Source/EpicSurvey/Private/SurveyTitleLocalStorage.h
--- a/Engine/Plugins/Editor/EpicSurvey/Source/EpicSurvey/Private/SurveyTitleLocalStorage.h
+++ b/Engine/Plugins/Editor/EpicSurvey/Source/EpicSurvey/Private/SurveyTitleLocalStorage.h
@@ -18,6 +18,8 @@ public:
 
 	virtual bool EnumerateFiles() OVERRIDE;
 
+	virtual bool EnumerateFiles(int32 Start, int32 Count) OVERRIDE;
+
 	virtual void GetFileList(TArray<FCloudFileHeader>& InFileHeaders) OVERRIDE;
 
 	virtual bool ReadFile(const FString& DLName) OVERRIDE;
diff --git a/Engine/Source/Editor/NewsFeed/Private/Implementation/CdnNewsFeedTitleFile.cpp b/Engine/Source/Editor/NewsFeed/Private/Implementation/CdnNewsFeedTitleFile.cpp
--- a/Engine/Source/Editor/NewsFeed/Private/Implementation/CdnNewsFeedTitleFile.cpp
+++ b/Engine/Source/Editor/NewsFeed/Private/Implementation/CdnNewsFeedTitleFile.cpp
@@ -92,6 +92,11 @@ bool FCdnNewsFeedTitleFile::EnumerateFiles()
 	
 }	
 
+bool FCdnNewsFeedTitleFile::EnumerateFiles(int32 Start, int32 Count)
+{
+	EnumerateFiles();
+}
+
 void FCdnNewsFeedTitleFile::GetFileList(TArray<FCloudFileHeader>& OutFiles) 
 {
 	TArray<FCloudFileHeader>* FilesPtr = &FileHeaders;
diff --git a/Engine/Source/Editor/NewsFeed/Private/Implementation/CdnNewsFeedTitleFile.h b/Engine/Source/Editor/NewsFeed/Private/Implementation/CdnNewsFeedTitleFile.h
--- a/Engine/Source/Editor/NewsFeed/Private/Implementation/CdnNewsFeedTitleFile.h
+++ b/Engine/Source/Editor/NewsFeed/Private/Implementation/CdnNewsFeedTitleFile.h
@@ -49,6 +49,8 @@ public:
 
 	virtual bool EnumerateFiles() OVERRIDE;
 
+	virtual bool EnumerateFiles(int32 Start, int32 Count) OVERRIDE;
+
 	virtual void GetFileList(TArray<FCloudFileHeader>& InFileHeaders) OVERRIDE;
 
 	virtual bool ReadFile(const FString& DLName) OVERRIDE;
diff --git a/Engine/Source/Editor/NewsFeed/Private/Implementation/LocalNewsFeedTitleFile.cpp b/Engine/Source/Editor/NewsFeed/Private/Implementation/LocalNewsFeedTitleFile.cpp
--- a/Engine/Source/Editor/NewsFeed/Private/Implementation/LocalNewsFeedTitleFile.cpp
+++ b/Engine/Source/Editor/NewsFeed/Private/Implementation/LocalNewsFeedTitleFile.cpp
@@ -88,6 +88,11 @@ bool FLocalNewsFeedTitleFile::EnumerateFiles()
 	return Success;
 }
 
+bool FLocalNewsFeedTitleFile::EnumerateFiles(int32 Start, int32 Count)
+{
+	EnumerateFiles();
+}
+
 
 void FLocalNewsFeedTitleFile::GetFileList(TArray<FCloudFileHeader>& InFileHeaders)
 {
diff --git a/Engine/Source/Editor/NewsFeed/Private/Implementation/LocalNewsFeedTitleFile.h b/Engine/Source/Editor/NewsFeed/Private/Implementation/LocalNewsFeedTitleFile.h
--- a/Engine/Source/Editor/NewsFeed/Private/Implementation/LocalNewsFeedTitleFile.h
+++ b/Engine/Source/Editor/NewsFeed/Private/Implementation/LocalNewsFeedTitleFile.h
@@ -25,6 +25,8 @@ public:
 
 	virtual bool EnumerateFiles() OVERRIDE;
 
+	virtual bool EnumerateFiles(int32 Start, int32 Count) OVERRIDE;
+
 	virtual void GetFileList(TArray<FCloudFileHeader>& InFileHeaders) OVERRIDE;
 
 	virtual bool ReadFile(const FString& DLName) OVERRIDE;
diff --git a/Engine/Source/Runtime/Online/OnlineSubsystem/Public/Interfaces/OnlineTitleFileInterface.h b/Engine/Source/Runtime/Online/OnlineSubsystem/Public/Interfaces/OnlineTitleFileInterface.h
--- a/Engine/Source/Runtime/Online/OnlineSubsystem/Public/Interfaces/OnlineTitleFileInterface.h
+++ b/Engine/Source/Runtime/Online/OnlineSubsystem/Public/Interfaces/OnlineTitleFileInterface.h
@@ -56,12 +56,25 @@ public:
 	 */
 	virtual bool ClearFile(const FString& FileName) = 0;
 
+	/**
+	* Requests a list of available files from the network store
+	*
+	* @param Start first file entry to enumerate (default 0)
+	* @param Count number of files to enumerate (default -1 for all)
+	*
+	* @return true if the request has started, false if not
+	*/
+	virtual bool EnumerateFiles() = 0;
+
 	/**
 	 * Requests a list of available files from the network store
 	 *
+	 * @param Start first file entry to enumerate (default 0)
+	 * @param Count number of files to enumerate (default -1 for all)
+	 *
 	 * @return true if the request has started, false if not
 	 */
-	virtual bool EnumerateFiles() = 0;
+	virtual bool EnumerateFiles(int32 Start, int32 Count) = 0;
 
 	/**
 	 * Delegate fired when the list of files has been returned from the network store
