#UE4 - added a common impl between OSS interface and concrete implementations

diff --git a/Engine/Source/Runtime/Online/IOS/OnlineSubsystemIOS/Public/OnlineSubsystemIOS.h b/Engine/Source/Runtime/Online/IOS/OnlineSubsystemIOS/Public/OnlineSubsystemIOS.h
--- a/Engine/Source/Runtime/Online/IOS/OnlineSubsystemIOS/Public/OnlineSubsystemIOS.h
+++ b/Engine/Source/Runtime/Online/IOS/OnlineSubsystemIOS/Public/OnlineSubsystemIOS.h
@@ -3,6 +3,7 @@
 #pragma once
 
 #include "OnlineSubsystem.h"
+#include "OnlineSubsystemImpl.h"
 #include "OnlineSubsystemIOSPackage.h"
 
 
@@ -10,7 +11,7 @@
  *	OnlineSubsystemIOS - Implementation of the online subsystem for IOS services
  */
 class ONLINESUBSYSTEMIOS_API FOnlineSubsystemIOS : 
-	public IOnlineSubsystem,
+	public FOnlineSubsystemImpl,
 	public FTickerObjectBase
 {
 
diff --git a/Engine/Source/Runtime/Online/OnlineSubsystemAmazon/Public/OnlineSubsystemAmazon.h b/Engine/Source/Runtime/Online/OnlineSubsystemAmazon/Public/OnlineSubsystemAmazon.h
--- a/Engine/Source/Runtime/Online/OnlineSubsystemAmazon/Public/OnlineSubsystemAmazon.h
+++ b/Engine/Source/Runtime/Online/OnlineSubsystemAmazon/Public/OnlineSubsystemAmazon.h
@@ -3,6 +3,7 @@
 #pragma once
 
 #include "OnlineSubsystem.h"
+#include "OnlineSubsystemImpl.h"
 #include "OnlineSubsystemAmazonModule.h"
 #include "OnlineSubsystemAmazonPackage.h"
 
@@ -13,7 +14,7 @@ typedef TSharedPtr<class FOnlineIdentityAmazon, ESPMode::ThreadSafe> FOnlineIden
  * Amazon subsystem
  */
 class ONLINESUBSYSTEMAMAZON_API FOnlineSubsystemAmazon :
-	public IOnlineSubsystem,
+	public FOnlineSubsystemImpl,
 	public FTickerObjectBase
 
 {
diff --git a/Engine/Source/Runtime/Online/OnlineSubsystemFacebook/Public/OnlineSubsystemFacebook.h b/Engine/Source/Runtime/Online/OnlineSubsystemFacebook/Public/OnlineSubsystemFacebook.h
--- a/Engine/Source/Runtime/Online/OnlineSubsystemFacebook/Public/OnlineSubsystemFacebook.h
+++ b/Engine/Source/Runtime/Online/OnlineSubsystemFacebook/Public/OnlineSubsystemFacebook.h
@@ -3,6 +3,7 @@
 #pragma once
 
 #include "OnlineSubsystem.h"
+#include "OnlineSubsystemImpl.h"
 #include "OnlineSubsystemFacebookPackage.h"
 
 /** Forward declarations of all interface classes */
@@ -15,7 +16,7 @@ typedef TSharedPtr<class FOnlineUserFacebook, ESPMode::ThreadSafe> FOnlineUserFa
  *	OnlineSubsystemFacebook - Implementation of the online subsystem for Facebook services
  */
 class ONLINESUBSYSTEMFACEBOOK_API FOnlineSubsystemFacebook 
-	: public IOnlineSubsystem
+	: public FOnlineSubsystemImpl
 	, public FTickerObjectBase
 {
 public:
