#UE4 - const function and added some name parsing to IsOnlineModuleLoaded function

diff --git a/Engine/Source/Runtime/Online/OnlineSubsystem/Private/OnlineSubsystemModule.cpp b/Engine/Source/Runtime/Online/OnlineSubsystem/Private/OnlineSubsystemModule.cpp
--- a/Engine/Source/Runtime/Online/OnlineSubsystem/Private/OnlineSubsystemModule.cpp
+++ b/Engine/Source/Runtime/Online/OnlineSubsystem/Private/OnlineSubsystemModule.cpp
@@ -140,7 +140,7 @@ void FOnlineSubsystemModule::UnregisterPlatformService(const FName FactoryName)
 	}
 }
 
-void FOnlineSubsystemModule::ParseOnlineSubsystemName(const FName& FullName, FName& SubsystemName, FName& InstanceName)
+void FOnlineSubsystemModule::ParseOnlineSubsystemName(const FName& FullName, FName& SubsystemName, FName& InstanceName) const
 {
 	SubsystemName = DefaultPlatformService;
 	InstanceName = DEFAULT_INSTANCE;
@@ -241,11 +241,9 @@ bool FOnlineSubsystemModule::IsOnlineSubsystemLoaded(const FName InSubsystemName
 {
 	bool bIsLoaded = false;
 
-	FName SubsystemName = InSubsystemName;
-	if (SubsystemName == NAME_None)
-	{
-		SubsystemName = DefaultPlatformService;
-	}
+	FName SubsystemName, InstanceName;
+	ParseOnlineSubsystemName(InSubsystemName, SubsystemName, InstanceName);
+
 	if (SubsystemName != NAME_None)
 	{
 		if (FModuleManager::Get().IsModuleLoaded(GetOnlineModuleName(SubsystemName.ToString())))
diff --git a/Engine/Source/Runtime/Online/OnlineSubsystem/Public/OnlineSubsystemModule.h b/Engine/Source/Runtime/Online/OnlineSubsystem/Public/OnlineSubsystemModule.h
--- a/Engine/Source/Runtime/Online/OnlineSubsystem/Public/OnlineSubsystemModule.h
+++ b/Engine/Source/Runtime/Online/OnlineSubsystem/Public/OnlineSubsystemModule.h
@@ -44,7 +44,7 @@ private:
 	 * @param InstanceName parsed out value or default instance name
 	 *
 	 */
-	void ParseOnlineSubsystemName(const FName& FullName, FName& SubsystemName, FName& InstanceName);
+	void ParseOnlineSubsystemName(const FName& FullName, FName& SubsystemName, FName& InstanceName) const;
 
 	/**
 	 *	Shuts down all registered online subsystem platforms and unloads their modules
