#ue4
- Changed iOS achievements to always show a toast

diff --git a/Engine/Source/Runtime/Online/IOS/OnlineSubsystemIOS/Private/OnlineAchievementsInterfaceIOS.cpp b/Engine/Source/Runtime/Online/IOS/OnlineSubsystemIOS/Private/OnlineAchievementsInterfaceIOS.cpp
--- a/Engine/Source/Runtime/Online/IOS/OnlineSubsystemIOS/Private/OnlineAchievementsInterfaceIOS.cpp
+++ b/Engine/Source/Runtime/Online/IOS/OnlineSubsystemIOS/Private/OnlineAchievementsInterfaceIOS.cpp
@@ -147,6 +147,7 @@ void FOnlineAchievementsIOS::WriteAchievements(const FUniqueNetId& PlayerId, FOn
 		NSString* AchievementID = [NSString stringWithFString:AchievementName];
 		
 		GKAchievement* Achievement = [[[GKAchievement alloc] initWithIdentifier:AchievementID] autorelease];
+		Achievement.showsCompletionBanner = YES;
 
 		// Setup the percentage complete with the value we are writing from the variant type
 		switch (Stat.GetType())
