More logs to help determine the problem with Steam OSS seemingly filtering out a valid search result on OS X (TTP #333177).

diff --git a/Engine/Source/Runtime/Online/OnlineSubsystemSteam/Private/OnlineSessionAsyncLobbySteam.cpp b/Engine/Source/Runtime/Online/OnlineSubsystemSteam/Private/OnlineSessionAsyncLobbySteam.cpp
--- a/Engine/Source/Runtime/Online/OnlineSubsystemSteam/Private/OnlineSessionAsyncLobbySteam.cpp
+++ b/Engine/Source/Runtime/Online/OnlineSubsystemSteam/Private/OnlineSessionAsyncLobbySteam.cpp
@@ -421,7 +421,7 @@ bool FillMembersFromLobbyData(FUniqueNetIdSteam& LobbyId, FNamedOnlineSession& S
  */
 FString FOnlineAsyncTaskSteamCreateLobby::ToString() const
 {
-	return FString::Printf(TEXT("FOnlineAsyncTaskSteamCreateLobby bWasSuccessful: %d LobbyId: %I64u LobbyType: %d Result: %s"), bWasSuccessful, CallbackResults.m_ulSteamIDLobby, (int32)LobbyType, *SteamResultString(CallbackResults.m_eResult));
+	return FString::Printf(TEXT("FOnlineAsyncTaskSteamCreateLobby bWasSuccessful: %d LobbyId: %llu LobbyType: %d Result: %s"), bWasSuccessful, CallbackResults.m_ulSteamIDLobby, (int32)LobbyType, *SteamResultString(CallbackResults.m_eResult));
 }
 
 /**
@@ -885,6 +885,7 @@ void FOnlineAsyncTaskSteamFindLobbies::ParseSearchResult(FUniqueNetIdSteam& Lobb
 	FOnlineSessionSearchResult* NewSearchResult = new (SearchSettings->SearchResults) FOnlineSessionSearchResult();
 	if (!FillSessionFromLobbyData(LobbyId, NewSearchResult->Session))
 	{
+		UE_LOG_ONLINE(Warning, TEXT("Unable to parse search result for lobby '%s'"), *LobbyId.ToDebugString());
 		// Remove the failed element
 		SearchSettings->SearchResults.RemoveAtSwap(SearchSettings->SearchResults.Num() - 1);
 	}
@@ -1007,16 +1008,25 @@ void FOnlineAsyncTaskSteamFindLobbies::Finalize()
 {
 	FOnlineSessionSteamPtr SessionInt = StaticCastSharedPtr<FOnlineSessionSteam>(Subsystem->GetSessionInterface());
 
+	UE_LOG_ONLINE(Log, TEXT("Found %d lobbies, finalizing the search"), SessionInt->PendingSearchLobbyIds.Num());
+
 	if (bWasSuccessful)
 	{
 		// Parse any ready search results
 		for (int32 LobbyIdx=0; LobbyIdx < SessionInt->PendingSearchLobbyIds.Num(); LobbyIdx++)
 		{
 			FUniqueNetIdSteam& LobbyId = SessionInt->PendingSearchLobbyIds[LobbyIdx];
+			UE_LOG_ONLINE(Log, TEXT("Search result %d: LobbyId=%s, LobbyId.IsValid()=%s, CSteamID(LobbyId).IsLobby()=%s"),
+				LobbyIdx, *LobbyId.ToDebugString(), LobbyId.IsValid() ? TEXT("true") : TEXT("false"), CSteamID(LobbyId).IsLobby() ? TEXT("true") : TEXT("false")
+				);
 			if (LobbyId.IsValid() && CSteamID(LobbyId).IsLobby())
 			{
 				ParseSearchResult(LobbyId);
 			}
+			else
+			{
+				UE_LOG_ONLINE(Warning, TEXT("Lobby %d is invalid (or not a lobby), skipping."), LobbyIdx);
+			}
 		}
 
 		SearchSettings->SortSearchResults();
diff --git a/Engine/Source/Runtime/Online/OnlineSubsystemSteam/Private/OnlineSubsystemSteamTypes.h b/Engine/Source/Runtime/Online/OnlineSubsystemSteam/Private/OnlineSubsystemSteamTypes.h
--- a/Engine/Source/Runtime/Online/OnlineSubsystemSteam/Private/OnlineSubsystemSteamTypes.h
+++ b/Engine/Source/Runtime/Online/OnlineSubsystemSteam/Private/OnlineSubsystemSteamTypes.h
@@ -159,7 +159,7 @@ public:
 	 */
 	virtual FString ToString() const OVERRIDE
 	{
-		return FString::Printf(TEXT("%I64u"), UniqueNetId);
+		return FString::Printf(TEXT("%llu"), UniqueNetId);
 	}
 
 	/**
@@ -173,20 +173,20 @@ public:
 		CSteamID SteamID(UniqueNetId);
 		if (SteamID.IsLobby())
 		{
-			return FString::Printf(TEXT("Lobby [0x%I64X]"), UniqueNetId);
+			return FString::Printf(TEXT("Lobby [0x%llX]"), UniqueNetId);
 		}
 		else if (SteamID.BAnonGameServerAccount())
 		{
-			return FString::Printf(TEXT("Server [0x%I64X]"), UniqueNetId);
+			return FString::Printf(TEXT("Server [0x%llX]"), UniqueNetId);
 		}
 		else if (SteamID.IsValid())
 		{
 			const FString NickName(SteamFriends() ? UTF8_TO_TCHAR(SteamFriends()->GetFriendPersonaName(UniqueNetId)) : TEXT("UNKNOWN"));
-			return FString::Printf(TEXT("%s [0x%I64X]"), *NickName, UniqueNetId);
+			return FString::Printf(TEXT("%s [0x%llX]"), *NickName, UniqueNetId);
 		}
 		else
 		{
-			return FString::Printf(TEXT("INVALID [0x%I64X]"), UniqueNetId);
+			return FString::Printf(TEXT("INVALID [0x%llX]"), UniqueNetId);
 		}
 	}
 
@@ -402,7 +402,7 @@ public:
 	 */
 	virtual FString ToString() const OVERRIDE
 	{
-		return FString::Printf(TEXT("%I64d"), SharedContentHandle);
+		return FString::Printf(TEXT("%lld"), SharedContentHandle);
 	}
 
 	/** 
@@ -413,7 +413,7 @@ public:
 	 */
 	virtual FString ToDebugString() const OVERRIDE
 	{
-		return FString::Printf(TEXT("0x%I64X"), SharedContentHandle);
+		return FString::Printf(TEXT("0x%llX"), SharedContentHandle);
 	}
 };
 
