#UE4 - Big OSS "access to World/WorldContext" refactors
- removed OSS only UWorld accessors on UEngine
- added GetGameWorld accessor to UGameEngine (should be used sparingly)
- added Force Dedicated mode so that OSS can work with PIE in dedicated
- exposed instance name throughout OSS so that we can find the right UWorld in various IsServer / IsDedicated queries
- moved IsServer inside OnlineSubsystem
- moved IsLocalPlayer inside OnlineSubsystem interface
- updated the Core delegate for IsServerForOnlineSubsystems (keeps engine out of base OSS)

diff --git a/Engine/Source/Runtime/Core/Private/Misc/CoreMisc.cpp b/Engine/Source/Runtime/Core/Private/Misc/CoreMisc.cpp
--- a/Engine/Source/Runtime/Core/Private/Misc/CoreMisc.cpp
+++ b/Engine/Source/Runtime/Core/Private/Misc/CoreMisc.cpp
@@ -790,13 +790,13 @@ struct FTestTicker
 	Runtime functions.
 ----------------------------------------------------------------------------*/
 
-FQueryIsRunningServer IsServerDelegate;
+FQueryIsRunningServer GIsServerDelegate;
 
-bool IsServerForOnlineSubsystems()
+bool IsServerForOnlineSubsystems(int32 WorldContextHandle)
 {
-	if (IsServerDelegate.IsBound())
+	if (GIsServerDelegate.IsBound())
 	{
-		return IsServerDelegate.Execute();
+		return GIsServerDelegate.Execute(WorldContextHandle);
 	}
 	else
 	{
@@ -806,7 +806,7 @@ bool IsServerForOnlineSubsystems()
 
 void SetIsServerForOnlineSubsystemsDelegate(FQueryIsRunningServer NewDelegate)
 {
-	IsServerDelegate = NewDelegate;
+	GIsServerDelegate = NewDelegate;
 }
 
 #if UE_EDITOR
diff --git a/Engine/Source/Runtime/Core/Public/Misc/App.h b/Engine/Source/Runtime/Core/Public/Misc/App.h
--- a/Engine/Source/Runtime/Core/Public/Misc/App.h
+++ b/Engine/Source/Runtime/Core/Public/Misc/App.h
@@ -399,12 +399,12 @@ private:
 /**
  * Called to determine the result of IsServerForOnlineSubsystems()
  */
-DECLARE_DELEGATE_RetVal(bool, FQueryIsRunningServer);
+DECLARE_DELEGATE_RetVal_OneParam(bool, FQueryIsRunningServer, int32);
 
 /**
  * @return true if there is a running game world that is a server (including listen servers), false otherwise
  */
-CORE_API bool IsServerForOnlineSubsystems();
+CORE_API bool IsServerForOnlineSubsystems(int32 WorldContextHandle);
 
 /**
  * Sets the delegate used for IsServerForOnlineSubsystems
diff --git a/Engine/Source/Runtime/Engine/Classes/Engine/Engine.h b/Engine/Source/Runtime/Engine/Classes/Engine/Engine.h
--- a/Engine/Source/Runtime/Engine/Classes/Engine/Engine.h
+++ b/Engine/Source/Runtime/Engine/Classes/Engine/Engine.h
@@ -2161,30 +2161,8 @@ public:
 	 * Note: if there is no valid net driver, returns NM_StandAlone
 	 */
 	//virtual ENetMode GetNetMode(FName NetDriverName = NAME_GameNetDriver) const;
-
 	ENetMode GetNetMode(const UWorld *World) const;
 
-	/**
-	 * This is a global, parameterless function used by the online subsystem modules.
-	 * It should never be used in gamecode - instead use GetNetMode(UWorld*) in order
-	 * to properly support multiple concurrent UWorlds.
-	 */
-	ENetMode GetNetModeForOnlineSubsystems() const;
-
-	/**
-	 * This is a global, parameterless function used by the online subsystem modules.
-	 * It should never be used in gamecode - instead use FindNamedNetDriver(UWorld*) in order
-	 * to properly support multiple concurrent UWorlds.
-	 */
-	UNetDriver* GetNetDriverForOnlineSubsystems(FName NetDriverName);
-
-	/**
-	 * This is a global, parameterless function used by the online subsystem modules.
-	 * It should never be used in gamecode - instead use FindNamedNetDriver(UWorld*) in order
-	 * to properly support multiple concurrent UWorlds.
-	 */
-	UWorld* GetWorldForOnlineSubsystem();
-
 	/**
 	 * Creates a UNetDriver and associates a name with it.
 	 *
diff --git a/Engine/Source/Runtime/Engine/Classes/Engine/GameEngine.h b/Engine/Source/Runtime/Engine/Classes/Engine/GameEngine.h
--- a/Engine/Source/Runtime/Engine/Classes/Engine/GameEngine.h
+++ b/Engine/Source/Runtime/Engine/Classes/Engine/GameEngine.h
@@ -146,5 +146,12 @@ public:
 	{
 		return GameViewportWidget;
 	}
+
+	/**
+	 * This is a global, parameterless function used by the online subsystem modules.
+	 * It should never be used in gamecode - instead use the appropriate world context function 
+	 * in order to properly support multiple concurrent UWorlds.
+	 */
+	UWorld* GetGameWorld();
 };
 
diff --git a/Engine/Source/Runtime/Engine/Private/GameEngine.cpp b/Engine/Source/Runtime/Engine/Private/GameEngine.cpp
--- a/Engine/Source/Runtime/Engine/Private/GameEngine.cpp
+++ b/Engine/Source/Runtime/Engine/Private/GameEngine.cpp
@@ -1170,3 +1170,19 @@ void UGameEngine::CloseSecondaryViewports()
 	}
 	SecondaryViewportClients.Empty();*/
 }
+
+UWorld* UGameEngine::GetGameWorld()
+{
+	for (auto It = WorldList.CreateConstIterator(); It; ++It)
+	{
+		const FWorldContext& Context = *It;
+        // Explicitly not checking for PIE worlds here, this should only 
+        // be called outside of editor (and thus is in UGameEngine
+		if (Context.WorldType == EWorldType::Game && Context.World())
+		{
+			return Context.World();
+		}
+	}
+
+	return NULL;
+}
diff --git a/Engine/Source/Runtime/Engine/Private/UnrealEngine.cpp b/Engine/Source/Runtime/Engine/Private/UnrealEngine.cpp
--- a/Engine/Source/Runtime/Engine/Private/UnrealEngine.cpp
+++ b/Engine/Source/Runtime/Engine/Private/UnrealEngine.cpp
@@ -8755,48 +8755,6 @@ ENetMode UEngine::GetNetMode(const UWorld *World) const
 	return NM_Standalone;
 }
 
-ENetMode UEngine::GetNetModeForOnlineSubsystems() const
-{
-	for (auto It = WorldList.CreateConstIterator(); It; ++It)
-	{
-		const FWorldContext& Context = *It;
-		if ((Context.WorldType == EWorldType::Game || Context.WorldType == EWorldType::PIE) && Context.World())
-		{
-			return Context.World()->GetNetMode();
-		}
-	}
-
-	return NM_Standalone;
-}
-
-UNetDriver* UEngine::GetNetDriverForOnlineSubsystems(FName NetDriverName)
-{
-	for (auto It = WorldList.CreateConstIterator(); It; ++It)
-	{
-		const FWorldContext& Context = *It;
-		if ((Context.WorldType == EWorldType::Game || Context.WorldType == EWorldType::PIE) && Context.World())
-		{
-			return FindNamedNetDriver(Context.World(), NetDriverName);
-		}
-	}
-
-	return NULL;
-}
-
-UWorld* UEngine::GetWorldForOnlineSubsystem()
-{
-	for (auto It = WorldList.CreateConstIterator(); It; ++It)
-	{
-		const FWorldContext& Context = *It;
-		if ((Context.WorldType == EWorldType::Game || Context.WorldType == EWorldType::PIE) && Context.World())
-		{
-			return Context.World();
-		}
-	}
-
-	return NULL;
-}
-
 static inline void CallHandleDisconnectForFailure(UWorld* InWorld, UNetDriver* NetDriver)
 {
 	// No world will be created yet if you fail to initialize network driver while trying to connect via cmd line arg.
diff --git a/Engine/Source/Runtime/Launch/Private/LaunchEngineLoop.cpp b/Engine/Source/Runtime/Launch/Private/LaunchEngineLoop.cpp
--- a/Engine/Source/Runtime/Launch/Private/LaunchEngineLoop.cpp
+++ b/Engine/Source/Runtime/Launch/Private/LaunchEngineLoop.cpp
@@ -518,9 +518,28 @@ int32 FEngineLoop::PreInit(int32 ArgC, ANSICHAR* ArgV[], const TCHAR* Additional
 }
 
 #if WITH_ENGINE
-bool IsServerDelegateForOSS()
+bool IsServerDelegateForOSS(int32 WorldContextHandle)
 {
-	return IsRunningDedicatedServer() || (GEngine && (GEngine->GetNetModeForOnlineSubsystems() == NM_DedicatedServer || GEngine->GetNetModeForOnlineSubsystems() == NM_ListenServer));
+	if (IsRunningDedicatedServer())
+	{
+		return true;
+	}
+
+	UWorld* World = NULL;
+	if (WorldContextHandle != INDEX_NONE)
+	{
+		FWorldContext& WorldContext = GEngine->GetWorldContextFromHandleChecked(WorldContextHandle);
+		check(WorldContext.WorldType == EWorldType::Game || WorldContext.WorldType == EWorldType::PIE);
+		World = WorldContext.World();
+	}
+	else
+	{
+		UGameEngine* GameEngine = Cast<UGameEngine>(GEngine);
+		World = GameEngine->GetGameWorld();
+	}
+
+	ENetMode NetMode = World ? World->GetNetMode() : NM_Standalone;
+	return (NetMode == NM_ListenServer || NetMode == NM_DedicatedServer);
 }
 #endif
 
@@ -1030,7 +1049,7 @@ int32 FEngineLoop::PreInit( const TCHAR* CmdLine )
 		bIsSeekFreeDedicatedServer = FPlatformProperties::RequiresCookedData();
 	}
 
-	// If std out device hasn't been initialzied yet (there was no -stdout param in the command line) and
+	// If std out device hasn't been initialized yet (there was no -stdout param in the command line) and
 	// we meet all the criteria, initialize it now.
 	if (!GScopedStdOut.IsValid() && !bHasEditorToken && !bIsRegularClient && !IsRunningDedicatedServer())
 	{
diff --git a/Engine/Source/Runtime/Online/OnlineSubsystem/Private/OnlineSubsystem.cpp b/Engine/Source/Runtime/Online/OnlineSubsystem/Private/OnlineSubsystem.cpp
--- a/Engine/Source/Runtime/Online/OnlineSubsystem/Private/OnlineSubsystem.cpp
+++ b/Engine/Source/Runtime/Online/OnlineSubsystem/Private/OnlineSubsystem.cpp
@@ -88,20 +88,3 @@ bool IsPlayerInSessionImpl(IOnlineSession* SessionInt, FName SessionName, const
 	return bFound;
 }
 
-bool IsLocalPlayerImpl(IOnlineIdentityPtr IdentityInt, const FUniqueNetId& UniqueId)
-{
-	if (!IsRunningDedicatedServer() && IdentityInt.IsValid())
-	{
-		for (int32 LocalUserNum = 0; LocalUserNum < MAX_LOCAL_PLAYERS; LocalUserNum++)
-		{
-			TSharedPtr<FUniqueNetId> LocalUniqueId = IdentityInt->GetUniquePlayerId(LocalUserNum);
-			if (LocalUniqueId.IsValid() && UniqueId == *LocalUniqueId)
-			{
-				return true;
-			}
-		}
-	}
-
-	return false;
-}
-
diff --git a/Engine/Source/Runtime/Online/OnlineSubsystem/Private/OnlineSubsystemImpl.cpp b/Engine/Source/Runtime/Online/OnlineSubsystem/Private/OnlineSubsystemImpl.cpp
--- a/Engine/Source/Runtime/Online/OnlineSubsystem/Private/OnlineSubsystemImpl.cpp
+++ b/Engine/Source/Runtime/Online/OnlineSubsystem/Private/OnlineSubsystemImpl.cpp
@@ -5,7 +5,7 @@
 #include "NamedInterfaces.h"
 
 FOnlineSubsystemImpl::FOnlineSubsystemImpl() :
-	InstanceName(NAME_None),
+	InstanceName(DEFAULT_INSTANCE),
 	bForceDedicated(false),
 	NamedInterfaces(NULL)
 {
@@ -65,3 +65,35 @@ void FOnlineSubsystemImpl::SetNamedInterface(FName InterfaceName, UObject* NewIn
 	}
 }
 
+bool FOnlineSubsystemImpl::IsServer() const
+{
+	int32 WorldContextHandle = INDEX_NONE;
+	if (InstanceName != DEFAULT_INSTANCE && InstanceName != NAME_None)
+	{
+		WorldContextHandle = atoi(TCHAR_TO_ANSI(*InstanceName.ToString()));
+	}
+
+	return IsServerForOnlineSubsystems(WorldContextHandle);
+}
+
+bool FOnlineSubsystemImpl::IsLocalPlayer(const FUniqueNetId& UniqueId) const
+{
+	if (!IsDedicated())
+	{
+		IOnlineIdentityPtr IdentityInt = GetIdentityInterface();
+		if (IdentityInt.IsValid())
+		{
+			for (int32 LocalUserNum = 0; LocalUserNum < MAX_LOCAL_PLAYERS; LocalUserNum++)
+			{
+				TSharedPtr<FUniqueNetId> LocalUniqueId = IdentityInt->GetUniquePlayerId(LocalUserNum);
+				if (LocalUniqueId.IsValid() && UniqueId == *LocalUniqueId)
+				{
+					return true;
+				}
+			}
+		}
+	}
+
+	return false;
+}
+
diff --git a/Engine/Source/Runtime/Online/OnlineSubsystem/Public/OnlineSubsystem.h b/Engine/Source/Runtime/Online/OnlineSubsystem/Public/OnlineSubsystem.h
--- a/Engine/Source/Runtime/Online/OnlineSubsystem/Public/OnlineSubsystem.h
+++ b/Engine/Source/Runtime/Online/OnlineSubsystem/Public/OnlineSubsystem.h
@@ -222,7 +222,7 @@ public:
 	 *
 	 * @param InterfaceName key to the custom data
 	 */
-	virtual class UObject* GetNamedInterface(FName InterfaceName) { return NULL; }
+	virtual class UObject* GetNamedInterface(FName InterfaceName) = 0;
 
 	/**
 	 * Set a custom UObject to be preserved by the online subsystem
@@ -230,19 +230,40 @@ public:
 	 * @param InterfaceName key to the custom data
 	 * @param NewInterface object to preserve
 	 */
-	virtual void SetNamedInterface(FName InterfaceName, class UObject* NewInterface) {};
+	virtual void SetNamedInterface(FName InterfaceName, class UObject* NewInterface) = 0;
+
+	/**
+	 * Is the online subsystem associated with the game/editor/engine running as dedicated.
+	 * May be forced into this mode by EditorPIE, but basically asks if the OSS is serving
+	 * in a dedicated capacity
+	 *
+	 * @return true if the online subsystem is in dedicated server mode, false otherwise
+	 */
+	virtual bool IsDedicated() const = 0;
+
+	/**
+	 * Is this instance of the game running as a server (dedicated OR listen)
+	 * checks the Engine if possible for netmode status
+	 *
+	 * @return true if this is the server, false otherwise
+	 */
+	virtual bool IsServer() const = 0;
 
 	/**
 	 * Force the online subsystem to behave as if it's associated with running a dedicated server
 	 *
 	 * @param bForce force dedicated mode if true
 	 */
-	virtual void SetForceDedicated(bool bForce) {}
+	virtual void SetForceDedicated(bool bForce) = 0;
 
 	/**
-	 * @return true if the online subsystem is in forced dedicated server mode, false otherwise
+	 * Is a player local to this machine by unique id
+	 *
+	 * @param UniqueId UniqueId of the player
+	 *
+	 * @return true if unique id is local to this machine, false otherwise
 	 */
-	virtual bool IsDedicated() const { return false; }
+	virtual bool IsLocalPlayer(const FUniqueNetId& UniqueId) const = 0;
 
 	/** 
 	 * Initialize the underlying subsystem APIs
@@ -312,11 +333,6 @@ public:
  */
 ONLINESUBSYSTEM_API uint32 GetBuildUniqueId();
 
-/**
- * @return true if this is the server, false otherwise
- */
-inline bool IsServer() { return IsServerForOnlineSubsystems(); }
-
 /**
  * Common implementation for finding a player in a session
  *
@@ -328,13 +344,3 @@ inline bool IsServer() { return IsServerForOnlineSubsystems(); }
  */
 ONLINESUBSYSTEM_API bool IsPlayerInSessionImpl(class IOnlineSession* SessionInt, FName SessionName, const FUniqueNetId& UniqueId);
 
-/**
- * Common implementation for asking if a player is local by unique id
- *
- * @param IdentityInt Identity interface to use
- * @param UniqueId UniqueId of the player
- * 
- * @return true if unique id is local to this machine, false otherwise
- */
-ONLINESUBSYSTEM_API bool IsLocalPlayerImpl(IOnlineIdentityPtr IdentityInt, const FUniqueNetId& UniqueId);
-
diff --git a/Engine/Source/Runtime/Online/OnlineSubsystem/Public/OnlineSubsystemImpl.h b/Engine/Source/Runtime/Online/OnlineSubsystem/Public/OnlineSubsystemImpl.h
--- a/Engine/Source/Runtime/Online/OnlineSubsystem/Public/OnlineSubsystemImpl.h
+++ b/Engine/Source/Runtime/Online/OnlineSubsystem/Public/OnlineSubsystemImpl.h
@@ -33,12 +33,15 @@ public:
 	virtual ~FOnlineSubsystemImpl();
 
 	// IOnlineSubsystem
-	virtual void SetForceDedicated(bool bForce) OVERRIDE { bForceDedicated = bForce; }
+	virtual bool IsServer() const OVERRIDE;
 	virtual bool IsDedicated() const OVERRIDE{ return bForceDedicated || IsRunningDedicatedServer(); }
-
+	virtual void SetForceDedicated(bool bForce) OVERRIDE { bForceDedicated = bForce; }
+	
 	virtual class UObject* GetNamedInterface(FName InterfaceName) OVERRIDE;
 	virtual void SetNamedInterface(FName InterfaceName, class UObject* NewInterface) OVERRIDE;
 
+	virtual bool IsLocalPlayer(const FUniqueNetId& UniqueId) const OVERRIDE;
+
 	// FOnlineSubsystemImpl
 
 	/**
diff --git a/Engine/Source/Runtime/Online/OnlineSubsystemNull/Private/OnlineSessionInterfaceNull.cpp b/Engine/Source/Runtime/Online/OnlineSubsystemNull/Private/OnlineSessionInterfaceNull.cpp
--- a/Engine/Source/Runtime/Online/OnlineSubsystemNull/Private/OnlineSessionInterfaceNull.cpp
+++ b/Engine/Source/Runtime/Online/OnlineSubsystemNull/Private/OnlineSessionInterfaceNull.cpp
@@ -260,7 +260,7 @@ bool FOnlineSessionNull::NeedsToAdvertise( FNamedOnlineSession& Session )
 	// b) Not started public LAN session (same as usually)
 	// d) Joinable presence-enabled session that would be advertised with in an online service
 	// (all of that only if we're server)
-	return Session.SessionSettings.bShouldAdvertise && IsServer() &&
+	return Session.SessionSettings.bShouldAdvertise && NullSubsystem->IsServer() &&
 		(
 			(
 			  Session.SessionSettings.bIsLANMatch && 			  
@@ -575,6 +575,8 @@ bool FOnlineSessionNull::JoinSession(int32 PlayerNum, FName SessionName, const F
 bool FOnlineSessionNull::FindFriendSession(int32 LocalUserNum, const FUniqueNetId& Friend)
 {
 	// this function has to exist due to interface definition, but it does not have a meaningful implementation in Null subsystem
+	FOnlineSessionSearchResult EmptySearchResult;
+	TriggerOnFindFriendSessionCompleteDelegates(LocalUserNum, false, EmptySearchResult);
 	return false;
 };
 
@@ -705,7 +707,7 @@ FOnlineSessionSettings* FOnlineSessionNull::GetSessionSettings(FName SessionName
 
 void FOnlineSessionNull::RegisterLocalPlayers(FNamedOnlineSession* Session)
 {
-	if (!IsRunningDedicatedServer())
+	if (!NullSubsystem->IsDedicated())
 	{
 		IOnlineVoicePtr VoiceInt = NullSubsystem->GetVoiceInterface();
 		if (VoiceInt.IsValid())
diff --git a/Engine/Source/Runtime/Online/OnlineSubsystemNull/Private/OnlineSubsystemNull.cpp b/Engine/Source/Runtime/Online/OnlineSubsystemNull/Private/OnlineSubsystemNull.cpp
--- a/Engine/Source/Runtime/Online/OnlineSubsystemNull/Private/OnlineSubsystemNull.cpp
+++ b/Engine/Source/Runtime/Online/OnlineSubsystemNull/Private/OnlineSubsystemNull.cpp
@@ -189,12 +189,6 @@ bool FOnlineSubsystemNull::Shutdown()
 	return true;
 }
 
-bool FOnlineSubsystemNull::IsLocalPlayer(const FUniqueNetId& UniqueId) const
-{
-	IOnlineIdentityPtr IdentityInt = GetIdentityInterface();
-	return IsLocalPlayerImpl(IdentityInt, UniqueId);
-}
-
 FString FOnlineSubsystemNull::GetAppId() const
 {
 	return TEXT("");
diff --git a/Engine/Source/Runtime/Online/OnlineSubsystemNull/Public/OnlineSubsystemNull.h b/Engine/Source/Runtime/Online/OnlineSubsystemNull/Public/OnlineSubsystemNull.h
--- a/Engine/Source/Runtime/Online/OnlineSubsystemNull/Public/OnlineSubsystemNull.h
+++ b/Engine/Source/Runtime/Online/OnlineSubsystemNull/Public/OnlineSubsystemNull.h
@@ -92,10 +92,6 @@ PACKAGE_SCOPE:
 		OnlineAsyncTaskThreadRunnable(NULL),
 		OnlineAsyncTaskThread(NULL)
 	{}
-	/**
-	 *	@return true if the specified Id is considered a local player, false otherwise
-	 */
-	bool IsLocalPlayer(const FUniqueNetId& UniqueId) const;
 
 private:
 
diff --git a/Engine/Source/Runtime/Online/OnlineSubsystemSteam/Private/OnlineSessionAsyncServerSteam.cpp b/Engine/Source/Runtime/Online/OnlineSubsystemSteam/Private/OnlineSessionAsyncServerSteam.cpp
--- a/Engine/Source/Runtime/Online/OnlineSubsystemSteam/Private/OnlineSessionAsyncServerSteam.cpp
+++ b/Engine/Source/Runtime/Online/OnlineSubsystemSteam/Private/OnlineSessionAsyncServerSteam.cpp
@@ -138,9 +138,10 @@ void GetServerKeyValuePairsFromSession(const FOnlineSession* Session, FSteamSess
 /** 
  *  Update the backend with the currently defined settings
  *
+ * @param World current running world instance
  * @param SessionName name of session to update published settings
  */
-void UpdatePublishedSettings(FNamedOnlineSession* Session)
+void UpdatePublishedSettings(UWorld* World, FNamedOnlineSession* Session)
 {
 	ISteamGameServer* SteamGameServerPtr = SteamGameServer();
 	check(SteamGameServerPtr);
@@ -183,7 +184,6 @@ void UpdatePublishedSettings(FNamedOnlineSession* Session)
 	TempSessionSettings.Remove(SETTING_NUMBOTS);
 
 	// Update all the players names/scores
-	UWorld* World = GEngine->GetWorldForOnlineSubsystem();
 	if (World)
 	{
 		AGameState const* const GameState = World->GameState;
@@ -413,7 +413,8 @@ void FOnlineAsyncTaskSteamCreateServer::Finalize()
 
 			Session->SessionState = EOnlineSessionState::Pending;
 
-			UpdatePublishedSettings(Session);
+			UWorld* World = GetWorldForOnline(Subsystem->GetInstanceName());
+			UpdatePublishedSettings(World, Session);
 
 			SessionInt->RegisterLocalPlayers(Session);
 
@@ -483,8 +484,10 @@ void FOnlineAsyncTaskSteamUpdateServer::Tick()
 
 		if (bUpdateOnlineData)
 		{
+			UWorld* World = GetWorldForOnline(Subsystem->GetInstanceName());
+
 			// Master server update
-			UpdatePublishedSettings(Session);
+			UpdatePublishedSettings(World, Session);
 		}
 
 		bWasSuccessful = true;
diff --git a/Engine/Source/Runtime/Online/OnlineSubsystemSteam/Private/OnlineSessionInterfaceSteam.cpp b/Engine/Source/Runtime/Online/OnlineSubsystemSteam/Private/OnlineSessionInterfaceSteam.cpp
--- a/Engine/Source/Runtime/Online/OnlineSubsystemSteam/Private/OnlineSessionInterfaceSteam.cpp
+++ b/Engine/Source/Runtime/Online/OnlineSubsystemSteam/Private/OnlineSessionInterfaceSteam.cpp
@@ -175,7 +175,7 @@ public:
 					IOnlineVoicePtr VoiceInt = Subsystem->GetVoiceInterface();
 					if (VoiceInt.IsValid())
 					{
-						if (!IsRunningDedicatedServer())
+						if (!Subsystem->IsDedicated())
 						{
 							// Stop local talkers
 							VoiceInt->UnregisterLocalTalkers();
@@ -467,7 +467,7 @@ bool FOnlineSessionSteam::EndSession(FName SessionName)
 				// If the session should be advertised and the lan beacon was destroyed, recreate
 				if (Session->SessionSettings.bShouldAdvertise && 
 					LANSession->LanBeacon == NULL &&
-					IsServer())
+					SteamSubsystem->IsServer())
 				{
 					// Recreate the beacon
 					Result = CreateLANSession(Session->HostingPlayerNum, Session);
@@ -1272,7 +1272,7 @@ bool FOnlineSessionSteam::RegisterPlayers(FName SessionName, const TArray< TShar
 
 void FOnlineSessionSteam::RegisterLocalPlayers(FNamedOnlineSession* Session)
 {
-	if (!IsRunningDedicatedServer())
+	if (!SteamSubsystem->IsDedicated())
 	{
 		IOnlineVoicePtr VoiceInt = SteamSubsystem->GetVoiceInterface();
 		if (VoiceInt.IsValid())
@@ -1601,7 +1601,7 @@ void FOnlineSessionSteam::OnLANSearchTimeout()
 			FNamedOnlineSession& Session = Sessions[SessionIdx];
 			if (Session.SessionSettings.bShouldAdvertise &&
 				Session.SessionSettings.bIsLANMatch &&
-				IsServer())
+				SteamSubsystem->IsServer())
 			{
 				bWasHosting = true;
 				break;
diff --git a/Engine/Source/Runtime/Online/OnlineSubsystemSteam/Private/OnlineSubsystemSteamPrivatePCH.h b/Engine/Source/Runtime/Online/OnlineSubsystemSteam/Private/OnlineSubsystemSteamPrivatePCH.h
--- a/Engine/Source/Runtime/Online/OnlineSubsystemSteam/Private/OnlineSubsystemSteamPrivatePCH.h
+++ b/Engine/Source/Runtime/Online/OnlineSubsystemSteam/Private/OnlineSubsystemSteamPrivatePCH.h
@@ -6,6 +6,7 @@
 #include "OnlineSubsystemSteamModule.h"
 #include "OnlineSubsystemModule.h"
 #include "OnlineSubsystem.h"
+#include "OnlineSubsystemUtils.h"
 #include "SocketSubsystem.h"
 
 #include "ModuleManager.h"
diff --git a/Engine/Source/Runtime/Online/OnlineSubsystemSteam/Private/VoiceInterfaceSteam.cpp b/Engine/Source/Runtime/Online/OnlineSubsystemSteam/Private/VoiceInterfaceSteam.cpp
--- a/Engine/Source/Runtime/Online/OnlineSubsystemSteam/Private/VoiceInterfaceSteam.cpp
+++ b/Engine/Source/Runtime/Online/OnlineSubsystemSteam/Private/VoiceInterfaceSteam.cpp
@@ -53,7 +53,7 @@ bool FOnlineVoiceSteam::Init()
 		bSuccess = SessionInt && IdentityInt;
 	}
 
-	if (bSuccess && !IsRunningDedicatedServer())
+	if (bSuccess && !SteamSubsystem->IsDedicated())
 	{
 		VoiceEngine = MakeShareable(new FVoiceEngineSteam(SteamSubsystem));
 		bSuccess = VoiceEngine->Init(MaxLocalTalkers, MaxRemoteTalkers);
@@ -538,7 +538,7 @@ TSharedPtr<FVoicePacket> FOnlineVoiceSteam::SerializeRemotePacket(FArchive& Ar)
 	NewPacket->Serialize(Ar);
 	if (Ar.IsError() == false && NewPacket->GetBufferSize() > 0)
 	{
-		if (!IsRunningDedicatedServer())
+		if (!SteamSubsystem->IsDedicated())
 		{
 			FUniqueNetIdMatcher PlayerMatch(*NewPacket->GetSender());
 			if (MuteList.FindMatch(PlayerMatch) == INDEX_NONE)
diff --git a/Engine/Source/Runtime/Online/OnlineSubsystemSteam/Public/OnlineSubsystemSteam.h b/Engine/Source/Runtime/Online/OnlineSubsystemSteam/Public/OnlineSubsystemSteam.h
--- a/Engine/Source/Runtime/Online/OnlineSubsystemSteam/Public/OnlineSubsystemSteam.h
+++ b/Engine/Source/Runtime/Online/OnlineSubsystemSteam/Public/OnlineSubsystemSteam.h
@@ -195,11 +195,6 @@ PACKAGE_SCOPE:
 	 */
 	void ClearUserCloudFiles();
 
-	/**
-	 *	@return true if the specified Id is considered a local player, false otherwise
-	 */
-	bool IsLocalPlayer(const FUniqueNetId& UniqueId) const;
-
 	/** 
 	 * **INTERNAL** 
 	 * Get the interface for accessing leaderboards/stats
@@ -234,6 +229,7 @@ public:
 	virtual IOnlineUserPtr GetUserInterface() const OVERRIDE;
 	virtual IOnlineMessagePtr GetMessageInterface() const OVERRIDE;
 	virtual IOnlinePresencePtr GetPresenceInterface() const OVERRIDE;
+	virtual bool IsLocalPlayer(const FUniqueNetId& UniqueId) const OVERRIDE;
 	virtual bool Init() OVERRIDE;
 	virtual bool Shutdown() OVERRIDE;
 	virtual bool Exec(class UWorld* InWorld, const TCHAR* Cmd, FOutputDevice& Ar) OVERRIDE;
diff --git a/Engine/Source/Runtime/Online/OnlineSubsystemUtils/Private/OnlineSubsystemUtils.cpp b/Engine/Source/Runtime/Online/OnlineSubsystemUtils/Private/OnlineSubsystemUtils.cpp
--- a/Engine/Source/Runtime/Online/OnlineSubsystemUtils/Private/OnlineSubsystemUtils.cpp
+++ b/Engine/Source/Runtime/Online/OnlineSubsystemUtils/Private/OnlineSubsystemUtils.cpp
@@ -28,19 +28,11 @@
 
 IMPLEMENT_MODULE( FOnlineSubsystemUtilsModule, OnlineSubsystemUtils );
 
-/**
- * Called right after the module DLL has been loaded and the module object has been created
- * Overloaded to allow the default subsystem a chance to load
- */
 void FOnlineSubsystemUtilsModule::StartupModule()
 {
 	
 }
 
-/**
- * Called before the module is unloaded, right before the module object is destroyed.
- * Overloaded to shut down all loaded online subsystems
- */
 void FOnlineSubsystemUtilsModule::ShutdownModule()
 {
 	
@@ -77,6 +69,30 @@ UAudioComponent* CreateVoiceAudioComponent(uint32 SampleRate)
 	return AudioComponent;
 }
 
+UWorld* GetWorldForOnline(FName InstanceName)
+{
+	int32 WorldContextHandle = INDEX_NONE;
+	if (InstanceName != DEFAULT_INSTANCE && InstanceName != NAME_None)
+	{
+		WorldContextHandle = atoi(TCHAR_TO_ANSI(*InstanceName.ToString()));
+	}
+
+	UWorld* World = NULL;
+	if (WorldContextHandle != INDEX_NONE)
+	{
+		FWorldContext& WorldContext = GEngine->GetWorldContextFromHandleChecked(WorldContextHandle);
+		check(WorldContext.WorldType == EWorldType::Game || WorldContext.WorldType == EWorldType::PIE);
+		World = WorldContext.World();
+	}
+	else
+	{
+		UGameEngine* GameEngine = Cast<UGameEngine>(GEngine);
+		World = GameEngine ? GameEngine->GetGameWorld() : NULL;
+	}
+
+	return World;
+}
+
 /**
  * Exec handler that routes online specific execs to the proper subsystem
  *
diff --git a/Engine/Source/Runtime/Online/OnlineSubsystemUtils/Private/VoiceEngineImpl.cpp b/Engine/Source/Runtime/Online/OnlineSubsystemUtils/Private/VoiceEngineImpl.cpp
--- a/Engine/Source/Runtime/Online/OnlineSubsystemUtils/Private/VoiceEngineImpl.cpp
+++ b/Engine/Source/Runtime/Online/OnlineSubsystemUtils/Private/VoiceEngineImpl.cpp
@@ -115,7 +115,7 @@ bool FVoiceEngineImpl::Init(int32 MaxLocalTalkers, int32 MaxRemoteTalkers)
 {
 	bool bSuccess = false;
 
-	if (!IsRunningDedicatedServer())
+	if (!OnlineSubsystem->IsDedicated())
 	{
 		bool bHasVoiceEnabled = false;
 		if (GConfig->GetBool(TEXT("OnlineSubsystem"),TEXT("bHasVoiceEnabled"), bHasVoiceEnabled, GEngineIni) && bHasVoiceEnabled)
diff --git a/Engine/Source/Runtime/Online/OnlineSubsystemUtils/Private/VoiceInterfaceImpl.cpp b/Engine/Source/Runtime/Online/OnlineSubsystemUtils/Private/VoiceInterfaceImpl.cpp
--- a/Engine/Source/Runtime/Online/OnlineSubsystemUtils/Private/VoiceInterfaceImpl.cpp
+++ b/Engine/Source/Runtime/Online/OnlineSubsystemUtils/Private/VoiceInterfaceImpl.cpp
@@ -52,7 +52,7 @@ bool FOnlineVoiceImpl::Init()
 		bSuccess = SessionInt && IdentityInt;
 	}
 
-	if (bSuccess && !IsRunningDedicatedServer())
+	if (bSuccess && !OnlineSubsystem->IsDedicated())
 	{
 		VoiceEngine = MakeShareable(new FVoiceEngineImpl(OnlineSubsystem));
 		bSuccess = VoiceEngine->Init(MaxLocalTalkers, MaxRemoteTalkers);
@@ -537,7 +537,7 @@ TSharedPtr<FVoicePacket> FOnlineVoiceImpl::SerializeRemotePacket(FArchive& Ar)
 	NewPacket->Serialize(Ar);
 	if (Ar.IsError() == false && NewPacket->GetBufferSize() > 0)
 	{
-		if (!IsRunningDedicatedServer())
+		if (!OnlineSubsystem->IsDedicated())
 		{
 			FUniqueNetIdMatcher PlayerMatch(*NewPacket->GetSender());
 			if (MuteList.FindMatch(PlayerMatch) == INDEX_NONE)
diff --git a/Engine/Source/Runtime/Online/OnlineSubsystemUtils/Public/OnlineSubsystemUtils.h b/Engine/Source/Runtime/Online/OnlineSubsystemUtils/Public/OnlineSubsystemUtils.h
--- a/Engine/Source/Runtime/Online/OnlineSubsystemUtils/Public/OnlineSubsystemUtils.h
+++ b/Engine/Source/Runtime/Online/OnlineSubsystemUtils/Public/OnlineSubsystemUtils.h
@@ -10,6 +10,7 @@
 /** @return an initialized audio component specifically for use with VoIP */
 #ifdef ONLINESUBSYSTEMUTILS_API
 ONLINESUBSYSTEMUTILS_API class UAudioComponent* CreateVoiceAudioComponent(uint32 SampleRate);
+ONLINESUBSYSTEMUTILS_API UWorld* GetWorldForOnline(FName InstanceName);
 #endif
 
 /** Macro to handle the boilerplate of accessing the proper online subsystem and getting the requested interface (UWorld version) */
