22nd Century Toys - MS2 CL

PS4 Online Sharing interface is supported through the activity feed WebApi
 - Support for different activity feeds types is supported
 - Titles can create their own Title/Content/OSS/PS4/AcvityFeeds/{TypeName}.json files to add custom captions and actions to the activity feed they wish to post

PS4 ExternalUI Sharing interface has been fleshed out
 - achievements, friends, and login are not supported on PS4 due to limitations/design of that platform
 - ShowWebURL(), ShowInviteUI(), and ShowProfileUI() have been fleshed out
 - ShowInviteUI() has a dependency on SessionInterface tech that is not yet available
   - #define'ing EXTERNAL_UI_CREATE_TEST_SESSION will auto-create a session to join
   - This also demonstrates how the SessionInterface could create such a session (join, leave, etc. would need to be implemented as well)
   - Title/Content/OSS/PS4/Invites/{SomeName}.jpg is a good place to put an image for a session invite
   	- Note that these images have to be JPEGs, cannot exceed 160kb, and are recommended to be 457x257

Other notes:

FOnlineSubsystemPS4 initializes the WebApi libraries and pools so we can make WebApi calls going forward (Also handles shutdown of these libraries)
 - Added a helper class (FAsyncWebApiPS4) that wraps "simple" WebApi calls into a task that can be run in a background thread (since the WebApi methods are blocking).
 - The PS4 Sharing Interface is a good example of using this class
 - WebApi requests that are multi-part are not general enough to wrap. The PS4 Message interface and the ExternalUI interface (see CreateWebApiSessionForTesting) are good examples of multi-part requests.

IOnlineExternalUI interface (and inherited children) has been modified so ShowProfileUI returns bool like other APIs
 - Steam and Live implementations have been updated

IOnlineSharingInterface now supports reading (aysnc) and getting posts (activity feeds on the PS4)

TestExternalUIInterface class has been added to test to the ExternalUI interface
 - You can choose which external UIs you want to test (ACHIEVEMENTS FRIENDS INVITE LOGIN PROFILE WEBURL)

diff --git a/Engine/Source/Runtime/Online/OnlineSubsystem/Public/Interfaces/OnlineExternalUIInterface.h b/Engine/Source/Runtime/Online/OnlineSubsystem/Public/Interfaces/OnlineExternalUIInterface.h
--- a/Engine/Source/Runtime/Online/OnlineSubsystem/Public/Interfaces/OnlineExternalUIInterface.h
+++ b/Engine/Source/Runtime/Online/OnlineSubsystem/Public/Interfaces/OnlineExternalUIInterface.h
@@ -88,8 +88,10 @@ public:
 	 *
 	 * @param Requestor The user requesting the profile.
 	 * @param Requestee The user for whom to show the profile.
+	 *
+	 * @return true if it was able to show the UI, false if it failed
 	 */
-	virtual void ShowProfileUI(const FUniqueNetId& Requestor, const FUniqueNetId& Requestee, const FOnProfileUIClosedDelegate& Delegate) = 0;
+	virtual bool ShowProfileUI(const FUniqueNetId& Requestor, const FUniqueNetId& Requestee, const FOnProfileUIClosedDelegate& Delegate) = 0;
 
 	/**
 	 * Delegate called when the external UI is opened or closed
diff --git a/Engine/Source/Runtime/Online/OnlineSubsystem/Public/Interfaces/OnlineMessageInterface.h b/Engine/Source/Runtime/Online/OnlineSubsystem/Public/Interfaces/OnlineMessageInterface.h
--- a/Engine/Source/Runtime/Online/OnlineSubsystem/Public/Interfaces/OnlineMessageInterface.h
+++ b/Engine/Source/Runtime/Online/OnlineSubsystem/Public/Interfaces/OnlineMessageInterface.h
@@ -216,6 +216,15 @@ public:
 	 */
 	virtual bool GetMessageHeaders(int32 LocalUserNum, TArray< TSharedRef<class FOnlineMessageHeader> >& OutHeaders) = 0;
 
+	/**
+	 * Clear the cached list of message headers
+	 *
+	 * @param LocalUserNum the controller number of the associated user that made the request
+	 *
+	 * @return true if messages were clear, false if there was a problem
+	 */
+	virtual bool ClearMessageHeaders(int32 LocalUserNum) = 0;
+
 	/**
 	 * Download a message and its payload from user's inbox
 	 *
@@ -246,6 +255,24 @@ public:
 	 */
 	virtual TSharedPtr<class FOnlineMessage> GetMessage(int32 LocalUserNum, const FUniqueMessageId& MessageId) = 0;
 
+	/**
+	 * Clear the given cached message
+	 *
+	 * @param LocalUserNum the controller number of the associated user that made the request
+	 *
+	 * @return true if messages were clear, false if there was a problem
+	 */
+	virtual bool ClearMessage(int32 LocalUserNum, const FUniqueMessageId& MessageId) = 0;
+
+	/**
+	 * Clear all the cached messages
+	 *
+	 * @param LocalUserNum the controller number of the associated user that made the request
+	 *
+	 * @return true if messages were clear, false if there was a problem
+	 */
+	virtual bool ClearMessages(int32 LocalUserNum) = 0;
+
 	/**
 	 * Send a message from the currently logged in user to a list of recipients
 	 *
diff --git a/Engine/Source/Runtime/Online/OnlineSubsystem/Public/Interfaces/OnlineSharingInterface.h b/Engine/Source/Runtime/Online/OnlineSubsystem/Public/Interfaces/OnlineSharingInterface.h
--- a/Engine/Source/Runtime/Online/OnlineSubsystem/Public/Interfaces/OnlineSharingInterface.h
+++ b/Engine/Source/Runtime/Online/OnlineSubsystem/Public/Interfaces/OnlineSharingInterface.h
@@ -53,6 +53,9 @@ typedef FOnRequestNewPublishPermissionsComplete::FDelegate FOnRequestNewPublishP
  */
 struct FOnlineStatusUpdate
 {
+	/** Which type of status update for this post. May be ignored for some platforms */
+	FString Type;
+
 	/** The text which forms this post */
 	FString Message;
 	
@@ -69,7 +72,8 @@ struct FOnlineStatusUpdate
 	 * Default Constructor
 	 */
 	FOnlineStatusUpdate() 
-		: Image(NULL)
+		: Type(TEXT("Default"))
+		, Image(NULL)
 		, PostPrivacy(EOnlineStatusUpdatePrivacy::OnlyMe)
 	{}
 };
@@ -147,6 +151,28 @@ public:
 	virtual bool ReadNewsFeed(int32 LocalUserNum, int32 NumPostsToRead = 50) = 0;
 
 
+	/**
+	 * Get a news feed object which was previously synced from the server
+	 *
+	 * @param LocalUserId - The id of the user we are reading news feeds for
+	 * @param NewsFeedIdx - The index of the news feed we are looking up
+	 * @param OutNewsFeed - The news feed object we are searching for
+	 *
+	 * @return Whether the news feed was obtained
+	 */
+	virtual EOnlineCachedResult::Type GetCachedNewsFeed(int32 LocalUserNum, int32 NewsFeedIdx, FOnlineStatusUpdate& OutNewsFeed) = 0;
+
+
+	/**
+	 * Get all the status update objects for the specified local user
+	 *
+	 * @param LocalUserId - The local user Id of the user we are obtaining status updates for
+	 * @param OutNewsFeeds - The collection of news feeds obtained from the server for the given player
+	 *
+	 * @return Whether news feeds were obtained
+	 */
+	virtual EOnlineCachedResult::Type GetCachedNewsFeeds(int32 LocalUserNum, TArray<FOnlineStatusUpdate>& OutNewsFeeds) = 0;
+
 public:
 
 	///////////////////////////////////////////////////////////
diff --git a/Engine/Source/Runtime/Online/OnlineSubsystem/Public/OnlineSubsystemTypes.h b/Engine/Source/Runtime/Online/OnlineSubsystem/Public/OnlineSubsystemTypes.h
--- a/Engine/Source/Runtime/Online/OnlineSubsystem/Public/OnlineSubsystemTypes.h
+++ b/Engine/Source/Runtime/Online/OnlineSubsystem/Public/OnlineSubsystemTypes.h
@@ -721,7 +721,7 @@ public:
 	/**
 	 * @return hex encoded string representation of unique id
 	 */
-	FString GetHexEncodedString()
+	FString GetHexEncodedString() const
 	{
 		if (GetSize() > 0 && GetBytes() != NULL)
 		{
diff --git a/Engine/Source/Runtime/Online/OnlineSubsystemFacebook/Private/IOS/OnlineSharingFacebook.cpp b/Engine/Source/Runtime/Online/OnlineSubsystemFacebook/Private/IOS/OnlineSharingFacebook.cpp
--- a/Engine/Source/Runtime/Online/OnlineSubsystemFacebook/Private/IOS/OnlineSharingFacebook.cpp
+++ b/Engine/Source/Runtime/Online/OnlineSubsystemFacebook/Private/IOS/OnlineSharingFacebook.cpp
@@ -302,3 +302,17 @@ bool FOnlineSharingFacebook::ReadNewsFeed(int32 LocalUserNum, int32 NumPostsToRe
 
 	return bTriggeredRequest;
 }
+
+
+EOnlineCachedResult::Type FOnlineSharingFacebook::GetCachedNewsFeed(int32 LocalUserNum, int32 NewsFeedIdx, FOnlineStatusUpdate& OutNewsFeed)
+{
+	check(NewsFeedIdx >= 0);
+	UE_LOG(LogOnline, Error, TEXT("FOnlineSharingFacebook::GetCachedNewsFeed not yet implemented"));
+	return EOnlineCachedResult::NotFound;
+}
+
+EOnlineCachedResult::Type FOnlineSharingFacebook::GetCachedNewsFeeds(int32 LocalUserNum, TArray<FOnlineStatusUpdate>& OutNewsFeeds)
+{
+	UE_LOG(LogOnline, Error, TEXT("FOnlineSharingFacebook::GetCachedNewsFeeds not yet implemented"));
+	return EOnlineCachedResult::NotFound;
+}
\ No newline at end of file
diff --git a/Engine/Source/Runtime/Online/OnlineSubsystemFacebook/Private/IOS/OnlineSharingFacebook.h b/Engine/Source/Runtime/Online/OnlineSubsystemFacebook/Private/IOS/OnlineSharingFacebook.h
--- a/Engine/Source/Runtime/Online/OnlineSubsystemFacebook/Private/IOS/OnlineSharingFacebook.h
+++ b/Engine/Source/Runtime/Online/OnlineSubsystemFacebook/Private/IOS/OnlineSharingFacebook.h
@@ -21,6 +21,8 @@ public:
 	virtual bool RequestNewReadPermissions(int32 LocalUserNum, EOnlineSharingReadCategory::Type NewPermissions) OVERRIDE;
 	virtual bool ShareStatusUpdate(int32 LocalUserNum, const FOnlineStatusUpdate& StatusUpdate) OVERRIDE;
 	virtual bool RequestNewPublishPermissions(int32 LocalUserNum, EOnlineSharingPublishingCategory::Type NewPermissions, EOnlineStatusUpdatePrivacy::Type Privacy) OVERRIDE;
+	virtual EOnlineCachedResult::Type GetCachedNewsFeed(int32 LocalUserNum, int32 NewsFeedIdx, FOnlineStatusUpdate& OutNewsFeed) OVERRIDE;
+	virtual EOnlineCachedResult::Type GetCachedNewsFeeds(int32 LocalUserNum, TArray<FOnlineStatusUpdate>& OutNewsFeeds) OVERRIDE;
 	// End IOnlineSharing Interface
 
 	
diff --git a/Engine/Source/Runtime/Online/OnlineSubsystemSteam/Private/OnlineExternalUIInterfaceSteam.cpp b/Engine/Source/Runtime/Online/OnlineSubsystemSteam/Private/OnlineExternalUIInterfaceSteam.cpp
--- a/Engine/Source/Runtime/Online/OnlineSubsystemSteam/Private/OnlineExternalUIInterfaceSteam.cpp
+++ b/Engine/Source/Runtime/Online/OnlineSubsystemSteam/Private/OnlineExternalUIInterfaceSteam.cpp
@@ -67,9 +67,10 @@ bool FOnlineExternalUISteam::ShowWebURL(const FString& WebURL)
 	return true;
 }
 
-void FOnlineExternalUISteam::ShowProfileUI( const FUniqueNetId& Requestor, const FUniqueNetId& Requestee, const FOnProfileUIClosedDelegate& Delegate )
+bool FOnlineExternalUISteam::ShowProfileUI( const FUniqueNetId& Requestor, const FUniqueNetId& Requestee, const FOnProfileUIClosedDelegate& Delegate )
 {
 	//@todo
+	return false;
 }
 
 
diff --git a/Engine/Source/Runtime/Online/OnlineSubsystemSteam/Private/OnlineExternalUIInterfaceSteam.h b/Engine/Source/Runtime/Online/OnlineSubsystemSteam/Private/OnlineExternalUIInterfaceSteam.h
--- a/Engine/Source/Runtime/Online/OnlineSubsystemSteam/Private/OnlineExternalUIInterfaceSteam.h
+++ b/Engine/Source/Runtime/Online/OnlineSubsystemSteam/Private/OnlineExternalUIInterfaceSteam.h
@@ -80,7 +80,7 @@ public:
 	virtual bool ShowInviteUI(int32 LocalUserNum) OVERRIDE;
 	virtual bool ShowAchievementsUI(int32 LocalUserNum) OVERRIDE;
 	virtual bool ShowWebURL(const FString& WebURL) OVERRIDE;
-	virtual void ShowProfileUI(const FUniqueNetId& Requestor, const FUniqueNetId& Requestee, const FOnProfileUIClosedDelegate& Delegate) OVERRIDE;
+	virtual bool ShowProfileUI(const FUniqueNetId& Requestor, const FUniqueNetId& Requestee, const FOnProfileUIClosedDelegate& Delegate) OVERRIDE;
 };
 
 typedef TSharedPtr<FOnlineExternalUISteam, ESPMode::ThreadSafe> FOnlineExternalUISteamPtr;
diff --git a/Engine/Source/Runtime/Online/OnlineSubsystemUtils/Private/OnlineSubsystemUtils.cpp b/Engine/Source/Runtime/Online/OnlineSubsystemUtils/Private/OnlineSubsystemUtils.cpp
--- a/Engine/Source/Runtime/Online/OnlineSubsystemUtils/Private/OnlineSubsystemUtils.cpp
+++ b/Engine/Source/Runtime/Online/OnlineSubsystemUtils/Private/OnlineSubsystemUtils.cpp
@@ -25,6 +25,7 @@
 #include "Tests/TestUserInterface.h"
 #include "Tests/TestMessageInterface.h"
 #include "Tests/TestVoice.h"
+#include "Tests/TestExternalUIInterface.h"
 
 IMPLEMENT_MODULE( FOnlineSubsystemUtilsModule, OnlineSubsystemUtils );
 
@@ -268,6 +269,22 @@ static bool OnlineExec( UWorld* InWorld, const TCHAR* Cmd, FOutputDevice& Ar )
 						(new FTestMessageInterface(SubName))->Test(InWorld, RecipientIds);
 						bWasHandled = true;
 					}
+					else if (FParse::Command(&Cmd, TEXT("EXTERNALUI")))
+					{
+						// Full command usage:    EXTERNALUI ACHIEVEMENTS FRIENDS INVITE LOGIN PROFILE WEBURL
+						// Example for one test:  EXTERNALUI WEBURL
+						// Note that tests are enabled in alphabetical order
+						bool bTestAchievementsUI = FParse::Command(&Cmd, TEXT("ACHIEVEMENTS")) ? true : false;
+						bool bTestFriendsUI = FParse::Command(&Cmd, TEXT("FRIENDS")) ? true : false;
+						bool bTestInviteUI = FParse::Command(&Cmd, TEXT("INVITE")) ? true : false;
+						bool bTestLoginUI = FParse::Command(&Cmd, TEXT("LOGIN")) ? true : false;
+						bool bTestProfileUI = FParse::Command(&Cmd, TEXT("PROFILE")) ? true : false;
+						bool bTestWebURL = FParse::Command(&Cmd, TEXT("WEBURL")) ? true : false;
+
+						// This class also deletes itself once done
+						(new FTestExternalUIInterface(SubName, bTestLoginUI, bTestFriendsUI, bTestInviteUI, bTestAchievementsUI, bTestWebURL, bTestProfileUI))->Test();
+						bWasHandled = true;
+					}
 				}
 			}
 		}
diff --git a/Engine/Source/Runtime/Online/OnlineSubsystemUtils/Private/Tests/TestExternalUIInterface.cpp b/Engine/Source/Runtime/Online/OnlineSubsystemUtils/Private/Tests/TestExternalUIInterface.cpp
new file mode 100644
--- /dev/null
+++ b/Engine/Source/Runtime/Online/OnlineSubsystemUtils/Private/Tests/TestExternalUIInterface.cpp
@@ -0,0 +1,190 @@
+// Copyright 1998-2014 Epic Games, Inc. All Rights Reserved.
+
+#include "OnlineSubsystemUtilsPrivatePCH.h"
+#include "ModuleManager.h"
+#include "Online.h"
+#include "TestExternalUIInterface.h"
+
+
+void FTestExternalUIInterface::Test()
+{
+	// Cache interfaces
+	OnlineSub = IOnlineSubsystem::Get(FName(*SubsystemName));
+	check(OnlineSub); 
+
+	ExternalUI = OnlineSub->GetExternalUIInterface();
+	check(ExternalUI.IsValid());
+
+	// Define and register delegates
+	ExternalUIChangeDelegate = FOnExternalUIChangeDelegate::CreateRaw(this, &FTestExternalUIInterface::OnExternalUIChange);
+	ExternalUI->AddOnExternalUIChangeDelegate(ExternalUIChangeDelegate);
+
+	// Are we testing at least one of our external UI?
+	if (bTestLoginUI == false && bTestFriendsUI == false && bTestInviteUI == false && bTestAchievementsUI == false && bTestWebURL == false && bTestProfileUI == false)
+	{
+		UE_LOG(LogOnline, Error, TEXT("ExternalUI test -- No UIs selected to test"));
+		FinishTest();
+	}
+	else
+	{
+		StartNextTest();
+	}
+}
+
+void FTestExternalUIInterface::FinishTest()
+{
+	UE_LOG(LogOnline, Log, TEXT("FTestExternalUIInterface::FinishTest -- completed testing"));
+	ExternalUI->ClearOnExternalUIChangeDelegate(ExternalUIChangeDelegate);
+	delete this;
+}
+
+void FTestExternalUIInterface::StartNextTest()
+{
+	State = ETestExternalUIInterfaceState::Type(int(State) + 1);
+	bool bShowedUI = false;
+
+	if (State >= ETestExternalUIInterfaceState::End)
+	{
+		// We're done. We've tested everything.
+		FinishTest();
+		return;
+	}
+	else if (State == ETestExternalUIInterfaceState::ShowLoginUI)
+	{
+		bShowedUI = TestLoginUI();
+	}
+	else if (State == ETestExternalUIInterfaceState::ShowFriendsUI)
+	{
+		bShowedUI = TestFriendsUI();
+	}
+	else if (State == ETestExternalUIInterfaceState::ShowInviteUI)
+	{
+		bShowedUI = TestInviteUI();
+	}
+	else if (State == ETestExternalUIInterfaceState::ShowAchievementsUI)
+	{
+		bShowedUI = TestAchievementsUI();
+	}
+	else if (State == ETestExternalUIInterfaceState::ShowWebURL)
+	{
+		bShowedUI = TestWebURL();
+	}
+	else if (State == ETestExternalUIInterfaceState::ShowProfileUI)
+	{
+		bShowedUI = TestProfileUI();
+	}
+
+	if (!bShowedUI)
+	{
+		// Either the test was not enabled or there was an error. Go to the next test.
+		StartNextTest();
+	}
+}
+
+bool FTestExternalUIInterface::TestLoginUI()
+{
+	if (bTestLoginUI == false)
+	{
+		UE_LOG(LogOnline, Log, TEXT("TestLoginUI (skipping)"));
+		return false;
+	}
+
+	bool bShowingUI = ExternalUI->ShowLoginUI(0, true, IOnlineExternalUI::FOnLoginUIClosedDelegate::CreateRaw(this, &FTestExternalUIInterface::OnLoginUIClosed));
+	UE_LOG(LogOnline, Log, TEXT("TestLoginUI bShowingUI: %d"), bShowingUI);
+	return bShowingUI;
+}
+
+bool FTestExternalUIInterface::TestFriendsUI()
+{
+	if (bTestFriendsUI == false)
+	{
+		UE_LOG(LogOnline, Log, TEXT("TestFriendsUI (skipping)"));
+		return false;
+	}
+
+	bool bShowingUI = ExternalUI->ShowFriendsUI(0);
+	UE_LOG(LogOnline, Log, TEXT("TestFriendsUI bShowingUI: %d"), bShowingUI);
+	return bShowingUI;
+}
+
+bool FTestExternalUIInterface::TestInviteUI()
+{
+	if (bTestInviteUI == false)
+	{
+		UE_LOG(LogOnline, Log, TEXT("TestInviteUI (skipping)"));
+		return false;
+	}
+
+	bool bShowingUI = ExternalUI->ShowInviteUI(0);
+	UE_LOG(LogOnline, Log, TEXT("TestInviteUI bShowingUI: %d"), bShowingUI);
+	return bShowingUI;
+}
+
+
+bool FTestExternalUIInterface::TestAchievementsUI()
+{
+	if (bTestAchievementsUI == false)
+	{
+		UE_LOG(LogOnline, Log, TEXT("TestAchievementsUI (skipping)"));
+		return false;
+	}
+
+	bool bShowingUI = ExternalUI->ShowAchievementsUI(0);
+	UE_LOG(LogOnline, Log, TEXT("TestAchievementsUI bShowingUI: %d"), bShowingUI);
+	return bShowingUI;
+}
+
+
+bool FTestExternalUIInterface::TestWebURL()
+{
+	if (bTestWebURL == false)
+	{
+		UE_LOG(LogOnline, Log, TEXT("TestWebURL (skipping)"));
+		return false;
+	}
+
+	bool bShowingUI = ExternalUI->ShowWebURL("https://www.unrealengine.com");
+	UE_LOG(LogOnline, Log, TEXT("TestWebURL bShowingUI: %d"), bShowingUI);
+	return bShowingUI;
+}
+
+bool FTestExternalUIInterface::TestProfileUI()
+{
+	if (bTestProfileUI == false)
+	{
+		UE_LOG(LogOnline, Log, TEXT("TestProfileUI (skipping)"));
+		return false;
+	}
+
+	// Show our own profile
+	TSharedPtr<FUniqueNetId> UserId = OnlineSub->GetIdentityInterface()->GetUniquePlayerId(0);
+	bool bShowingUI = ExternalUI->ShowProfileUI(
+		*UserId.Get(),
+		*UserId.Get(),
+		IOnlineExternalUI::FOnProfileUIClosedDelegate::CreateRaw(this, &FTestExternalUIInterface::OnProfileUIClosed));
+
+	UE_LOG(LogOnline, Log, TEXT("TestProfileUI bShowingUI: %d"), bShowingUI);
+	return bShowingUI;
+}
+
+void FTestExternalUIInterface::OnExternalUIChange(bool bIsOpening)
+{
+	UE_LOG(LogOnline, Log, TEXT("OnExternalUIChange delegate invoked. bIsOpening = %d"), bIsOpening);
+
+	if (bIsOpening == false)
+	{
+		// The external UI is no longer active
+		// Move on to the next test
+		StartNextTest();
+	}
+}
+
+void FTestExternalUIInterface::OnLoginUIClosed(TSharedPtr<FUniqueNetId> LoggedInUserId, const int LocalUserId)
+{
+	UE_LOG(LogOnline, Log, TEXT("Login UI closed by local user %d. Logged-in user = %s"), LocalUserId, *LoggedInUserId->ToString());
+}
+
+void FTestExternalUIInterface::OnProfileUIClosed()
+{
+	UE_LOG(LogOnline, Log, TEXT("Profile UI closed by user."));
+}
\ No newline at end of file
diff --git a/Engine/Source/Runtime/Online/OnlineSubsystemUtils/Private/Tests/TestExternalUIInterface.h b/Engine/Source/Runtime/Online/OnlineSubsystemUtils/Private/Tests/TestExternalUIInterface.h
new file mode 100644
--- /dev/null
+++ b/Engine/Source/Runtime/Online/OnlineSubsystemUtils/Private/Tests/TestExternalUIInterface.h
@@ -0,0 +1,103 @@
+// Copyright 1998-2014 Epic Games, Inc. All Rights Reserved.
+
+#pragma once
+
+#include "Core.h"
+#include "OnlineExternalUIInterface.h"
+
+/** Enumeration of external UI tests */
+namespace ETestExternalUIInterfaceState
+{
+	enum Type
+	{
+		Begin,
+		ShowLoginUI,
+		ShowFriendsUI,
+		ShowInviteUI,
+		ShowAchievementsUI,
+		ShowWebURL,
+		ShowProfileUI,
+		End,
+	};
+}
+
+/**
+ * Class used to test the external UI interface
+ */
+ class FTestExternalUIInterface
+ {
+
+	/** The subsystem that was requested to be tested or the default if empty */
+	const FString SubsystemName;
+
+	/** Booleans that control which external UIs to test */
+	bool bTestLoginUI;
+	bool bTestFriendsUI;
+	bool bTestInviteUI;
+	bool bTestAchievementsUI;
+	bool bTestWebURL;
+	bool bTestProfileUI;
+
+	/** The online interface to use for testing */
+	IOnlineSubsystem* OnlineSub;
+
+	/** Convenient access to the external UI interfaces */
+	IOnlineExternalUIPtr ExternalUI;
+
+	/** Delegate for external UI opening and closing */
+	FOnExternalUIChangeDelegate ExternalUIChangeDelegate;
+
+	/** Current external UI test */
+	ETestExternalUIInterfaceState::Type State;
+
+	/** Completes testing and cleans up after itself */
+	void FinishTest();
+
+	/** Go to the next test */
+	void StartNextTest();
+
+	/** Specific test functions */
+	bool TestLoginUI();
+	bool TestFriendsUI();
+	bool TestInviteUI();
+	bool TestAchievementsUI();
+	bool TestWebURL();
+	bool TestProfileUI();
+
+	/** Delegate called when external UI is opening and closing */
+	void OnExternalUIChange(bool bIsOpening);
+
+	/** Delegate executed when the user login UI has been closed. */
+	void OnLoginUIClosed(TSharedPtr<FUniqueNetId> LoggedInUserId, const int LocalUserId);
+
+	/** Delegate executed when the user profile UI has been closed. */
+	void OnProfileUIClosed();
+
+ public:
+
+	/**
+	 * Constructor
+	 *
+	 */
+	FTestExternalUIInterface(const FString& InSubsystemName, bool bInTestLoginUI, bool bInTestFriendsUI, bool bInTestInviteUI, bool bInTestAchievementsUI, bool bInTestWebURL, bool bInTestProfileUI)
+		:	SubsystemName(InSubsystemName)
+		,	bTestLoginUI(bInTestLoginUI)
+		,	bTestFriendsUI(bInTestFriendsUI)
+		,	bTestInviteUI(bInTestInviteUI)
+		,	bTestAchievementsUI(bInTestAchievementsUI)
+		,	bTestWebURL(bInTestWebURL)
+		,	bTestProfileUI(bInTestProfileUI)
+		,	OnlineSub(NULL)
+		,	ExternalUI(NULL)
+		,	State(ETestExternalUIInterfaceState::Begin)
+	{
+	}
+
+	/**
+	 * Kicks off all of the testing process
+	 *
+	 */
+	void Test();
+
+ };
+
diff --git a/Engine/Source/Runtime/Online/OnlineSubsystemUtils/Private/Tests/TestMessageInterface.cpp b/Engine/Source/Runtime/Online/OnlineSubsystemUtils/Private/Tests/TestMessageInterface.cpp
--- a/Engine/Source/Runtime/Online/OnlineSubsystemUtils/Private/Tests/TestMessageInterface.cpp
+++ b/Engine/Source/Runtime/Online/OnlineSubsystemUtils/Private/Tests/TestMessageInterface.cpp
@@ -7,6 +7,31 @@
 #include "ModuleManager.h"
 #include "TestMessageInterface.h"
 
+struct TestAttribute
+{
+	FString Name;
+	FVariantData Value;
+
+	TestAttribute (const FString &InName, const FVariantData &InValue) :
+		Name(InName),
+		Value(InValue)
+	{
+	}
+};
+
+static TestAttribute TestAttributeList [] =
+{
+	TestAttribute("INTValue", FVariantData(512)),
+	TestAttribute("FLOATValue", FVariantData(512.0f)),
+	TestAttribute("QWORDValue", FVariantData((uint64)512)),
+	TestAttribute("DOUBLEValue", FVariantData(512000.0)),
+	TestAttribute("STRINGValue", FVariantData(TEXT("This Is A Test!")))
+};
+
+static uint8 BLOBTestValue [] = 
+{
+	0xde, 0xad, 0xbe, 0xef, 0xfa, 0xde, 0xbe, 0xad
+};
 
 FTestMessageInterface::FTestMessageInterface(const FString& InSubsystem)
 	: OnlineSub(NULL)
@@ -84,14 +109,17 @@ void FTestMessageInterface::StartNextTest()
 		{
 			FOnlineMessagePayload TestPayload;
 
-			TestPayload.SetAttribute(TEXT("INTValue"), FVariantData(512)); 
- 			TestPayload.SetAttribute(TEXT("FLOATValue"), FVariantData(512.0f)); 
- 			TestPayload.SetAttribute(TEXT("QWORDValue"), FVariantData((uint64)512)); 
- 			TestPayload.SetAttribute(TEXT("DOUBLEValue"), FVariantData(512000.0)); 
- 			TestPayload.SetAttribute(TEXT("STRINGValue"), FVariantData(TEXT("This Is A Test!"))); 
+			int AttributeCount = sizeof(TestAttributeList)/sizeof(TestAttribute);
+			for (int i = 0; i < AttributeCount; ++i)
+			{
+				TestPayload.SetAttribute(TestAttributeList[i].Name, TestAttributeList[i].Value);
+			}
 
 			TArray<uint8> TestData;
-			TestData.Add((uint8)200);
+			for (int i = 0; i < sizeof(BLOBTestValue); ++i)
+			{
+				TestData.Add(BLOBTestValue[i]);
+			}
  			TestPayload.SetAttribute(TEXT("BLOBValue"), FVariantData(TestData));
 
 			OnlineSub->GetMessageInterface()->SendMessage(0, Recipients, TEXT("TestType"), TestPayload);
@@ -181,6 +209,74 @@ void FTestMessageInterface::OnReadMessageComplete(int32 LocalPlayer, bool bWasSu
 	UE_LOG(LogOnline, Log,
 		TEXT("ReadMessage() for player (%d) was success=%d"), LocalPlayer, bWasSuccessful);
 
+	// Dump the message content back out
+	if (bWasSuccessful)
+	{
+		TSharedPtr<class FOnlineMessage> Message = OnlineSub->GetMessageInterface()->GetMessage(LocalPlayer, MessageId);
+
+		int AttributeCount = sizeof(TestAttributeList)/sizeof(TestAttribute);
+		for (int i = 0; i < AttributeCount; ++i)
+		{
+			FVariantData Value;
+			if (Message->Payload.GetAttribute(TestAttributeList[i].Name, Value))
+			{
+				if (Value != TestAttributeList[i].Value)
+				{
+					UE_LOG(LogOnline, Log,
+						TEXT("Attribute %s is the wrong value in the received message payload"), *TestAttributeList[i].Name);
+				}
+				else
+				{
+					UE_LOG(LogOnline, Log,
+						TEXT("Attribute %s MATCHED in the received message payload"), *TestAttributeList[i].Name);
+				}
+			}
+			else
+			{
+				UE_LOG(LogOnline, Log,
+					TEXT("Attribute %s is missing from the received message payload"), *TestAttributeList[i].Name);
+			}
+		}
+
+		{
+			FVariantData BLOBValue;
+			if (Message->Payload.GetAttribute(TEXT("BLOBValue"), BLOBValue))
+			{
+				TArray<uint8> TestData;
+				BLOBValue.GetValue(TestData);
+				if (TestData.Num() != sizeof(BLOBTestValue))
+				{
+					UE_LOG(LogOnline, Log,
+						TEXT("Attribute BLOBValue is the wrong size in the received message payload"));
+				}
+				else
+				{
+					bool bIsDataGood = true;
+					for (int i = 0; i < sizeof(BLOBTestValue); ++i)
+					{
+						if (TestData[i] != BLOBTestValue[i])
+						{
+							UE_LOG(LogOnline, Log,
+								TEXT("Attribute BLOBValue contains the wrong data at position %d in the received message payload"), i);
+							bIsDataGood = false;
+							break;
+						}
+					}
+
+					if (bIsDataGood)
+					{
+						UE_LOG(LogOnline, Log,
+							TEXT("Attribute BLOBValue MATCHED in the received message payload"));
+					}
+				}
+			}
+			else
+			{
+				UE_LOG(LogOnline, Log, TEXT("Attribute BLOBValue is missing from the received message payload"));
+			}
+		}
+	}
+
 	// done with this part of the test if no more messages to download
 	MessagesToRead.RemoveAt(0);
 	if (MessagesToRead.Num() == 0)
diff --git a/Engine/Source/Runtime/Online/OnlineSubsystemUtils/Private/Tests/TestSharingInterface.cpp b/Engine/Source/Runtime/Online/OnlineSubsystemUtils/Private/Tests/TestSharingInterface.cpp
--- a/Engine/Source/Runtime/Online/OnlineSubsystemUtils/Private/Tests/TestSharingInterface.cpp
+++ b/Engine/Source/Runtime/Online/OnlineSubsystemUtils/Private/Tests/TestSharingInterface.cpp
@@ -36,7 +36,7 @@ void FTestSharingInterface::Test(UWorld* InWorld, bool bWithImage)
 	SharingInterface = OnlineSub->GetSharingInterface();
 	check(SharingInterface.IsValid());
 
-	TestStatusUpdate.Message = TEXT("This is a test post for UE4 Facebook support!!");
+	TestStatusUpdate.Message = FString::Printf(TEXT("This is a test post for UE4 Sharing support! Date = %s"), *FDateTime::Now().ToString());
 	TestStatusUpdate.PostPrivacy = EOnlineStatusUpdatePrivacy::OnlyMe;
 	if( bWithImage )
 	{
@@ -125,7 +125,7 @@ void FTestSharingInterface::RequestPermissionsToReadNewsFeed()
 
 void FTestSharingInterface::OnReadFeedPermissionsUpdated(int32 LocalUserNum, bool bWasSuccessful)
 {
-	UE_LOG(LogOnline, Display, TEXT("FOnlineSharingFacebook::OnReadFeedPermissionsUpdated() - %d"), bWasSuccessful);
+	UE_LOG(LogOnline, Display, TEXT("FTestSharingInterface::OnReadFeedPermissionsUpdated() - %d"), bWasSuccessful);
 	SharingInterface->ClearOnRequestNewReadPermissionsCompleteDelegate(LocalUserNum, RequestPermissionsToReadFeedDelegate);
 
 	if( ++ResponsesReceived == MAX_LOCAL_PLAYERS )
@@ -154,9 +154,29 @@ void FTestSharingInterface::OnNewsFeedRead(int32 LocalPlayer, bool bWasSuccessfu
 {
 	UE_LOG(LogOnline, Display, TEXT("FTestSharingInterface::OnNewsFeedRead[PlayerIdx:%i - Successful:%i]"), LocalPlayer, bWasSuccessful);
 
+	if (bWasSuccessful)
+	{
+		// Get the 1st cached post
+		FOnlineStatusUpdate FirstReadStatusUpdate;
+		SharingInterface->GetCachedNewsFeed(LocalPlayer, 0, FirstReadStatusUpdate);
+		UE_LOG(LogOnline, Display, TEXT("FTestSharingInterface first read update: %s"), *FirstReadStatusUpdate.Message);
+
+		// Get all the cached posts
+		TArray<FOnlineStatusUpdate> AllReadStatusUpdates;
+		SharingInterface->GetCachedNewsFeeds(LocalPlayer, AllReadStatusUpdates);
+		UE_LOG(LogOnline, Display, TEXT("FTestSharingInterface number of read updates: %d"), AllReadStatusUpdates.Num());
+
+		for (int Idx = 0; Idx < AllReadStatusUpdates.Num(); ++Idx)
+		{
+			const FOnlineStatusUpdate& StatusUpdate = AllReadStatusUpdates[Idx];
+			UE_LOG(LogOnline, Display, TEXT("FTestSharingInterface status update [%d]: %s"), Idx, *StatusUpdate.Message);
+		}
+	}
+
 	SharingInterface->ClearOnReadNewsFeedCompleteDelegate(LocalPlayer, OnNewsFeedReadDelegate);
 	if( ++ResponsesReceived == MAX_LOCAL_PLAYERS )
 	{
+		UE_LOG(LogOnline, Display, TEXT("FTestSharingInterface TESTS COMPLETED"), LocalPlayer, bWasSuccessful);
 		delete this;
 	}
 }
\ No newline at end of file
