fixed FSocket::HasPendingData wrongly returning true if ioctrl is unavailable and/or the pending data size is zero

diff --git a/Engine/Source/Runtime/HTML5/Simulator/HTML5Win32/Private/SocketRaw.cpp b/Engine/Source/Runtime/HTML5/Simulator/HTML5Win32/Private/SocketRaw.cpp
--- a/Engine/Source/Runtime/HTML5/Simulator/HTML5Win32/Private/SocketRaw.cpp
+++ b/Engine/Source/Runtime/HTML5/Simulator/HTML5Win32/Private/SocketRaw.cpp
@@ -109,14 +109,14 @@ bool FSocketRaw::HasPendingConnection( bool& bHasPendingConnection )
 
 bool FSocketRaw::HasPendingData( unsigned int& PendingDataSize )
 {
-	PendingDataSize = 0; 
-	if ( SocketCurrentState( SocketRawData->socket, SocketRawEnum::CanRead ) == SocketRawEnum::Yes )
+	if (SocketCurrentState(SocketRawData->socket, SocketRawEnum::CanRead) == SocketRawEnum::Yes)
 	{
-		if(ioctlsocket(SocketRawData->socket, FIONREAD, (u_long*)(&PendingDataSize)) == 0)
+		if (ioctlsocket(SocketRawData->socket, FIONREAD, (u_long*)(&PendingDataSize)) == 0)
 		{
-			return true; 
+			return (PendingDataSize > 0);
 		}
 	}
+
 	return false; 
 }
 
diff --git a/Engine/Source/Runtime/Online/OnlineSubsystemSteam/Private/SocketsSteam.cpp b/Engine/Source/Runtime/Online/OnlineSubsystemSteam/Private/SocketsSteam.cpp
--- a/Engine/Source/Runtime/Online/OnlineSubsystemSteam/Private/SocketsSteam.cpp
+++ b/Engine/Source/Runtime/Online/OnlineSubsystemSteam/Private/SocketsSteam.cpp
@@ -68,23 +68,14 @@ bool FSocketSteam::HasPendingConnection(bool& bHasPendingConnection)
 	return false;
 }
 
-/**
- * Queries the socket to determine if there is pending data on the queue
- *
- * @param PendingDataSize out parameter indicating how much data is on the pipe for a single recv call
- *
- * @return true if successful, false otherwise
- */
 bool FSocketSteam::HasPendingData(uint32& PendingDataSize) 
 {
-	bool bSuccess = true;
-	if (!SteamNetworkingPtr->IsP2PPacketAvailable(&PendingDataSize, SteamChannel))
+	if (SteamNetworkingPtr->IsP2PPacketAvailable(&PendingDataSize, SteamChannel))
 	{
-		PendingDataSize = 0;
-		bSuccess = false;
+		return (PendingDataSize > 0);
 	}
 
-	return bSuccess;
+	return false;
 }
 
 /**
diff --git a/Engine/Source/Runtime/Online/OnlineSubsystemSteam/Private/SocketsSteam.h b/Engine/Source/Runtime/Online/OnlineSubsystemSteam/Private/SocketsSteam.h
--- a/Engine/Source/Runtime/Online/OnlineSubsystemSteam/Private/SocketsSteam.h
+++ b/Engine/Source/Runtime/Online/OnlineSubsystemSteam/Private/SocketsSteam.h
@@ -113,7 +113,7 @@ public:
 	*
 	* @param PendingDataSize out parameter indicating how much data is on the pipe for a single recv call
 	*
-	* @return true if successful, false otherwise
+	* @return true if the socket has data, false otherwise
 	*/
 	virtual bool HasPendingData(uint32& PendingDataSize) OVERRIDE;
 
diff --git a/Engine/Source/Runtime/Sockets/Private/BSDIPv6Sockets/SocketsBSDIPv6.cpp b/Engine/Source/Runtime/Sockets/Private/BSDIPv6Sockets/SocketsBSDIPv6.cpp
--- a/Engine/Source/Runtime/Sockets/Private/BSDIPv6Sockets/SocketsBSDIPv6.cpp
+++ b/Engine/Source/Runtime/Sockets/Private/BSDIPv6Sockets/SocketsBSDIPv6.cpp
@@ -98,9 +98,6 @@ bool FSocketBSDIPv6::HasPendingConnection(bool& bHasPendingConnection)
 
 bool FSocketBSDIPv6::HasPendingData(uint32& PendingDataSize)
 {
-	bool bHasSucceeded = false;
-	PendingDataSize = 0;
-
 	// make sure socket has no error state
 	if (HasState(EIPv6SocketInternalState::CanRead) == EIPv6SocketInternalState::Yes)
 	{
@@ -109,11 +106,11 @@ bool FSocketBSDIPv6::HasPendingData(uint32& PendingDataSize)
 		if(ioctlsocket(Socket, FIONREAD, (u_long*)(&PendingDataSize)) == 0)
 #endif
 		{
-			bHasSucceeded = true;
+			return (PendingDataSize > 0);
 		}
 	}
 
-	return bHasSucceeded;
+	return false;
 }
 
 
diff --git a/Engine/Source/Runtime/Sockets/Private/BSDSockets/SocketsBSD.cpp b/Engine/Source/Runtime/Sockets/Private/BSDSockets/SocketsBSD.cpp
--- a/Engine/Source/Runtime/Sockets/Private/BSDSockets/SocketsBSD.cpp
+++ b/Engine/Source/Runtime/Sockets/Private/BSDSockets/SocketsBSD.cpp
@@ -115,7 +115,6 @@ bool FSocketBSD::HasPendingConnection(bool& bHasPendingConnection)
 
 bool FSocketBSD::HasPendingData(uint32& PendingDataSize)
 {
-	bool bHasSucceeded = false;
 	PendingDataSize = 0;
 
 	// make sure socket has no error state
@@ -123,14 +122,14 @@ bool FSocketBSD::HasPendingData(uint32& PendingDataSize)
 	{
 #if PLATFORM_HAS_BSD_SOCKET_FEATURE_IOCTL
 		// See if there is any pending data on the read socket
-		if(ioctlsocket(Socket, FIONREAD, (u_long*)(&PendingDataSize)) == 0)
+		if (ioctlsocket(Socket, FIONREAD, (u_long*)(&PendingDataSize)) == 0)
 #endif
 		{
-			bHasSucceeded = true;
+			return (PendingDataSize > 0);
 		}
 	}
 
-	return bHasSucceeded;
+	return false;
 }
 
 
diff --git a/Engine/Source/Runtime/Sockets/Public/Sockets.h b/Engine/Source/Runtime/Sockets/Public/Sockets.h
--- a/Engine/Source/Runtime/Sockets/Public/Sockets.h
+++ b/Engine/Source/Runtime/Sockets/Public/Sockets.h
@@ -93,7 +93,7 @@ public:
 	*
 	* @param PendingDataSize out parameter indicating how much data is on the pipe for a single recv call
 	*
-	* @return true if successful, false otherwise
+	* @return true if the socket has data, false otherwise
 	*/
 	virtual bool HasPendingData(uint32& PendingDataSize) = 0;
 
