#UE4 - Beacon persistence changes
- added BeaconState class to store data in the beacon that needs to persist across levels
- uses SetNamedInterface on OSS to store this data safely across travel
- made initial dedicated server state hard travel into lobby / game on first join (no more initial seamless travel)

diff --git a/Engine/Source/Runtime/Online/OnlineSubsystemUtils/Classes/PartyBeaconHost.h b/Engine/Source/Runtime/Online/OnlineSubsystemUtils/Classes/PartyBeaconHost.h
--- a/Engine/Source/Runtime/Online/OnlineSubsystemUtils/Classes/PartyBeaconHost.h
+++ b/Engine/Source/Runtime/Online/OnlineSubsystemUtils/Classes/PartyBeaconHost.h
@@ -2,139 +2,9 @@
 
 #pragma once
 #include "Runtime/Online/OnlineSubsystemUtils/Classes/OnlineBeaconHost.h"
+#include "Runtime/Online/OnlineSubsystemUtils/Classes/PartyBeaconState.h"
 #include "PartyBeaconHost.generated.h"
 
-/** The result code that will be returned during party reservation */
-UENUM()
-namespace EPartyReservationResult
-{
-	enum Type
-	{
-		// Pending request due to async operation
-		RequestPending,
-		// An unknown error happened
-		GeneralError,
-		// All available reservations are booked
-		PartyLimitReached,
-		// Wrong number of players to join the session
-		IncorrectPlayerCount,
-		// No response from the host
-		RequestTimedOut,
-		// Already have a reservation entry for the requesting party leader
-		ReservationDuplicate,
-		// Couldn't find the party leader specified for a reservation update request 
-		ReservationNotFound,
-		// Space was available and it's time to join
-		ReservationAccepted,
-		// The beacon is paused and not accepting new connections
-		ReservationDenied,
-		// This player is banned
-		ReservationDenied_Banned,
-		// The reservation request was canceled before being sent
-		ReservationRequestCanceled
-	};
-}
-
-namespace EPartyReservationResult
-{
-	/** @return the stringified version of the enum passed in */
-	inline const TCHAR* ToString(EPartyReservationResult::Type SessionType)
-	{
-		switch (SessionType)
-		{
-		case RequestPending:
-			{
-				return TEXT("Pending Request");
-			}
-		case GeneralError:
-			{
-				return TEXT("General Error");
-			}
-		case PartyLimitReached:
-			{
-				return TEXT("Party Limit Reached");
-			}
-		case IncorrectPlayerCount:
-			{
-				return TEXT("Incorrect Player Count");
-			}
-		case RequestTimedOut:
-			{
-				return TEXT("Request Timed Out");
-			}
-		case ReservationDuplicate:
-			{
-				return TEXT("Reservation Duplicate");
-			}
-		case ReservationNotFound:
-			{
-				return TEXT("Reservation Not Found");
-			}
-		case ReservationAccepted:
-			{
-				return TEXT("Reservation Accepted");
-			}
-		case ReservationDenied:
-			{
-				return TEXT("Reservation Denied");
-			}
-		case ReservationDenied_Banned:
-			{
-				return TEXT("Reservation Banned");
-			}
-		case ReservationRequestCanceled:
-			{
-				return TEXT("Request Canceled");
-			}
-		}
-		return TEXT("");
-	}
-}
-
-/** A single player reservation */
-USTRUCT()
-struct FPlayerReservation
-{
-	GENERATED_USTRUCT_BODY()
-	
-	#if CPP
-	FPlayerReservation() :
-		ElapsedTime(0.0f)
-	{}
-	#endif
-
-	/** Unique id for this reservation */
-	UPROPERTY(Transient)
-	FUniqueNetIdRepl UniqueId;
-
-	/** Info needed to validate user credentials when joining a server */
-	UPROPERTY(Transient)
-	FString ValidationStr;
-
-	/** Elapsed time since player made reservation and was last seen */
-	UPROPERTY(Transient)
-	float ElapsedTime;
-};
-
-/** A whole party reservation */
-USTRUCT()
-struct FPartyReservation
-{
-	GENERATED_USTRUCT_BODY()
-
-	/** Team assigned to this party */
-	UPROPERTY(Transient)
-	int32 TeamNum;
-
-	/** Player initiating the request */
-	UPROPERTY(Transient)
-	FUniqueNetIdRepl PartyLeader;
-
-	/** All party members (including party leader) in the reservation */
-	UPROPERTY(Transient)
-	TArray<FPlayerReservation> PartyMembers;
-};
-
 /**
  * Delegate fired when a the beacon host detects a reservation addition/removal
  */
@@ -191,11 +61,20 @@ class ONLINESUBSYSTEMUTILS_API APartyBeaconHost : public AOnlineBeaconHost
 	 * @param InTeamSize size of each team
 	 * @param InMaxReservation max number of reservations allowed
 	 * @param InSessionName name of session related to the beacon
+	 * @param InForceTeamNum team to force players on if applicable (usually only 1 team games)
 	 *
 	 * @return true if successful created, false otherwise
 	 */
 	virtual bool InitHostBeacon(int32 InTeamCount, int32 InTeamSize, int32 InMaxReservations, FName InSessionName, int32 InForceTeamNum = 0);
 
+	/**
+	 * Initialize the party host beacon from a previous state/configuration
+	 * all existing reservations and configuration values are preserved
+     *
+	 * @return true if successful created, false otherwise
+	 */
+	virtual bool InitFromBeaconState(class UPartyBeaconState* PrevState);
+
 	/** 
 	 * Reconfigures the beacon for a different team/player count configuration
 	 * Allows dedicated server to change beacon parameters after a playlist configuration has been made
@@ -207,6 +86,11 @@ class ONLINESUBSYSTEMUTILS_API APartyBeaconHost : public AOnlineBeaconHost
 	 */
 	virtual bool ReconfigureTeamAndPlayerCount(int32 InNumTeams, int32 InNumPlayersPerTeam, int32 InNumReservations);
 
+	/**
+	 * @return reference to the state of the PartyBeacon
+	 */
+	class UPartyBeaconState* GetState() const { return State; }
+
 	/**
 	 * Notify the beacon of a player logout
 	 *
@@ -219,7 +103,7 @@ class ONLINESUBSYSTEMUTILS_API APartyBeaconHost : public AOnlineBeaconHost
 	 *
 	 * @return number of consumed reservations
 	 */
-	virtual int32 GetReservationCount() const { return Reservations.Num(); }
+	virtual int32 GetReservationCount() const { return State->GetReservationCount(); }
 
 	/**
 	 * Get the number of players on a team across all existing reservations
@@ -237,14 +121,14 @@ class ONLINESUBSYSTEMUTILS_API APartyBeaconHost : public AOnlineBeaconHost
 	 * @param PlayerId net id of player to find
 	 * @return index of team assignment for the given player, INDEX_NONE if not found
 	 */
-	int32 GetTeamForCurrentPlayer(const FUniqueNetId& PlayerId);
+	int32 GetTeamForCurrentPlayer(const FUniqueNetId& PlayerId) const;
 
 	/**
 	 * Get the number of teams.
 	 *
 	 * @return The number of teams.
 	 */
-	int32 GetNumTeams() const { return NumTeams; }
+	int32 GetNumTeams() const { return State->GetNumTeams(); }
 
 	/**
 	 * Does a given player id have an existing reservation
@@ -291,7 +175,7 @@ class ONLINESUBSYSTEMUTILS_API APartyBeaconHost : public AOnlineBeaconHost
 	virtual void ProcessReservationRequest(class APartyBeaconClient* Client, const FString& SessionId, const FPartyReservation& ReservationRequest);
 
 	/**
-	 * Handle a reservation cancelation request received from an incoming client
+	 * Handle a reservation cancellation request received from an incoming client
 	 *
 	 * @param Client client beacon making the request
 	 * @param PartyLeader reservation leader
@@ -338,43 +222,21 @@ class ONLINESUBSYSTEMUTILS_API APartyBeaconHost : public AOnlineBeaconHost
 
 protected:
 
-	/** Session tied to the beacon */
-	UPROPERTY(Transient)
-	FName SessionName;
-	/** Number of currently consumed reservations */
-	UPROPERTY(Transient)
-	int32 NumConsumedReservations;
-	/** Maximum allowed reservations */
-	UPROPERTY(Transient)
-	int32 MaxReservations;
-	/** Number of teams in the game */
-	UPROPERTY(Transient)
-	int32 NumTeams;
-	/** Number of players on each team for balancing */
-	UPROPERTY(Transient)
-	int32 NumPlayersPerTeam;
-	/** Team that the host has been assigned to */
-	UPROPERTY(Transient)
-	int32 ReservedHostTeamNum;
-	/** Team that everyone is forced to in single team games */
-	UPROPERTY(Transient)
-	int32 ForceTeamNum;
+	/** State of the beacon */
+	UPROPERTY()
+	class UPartyBeaconState* State;
 
 	/** Delegate fired when the beacon indicates all reservations are taken */
 	FOnReservationsFull ReservationsFull;
 	/** Delegate fired when the beacon indicates a reservation add/remove */
 	FOnReservationChanged ReservationChanged;
-	/** Delegate fired when the beacon indicates a reservation cancelation */
+	/** Delegate fired when the beacon indicates a reservation cancellation */
 	FOnCancelationReceived CancelationReceived;
 	/** Delegate fired when the beacon detects a duplicate reservation */
 	FOnDuplicateReservation DuplicateReservation;
 	/** Delegate fired when asking the beacon owner if this reservation is legit */
 	FOnValidatePlayers ValidatePlayers;
 
-	/** Current reservations in the system */
-	UPROPERTY(Transient)
-	TArray<FPartyReservation> Reservations;
-
 	/** Seconds that can elapse before a reservation is removed due to player not being registered with the session */
 	UPROPERTY(Transient, Config)
 	float SessionTimeoutSecs;
@@ -382,8 +244,10 @@ protected:
 	UPROPERTY(Transient, Config)
 	float TravelSessionTimeoutSecs;
 
-	/** Players that are expected to join shortly */
-	TArray< TSharedPtr<FUniqueNetId> > PlayersPendingJoin;
+	/**
+	 * @return the class of the state object inside this beacon
+	 */
+	virtual TSubclassOf<UPartyBeaconState> GetPartyBeaconHostClass() const { return UPartyBeaconState::StaticClass(); }
 
 	/**
 	 * Handle a newly added player
@@ -392,38 +256,6 @@ protected:
 	 */
 	void NewPlayerAdded(const FPlayerReservation& NewPlayer);
 
-	/**
-	 * Get an existing reservation for a given party
-	 * 
-	 * @param PartyLeader UniqueId of party leader for a reservation
-	 *
-	 * @return index of reservation, INDEX_NONE otherwise
-	 */
-	int32 GetExistingReservation(const FUniqueNetIdRepl& PartyLeader);
-
-	/**
-	 * Initializes the team array so that random choices can be made from it
-	 * Also initializes the host's team number (random from range)
-	 */
-	void InitTeamArray();
-
-	/** 
-	 * Determine if there are any teams that can fit the current party request.
-	 * 
-	 * @param PartySize number of players in the party making a reservation request
-	 * @return true if there are teams available, false otherwise 
-	 */
-	bool AreTeamsAvailable(int32 PartySize) const;
-
-	/**
-	 * Determine the team number for the given party reservation request.
-	 * Uses the list of current reservations to determine what teams have open slots.
-	 *
-	 * @param PartyRequest the party reservation request received from the client beacon
-	 * @return index of the team to assign to all members of this party
-	 */
-	int32 GetTeamAssignment(const FPartyReservation& Party);
-
 	/**
 	 * Does the session match the one associated with this beacon
 	 *
diff --git a/Engine/Source/Runtime/Online/OnlineSubsystemUtils/Classes/PartyBeaconState.h b/Engine/Source/Runtime/Online/OnlineSubsystemUtils/Classes/PartyBeaconState.h
new file mode 100644
--- /dev/null
+++ b/Engine/Source/Runtime/Online/OnlineSubsystemUtils/Classes/PartyBeaconState.h
@@ -0,0 +1,346 @@
+// Copyright 1998-2014 Epic Games, Inc. All Rights Reserved.
+
+#pragma once
+#include "PartyBeaconState.generated.h"
+
+/** The result code that will be returned during party reservation */
+UENUM()
+namespace EPartyReservationResult
+{
+	enum Type
+	{
+		// Pending request due to async operation
+		RequestPending,
+		// An unknown error happened
+		GeneralError,
+		// All available reservations are booked
+		PartyLimitReached,
+		// Wrong number of players to join the session
+		IncorrectPlayerCount,
+		// No response from the host
+		RequestTimedOut,
+		// Already have a reservation entry for the requesting party leader
+		ReservationDuplicate,
+		// Couldn't find the party leader specified for a reservation update request 
+		ReservationNotFound,
+		// Space was available and it's time to join
+		ReservationAccepted,
+		// The beacon is paused and not accepting new connections
+		ReservationDenied,
+		// This player is banned
+		ReservationDenied_Banned,
+		// The reservation request was canceled before being sent
+		ReservationRequestCanceled
+	};
+}
+
+namespace EPartyReservationResult
+{
+	/** @return the stringified version of the enum passed in */
+	inline const TCHAR* ToString(EPartyReservationResult::Type SessionType)
+	{
+		switch (SessionType)
+		{
+		case RequestPending:
+			{
+				return TEXT("Pending Request");
+			}
+		case GeneralError:
+			{
+				return TEXT("General Error");
+			}
+		case PartyLimitReached:
+			{
+				return TEXT("Party Limit Reached");
+			}
+		case IncorrectPlayerCount:
+			{
+				return TEXT("Incorrect Player Count");
+			}
+		case RequestTimedOut:
+			{
+				return TEXT("Request Timed Out");
+			}
+		case ReservationDuplicate:
+			{
+				return TEXT("Reservation Duplicate");
+			}
+		case ReservationNotFound:
+			{
+				return TEXT("Reservation Not Found");
+			}
+		case ReservationAccepted:
+			{
+				return TEXT("Reservation Accepted");
+			}
+		case ReservationDenied:
+			{
+				return TEXT("Reservation Denied");
+			}
+		case ReservationDenied_Banned:
+			{
+				return TEXT("Reservation Banned");
+			}
+		case ReservationRequestCanceled:
+			{
+				return TEXT("Request Canceled");
+			}
+		}
+		return TEXT("");
+	}
+}
+
+/** A single player reservation */
+USTRUCT()
+struct FPlayerReservation
+{
+	GENERATED_USTRUCT_BODY()
+	
+	#if CPP
+	FPlayerReservation() :
+		ElapsedTime(0.0f)
+	{}
+	#endif
+
+	/** Unique id for this reservation */
+	UPROPERTY(Transient)
+	FUniqueNetIdRepl UniqueId;
+
+	/** Info needed to validate user credentials when joining a server */
+	UPROPERTY(Transient)
+	FString ValidationStr;
+
+	/** Elapsed time since player made reservation and was last seen */
+	UPROPERTY(Transient)
+	float ElapsedTime;
+};
+
+/** A whole party reservation */
+USTRUCT()
+struct FPartyReservation
+{
+	GENERATED_USTRUCT_BODY()
+
+	/** Team assigned to this party */
+	UPROPERTY(Transient)
+	int32 TeamNum;
+
+	/** Player initiating the request */
+	UPROPERTY(Transient)
+	FUniqueNetIdRepl PartyLeader;
+
+	/** All party members (including party leader) in the reservation */
+	UPROPERTY(Transient)
+	TArray<FPlayerReservation> PartyMembers;
+};
+
+/**
+ * A beacon host used for taking reservations for an existing game session
+ */
+UCLASS(transient, notplaceable, config=Engine)
+class ONLINESUBSYSTEMUTILS_API UPartyBeaconState : public UObject
+{
+	GENERATED_UCLASS_BODY()
+
+	/**
+	 * Initialize this state object 
+	 *
+	 * @param InTeamCount number of teams to make room for
+	 * @param InTeamSize size of each team
+	 * @param InMaxReservation max number of reservations allowed
+	 * @param InSessionName name of session related to the beacon
+	 * @param InForceTeamNum team to force players on if applicable (usually only 1 team games)
+	 *
+	 * @return true if successful created, false otherwise
+	 */
+	virtual bool InitState(int32 InTeamCount, int32 InTeamSize, int32 InMaxReservations, FName InSessionName, int32 InForceTeamNum);
+
+	/**
+	 * Reconfigures the beacon for a different team/player count configuration
+	 * Allows dedicated server to change beacon parameters after a playlist configuration has been made
+	 * Does no real checking against current reservations because we assume the UI wouldn't let
+	 * this party start a gametype if they were too big to fit on a team together
+	 *
+	 * @param InNumTeams the number of teams that are expected to join
+	 * @param InNumPlayersPerTeam the number of players that are allowed to be on each team
+	 * @param InNumReservations the total number of players to allow to join (if different than team * players)
+	 */
+	virtual bool ReconfigureTeamAndPlayerCount(int32 InNumTeams, int32 InNumPlayersPerTeam, int32 InNumReservations);
+
+	/**
+	 * Add a reservation to the beacon state, tries to assign a team
+	 * 
+	 * @param ReservationRequest reservation to possibly add to this state
+	 *
+	 * @return true if successful, false otherwise
+	 */
+	virtual bool AddReservation(const FPartyReservation& ReservationRequest);
+
+	/**
+	 * Remove an entire reservation from this state object
+	 *
+	 * @param PartyLeader leader holding the reservation for the party 
+	 *
+	 * @return true if successful, false if reservation not found
+	 */
+	virtual bool RemoveReservation(const FUniqueNetIdRepl& PartyLeader);
+
+	/**
+	 * Remove a single player from their party's reservation
+	 *
+	 * PlayerId player to remove
+	 *
+	 * @return true if player found and removed, false otherwise
+	 */
+	virtual bool RemovePlayer(const FUniqueNetIdRepl& PlayerId);
+
+	/**
+	 * @return the name of the session associated with this beacon state
+	 */
+	virtual FName GetSessionName() const { return SessionName; }
+
+	/**
+	 * @return all reservations in this beacon state
+	 */
+	virtual TArray<FPartyReservation>& GetReservations() { return Reservations; }
+
+	/**
+	 * Get an existing reservation for a given party
+	 *
+	 * @param PartyLeader UniqueId of party leader for a reservation
+	 *
+	 * @return index of reservation, INDEX_NONE otherwise
+	 */
+	virtual int32 GetExistingReservation(const FUniqueNetIdRepl& PartyLeader);
+
+	/**
+	* Get the current reservation count inside the beacon
+	* this is NOT the number of players in the game
+	*
+	* @return number of consumed reservations
+	*/
+	virtual int32 GetReservationCount() const { return Reservations.Num(); }
+
+	/**
+	 * @return the number of actually used reservations across all parties
+	 */
+	virtual int32 GetNumConsumedReservations() const { return NumConsumedReservations; }
+
+	/**
+	 * Determine if this reservation fits all rules for fitting in the game
+	 * 
+	 * @param ReservationRequest reservation to test for 
+	 *
+	 * @return true if reservation can be added to this state, false otherwise
+	 */
+	virtual bool DoesReservationFit(const FPartyReservation& ReservationRequest) const;
+
+	/**
+	 * @return true if the beacon is currently at max capacity
+	 */
+	virtual bool IsBeaconFull() const { return NumConsumedReservations == MaxReservations; }
+
+	/**
+	 * Get the number of teams.
+	 *
+	 * @return The number of teams.
+	 */
+	virtual int32 GetNumTeams() const { return NumTeams; }
+
+	/**
+	 * Get the number of current players on a given team.
+	 *
+	 * @param TeamIdx team of interest
+	 *
+	 * @return The number of teams.
+	 */
+	virtual int32 GetNumPlayersOnTeam(int32 TeamIdx) const;
+
+	/**
+	 * Get the team index for a given player
+	 *
+	 * @param PlayerId uniqueid for the player of interest
+	 *
+	 * @return team index for the given player, INDEX_NONE otherwise
+	 */
+	virtual int32 GetTeamForCurrentPlayer(const FUniqueNetId& PlayerId) const;
+
+	/**
+	 * Does a given player id have an existing reservation
+	 *
+	 * @param PlayerId uniqueid of the player to check
+	 *
+	 * @return true if a reservation exists, false otherwise
+	 */
+	virtual bool PlayerHasReservation(const FUniqueNetId& PlayerId) const;
+
+	/**
+	 * Obtain player validation string from party reservation entry
+	 *
+	 * @param PlayerId unique id of player to find validation in an existing reservation
+	 * @param OutValidation [out] validation string used when player requested a reservation
+	 *
+	 * @return true if reservation exists for player
+	 *
+	 */
+	virtual bool GetPlayerValidation(const FUniqueNetId& PlayerId, FString& OutValidation) const;
+
+	/**
+	 * Randomly assign a team for the reservation configuring the beacon
+	 */
+	virtual void InitTeamArray();
+
+	/**
+	* Determine if there are any teams that can fit the current party request.
+	*
+	* @param ReservationRequest reservation request to fit
+	* @return true if there are teams available, false otherwise
+	*/
+	virtual bool AreTeamsAvailable(const FPartyReservation& ReservationRequest) const;
+
+	/**
+	* Determine the team number for the given party reservation request.
+	* Uses the list of current reservations to determine what teams have open slots.
+	*
+	* @param PartyRequest the party reservation request received from the client beacon
+	* @return index of the team to assign to all members of this party
+	*/
+	virtual int32 GetTeamAssignment(const FPartyReservation& Party);
+
+	/**
+	 * Output current state of reservations to log
+	 */
+	virtual void DumpReservations() const;
+
+protected:
+
+	/** Session tied to the beacon */
+	UPROPERTY(Transient)
+	FName SessionName;
+	/** Number of currently consumed reservations */
+	UPROPERTY(Transient)
+	int32 NumConsumedReservations;
+	/** Maximum allowed reservations */
+	UPROPERTY(Transient)
+	int32 MaxReservations;
+	/** Number of teams in the game */
+	UPROPERTY(Transient)
+	int32 NumTeams;
+	/** Number of players on each team for balancing */
+	UPROPERTY(Transient)
+	int32 NumPlayersPerTeam;
+	/** Team that the host has been assigned to */
+	UPROPERTY(Transient)
+	int32 ReservedHostTeamNum;
+	/** Team that everyone is forced to in single team games */
+	UPROPERTY(Transient)
+	int32 ForceTeamNum;
+
+	/** Current reservations in the system */
+	UPROPERTY(Transient)
+	TArray<FPartyReservation> Reservations;
+	/** Players that are expected to join shortly */
+	TArray< TSharedPtr<FUniqueNetId> > PlayersPendingJoin;
+
+	friend class APartyBeaconHost;
+};
diff --git a/Engine/Source/Runtime/Online/OnlineSubsystemUtils/Private/PartyBeaconHost.cpp b/Engine/Source/Runtime/Online/OnlineSubsystemUtils/Private/PartyBeaconHost.cpp
--- a/Engine/Source/Runtime/Online/OnlineSubsystemUtils/Private/PartyBeaconHost.cpp
+++ b/Engine/Source/Runtime/Online/OnlineSubsystemUtils/Private/PartyBeaconHost.cpp
@@ -4,152 +4,14 @@
 
 APartyBeaconHost::APartyBeaconHost(const FPostConstructInitializeProperties& PCIP) :
 	Super(PCIP),
-	SessionName(NAME_None),
-	NumConsumedReservations(0),
-	MaxReservations(0)
+	State(NULL)
 {
 	PrimaryActorTick.bCanEverTick = true;
 }
 
-void APartyBeaconHost::Tick(float DeltaTime)
-{
-	UWorld* World = GetWorld();
-	IOnlineSessionPtr SessionInt = Online::GetSessionInterface(World);
-	if (SessionInt.IsValid())
-	{
-		FNamedOnlineSession* Session = SessionInt->GetNamedSession(SessionName);
-		if (Session)
-		{
-			TArray< TSharedPtr<FUniqueNetId> > PlayersToLogout;
-			for (int32 ResIdx=0; ResIdx < Reservations.Num(); ResIdx++)
-			{
-				FPartyReservation& PartyRes = Reservations[ResIdx];
-
-				// Determine if we have a client connection for the reservation 
-				bool bIsConnectedReservation = false;
-				for (int32 ClientIdx=0; ClientIdx < ClientActors.Num(); ClientIdx++)
-				{
-					APartyBeaconClient* Client = Cast<APartyBeaconClient>(ClientActors[ClientIdx]);
-					if (Client)
-					{
-						const FPartyReservation& ClientRes = Client->GetPendingReservation();
-						if (ClientRes.PartyLeader == PartyRes.PartyLeader)
-						{
-							bIsConnectedReservation = true;
-							break;
-						}
-					}
-					else
-					{
-						UE_LOG(LogBeacon, Error, TEXT("Missing PartyBeaconClient found in ClientActors array"));
-						ClientActors.RemoveAtSwap(ClientIdx);
-						ClientIdx--;
-					}
-				}
-
-				// Don't update clients that are still connected
-				if (bIsConnectedReservation)
-				{
-					for (int32 PlayerIdx=0; PlayerIdx < PartyRes.PartyMembers.Num(); PlayerIdx++)
-					{
-						FPlayerReservation& PlayerEntry = PartyRes.PartyMembers[PlayerIdx];
-						PlayerEntry.ElapsedTime = 0.0f;
-					}
-				}
-				// Once a client disconnects, update the elapsed time since they were found as a registrant in the game session
-				else
-				{
-					for (int32 PlayerIdx=0; PlayerIdx < PartyRes.PartyMembers.Num(); PlayerIdx++)
-					{
-						FPlayerReservation& PlayerEntry = PartyRes.PartyMembers[PlayerIdx];
-
-						// Determine if the player is the owner of the session	
-						const bool bIsSessionOwner = Session->OwningUserId.IsValid() && (*Session->OwningUserId == *PlayerEntry.UniqueId);
-
-						// Determine if the player member is registered in the game session
-						if (SessionInt->IsPlayerInSession(SessionName, *PlayerEntry.UniqueId) ||
-							// Never timeout the session owner
-							bIsSessionOwner)
-						{
-							FUniqueNetIdMatcher PlayerMatch(*PlayerEntry.UniqueId);
-							int32 FoundIdx = PlayersPendingJoin.FindMatch(PlayerMatch);
-							if (FoundIdx != INDEX_NONE)
-							{
-								UE_LOG(LogBeacon, Display, TEXT("Beacon (%s): pending player %s found in session (%s)."),
-									*GetName(),
-									*PlayerEntry.UniqueId->ToDebugString(),
-									*SessionName.ToString());
-
-								// reset elapsed time since found
-								PlayerEntry.ElapsedTime = 0.0f;
-								// also remove from pending join list
-								PlayersPendingJoin.RemoveAtSwap(FoundIdx);
-							}
-						}
-						else
-						{
-							// update elapsed time
-							PlayerEntry.ElapsedTime += DeltaTime;
-
-							// if the player is pending it's initial join then check against TravelSessionTimeoutSecs instead
-							FUniqueNetIdMatcher PlayerMatch(*PlayerEntry.UniqueId);
-							int32 FoundIdx = PlayersPendingJoin.FindMatch(PlayerMatch);
-							const bool bIsPlayerPendingJoin = FoundIdx != INDEX_NONE;
-							// if the timeout has been exceeded then add to list of players 
-							// that need to be logged out from the beacon
-							if ((bIsPlayerPendingJoin && PlayerEntry.ElapsedTime > TravelSessionTimeoutSecs) ||
-								(!bIsPlayerPendingJoin && PlayerEntry.ElapsedTime > SessionTimeoutSecs))
-							{
-								PlayersToLogout.AddUnique(PlayerEntry.UniqueId.GetUniqueNetId());
-							}
-						}
-					}
-				}
-			}
-
-			// Logout any players that timed out
-			for (int32 LogoutIdx=0; LogoutIdx < PlayersToLogout.Num(); LogoutIdx++)
-			{
-				bool bFound = false;
-				const TSharedPtr<FUniqueNetId>& UniqueId = PlayersToLogout[LogoutIdx];
-				float ElapsedSessionTime = 0.f;
-				for (int32 ResIdx=0; ResIdx < Reservations.Num(); ResIdx++)
-				{
-					FPartyReservation& PartyRes = Reservations[ResIdx];
-					for (int32 PlayerIdx=0; PlayerIdx < PartyRes.PartyMembers.Num(); PlayerIdx++)
-					{
-						FPlayerReservation& PlayerEntry = PartyRes.PartyMembers[PlayerIdx];
-						if (*PlayerEntry.UniqueId == *UniqueId)
-						{
-							ElapsedSessionTime = PlayerEntry.ElapsedTime;
-							bFound = true;
-							break;
-						}
-					}
-
-					if (bFound)
-					{
-						break;
-					}
-				}
-
-				UE_LOG(LogBeacon, Display, TEXT("Beacon (%s): player logout due to timeout for %s, elapsed time = %0.3f"),
-					*GetName(),
-					*UniqueId->ToDebugString(),
-					ElapsedSessionTime);
-				// Also remove from pending join list
-				PlayersPendingJoin.RemoveSingleSwap(UniqueId);
-				// let the beacon handle the logout and notifications/delegates
-				FUniqueNetIdRepl RemovedId(UniqueId);
-				HandlePlayerLogout(RemovedId);
-			}
-		}
-	}
-}
-
 bool APartyBeaconHost::InitHost()
 {
-	if(AOnlineBeaconHost::InitHost())
+	if (AOnlineBeaconHost::InitHost())
 	{
 		OnBeaconConnected(FName(TEXT("PartyBeacon"))).BindUObject(this, &APartyBeaconHost::ClientConnected);
 		return true;
@@ -160,111 +22,59 @@ bool APartyBeaconHost::InitHost()
 
 bool APartyBeaconHost::InitHostBeacon(int32 InTeamCount, int32 InTeamSize, int32 InMaxReservations, FName InSessionName, int32 InForceTeamNum)
 {
-	UE_LOG(LogBeacon, Verbose, TEXT("InitPartyHost TeamCount:%d TeamSize:%d MaxSize:%d"), InTeamCount, InTeamSize, InMaxReservations);
+	UE_LOG(LogBeacon, Verbose, TEXT("InitHostBeacon TeamCount:%d TeamSize:%d MaxSize:%d"), InTeamCount, InTeamSize, InMaxReservations);
 	if (InMaxReservations > 0)
 	{
 		if (InitHost())
 		{
-			SessionName = InSessionName;
-			NumTeams = InTeamCount;
-			NumPlayersPerTeam = InTeamSize;
-			MaxReservations = InMaxReservations;
-			ForceTeamNum = InForceTeamNum;
-			Reservations.Empty(MaxReservations);
-
-			InitTeamArray();
-			return true;
+			State = ConstructObject<UPartyBeaconState>(GetPartyBeaconHostClass());
+			if (State->InitState(InTeamCount, InTeamSize, InMaxReservations, InSessionName, InForceTeamNum))
+			{
+				return true;
+			}
+		}
+		else
+		{
+			UE_LOG(LogBeacon, Warning, TEXT("Failed to initialize beacon"));
 		}
 	}
 
 	return false;
 }
 
-void APartyBeaconHost::InitTeamArray()
+bool APartyBeaconHost::InitFromBeaconState(UPartyBeaconState* PrevState)
 {
-	if (NumTeams > 1)
-	{
-		// Grab one for the host team
-		ReservedHostTeamNum = FMath::Rand() % NumTeams;
-	}
-	else
+	if (!State && PrevState)
 	{
-		// Only one team, so choose 'forced team' for everything
-		ReservedHostTeamNum = ForceTeamNum;
+		UE_LOG(LogBeacon, Verbose, TEXT("InitFromBeaconState TeamCount:%d TeamSize:%d MaxSize:%d"), PrevState->NumTeams, PrevState->NumPlayersPerTeam, PrevState->MaxReservations);
+
+		if (InitHost())
+		{
+			State = PrevState;
+			return true;
+		}
+		else
+		{
+			UE_LOG(LogBeacon, Warning, TEXT("Failed to initialize beacon"));
+		}
 	}
 
-	UE_LOG(LogBeacon, Display,
-		TEXT("Beacon (%s) team count (%d), team size (%d), host team (%d)"),
-		*BeaconName,
-		NumTeams,
-		NumPlayersPerTeam,
-		ReservedHostTeamNum);
+	return false;
 }
 
 bool APartyBeaconHost::ReconfigureTeamAndPlayerCount(int32 InNumTeams, int32 InNumPlayersPerTeam, int32 InNumReservations)
 {
 	bool bSuccess = false;
-	if (NetDriver != NULL)
+	if (NetDriver != NULL && State != NULL)
 	{
-		//Check total existing reservations against new total maximum
-		if (NumConsumedReservations < InNumReservations)
-		{	 
-			bool bTeamError = false;
-			// Check teams with reservations against new team count
-			if (NumTeams > InNumTeams)
-			{
-				// Any team about to be removed can't have players already there
-				for (int32 TeamIdx = InNumTeams; TeamIdx < NumTeams; TeamIdx++)
-				{
-					if (GetNumPlayersOnTeam(TeamIdx) > 0)
-					{
-						bTeamError = true;
-						UE_LOG(LogBeacon, Warning, TEXT("Beacon (%s) has players on a team about to be removed."),
-							*BeaconName);
-					}
-				}
-			}
-
-			bool bTeamSizeError = false;
-			// Check num players per team against new team size
-			if (NumPlayersPerTeam > InNumPlayersPerTeam)
-			{
-				for (int32 TeamIdx=0; TeamIdx<NumTeams; TeamIdx++)
-				{
-					if (GetNumPlayersOnTeam(TeamIdx) > InNumPlayersPerTeam)
-					{
-						bTeamSizeError = true;
-						UE_LOG(LogBeacon, Warning, TEXT("Beacon (%s) has too many players on a team about to be resized."),
-							*BeaconName);
-					}
-				}
-			}
-
-			if (!bTeamError && !bTeamSizeError)
-			{
-				NumTeams = InNumTeams;
-				NumPlayersPerTeam = InNumPlayersPerTeam;
-				MaxReservations = InNumReservations;
-
-				InitTeamArray();
-				bSuccess = true;
-
-				UE_LOG(LogBeacon, Display,
-					TEXT("Reconfiguring beacon (%s) team count (%d), team size (%d)"),
-					*BeaconName,
-					NumTeams,
-					NumPlayersPerTeam);
-			}
-		}
-		else
-		{
-			UE_LOG(LogBeacon, Warning, TEXT("Beacon (%s) has too many consumed reservations for this reconfiguration, ignoring request."),
-				*BeaconName);
-		}
+		bSuccess = State->ReconfigureTeamAndPlayerCount(InNumTeams, InNumPlayersPerTeam, InNumReservations);
+		UE_LOG(LogBeacon, Log,
+			TEXT("Beacon (%s) reconfiguring team and player count."),
+			*BeaconName);
 	}
 	else
 	{
-		UE_LOG(LogBeacon, Warning, 
+		UE_LOG(LogBeacon, Warning,
 			TEXT("Beacon (%s) hasn't been initialized yet, can't change team and player count."),
 			*BeaconName);
 	}
@@ -272,152 +82,191 @@ bool APartyBeaconHost::ReconfigureTeamAndPlayerCount(int32 InNumTeams, int32 InN
 	return bSuccess;
 }
 
-int32 APartyBeaconHost::GetNumPlayersOnTeam(int32 TeamIdx) const
+void APartyBeaconHost::Tick(float DeltaTime)
 {
-	int32 Result = 0;
-	for (int32 ResIdx = 0; ResIdx < Reservations.Num(); ResIdx++)
+	if (State)
 	{
-		const FPartyReservation& Reservation = Reservations[ResIdx];
-		if (Reservation.TeamNum == TeamIdx)
+		UWorld* World = GetWorld();
+		IOnlineSessionPtr SessionInt = Online::GetSessionInterface(World);
+
+		if (SessionInt.IsValid())
 		{
-			for (int32 PlayerIdx = 0; PlayerIdx < Reservation.PartyMembers.Num(); PlayerIdx++)
+			FName SessionName = State->GetSessionName();
+			TArray<FPartyReservation>& Reservations = State->GetReservations();
+
+			FNamedOnlineSession* Session = SessionInt->GetNamedSession(SessionName);
+			if (Session)
 			{
-				const FPlayerReservation& PlayerEntry = Reservation.PartyMembers[PlayerIdx];
-				// only count valid player net ids
-				if (PlayerEntry.UniqueId.IsValid())
+				TArray< TSharedPtr<FUniqueNetId> > PlayersToLogout;
+				for (int32 ResIdx = 0; ResIdx < Reservations.Num(); ResIdx++)
+				{
+					FPartyReservation& PartyRes = Reservations[ResIdx];
+
+					// Determine if we have a client connection for the reservation 
+					bool bIsConnectedReservation = false;
+					for (int32 ClientIdx = 0; ClientIdx < ClientActors.Num(); ClientIdx++)
+					{
+						APartyBeaconClient* Client = Cast<APartyBeaconClient>(ClientActors[ClientIdx]);
+						if (Client)
+						{
+							const FPartyReservation& ClientRes = Client->GetPendingReservation();
+							if (ClientRes.PartyLeader == PartyRes.PartyLeader)
+							{
+								bIsConnectedReservation = true;
+								break;
+							}
+						}
+						else
+						{
+							UE_LOG(LogBeacon, Error, TEXT("Missing PartyBeaconClient found in ClientActors array"));
+							ClientActors.RemoveAtSwap(ClientIdx);
+							ClientIdx--;
+						}
+					}
+
+					// Don't update clients that are still connected
+					if (bIsConnectedReservation)
+					{
+						for (int32 PlayerIdx = 0; PlayerIdx < PartyRes.PartyMembers.Num(); PlayerIdx++)
+						{
+							FPlayerReservation& PlayerEntry = PartyRes.PartyMembers[PlayerIdx];
+							PlayerEntry.ElapsedTime = 0.0f;
+						}
+					}
+					// Once a client beacon disconnects, update the elapsed time since they were found as a registrant in the game session
+					else
+					{
+						for (int32 PlayerIdx = 0; PlayerIdx < PartyRes.PartyMembers.Num(); PlayerIdx++)
+						{
+							FPlayerReservation& PlayerEntry = PartyRes.PartyMembers[PlayerIdx];
+
+							// Determine if the player is the owner of the session	
+							const bool bIsSessionOwner = Session->OwningUserId.IsValid() && (*Session->OwningUserId == *PlayerEntry.UniqueId);
+
+							// Determine if the player member is registered in the game session
+							if (SessionInt->IsPlayerInSession(SessionName, *PlayerEntry.UniqueId) ||
+								// Never timeout the session owner
+								bIsSessionOwner)
+							{
+								FUniqueNetIdMatcher PlayerMatch(*PlayerEntry.UniqueId);
+								int32 FoundIdx = State->PlayersPendingJoin.FindMatch(PlayerMatch);
+								if (FoundIdx != INDEX_NONE)
+								{
+									UE_LOG(LogBeacon, Display, TEXT("Beacon (%s): pending player %s found in session (%s)."),
+										*GetName(),
+										*PlayerEntry.UniqueId->ToDebugString(),
+										*SessionName.ToString());
+
+									// reset elapsed time since found
+									PlayerEntry.ElapsedTime = 0.0f;
+									// also remove from pending join list
+									State->PlayersPendingJoin.RemoveAtSwap(FoundIdx);
+								}
+							}
+							else
+							{
+								// update elapsed time
+								PlayerEntry.ElapsedTime += DeltaTime;
+
+								// if the player is pending it's initial join then check against TravelSessionTimeoutSecs instead
+								FUniqueNetIdMatcher PlayerMatch(*PlayerEntry.UniqueId);
+								int32 FoundIdx = State->PlayersPendingJoin.FindMatch(PlayerMatch);
+								const bool bIsPlayerPendingJoin = FoundIdx != INDEX_NONE;
+								// if the timeout has been exceeded then add to list of players 
+								// that need to be logged out from the beacon
+								if ((bIsPlayerPendingJoin && PlayerEntry.ElapsedTime > TravelSessionTimeoutSecs) ||
+									(!bIsPlayerPendingJoin && PlayerEntry.ElapsedTime > SessionTimeoutSecs))
+								{
+									PlayersToLogout.AddUnique(PlayerEntry.UniqueId.GetUniqueNetId());
+								}
+							}
+						}
+					}
+				}
+
+				// Logout any players that timed out
+				for (int32 LogoutIdx = 0; LogoutIdx < PlayersToLogout.Num(); LogoutIdx++)
 				{
-					// count party members in each team (includes party leader)
-					Result++;
+					bool bFound = false;
+					const TSharedPtr<FUniqueNetId>& UniqueId = PlayersToLogout[LogoutIdx];
+					float ElapsedSessionTime = 0.f;
+					for (int32 ResIdx = 0; ResIdx < Reservations.Num(); ResIdx++)
+					{
+						const FPartyReservation& PartyRes = Reservations[ResIdx];
+						for (int32 PlayerIdx = 0; PlayerIdx < PartyRes.PartyMembers.Num(); PlayerIdx++)
+						{
+							const FPlayerReservation& PlayerEntry = PartyRes.PartyMembers[PlayerIdx];
+							if (*PlayerEntry.UniqueId == *UniqueId)
+							{
+								ElapsedSessionTime = PlayerEntry.ElapsedTime;
+								bFound = true;
+								break;
+							}
+						}
+
+						if (bFound)
+						{
+							break;
+						}
+					}
+
+					UE_LOG(LogBeacon, Display, TEXT("Beacon (%s): player logout due to timeout for %s, elapsed time = %0.3f"),
+						*GetName(),
+						*UniqueId->ToDebugString(),
+						ElapsedSessionTime);
+					// Also remove from pending join list
+					State->PlayersPendingJoin.RemoveSingleSwap(UniqueId);
+					// let the beacon handle the logout and notifications/delegates
+					FUniqueNetIdRepl RemovedId(UniqueId);
+					HandlePlayerLogout(RemovedId);
 				}
 			}
 		}
 	}
-	return Result;
 }
 
-int32 APartyBeaconHost::GetTeamForCurrentPlayer(const FUniqueNetId& PlayerId)
+int32 APartyBeaconHost::GetNumPlayersOnTeam(int32 TeamIdx) const
 {
-	int32 TeamNum = INDEX_NONE;
-	if (PlayerId.IsValid())
+	int32 Result = 0;
+	if (NetDriver != NULL && State != NULL)
 	{
-		for (int32 ResIdx = 0; ResIdx < Reservations.Num(); ResIdx++)
-		{
-			const FPartyReservation& Reservation = Reservations[ResIdx];
-			for (int32 PlayerIdx = 0; PlayerIdx < Reservation.PartyMembers.Num(); PlayerIdx++)
-			{
-				// find the player id in the existing list of reservations
-				if (*Reservation.PartyMembers[PlayerIdx].UniqueId == PlayerId)
-				{
-					TeamNum = Reservation.TeamNum;
-					break;
-				}
-			}
-		}
-
-		UE_LOG(LogBeacon, Display, TEXT("Assigning player %s to team %d"),
-			*PlayerId.ToString(),
-			TeamNum);
+		Result = State->GetNumPlayersOnTeam(TeamIdx);
 	}
 	else
 	{
-		UE_LOG(LogBeacon, Display, TEXT("Invalid player when attempting to find team assignment"));
+		UE_LOG(LogBeacon, Warning,
+			TEXT("Beacon (%s) hasn't been initialized yet, can't get team player count."),
+			*BeaconName);
 	}
 
-	return TeamNum;
-}
-
-bool APartyBeaconHost::AreTeamsAvailable(int32 PartySize) const
-{
-	for (int32 TeamIdx = 0; TeamIdx < NumTeams; TeamIdx++)
-	{
-		const int32 CurrentPlayersOnTeam = GetNumPlayersOnTeam(TeamIdx);
-		if ((CurrentPlayersOnTeam + PartySize) <= NumPlayersPerTeam)
-		{
-			return true;
-		}
-	}
-	return false;
+	return Result;
 }
 
-/**
- * Helper for sorting team sizes
- */
-struct FTeamBalanceInfo
-{
-	/** Index of team */
-	int32 TeamIdx;
-	/** Current size of team */
-	int32 TeamSize;
-
-	FTeamBalanceInfo(int32 InTeamIdx, int32 InTeamSize)
-		: TeamIdx(InTeamIdx),
-		  TeamSize(InTeamSize)
-	{}
-};
-
-/**
- * Sort teams by size (equal teams are randomly mixed)
- */
-struct FSortTeamSizeSmallestToLargest
+int32 APartyBeaconHost::GetTeamForCurrentPlayer(const FUniqueNetId& PlayerId) const
 {
-	bool operator()(const FTeamBalanceInfo& A, const FTeamBalanceInfo& B) const
+	int32 TeamNum = INDEX_NONE;
+	if (PlayerId.IsValid())
 	{
-		if (A.TeamSize == B.TeamSize)
+		if (State)
 		{
-			return FMath::Rand() % 2 ? true : false;
-		}
-		else
-		{
-			return A.TeamSize < B.TeamSize;
+			TeamNum = State->GetTeamForCurrentPlayer(PlayerId);
 		}
 	}
-};
-
-int32 APartyBeaconHost::GetTeamAssignment(const FPartyReservation& Party)
-{
-	if (NumTeams > 1)
+	else
 	{
-		TArray<FTeamBalanceInfo> PotentialTeamChoices;
-		for (int32 TeamIdx = 0; TeamIdx < NumTeams; TeamIdx++)
-		{
-			const int32 CurrentPlayersOnTeam = GetNumPlayersOnTeam(TeamIdx);
-			if ((CurrentPlayersOnTeam + Party.PartyMembers.Num()) <= NumPlayersPerTeam)
-			{
-				new (PotentialTeamChoices) FTeamBalanceInfo(TeamIdx, CurrentPlayersOnTeam);
-			}
-		}
-
-		// Grab one from our list of choices
-		if (PotentialTeamChoices.Num() > 0)
-		{
-			if (1)
-			{
-				// Choose smallest team
-				PotentialTeamChoices.Sort(FSortTeamSizeSmallestToLargest());
-				return PotentialTeamChoices[0].TeamIdx;
-			}
-			else
-			{
-				// Random choice from set of choices
-				int32 TeamIndex = FMath::Rand() % PotentialTeamChoices.Num();
-				return PotentialTeamChoices[TeamIndex].TeamIdx;
-			}
-		}
-		else
-		{
-			UE_LOG(LogBeacon, Warning, TEXT("(UPartyBeaconHost.GetTeamAssignment): couldn't find an open team for party members."));
-			return INDEX_NONE;
-		}
+		UE_LOG(LogBeacon, Display, TEXT("Invalid player when attempting to find team assignment"));
 	}
 
-	return ForceTeamNum;
+	return TeamNum;
 }
 
 void APartyBeaconHost::NewPlayerAdded(const FPlayerReservation& NewPlayer)
 {
 	UE_LOG(LogBeacon, Verbose, TEXT("Beacon adding player %s"), *NewPlayer.UniqueId->ToDebugString());
-
-	PlayersPendingJoin.Add(NewPlayer.UniqueId.GetUniqueNetId());
+	if (State)
+	{
+		State->PlayersPendingJoin.Add(NewPlayer.UniqueId.GetUniqueNetId());
+	}
 }
 
 void APartyBeaconHost::HandlePlayerLogout(const FUniqueNetIdRepl& PlayerId)
@@ -426,36 +275,11 @@ void APartyBeaconHost::HandlePlayerLogout(const FUniqueNetIdRepl& PlayerId)
 	{
 		UE_LOG(LogBeacon, Verbose, TEXT("HandlePlayerLogout %s"), *PlayerId->ToDebugString());
 
-		bool bWasRemoved = false;
-		for (int32 ResIdx=0; ResIdx < Reservations.Num(); ResIdx++)
+		if (State)
 		{
-			FPartyReservation& Reservation = Reservations[ResIdx];
-
-			// find the player in an existing reservation slot
-			for (int32 PlayerIdx=0; PlayerIdx < Reservation.PartyMembers.Num(); PlayerIdx++)
-			{	
-				FPlayerReservation& PlayerEntry = Reservation.PartyMembers[PlayerIdx];
-				if (PlayerEntry.UniqueId == PlayerId)
-				{
-					// player removed
-					Reservation.PartyMembers.RemoveAtSwap(PlayerIdx--);
-					bWasRemoved = true;
-
-					// free up a consumed entry
-					NumConsumedReservations--;
-				}
-			}
-
-			// remove the entire party reservation slot if no more party members
-			if (Reservation.PartyMembers.Num() == 0)
-			{
-				Reservations.RemoveAtSwap(ResIdx--);
-			}
-
-			if (bWasRemoved)
+			if (State->RemovePlayer(PlayerId))
 			{
 				ReservationChanged.ExecuteIfBound();
-				break;
 			}
 		}
 	}
@@ -463,74 +287,50 @@ void APartyBeaconHost::HandlePlayerLogout(const FUniqueNetIdRepl& PlayerId)
 
 bool APartyBeaconHost::PlayerHasReservation(const FUniqueNetId& PlayerId) const
 {
-	bool bFound = false;
-
-	for (int32 ResIdx=0; ResIdx < Reservations.Num(); ResIdx++)
+	bool bHasReservation = false;
+	if (State)
 	{
-		const FPartyReservation& ReservationEntry = Reservations[ResIdx];
-		for (int32 PlayerIdx=0; PlayerIdx < ReservationEntry.PartyMembers.Num(); PlayerIdx++)
-		{
-			if (*ReservationEntry.PartyMembers[PlayerIdx].UniqueId == PlayerId)
-			{
-				bFound = true;
-				break;
-			}
-		}
+		bHasReservation = State->PlayerHasReservation(PlayerId);
+	}
+	else
+	{
+		UE_LOG(LogBeacon, Warning,
+			TEXT("Beacon (%s) hasn't been initialized yet, no reservations."),
+			*BeaconName);
 	}
 
-	return bFound; 
+	return bHasReservation;
 }
 
 bool APartyBeaconHost::GetPlayerValidation(const FUniqueNetId& PlayerId, FString& OutValidation) const
 {
-	bool bFound = false;
 	OutValidation = FString();
 
-	for (int32 ResIdx=0; ResIdx < Reservations.Num(); ResIdx++)
+	bool bHasValidation = false;
+	if (State)
 	{
-		const FPartyReservation& ReservationEntry = Reservations[ResIdx];
-		for (int32 PlayerIdx=0; PlayerIdx < ReservationEntry.PartyMembers.Num(); PlayerIdx++)
-		{
-			if (*ReservationEntry.PartyMembers[PlayerIdx].UniqueId == PlayerId)
-			{
-				OutValidation = ReservationEntry.PartyMembers[PlayerIdx].ValidationStr;
-				bFound = true;
-				break;
-			}
-		}
+		bHasValidation = State->GetPlayerValidation(PlayerId, OutValidation);
 	}
-
-	return bFound; 
-}
-
-int32 APartyBeaconHost::GetExistingReservation(const FUniqueNetIdRepl& PartyLeader)
-{
-	int32 Result = INDEX_NONE;
-	for (int32 ResIdx=0; ResIdx < Reservations.Num(); ResIdx++)
+	else
 	{
-		const FPartyReservation& ReservationEntry = Reservations[ResIdx];
-		if (ReservationEntry.PartyLeader == PartyLeader)
-		{
-			Result = ResIdx;
-			break;
-		}
+		UE_LOG(LogBeacon, Warning,
+			TEXT("Beacon (%s) hasn't been initialized yet, no validation."),
+			*BeaconName);
 	}
 
-	return Result; 
+	return bHasValidation;
 }
 
 EPartyReservationResult::Type APartyBeaconHost::AddPartyReservation(const FPartyReservation& ReservationRequest)
 {
 	EPartyReservationResult::Type Result = EPartyReservationResult::GeneralError;
 
-	if (BeaconState == EBeaconState::DenyRequests)
+	if (!State || BeaconState == EBeaconState::DenyRequests)
 	{
 		return EPartyReservationResult::ReservationDenied;
 	}
 
-	int32 IncomingPartySize = ReservationRequest.PartyMembers.Num();
-	if (IncomingPartySize <= NumPlayersPerTeam &&
-		NumConsumedReservations + IncomingPartySize <= MaxReservations)
+	if (State->DoesReservationFit(ReservationRequest))
 	{
 		bool bContinue = true;
 		if (ValidatePlayers.IsBound())
@@ -540,9 +340,10 @@ EPartyReservationResult::Type APartyBeaconHost::AddPartyReservation(const FParty
 
 		if (bContinue)
 		{
-			const int32 ExistingReservationIdx = GetExistingReservation(ReservationRequest.PartyLeader);
+			const int32 ExistingReservationIdx = State->GetExistingReservation(ReservationRequest.PartyLeader);
 			if (ExistingReservationIdx != INDEX_NONE)
 			{
+				const TArray<FPartyReservation>& Reservations = State->GetReservations();
 				const FPartyReservation& ExistingReservation = Reservations[ExistingReservationIdx];
 				if (ReservationRequest.PartyMembers.Num() == ExistingReservation.PartyMembers.Num())
 				{
@@ -564,23 +365,18 @@ EPartyReservationResult::Type APartyBeaconHost::AddPartyReservation(const FParty
 			}
 			else
 			{
-				if (AreTeamsAvailable(IncomingPartySize))
+				if (State->AreTeamsAvailable(ReservationRequest))
 				{
-					int32 TeamAssignment = GetTeamAssignment(ReservationRequest);
-					if (TeamAssignment != INDEX_NONE)
+					if (State->AddReservation(ReservationRequest))
 					{
-						NumConsumedReservations += IncomingPartySize;
-						int32 ResIdx = Reservations.Add(ReservationRequest);
-						Reservations[ResIdx].TeamNum = TeamAssignment;
-
 						// Keep track of newly added players
-						for (int32 Count = 0; Count < ReservationRequest.PartyMembers.Num(); Count++)
+						for (const FPlayerReservation& PartyMember : ReservationRequest.PartyMembers)
 						{
-							NewPlayerAdded(ReservationRequest.PartyMembers[Count]);
+							NewPlayerAdded(PartyMember);
 						}
 
 						ReservationChanged.ExecuteIfBound();
-						if (NumConsumedReservations == MaxReservations)
+						if (State->IsBeaconFull())
 						{
 							ReservationsFull.ExecuteIfBound();
 						}
@@ -613,12 +409,8 @@ EPartyReservationResult::Type APartyBeaconHost::AddPartyReservation(const FParty
 
 void APartyBeaconHost::RemovePartyReservation(const FUniqueNetIdRepl& PartyLeader)
 {
-	const int32 ExistingReservationIdx = GetExistingReservation(PartyLeader);
-	if (ExistingReservationIdx != INDEX_NONE)
+	if (State && State->RemoveReservation(PartyLeader))
 	{
-		NumConsumedReservations -= Reservations[ExistingReservationIdx].PartyMembers.Num();
-		Reservations.RemoveAtSwap(ExistingReservationIdx);
-
 		CancelationReceived.ExecuteIfBound(*PartyLeader);
 		ReservationChanged.ExecuteIfBound();
 	}
@@ -630,15 +422,19 @@ void APartyBeaconHost::RemovePartyReservation(const FUniqueNetIdRepl& PartyLeade
 
 bool APartyBeaconHost::DoesSessionMatch(const FString& SessionId) const
 {
-	UWorld* World = GetWorld();
-
-	IOnlineSessionPtr SessionInt = Online::GetSessionInterface(World);
-	FNamedOnlineSession* Session = SessionInt.IsValid() ? SessionInt->GetNamedSession(SessionName) : NULL;
-	if (Session && Session->SessionInfo.IsValid() && !SessionId.IsEmpty() && Session->SessionInfo->GetSessionId().ToString() == SessionId)
+	if (State)
 	{
-		return true;
+		UWorld* World = GetWorld();
+		FName SessionName = State->GetSessionName();
+
+		IOnlineSessionPtr SessionInt = Online::GetSessionInterface(World);
+		FNamedOnlineSession* Session = SessionInt.IsValid() ? SessionInt->GetNamedSession(SessionName) : NULL;
+		if (Session && Session->SessionInfo.IsValid() && !SessionId.IsEmpty() && Session->SessionInfo->GetSessionId().ToString() == SessionId)
+		{
+			return true;
+		}
 	}
-	
+
 	return false;
 }
 
@@ -697,30 +493,10 @@ AOnlineBeaconClient* APartyBeaconHost::SpawnBeaconActor()
 
 void APartyBeaconHost::DumpReservations() const
 {
-	FUniqueNetIdRepl NetId;
-	FPlayerReservation PlayerRes;
-
 	UE_LOG(LogBeacon, Display, TEXT("Debug info for Beacon: %s"), *BeaconName);
-	UE_LOG(LogBeacon, Display, TEXT("Session that reservations are for: %s"), *SessionName.ToString());
-	UE_LOG(LogBeacon, Display, TEXT("Number of teams: %d"), NumTeams);
-	UE_LOG(LogBeacon, Display, TEXT("Number players per team: %d"), NumPlayersPerTeam);
-	UE_LOG(LogBeacon, Display, TEXT("Number total reservations: %d"), MaxReservations);
-	UE_LOG(LogBeacon, Display, TEXT("Number consumed reservations: %d"), NumConsumedReservations);
-	UE_LOG(LogBeacon, Display, TEXT("Number of party reservations: %d"), Reservations.Num());
-
-	// Log each party that has a reservation
-	for (int32 PartyIndex = 0; PartyIndex < Reservations.Num(); PartyIndex++)
+	if (State)
 	{
-		NetId = Reservations[PartyIndex].PartyLeader;
-		UE_LOG(LogBeacon, Display, TEXT("\t Party leader: %s"), *NetId->ToString());
-		UE_LOG(LogBeacon, Display, TEXT("\t Party team: %d"), Reservations[PartyIndex].TeamNum);
-		UE_LOG(LogBeacon, Display, TEXT("\t Party size: %d"), Reservations[PartyIndex].PartyMembers.Num());
-		// Log each member of the party
-		for (int32 MemberIndex = 0; MemberIndex < Reservations[PartyIndex].PartyMembers.Num(); MemberIndex++)
-		{
-			PlayerRes = Reservations[PartyIndex].PartyMembers[MemberIndex];
-			UE_LOG(LogBeacon, Display, TEXT("\t  Party member: %s"), *PlayerRes.UniqueId->ToString());
-		}
+		State->DumpReservations();
 	}
 	UE_LOG(LogBeacon, Display, TEXT(""));
 }
diff --git a/Engine/Source/Runtime/Online/OnlineSubsystemUtils/Private/PartyBeaconState.cpp b/Engine/Source/Runtime/Online/OnlineSubsystemUtils/Private/PartyBeaconState.cpp
new file mode 100644
--- /dev/null
+++ b/Engine/Source/Runtime/Online/OnlineSubsystemUtils/Private/PartyBeaconState.cpp
@@ -0,0 +1,412 @@
+// Copyright 1998-2014 Epic Games, Inc. All Rights Reserved.
+
+#include "OnlineSubsystemUtilsPrivatePCH.h"
+#include "OnlineSubsystemUtilsClasses.h"
+
+
+UPartyBeaconState::UPartyBeaconState(const FPostConstructInitializeProperties& PCIP) :
+	Super(PCIP),
+	SessionName(NAME_None),
+	NumConsumedReservations(0),
+	MaxReservations(0)
+{
+}
+
+bool UPartyBeaconState::InitState(int32 InTeamCount, int32 InTeamSize, int32 InMaxReservations, FName InSessionName, int32 InForceTeamNum)
+{
+	if (InMaxReservations > 0)
+	{
+		SessionName = InSessionName;
+		NumTeams = InTeamCount;
+		NumPlayersPerTeam = InTeamSize;
+		MaxReservations = InMaxReservations;
+		ForceTeamNum = InForceTeamNum;
+		Reservations.Empty(MaxReservations);
+
+		InitTeamArray();
+		return true;
+	}
+
+	return false;
+}
+
+void UPartyBeaconState::InitTeamArray()
+{
+	if (NumTeams > 1)
+	{
+		// Grab one for the host team
+		ReservedHostTeamNum = FMath::Rand() % NumTeams;
+	}
+	else
+	{
+		// Only one team, so choose 'forced team' for everything
+		ReservedHostTeamNum = ForceTeamNum;
+	}
+
+	UE_LOG(LogBeacon, Display,
+		TEXT("Beacon State: team count (%d), team size (%d), host team (%d)"),
+		NumTeams,
+		NumPlayersPerTeam,
+		ReservedHostTeamNum);
+}
+
+bool UPartyBeaconState::ReconfigureTeamAndPlayerCount(int32 InNumTeams, int32 InNumPlayersPerTeam, int32 InNumReservations)
+{
+	bool bSuccess = false;
+
+	//Check total existing reservations against new total maximum
+	if (NumConsumedReservations < InNumReservations)
+	{
+		bool bTeamError = false;
+		// Check teams with reservations against new team count
+		if (NumTeams > InNumTeams)
+		{
+			// Any team about to be removed can't have players already there
+			for (int32 TeamIdx = InNumTeams; TeamIdx < NumTeams; TeamIdx++)
+			{
+				if (GetNumPlayersOnTeam(TeamIdx) > 0)
+				{
+					bTeamError = true;
+					UE_LOG(LogBeacon, Warning, TEXT("Beacon has players on a team about to be removed."));
+				}
+			}
+		}
+
+		bool bTeamSizeError = false;
+		// Check num players per team against new team size
+		if (NumPlayersPerTeam > InNumPlayersPerTeam)
+		{
+			for (int32 TeamIdx = 0; TeamIdx<NumTeams; TeamIdx++)
+			{
+				if (GetNumPlayersOnTeam(TeamIdx) > InNumPlayersPerTeam)
+				{
+					bTeamSizeError = true;
+					UE_LOG(LogBeacon, Warning, TEXT("Beacon has too many players on a team about to be resized."));
+				}
+			}
+		}
+
+		if (!bTeamError && !bTeamSizeError)
+		{
+			NumTeams = InNumTeams;
+			NumPlayersPerTeam = InNumPlayersPerTeam;
+			MaxReservations = InNumReservations;
+
+			InitTeamArray();
+			bSuccess = true;
+
+			UE_LOG(LogBeacon, Display,
+				TEXT("Reconfiguring to team count (%d), team size (%d)"),
+				NumTeams,
+				NumPlayersPerTeam);
+		}
+	}
+	else
+	{
+		UE_LOG(LogBeacon, Warning, TEXT("Beacon has too many consumed reservations for this reconfiguration, ignoring request."));
+	}
+
+	return bSuccess;
+}
+
+int32 UPartyBeaconState::GetNumPlayersOnTeam(int32 TeamIdx) const
+{
+	int32 Result = 0;
+	for (int32 ResIdx = 0; ResIdx < Reservations.Num(); ResIdx++)
+	{
+		const FPartyReservation& Reservation = Reservations[ResIdx];
+		if (Reservation.TeamNum == TeamIdx)
+		{
+			for (int32 PlayerIdx = 0; PlayerIdx < Reservation.PartyMembers.Num(); PlayerIdx++)
+			{
+				const FPlayerReservation& PlayerEntry = Reservation.PartyMembers[PlayerIdx];
+				// only count valid player net ids
+				if (PlayerEntry.UniqueId.IsValid())
+				{
+					// count party members in each team (includes party leader)
+					Result++;
+				}
+			}
+		}
+	}
+	return Result;
+}
+
+int32 UPartyBeaconState::GetTeamForCurrentPlayer(const FUniqueNetId& PlayerId) const
+{
+	int32 TeamNum = INDEX_NONE;
+	if (PlayerId.IsValid())
+	{
+		for (int32 ResIdx = 0; ResIdx < Reservations.Num(); ResIdx++)
+		{
+			const FPartyReservation& Reservation = Reservations[ResIdx];
+			for (int32 PlayerIdx = 0; PlayerIdx < Reservation.PartyMembers.Num(); PlayerIdx++)
+			{
+				// find the player id in the existing list of reservations
+				if (*Reservation.PartyMembers[PlayerIdx].UniqueId == PlayerId)
+				{
+					TeamNum = Reservation.TeamNum;
+					break;
+				}
+			}
+		}
+
+		UE_LOG(LogBeacon, Display, TEXT("Assigning player %s to team %d"),
+			*PlayerId.ToString(),
+			TeamNum);
+	}
+	else
+	{
+		UE_LOG(LogBeacon, Display, TEXT("Invalid player when attempting to find team assignment"));
+	}
+
+	return TeamNum;
+}
+
+
+
+/**
+* Helper for sorting team sizes
+*/
+struct FTeamBalanceInfo
+{
+	/** Index of team */
+	int32 TeamIdx;
+	/** Current size of team */
+	int32 TeamSize;
+
+	FTeamBalanceInfo(int32 InTeamIdx, int32 InTeamSize)
+		: TeamIdx(InTeamIdx),
+		TeamSize(InTeamSize)
+	{}
+};
+
+/**
+* Sort teams by size (equal teams are randomly mixed)
+*/
+struct FSortTeamSizeSmallestToLargest
+{
+	bool operator()(const FTeamBalanceInfo& A, const FTeamBalanceInfo& B) const
+	{
+		if (A.TeamSize == B.TeamSize)
+		{
+			return FMath::Rand() % 2 ? true : false;
+		}
+		else
+		{
+			return A.TeamSize < B.TeamSize;
+		}
+	}
+};
+
+int32 UPartyBeaconState::GetTeamAssignment(const FPartyReservation& Party)
+{
+	if (NumTeams > 1)
+	{
+		TArray<FTeamBalanceInfo> PotentialTeamChoices;
+		for (int32 TeamIdx = 0; TeamIdx < NumTeams; TeamIdx++)
+		{
+			const int32 CurrentPlayersOnTeam = GetNumPlayersOnTeam(TeamIdx);
+			if ((CurrentPlayersOnTeam + Party.PartyMembers.Num()) <= NumPlayersPerTeam)
+			{
+				new (PotentialTeamChoices)FTeamBalanceInfo(TeamIdx, CurrentPlayersOnTeam);
+			}
+		}
+
+		// Grab one from our list of choices
+		if (PotentialTeamChoices.Num() > 0)
+		{
+			if (1)
+			{
+				// Choose smallest team
+				PotentialTeamChoices.Sort(FSortTeamSizeSmallestToLargest());
+				return PotentialTeamChoices[0].TeamIdx;
+			}
+			else
+			{
+				// Random choice from set of choices
+				int32 TeamIndex = FMath::Rand() % PotentialTeamChoices.Num();
+				return PotentialTeamChoices[TeamIndex].TeamIdx;
+			}
+		}
+		else
+		{
+			UE_LOG(LogBeacon, Warning, TEXT("(UPartyBeaconHost.GetTeamAssignment): couldn't find an open team for party members."));
+			return INDEX_NONE;
+		}
+	}
+
+	return ForceTeamNum;
+}
+
+bool UPartyBeaconState::AreTeamsAvailable(const FPartyReservation& ReservationRequest) const
+{
+	int32 IncomingPartySize = ReservationRequest.PartyMembers.Num();
+	for (int32 TeamIdx = 0; TeamIdx < NumTeams; TeamIdx++)
+	{
+		const int32 CurrentPlayersOnTeam = GetNumPlayersOnTeam(TeamIdx);
+		if ((CurrentPlayersOnTeam + IncomingPartySize) <= NumPlayersPerTeam)
+		{
+			return true;
+		}
+	}
+	return false;
+}
+
+bool UPartyBeaconState::DoesReservationFit(const FPartyReservation& ReservationRequest) const
+{
+	int32 IncomingPartySize = ReservationRequest.PartyMembers.Num();
+	bool bPartySizeOk = (IncomingPartySize > 0) && (IncomingPartySize <= NumPlayersPerTeam);
+	bool bRoomForReservation = (NumConsumedReservations + IncomingPartySize) <= MaxReservations;
+
+	return bPartySizeOk && bRoomForReservation;
+}
+
+bool UPartyBeaconState::AddReservation(const FPartyReservation& ReservationRequest)
+{
+	int32 TeamAssignment = GetTeamAssignment(ReservationRequest);
+	if (TeamAssignment != INDEX_NONE)
+	{
+		int32 IncomingPartySize = ReservationRequest.PartyMembers.Num();
+
+		NumConsumedReservations += IncomingPartySize;
+		int32 ResIdx = Reservations.Add(ReservationRequest);
+		Reservations[ResIdx].TeamNum = TeamAssignment;
+	}
+
+	return TeamAssignment != INDEX_NONE;
+}
+
+bool UPartyBeaconState::RemoveReservation(const FUniqueNetIdRepl& PartyLeader)
+{
+	const int32 ExistingReservationIdx = GetExistingReservation(PartyLeader);
+	if (ExistingReservationIdx != INDEX_NONE)
+	{
+		NumConsumedReservations -= Reservations[ExistingReservationIdx].PartyMembers.Num();
+		Reservations.RemoveAtSwap(ExistingReservationIdx);
+
+		return true;
+	}
+
+	return false;
+}
+
+bool UPartyBeaconState::RemovePlayer(const FUniqueNetIdRepl& PlayerId)
+{
+	bool bWasRemoved = false;
+
+	for (int32 ResIdx = 0; ResIdx < Reservations.Num() && !bWasRemoved; ResIdx++)
+	{
+		FPartyReservation& Reservation = Reservations[ResIdx];
+
+		// find the player in an existing reservation slot
+		for (int32 PlayerIdx = 0; PlayerIdx < Reservation.PartyMembers.Num(); PlayerIdx++)
+		{
+			FPlayerReservation& PlayerEntry = Reservation.PartyMembers[PlayerIdx];
+			if (PlayerEntry.UniqueId == PlayerId)
+			{
+				// player removed
+				Reservation.PartyMembers.RemoveAtSwap(PlayerIdx--);
+				bWasRemoved = true;
+
+				// free up a consumed entry
+				NumConsumedReservations--;
+			}
+		}
+
+		// remove the entire party reservation slot if no more party members
+		if (Reservation.PartyMembers.Num() == 0)
+		{
+			Reservations.RemoveAtSwap(ResIdx--);
+		}
+	}
+
+	return bWasRemoved;
+}
+
+int32 UPartyBeaconState::GetExistingReservation(const FUniqueNetIdRepl& PartyLeader)
+{
+	int32 Result = INDEX_NONE;
+	for (int32 ResIdx = 0; ResIdx < Reservations.Num(); ResIdx++)
+	{
+		const FPartyReservation& ReservationEntry = Reservations[ResIdx];
+		if (ReservationEntry.PartyLeader == PartyLeader)
+		{
+			Result = ResIdx;
+			break;
+		}
+	}
+
+	return Result;
+}
+
+bool UPartyBeaconState::PlayerHasReservation(const FUniqueNetId& PlayerId) const
+{
+	bool bFound = false;
+
+	for (int32 ResIdx = 0; ResIdx < Reservations.Num(); ResIdx++)
+	{
+		const FPartyReservation& ReservationEntry = Reservations[ResIdx];
+		for (int32 PlayerIdx = 0; PlayerIdx < ReservationEntry.PartyMembers.Num(); PlayerIdx++)
+		{
+			if (*ReservationEntry.PartyMembers[PlayerIdx].UniqueId == PlayerId)
+			{
+				bFound = true;
+				break;
+			}
+		}
+	}
+
+	return bFound;
+}
+
+bool UPartyBeaconState::GetPlayerValidation(const FUniqueNetId& PlayerId, FString& OutValidation) const
+{
+	bool bFound = false;
+	OutValidation = FString();
+
+	for (int32 ResIdx = 0; ResIdx < Reservations.Num(); ResIdx++)
+	{
+		const FPartyReservation& ReservationEntry = Reservations[ResIdx];
+		for (int32 PlayerIdx = 0; PlayerIdx < ReservationEntry.PartyMembers.Num(); PlayerIdx++)
+		{
+			if (*ReservationEntry.PartyMembers[PlayerIdx].UniqueId == PlayerId)
+			{
+				OutValidation = ReservationEntry.PartyMembers[PlayerIdx].ValidationStr;
+				bFound = true;
+				break;
+			}
+		}
+	}
+
+	return bFound;
+}
+
+void UPartyBeaconState::DumpReservations() const
+{
+	FUniqueNetIdRepl NetId;
+	FPlayerReservation PlayerRes;
+
+	UE_LOG(LogBeacon, Display, TEXT("Session that reservations are for: %s"), *SessionName.ToString());
+	UE_LOG(LogBeacon, Display, TEXT("Number of teams: %d"), NumTeams);
+	UE_LOG(LogBeacon, Display, TEXT("Number players per team: %d"), NumPlayersPerTeam);
+	UE_LOG(LogBeacon, Display, TEXT("Number total reservations: %d"), MaxReservations);
+	UE_LOG(LogBeacon, Display, TEXT("Number consumed reservations: %d"), NumConsumedReservations);
+	UE_LOG(LogBeacon, Display, TEXT("Number of party reservations: %d"), Reservations.Num());
+
+	// Log each party that has a reservation
+	for (int32 PartyIndex = 0; PartyIndex < Reservations.Num(); PartyIndex++)
+	{
+		NetId = Reservations[PartyIndex].PartyLeader;
+		UE_LOG(LogBeacon, Display, TEXT("\t Party leader: %s"), *NetId->ToString());
+		UE_LOG(LogBeacon, Display, TEXT("\t Party team: %d"), Reservations[PartyIndex].TeamNum);
+		UE_LOG(LogBeacon, Display, TEXT("\t Party size: %d"), Reservations[PartyIndex].PartyMembers.Num());
+		// Log each member of the party
+		for (int32 MemberIndex = 0; MemberIndex < Reservations[PartyIndex].PartyMembers.Num(); MemberIndex++)
+		{
+			PlayerRes = Reservations[PartyIndex].PartyMembers[MemberIndex];
+			UE_LOG(LogBeacon, Display, TEXT("\t  Party member: %s"), *PlayerRes.UniqueId->ToString());
+		}
+	}
+	UE_LOG(LogBeacon, Display, TEXT(""));
+}
