#UE4 - crash fix for bad world context handle when checking IsServer

diff --git a/Engine/Source/Runtime/Launch/Private/LaunchEngineLoop.cpp b/Engine/Source/Runtime/Launch/Private/LaunchEngineLoop.cpp
--- a/Engine/Source/Runtime/Launch/Private/LaunchEngineLoop.cpp
+++ b/Engine/Source/Runtime/Launch/Private/LaunchEngineLoop.cpp
@@ -534,7 +534,7 @@ bool IsServerDelegateForOSS(FName WorldContextHandle)
 	}
 	else
 	{
-		UGameEngine* GameEngine = Cast<UGameEngine>(GEngine);
+		UGameEngine* GameEngine = CastChecked<UGameEngine>(GEngine);
 		World = GameEngine->GetGameWorld();
 	}
 
diff --git a/Engine/Source/Runtime/Online/OnlineSubsystem/Private/OnlineSubsystemImpl.cpp b/Engine/Source/Runtime/Online/OnlineSubsystem/Private/OnlineSubsystemImpl.cpp
--- a/Engine/Source/Runtime/Online/OnlineSubsystem/Private/OnlineSubsystemImpl.cpp
+++ b/Engine/Source/Runtime/Online/OnlineSubsystem/Private/OnlineSubsystemImpl.cpp
@@ -68,7 +68,8 @@ void FOnlineSubsystemImpl::SetNamedInterface(FName InterfaceName, UObject* NewIn
 
 bool FOnlineSubsystemImpl::IsServer() const
 {
-	return IsServerForOnlineSubsystems(InstanceName);
+	FName WorldContextHandle = (InstanceName != NAME_None && InstanceName != DEFAULT_INSTANCE) ? InstanceName : NAME_None;
+	return IsServerForOnlineSubsystems(WorldContextHandle);
 }
 
 bool FOnlineSubsystemImpl::IsLocalPlayer(const FUniqueNetId& UniqueId) const
