#ue4
- Added some checks to try to catch intermittent crash on iOS when reading the leaderboard

diff --git a/Engine/Source/Runtime/Online/IOS/OnlineSubsystemIOS/Private/OnlineLeaderboardsInterfaceIOS.cpp b/Engine/Source/Runtime/Online/IOS/OnlineSubsystemIOS/Private/OnlineLeaderboardsInterfaceIOS.cpp
--- a/Engine/Source/Runtime/Online/IOS/OnlineSubsystemIOS/Private/OnlineLeaderboardsInterfaceIOS.cpp
+++ b/Engine/Source/Runtime/Online/IOS/OnlineSubsystemIOS/Private/OnlineLeaderboardsInterfaceIOS.cpp
@@ -14,7 +14,10 @@ FOnlineLeaderboardsIOS::FOnlineLeaderboardsIOS(FOnlineSubsystemIOS* InSubsystem)
 
 	// Cache a reference to the OSS Identity and Friends interfaces, we will need these when we are performing leaderboard actions
 	IdentityInterface = (FOnlineIdentityIOS*)InSubsystem->GetIdentityInterface().Get();
+	check(IdentityInterface);
+
 	FriendsInterface = (FOnlineFriendsIOS*)InSubsystem->GetFriendsInterface().Get();
+	check(FriendsInterface);
 
 	UnreportedScores = nil;
 }
@@ -33,13 +36,13 @@ bool FOnlineLeaderboardsIOS::ReadLeaderboards(const TArray< TSharedRef<FUniqueNe
 {
 	auto ReadObject = InReadObject;
 
-	UE_LOG(LogOnline, Display, TEXT("FOnlineLeaderboardsIOS::ReadLeaderboards() %p"), &ReadObject);
+	UE_LOG(LogOnline, Display, TEXT("FOnlineLeaderboardsIOS::ReadLeaderboards()"));
 	bool bTriggeredReadRequest = false;
 
 	ReadObject->ReadState = EOnlineAsyncTaskState::Failed;
 	ReadObject->Rows.Empty();
 	
-	if ((IdentityInterface->GetLocalGameCenterUser() != NULL) && IdentityInterface->GetLocalGameCenterUser().isAuthenticated)
+	if ((IdentityInterface != nullptr) && (IdentityInterface->GetLocalGameCenterUser() != NULL) && IdentityInterface->GetLocalGameCenterUser().isAuthenticated)
 	{
 		ReadObject->ReadState = EOnlineAsyncTaskState::InProgress;
 
@@ -47,8 +50,12 @@ bool FOnlineLeaderboardsIOS::ReadLeaderboards(const TArray< TSharedRef<FUniqueNe
 		NSMutableArray* FriendIds = [NSMutableArray arrayWithCapacity: (Players.Num() + 1)];
 		
 		// Add the local player to the list of ids to look up.
-		FriendIds[0] = [NSString stringWithFString:IdentityInterface->GetUniquePlayerId(0)->ToString()];
-		
+		TSharedPtr<FUniqueNetId> LocalPlayerUID = IdentityInterface->GetUniquePlayerId(0);
+		check(LocalPlayerUID.IsValid());
+
+		FriendIds[0] = [NSString stringWithFString:LocalPlayerUID->ToString()];
+
+		// Add the other requested players
 		for (int32 FriendIdx = 0; FriendIdx < Players.Num(); FriendIdx++)
 		{
 			FriendIds[FriendIdx + 1] = [NSString stringWithFString:Players[FriendIdx]->ToString()];
@@ -61,7 +68,7 @@ bool FOnlineLeaderboardsIOS::ReadLeaderboards(const TArray< TSharedRef<FUniqueNe
 			const FString LeaderboardName = ReadObject->LeaderboardName.ToString();
 
 			NSString* Category = [NSString stringWithFString:LeaderboardName];
-			UE_LOG(LogOnline, Display, TEXT("Attempting to read leaderboard: %s"),*LeaderboardName);
+			UE_LOG(LogOnline, Display, TEXT("Attempting to read leaderboard: %s"), *LeaderboardName);
 
 			LeaderboardRequest.playerScope = GKLeaderboardPlayerScopeGlobal;
 			LeaderboardRequest.timeScope = GKLeaderboardTimeScopeToday;
