Fix OSS hacks in blueprint-callable functions

Added ShowLeaderboardUI function to IOnlineExternalUI, blueprint-callable functions now use this to show UIs
Added ExternalUIInterface for iOS Game Center, blueprint-callable functions now use it instead of platform-specific hacks
Added leaderboard interface for Android/Google Play, blueprint-callable functions now use this instead of hacks
Expose the Google Play leaderboard ID mapping to the editor instead of being hidden in a random JSON file
Fix the TappyChicken leaderboard mapping and removed the JSON file

diff --git a/Engine/Source/Runtime/Android/AndroidRuntimeSettings/Classes/AndroidRuntimeSettings.h b/Engine/Source/Runtime/Android/AndroidRuntimeSettings/Classes/AndroidRuntimeSettings.h
--- a/Engine/Source/Runtime/Android/AndroidRuntimeSettings/Classes/AndroidRuntimeSettings.h
+++ b/Engine/Source/Runtime/Android/AndroidRuntimeSettings/Classes/AndroidRuntimeSettings.h
@@ -52,6 +52,23 @@ struct FGooglePlayAchievementMapping
 	FString AchievementID;
 };
 
+/**
+ * Holds the game-specific leaderboard name and corresponding ID from Google Play services.
+ */
+USTRUCT()
+struct FGooglePlayLeaderboardMapping
+{
+	GENERATED_USTRUCT_BODY()
+
+	// The game-specific leaderboard name (the one passed in to WriteLeaderboards calls)
+	UPROPERTY(EditAnywhere, Category = GooglePlayServices)
+	FString Name;
+
+	// The ID of the corresponding leaderboard, generated by the Google Play developer console.
+	UPROPERTY(EditAnywhere, Category = GooglePlayServices)
+	FString LeaderboardID;
+};
+
 /**
  * Implements the settings for the Android runtime platform.
  */
@@ -76,4 +93,8 @@ public:
 	// Mapping of game achievement names to IDs generated by Google Play.
 	UPROPERTY(GlobalConfig, EditAnywhere, Category = GooglePlayServices)
 	TArray<FGooglePlayAchievementMapping> AchievementMap;
+
+	// Mapping of game leaderboard names to IDs generated by Google Play.
+	UPROPERTY(GlobalConfig, EditAnywhere, Category = GooglePlayServices)
+	TArray<FGooglePlayLeaderboardMapping> LeaderboardMap;
 };
diff --git a/Engine/Source/Runtime/Engine/Private/KismetSystemLibrary.cpp b/Engine/Source/Runtime/Engine/Private/KismetSystemLibrary.cpp
--- a/Engine/Source/Runtime/Engine/Private/KismetSystemLibrary.cpp
+++ b/Engine/Source/Runtime/Engine/Private/KismetSystemLibrary.cpp
@@ -2280,24 +2280,15 @@ void UKismetSystemLibrary::EXPERIMENTAL_CloseAdBanner()
 
 void UKismetSystemLibrary::ShowPlatformSpecificLeaderboardScreen(const FString& CategoryName)
 {
-#if PLATFORM_IOS
-	extern CORE_API void IOSShowLeaderboardUI(const FString& CategoryName);
-	IOSShowLeaderboardUI(CategoryName);
-#elif PLATFORM_ANDROID
-	UE_LOG(LogTemp, Log, TEXT("Showing leaderboard!"));
-
-	extern void AndroidThunkCpp_ShowLeaderboard(const FString&);
-	extern FString GetLeaderboardID(const FString&);
-	AndroidThunkCpp_ShowLeaderboard(GetLeaderboardID(CategoryName));
-#endif
+	IOnlineExternalUIPtr ExternalUI = Online::GetExternalUIInterface();
+	if(ExternalUI.IsValid())
+	{
+		ExternalUI->ShowLeaderboardUI(CategoryName);
+	}
 }
 
 void UKismetSystemLibrary::ShowPlatformSpecificAchievementsScreen(class APlayerController* SpecificPlayer)
 {
-#if PLATFORM_IOS
-	extern CORE_API void IOSShowAchievementsUI();
-	IOSShowAchievementsUI();
-#else
 	IOnlineExternalUIPtr ExternalUI = Online::GetExternalUIInterface();
 	if(ExternalUI.IsValid())
 	{
@@ -2313,7 +2304,6 @@ void UKismetSystemLibrary::ShowPlatformSpecificAchievementsScreen(class APlayerC
 		}
 		ExternalUI->ShowAchievementsUI(LocalUserNum);
 	}
-#endif
 }
 void UKismetSystemLibrary::SetStructurePropertyByName(UObject* Object, FName PropertyName, const FGenericStruct& Value)
 {
diff --git a/Engine/Source/Runtime/Online/Android/OnlineSubsystemGooglePlay/Private/OnlineAchievementsInterfaceGooglePlay.cpp b/Engine/Source/Runtime/Online/Android/OnlineSubsystemGooglePlay/Private/OnlineAchievementsInterfaceGooglePlay.cpp
--- a/Engine/Source/Runtime/Online/Android/OnlineSubsystemGooglePlay/Private/OnlineAchievementsInterfaceGooglePlay.cpp
+++ b/Engine/Source/Runtime/Online/Android/OnlineSubsystemGooglePlay/Private/OnlineAchievementsInterfaceGooglePlay.cpp
@@ -1,56 +1,15 @@
 // Copyright 1998-2014 Epic Games, Inc. All Rights Reserved.
 
-// Module includes
 #include "OnlineSubsystemGooglePlayPrivatePCH.h"
 
-#include <jni.h>
-
 FOnlineAchievementsGooglePlay::FOnlineAchievementsGooglePlay( FOnlineSubsystemGooglePlay* InSubsystem )
-	: CurrentWriteObject(MakeShareable(new FOnlineAchievementsWrite()))
+	: AndroidSubsystem(InSubsystem)
 {
-	AndroidSubsystem = InSubsystem;
-	CurrentPlayerID = MakeShareable( new FUniqueNetIdString(FString("")));
-	
-	TArray<FString> MapAsStrings;
-
-	// Use the config to find the Google Play-assigned achievement ID corresponding to the name the game passed in.
-	GConfig->GetArray(TEXT("/Script/AndroidRuntimeSettings.AndroidRuntimeSettings"), TEXT("AchievementMap"), MapAsStrings, GEngineIni);
-
-	for(int32 x = 0; x < MapAsStrings.Num(); ++x)
-	{
-		FGooglePlayAchievementMapping Mapping;
-
-		if(!FParse::Value(*MapAsStrings[x], TEXT("Name="), Mapping.Name))
-		{
-			// Could not find an achievement name, cannot continue with this line
-			continue;
-		}
-
-		if(!FParse::Value(*MapAsStrings[x], TEXT("AchievementID="), Mapping.AchievementID))
-		{
-			// Could not find an achievement id, cannot continue with this line
-			continue;
-		}
-
-		UE_LOG(LogOnline, Log, TEXT("Found achievement mapping: %s to %s"), *Mapping.Name, *Mapping.AchievementID);
-		AchievementConfigMap.Add(Mapping);
-	}
-}
-
-void FOnlineAchievementsGooglePlay::SetUniqueID(const FUniqueNetId& InUserID)
-{
-	CurrentPlayerID = MakeShareable( new FUniqueNetIdString(InUserID));
-}
-
-void FOnlineAchievementsGooglePlay::ResetUniqueID()
-{
-	CurrentPlayerID.Reset();
 }
 
 void FOnlineAchievementsGooglePlay::QueryAchievements(const FUniqueNetId& PlayerId, const FOnQueryAchievementsCompleteDelegate& Delegate)
 {
 	Delegate.ExecuteIfBound( PlayerId, false );
-	return;
 }
 
 
@@ -97,7 +56,9 @@ void FOnlineAchievementsGooglePlay::WriteAchievements( const FUniqueNetId& Playe
 		}
 
 		// Find the ID in the mapping
-		for(const auto& Mapping : AchievementConfigMap)
+		auto DefaultSettings = GetDefault<UAndroidRuntimeSettings>();
+
+		for(const auto& Mapping : DefaultSettings->AchievementMap)
 		{
 			if(Mapping.Name == AchievementName)
 			{
diff --git a/Engine/Source/Runtime/Online/Android/OnlineSubsystemGooglePlay/Private/OnlineAchievementsInterfaceGooglePlay.h b/Engine/Source/Runtime/Online/Android/OnlineSubsystemGooglePlay/Private/OnlineAchievementsInterfaceGooglePlay.h
--- a/Engine/Source/Runtime/Online/Android/OnlineSubsystemGooglePlay/Private/OnlineAchievementsInterfaceGooglePlay.h
+++ b/Engine/Source/Runtime/Online/Android/OnlineSubsystemGooglePlay/Private/OnlineAchievementsInterfaceGooglePlay.h
@@ -17,13 +17,12 @@ private:
 	class FOnlineSubsystemGooglePlay* AndroidSubsystem;
 	
 	/** hide the default constructor, we need a reference to our OSS */
-	FOnlineAchievementsGooglePlay() : CurrentWriteObject(MakeShareable(new FOnlineAchievementsWrite()))
-	{}
+	FOnlineAchievementsGooglePlay() {}
 
-	// Android only supports loading achievements for local player. This is where they are cached.
+	/** Android only supports loading achievements for local player. This is where they are cached. */
 	TArray< FOnlineAchievement > Achievements;
 
-	// Cached achievement descriptions for an Id
+	/** Cached achievement descriptions for an Id */
 	TMap< FString, FOnlineAchievementDesc > AchievementDescriptions;
 
 public:
@@ -39,14 +38,6 @@ public:
 	virtual bool ResetAchievements( const FUniqueNetId& PlayerId ) OVERRIDE;
 #endif // !UE_BUILD_SHIPPING
 	// End IOnlineAchievements interface
-	
-	//platform specific members
-
-	/** Sets unique ID of current logged-in player */
-	void SetUniqueID(const FUniqueNetId& InUserID);
-	
-	/** Reset unique ID of current logged-in player */
-	void ResetUniqueID();
 
 	/**
 	 * Constructor
@@ -54,20 +45,6 @@ public:
 	 * @param InSubsystem - A reference to the owning subsystem
 	 */
 	FOnlineAchievementsGooglePlay( class FOnlineSubsystemGooglePlay* InSubsystem );
-
-	/**
-	 * Default destructor
-	 */
-	virtual ~FOnlineAchievementsGooglePlay() {}
-
-
-private:
-
-	TArray<FGooglePlayAchievementMapping> AchievementConfigMap;
-
-	TSharedPtr< FUniqueNetIdString > CurrentPlayerID;
-
-	TSharedRef<FOnlineAchievementsWrite> CurrentWriteObject;
 };
 
 typedef TSharedPtr<FOnlineAchievementsGooglePlay, ESPMode::ThreadSafe> FOnlineAchievementsGooglePlayPtr;
diff --git a/Engine/Source/Runtime/Online/Android/OnlineSubsystemGooglePlay/Private/OnlineExternalUIInterfaceGooglePlay.cpp b/Engine/Source/Runtime/Online/Android/OnlineSubsystemGooglePlay/Private/OnlineExternalUIInterfaceGooglePlay.cpp
--- a/Engine/Source/Runtime/Online/Android/OnlineSubsystemGooglePlay/Private/OnlineExternalUIInterfaceGooglePlay.cpp
+++ b/Engine/Source/Runtime/Online/Android/OnlineSubsystemGooglePlay/Private/OnlineExternalUIInterfaceGooglePlay.cpp
@@ -1,6 +1,8 @@
 // Copyright 1998-2014 Epic Games, Inc. All Rights Reserved.
 
 #include "OnlineSubsystemGooglePlayPrivatePCH.h"
+#include "AndroidRuntimeSettings.h"
+#include "CoreUObject.h"
 
 FOnlineExternalUIGooglePlay::FOnlineExternalUIGooglePlay() 
 {
@@ -30,6 +32,23 @@ bool FOnlineExternalUIGooglePlay::ShowAchievementsUI(int32 LocalUserNum)
 	return true;
 }
 
+bool FOnlineExternalUIGooglePlay::ShowLeaderboardUI(const FString& LeaderboardName)
+{
+	auto DefaultSettings = GetDefault<UAndroidRuntimeSettings>();
+
+	for(const auto& Mapping : DefaultSettings->LeaderboardMap)
+	{
+		if(Mapping.Name == LeaderboardName)
+		{
+			extern void AndroidThunkCpp_ShowLeaderboard(const FString&);
+			AndroidThunkCpp_ShowLeaderboard(Mapping.LeaderboardID);
+			return true;
+		}
+	}
+	
+	return false;
+}
+
 bool FOnlineExternalUIGooglePlay::ShowWebURL(const FString& WebURL) 
 {
 	return false;
diff --git a/Engine/Source/Runtime/Online/Android/OnlineSubsystemGooglePlay/Private/OnlineExternalUIInterfaceGooglePlay.h b/Engine/Source/Runtime/Online/Android/OnlineSubsystemGooglePlay/Private/OnlineExternalUIInterfaceGooglePlay.h
--- a/Engine/Source/Runtime/Online/Android/OnlineSubsystemGooglePlay/Private/OnlineExternalUIInterfaceGooglePlay.h
+++ b/Engine/Source/Runtime/Online/Android/OnlineSubsystemGooglePlay/Private/OnlineExternalUIInterfaceGooglePlay.h
@@ -21,6 +21,7 @@ public:
 	virtual bool ShowFriendsUI(int32 LocalUserNum) OVERRIDE;
 	virtual bool ShowInviteUI(int32 LocalUserNum) OVERRIDE;
 	virtual bool ShowAchievementsUI(int32 LocalUserNum) OVERRIDE;
+	virtual bool ShowLeaderboardUI(const FString& LeaderboardName) OVERRIDE;
 	virtual bool ShowWebURL(const FString& WebURL) OVERRIDE;
 	virtual bool ShowProfileUI(const FUniqueNetId& Requestor, const FUniqueNetId& Requestee, const FOnProfileUIClosedDelegate& Delegate) OVERRIDE;
 };
diff --git a/Engine/Source/Runtime/Online/Android/OnlineSubsystemGooglePlay/Private/OnlineLeaderboardInterfaceGooglePlay.cpp b/Engine/Source/Runtime/Online/Android/OnlineSubsystemGooglePlay/Private/OnlineLeaderboardInterfaceGooglePlay.cpp
new file mode 100644
--- /dev/null
+++ b/Engine/Source/Runtime/Online/Android/OnlineSubsystemGooglePlay/Private/OnlineLeaderboardInterfaceGooglePlay.cpp
@@ -0,0 +1,124 @@
+// Copyright 1998-2014 Epic Games, Inc. All Rights Reserved.
+
+#include "OnlineSubsystemGooglePlayPrivatePCH.h"
+#include "OnlineLeaderboardInterfaceGooglePlay.h"
+#include "Online.h"
+
+FOnlineLeaderboardsGooglePlay::FOnlineLeaderboardsGooglePlay(FOnlineSubsystemGooglePlay* InSubsystem)
+{
+}
+
+bool FOnlineLeaderboardsGooglePlay::ReadLeaderboards(const TArray< TSharedRef<FUniqueNetId> >& Players, FOnlineLeaderboardReadRef& ReadObject)
+{
+	TriggerOnLeaderboardReadCompleteDelegates(false);
+	return false;
+}
+
+bool FOnlineLeaderboardsGooglePlay::ReadLeaderboardsForFriends(int32 LocalUserNum, FOnlineLeaderboardReadRef& ReadObject)
+{
+	TriggerOnLeaderboardReadCompleteDelegates(false);
+	return false;
+}
+
+void FOnlineLeaderboardsGooglePlay::FreeStats(FOnlineLeaderboardRead& ReadObject)
+{
+	// iOS doesn't support this, and there is no Google Play functionality for this either
+}
+
+bool FOnlineLeaderboardsGooglePlay::WriteLeaderboards(const FName& SessionName, const FUniqueNetId& Player, FOnlineLeaderboardWrite& WriteObject)
+{
+	UE_LOG_ONLINE(Display, TEXT("WriteLeaderboards"));
+
+	bool bWroteAnyLeaderboard = false;
+
+	for(int32 LeaderboardIdx = 0; LeaderboardIdx < WriteObject.LeaderboardNames.Num(); ++LeaderboardIdx)
+	{
+		FString LeaderboardName = WriteObject.LeaderboardNames[LeaderboardIdx].ToString();
+		if(LeaderboardName.Equals(TEXT("TestLeaderboard")))
+		{
+			LeaderboardName = TEXT("leaderboard_00"); //hack to work around leaderboard name mismatch between test module and the format that java call expects
+		}
+		UE_LOG_ONLINE(Display, TEXT("Going through stats for leaderboard :  %s "), *LeaderboardName);
+		
+		for(FStatPropertyArray::TConstIterator It(WriteObject.Properties); It; ++It)
+		{
+			const FVariantData& Stat = It.Value();
+			uint64 Score;
+			FOnlinePendingLeaderboardWrite UnreportedScore;
+
+			UE_LOG_ONLINE(Display, TEXT("Here's a stat"));
+
+			//Google leaderboard stats are always a long/int64
+			if(Stat.GetType() == EOnlineKeyValuePairDataType::Int64)
+			{
+				Stat.GetValue(Score);
+				
+				FOnlinePendingLeaderboardWrite* UnreportedScore = new (UnreportedScores) FOnlinePendingLeaderboardWrite();
+				UnreportedScore->LeaderboardName = LeaderboardName;
+				UnreportedScore->Score = Score;
+				UE_LOG_ONLINE(Display, TEXT("FOnlineLeaderboardsAndroid::WriteLeaderboards() Int64 value Score: %d"), Score);
+
+				bWroteAnyLeaderboard = true;
+			}
+			else if (Stat.GetType() == EOnlineKeyValuePairDataType::Int32)
+			{
+				// cast from 32 to 64
+				int32 Score32 = 0;
+				Stat.GetValue(Score32);
+				Score = static_cast<uint64>(Score32);
+
+				FOnlinePendingLeaderboardWrite* UnreportedScore = new (UnreportedScores) FOnlinePendingLeaderboardWrite();
+				UnreportedScore->LeaderboardName = LeaderboardName;
+				UnreportedScore->Score = Score;
+				UE_LOG_ONLINE(Display, TEXT("FOnlineLeaderboardsAndroid::WriteLeaderboards() Int32 value Score: %d "), Score);
+
+				bWroteAnyLeaderboard = true;
+			}
+		}
+	}
+	
+	//Return whether any stat was cached
+	return bWroteAnyLeaderboard;
+}
+
+bool FOnlineLeaderboardsGooglePlay::FlushLeaderboards(const FName& SessionName)
+{
+	UE_LOG_ONLINE(Display, TEXT("flush leaderboards session name :%s"), *SessionName.ToString());
+
+	bool Success = true;
+	//@todo android: make ansynchronous?
+	for(int32 Index = 0; Index < UnreportedScores.Num(); ++Index)
+	{
+		UE_LOG_ONLINE(Display, TEXT("Submitting an unreported score to %s. Value: %d "), *UnreportedScores[Index].LeaderboardName);
+
+		extern void AndroidThunkCpp_WriteLeaderboardValue(const FString&, int64_t);
+		AndroidThunkCpp_WriteLeaderboardValue(GetLeaderboardID(UnreportedScores[Index].LeaderboardName), UnreportedScores[Index].Score);
+	}
+
+	UnreportedScores.Empty();
+
+	TriggerOnLeaderboardFlushCompleteDelegates(SessionName, Success);
+
+	return true;
+}
+
+bool FOnlineLeaderboardsGooglePlay::WriteOnlinePlayerRatings(const FName& SessionName, int32 LeaderboardId, const TArray<FOnlinePlayerScore>& PlayerScores)
+{
+	//iOS doesn't support this, and there is no Google Play functionality for this either
+	return false;
+}
+
+FString FOnlineLeaderboardsGooglePlay::GetLeaderboardID(const FString& LeaderboardName)
+{
+	auto DefaultSettings = GetDefault<UAndroidRuntimeSettings>();
+	for(const auto& Mapping : DefaultSettings->LeaderboardMap)
+	{
+		if(Mapping.Name == LeaderboardName)
+		{
+			return Mapping.LeaderboardID;
+		}
+	}
+
+	UE_LOG(LogOnline, Warning, TEXT( "GetLeaderboardID: No mapping for leaderboard %s"), *LeaderboardName );
+	return LeaderboardName;
+}
diff --git a/Engine/Source/Runtime/Online/Android/OnlineSubsystemGooglePlay/Private/OnlineLeaderboardInterfaceGooglePlay.h b/Engine/Source/Runtime/Online/Android/OnlineSubsystemGooglePlay/Private/OnlineLeaderboardInterfaceGooglePlay.h
new file mode 100644
--- /dev/null
+++ b/Engine/Source/Runtime/Online/Android/OnlineSubsystemGooglePlay/Private/OnlineLeaderboardInterfaceGooglePlay.h
@@ -0,0 +1,42 @@
+// Copyright 1998-2014 Epic Games, Inc. All Rights Reserved.
+
+#pragma once
+
+#include "OnlineLeaderboardInterface.h"
+#include "AndroidRuntimeSettings.h"
+
+struct FOnlinePendingLeaderboardWrite
+{
+	FString LeaderboardName;
+	uint64 Score;
+};
+
+/**
+ * Interface definition for the online services leaderboard services 
+ */
+class FOnlineLeaderboardsGooglePlay : public IOnlineLeaderboards
+{
+public:
+	FOnlineLeaderboardsGooglePlay(class FOnlineSubsystemGooglePlay* InSubsystem);
+
+	virtual bool ReadLeaderboards(const TArray< TSharedRef<FUniqueNetId> >& Players, FOnlineLeaderboardReadRef& ReadObject) OVERRIDE;
+	virtual bool ReadLeaderboardsForFriends(int32 LocalUserNum, FOnlineLeaderboardReadRef& ReadObject) OVERRIDE;
+	virtual void FreeStats(FOnlineLeaderboardRead& ReadObject) OVERRIDE;
+	virtual bool WriteLeaderboards(const FName& SessionName, const FUniqueNetId& Player, FOnlineLeaderboardWrite& WriteObject) OVERRIDE;
+	virtual bool FlushLeaderboards(const FName& SessionName) OVERRIDE;
+	virtual bool WriteOnlinePlayerRatings(const FName& SessionName, int32 LeaderboardId, const TArray<FOnlinePlayerScore>& PlayerScores) OVERRIDE;
+
+private:
+	/**
+	 * Helper function to get the platform- and game-specific leaderboard ID from the JSON config file.
+	 *
+	 * @param LeaderboardName the cross-platform name of the leaderboard to look up
+	 * @return The unique ID for the leaderboard as specified in the config file.
+	*/
+	FString GetLeaderboardID(const FString& LeaderboardName);
+
+	/** Scores are cached here in WriteLeaderboards until FlushLeaderboards is called */
+	TArray<FOnlinePendingLeaderboardWrite> UnreportedScores;
+};
+
+typedef TSharedPtr<FOnlineLeaderboardsGooglePlay, ESPMode::ThreadSafe> FOnlineLeaderboardsGooglePlayPtr;
diff --git a/Engine/Source/Runtime/Online/Android/OnlineSubsystemGooglePlay/Private/OnlineSubsystemGooglePlay.cpp b/Engine/Source/Runtime/Online/Android/OnlineSubsystemGooglePlay/Private/OnlineSubsystemGooglePlay.cpp
--- a/Engine/Source/Runtime/Online/Android/OnlineSubsystemGooglePlay/Private/OnlineSubsystemGooglePlay.cpp
+++ b/Engine/Source/Runtime/Online/Android/OnlineSubsystemGooglePlay/Private/OnlineSubsystemGooglePlay.cpp
@@ -6,9 +6,10 @@
 #include "AndroidApplication.h"
 
 FOnlineSubsystemGooglePlay::FOnlineSubsystemGooglePlay()
-	: OnlineAsyncTaskThread(NULL)
-	, IdentityInterface(NULL)
-	, AchievementsInterface(NULL)
+	: OnlineAsyncTaskThread(nullptr)
+	, IdentityInterface(nullptr)
+	, LeaderboardsInterface(nullptr)
+	, AchievementsInterface(nullptr)
 {
 
 }
@@ -50,7 +51,7 @@ IOnlineUserCloudPtr FOnlineSubsystemGooglePlay::GetUserCloudInterface() const
 
 IOnlineLeaderboardsPtr FOnlineSubsystemGooglePlay::GetLeaderboardsInterface() const
 {
-	return nullptr;
+	return LeaderboardsInterface;
 }
 
 
@@ -94,6 +95,7 @@ bool FOnlineSubsystemGooglePlay::Init()
 	FPlatformMisc::LowLevelOutputDebugStringf(TEXT("FOnlineSubsystemAndroid::Init"));
 	
 	IdentityInterface = MakeShareable(new FOnlineIdentityGooglePlay(this));
+	LeaderboardsInterface = MakeShareable(new FOnlineLeaderboardsGooglePlay(this));
 	AchievementsInterface = MakeShareable(new FOnlineAchievementsGooglePlay(this));
 	ExternalUIInterface = MakeShareable(new FOnlineExternalUIGooglePlay());
 
@@ -123,6 +125,7 @@ bool FOnlineSubsystemGooglePlay::Shutdown()
 	// Destruct the interfaces
 	DESTRUCT_INTERFACE(ExternalUIInterface);
 	DESTRUCT_INTERFACE(AchievementsInterface);
+	DESTRUCT_INTERFACE(LeaderboardsInterface);
 	DESTRUCT_INTERFACE(IdentityInterface);
 #undef DESTRUCT_INTERFACE
 
diff --git a/Engine/Source/Runtime/Online/Android/OnlineSubsystemGooglePlay/Public/OnlineSubsystemGooglePlay.h b/Engine/Source/Runtime/Online/Android/OnlineSubsystemGooglePlay/Public/OnlineSubsystemGooglePlay.h
--- a/Engine/Source/Runtime/Online/Android/OnlineSubsystemGooglePlay/Public/OnlineSubsystemGooglePlay.h
+++ b/Engine/Source/Runtime/Online/Android/OnlineSubsystemGooglePlay/Public/OnlineSubsystemGooglePlay.h
@@ -5,6 +5,7 @@
 #include "OnlineSubsystem.h"
 #include "OnlineIdentityInterfaceGooglePlay.h"
 #include "OnlineAchievementsInterfaceGooglePlay.h"
+#include "OnlineLeaderboardInterfaceGooglePlay.h"
 #include "OnlineExternalUIInterfaceGooglePlay.h"
 
 /**
@@ -68,6 +69,9 @@ private:
 	/** Interface to the online identity system */
 	FOnlineIdentityGooglePlayPtr IdentityInterface;
 
+	/** Interface to the online leaderboards */
+	FOnlineLeaderboardsGooglePlayPtr LeaderboardsInterface;
+
 	/** Interface to the online achievements */
 	FOnlineAchievementsGooglePlayPtr AchievementsInterface;
 
diff --git a/Engine/Source/Runtime/Online/IOS/OnlineSubsystemIOS/Private/OnlineExternalUIInterfaceIOS.cpp b/Engine/Source/Runtime/Online/IOS/OnlineSubsystemIOS/Private/OnlineExternalUIInterfaceIOS.cpp
new file mode 100644
--- /dev/null
+++ b/Engine/Source/Runtime/Online/IOS/OnlineSubsystemIOS/Private/OnlineExternalUIInterfaceIOS.cpp
@@ -0,0 +1,48 @@
+// Copyright 1998-2014 Epic Games, Inc. All Rights Reserved.
+
+// Module includes
+#include "OnlineSubsystemIOSPrivatePCH.h"
+
+FOnlineExternalUIIOS::FOnlineExternalUIIOS()
+{
+}
+
+bool FOnlineExternalUIIOS::ShowLoginUI(const int ControllerIndex, bool bShowOnlineOnly, const FOnLoginUIClosedDelegate& Delegate)
+{
+	return false;
+}
+
+bool FOnlineExternalUIIOS::ShowFriendsUI(int32 LocalUserNum)
+{
+	return false;
+}
+
+bool FOnlineExternalUIIOS::ShowInviteUI(int32 LocalUserNum) 
+{
+	return false;
+}
+
+bool FOnlineExternalUIIOS::ShowAchievementsUI(int32 LocalUserNum) 
+{
+	// Will always show the achievements UI for the current local signed-in user
+	extern CORE_API void IOSShowAchievementsUI();
+	IOSShowAchievementsUI();
+	return true;
+}
+
+bool FOnlineExternalUIIOS::ShowLeaderboardUI( const FString& LeaderboardName )
+{
+	extern CORE_API void IOSShowLeaderboardUI(const FString& CategoryName);
+	IOSShowLeaderboardUI(LeaderboardName);
+	return true;
+}
+
+bool FOnlineExternalUIIOS::ShowWebURL(const FString& WebURL) 
+{
+	return false;
+}
+
+bool FOnlineExternalUIIOS::ShowProfileUI( const FUniqueNetId& Requestor, const FUniqueNetId& Requestee, const FOnProfileUIClosedDelegate& Delegate )
+{
+	return false;
+}
diff --git a/Engine/Source/Runtime/Online/IOS/OnlineSubsystemIOS/Private/OnlineExternalUIInterfaceIOS.h b/Engine/Source/Runtime/Online/IOS/OnlineSubsystemIOS/Private/OnlineExternalUIInterfaceIOS.h
new file mode 100644
--- /dev/null
+++ b/Engine/Source/Runtime/Online/IOS/OnlineSubsystemIOS/Private/OnlineExternalUIInterfaceIOS.h
@@ -0,0 +1,27 @@
+// Copyright 1998-2014 Epic Games, Inc. All Rights Reserved.
+
+#pragma once
+
+#include "OnlineExternalUIInterface.h"
+#include "OnlineSubsystemIOSTypes.h"
+
+class FOnlineExternalUIIOS : public IOnlineExternalUI
+{
+PACKAGE_SCOPE:
+
+	FOnlineExternalUIIOS();
+
+public:
+
+	// Begin IOnlineExternalUI interface
+	virtual bool ShowLoginUI(const int ControllerIndex, bool bShowOnlineOnly, const FOnLoginUIClosedDelegate& Delegate) OVERRIDE;
+	virtual bool ShowFriendsUI(int32 LocalUserNum) OVERRIDE;
+	virtual bool ShowInviteUI(int32 LocalUserNum) OVERRIDE;
+	virtual bool ShowAchievementsUI(int32 LocalUserNum) OVERRIDE;
+	virtual bool ShowLeaderboardUI(const FString& LeaderboardName) OVERRIDE;
+	virtual bool ShowWebURL(const FString& WebURL) OVERRIDE;
+	virtual bool ShowProfileUI(const FUniqueNetId& Requestor, const FUniqueNetId& Requestee, const FOnProfileUIClosedDelegate& Delegate) OVERRIDE;
+	// End IOnlineExternalUI interface
+};
+
+typedef TSharedPtr<FOnlineExternalUIIOS, ESPMode::ThreadSafe> FOnlineExternalUIIOSPtr;
\ No newline at end of file
diff --git a/Engine/Source/Runtime/Online/IOS/OnlineSubsystemIOS/Private/OnlineSubsystemIOS.cpp b/Engine/Source/Runtime/Online/IOS/OnlineSubsystemIOS/Private/OnlineSubsystemIOS.cpp
--- a/Engine/Source/Runtime/Online/IOS/OnlineSubsystemIOS/Private/OnlineSubsystemIOS.cpp
+++ b/Engine/Source/Runtime/Online/IOS/OnlineSubsystemIOS/Private/OnlineSubsystemIOS.cpp
@@ -46,7 +46,7 @@ IOnlineVoicePtr FOnlineSubsystemIOS::GetVoiceInterface() const
 
 IOnlineExternalUIPtr FOnlineSubsystemIOS::GetExternalUIInterface() const
 {
-	return NULL;
+	return ExternalUIInterface;
 }
 
 
@@ -135,6 +135,7 @@ bool FOnlineSubsystemIOS::Init()
 		LeaderboardsInterface = MakeShareable(new FOnlineLeaderboardsIOS(this));
 		StoreInterface = MakeShareable(new FOnlineStoreInterfaceIOS());
 		AchievementsInterface = MakeShareable(new FOnlineAchievementsIOS(this));
+		ExternalUIInterface = MakeShareable(new FOnlineExternalUIIOS());
 	}
 
 	return bSuccessfullyStartedUp;
diff --git a/Engine/Source/Runtime/Online/IOS/OnlineSubsystemIOS/Private/OnlineSubsystemIOSPrivatePCH.h b/Engine/Source/Runtime/Online/IOS/OnlineSubsystemIOS/Private/OnlineSubsystemIOSPrivatePCH.h
--- a/Engine/Source/Runtime/Online/IOS/OnlineSubsystemIOS/Private/OnlineSubsystemIOSPrivatePCH.h
+++ b/Engine/Source/Runtime/Online/IOS/OnlineSubsystemIOS/Private/OnlineSubsystemIOSPrivatePCH.h
@@ -16,6 +16,7 @@
 #include "OnlineLeaderboardsInterfaceIOS.h"
 #include "OnlineStoreInterfaceIOS.h"
 #include "OnlineAchievementsInterfaceIOS.h"
+#include "OnlineExternalUIInterfaceIOS.h"
 
 #include "OnlineSubsystemIOS.h"
 #include "OnlineSubsystemIOSModule.h"
\ No newline at end of file
diff --git a/Engine/Source/Runtime/Online/IOS/OnlineSubsystemIOS/Public/OnlineSubsystemIOS.h b/Engine/Source/Runtime/Online/IOS/OnlineSubsystemIOS/Public/OnlineSubsystemIOS.h
--- a/Engine/Source/Runtime/Online/IOS/OnlineSubsystemIOS/Public/OnlineSubsystemIOS.h
+++ b/Engine/Source/Runtime/Online/IOS/OnlineSubsystemIOS/Public/OnlineSubsystemIOS.h
@@ -75,6 +75,9 @@ private:
 
 	/** Interface to the online achievements */
 	FOnlineAchievementsIOSPtr AchievementsInterface;
+
+	/** Interface to the external UI services */
+	FOnlineExternalUIIOSPtr ExternalUIInterface;
 };
 
 typedef TSharedPtr<FOnlineSubsystemIOS, ESPMode::ThreadSafe> FOnlineSubsystemIOSPtr;
diff --git a/Engine/Source/Runtime/Online/OnlineSubsystem/Public/Interfaces/OnlineAchievementsInterface.h b/Engine/Source/Runtime/Online/OnlineSubsystem/Public/Interfaces/OnlineAchievementsInterface.h
--- a/Engine/Source/Runtime/Online/OnlineSubsystem/Public/Interfaces/OnlineAchievementsInterface.h
+++ b/Engine/Source/Runtime/Online/OnlineSubsystem/Public/Interfaces/OnlineAchievementsInterface.h
@@ -89,6 +89,7 @@ struct FOnlineAchievementDesc
 class IOnlineAchievements
 {
 public:
+	virtual ~IOnlineAchievements() {}
 
 	/**
 	 * Write the achievements provided to the server
diff --git a/Engine/Source/Runtime/Online/OnlineSubsystem/Public/Interfaces/OnlineExternalUIInterface.h b/Engine/Source/Runtime/Online/OnlineSubsystem/Public/Interfaces/OnlineExternalUIInterface.h
--- a/Engine/Source/Runtime/Online/OnlineSubsystem/Public/Interfaces/OnlineExternalUIInterface.h
+++ b/Engine/Source/Runtime/Online/OnlineSubsystem/Public/Interfaces/OnlineExternalUIInterface.h
@@ -72,6 +72,15 @@ public:
 	 */
 	virtual bool ShowAchievementsUI(int32 LocalUserNum) = 0;
 
+	/**
+	 *	Displays the UI that shows a specific leaderboard
+	 *
+	 * @param LeaderboardName the name of the leaderboard to show
+	 *
+	 * @return true if it was able to show the UI, false if it failed
+	 */
+	virtual bool ShowLeaderboardUI(const FString& LeaderboardName) = 0;
+
 	/**
 	 *	Displays a web page in the external UI
 	 *
diff --git a/Engine/Source/Runtime/Online/OnlineSubsystemSteam/Private/OnlineExternalUIInterfaceSteam.cpp b/Engine/Source/Runtime/Online/OnlineSubsystemSteam/Private/OnlineExternalUIInterfaceSteam.cpp
--- a/Engine/Source/Runtime/Online/OnlineSubsystemSteam/Private/OnlineExternalUIInterfaceSteam.cpp
+++ b/Engine/Source/Runtime/Online/OnlineSubsystemSteam/Private/OnlineExternalUIInterfaceSteam.cpp
@@ -53,6 +53,11 @@ bool FOnlineExternalUISteam::ShowAchievementsUI(int32 LocalUserNum)
 	return true;
 }
 
+bool FOnlineExternalUISteam::ShowLeaderboardUI(const FString& LeaderboardName)
+{
+	return false;
+}
+
 bool FOnlineExternalUISteam::ShowWebURL(const FString& WebURL)
 {
 	if (!WebURL.StartsWith(TEXT("https://")))
diff --git a/Engine/Source/Runtime/Online/OnlineSubsystemSteam/Private/OnlineExternalUIInterfaceSteam.h b/Engine/Source/Runtime/Online/OnlineSubsystemSteam/Private/OnlineExternalUIInterfaceSteam.h
--- a/Engine/Source/Runtime/Online/OnlineSubsystemSteam/Private/OnlineExternalUIInterfaceSteam.h
+++ b/Engine/Source/Runtime/Online/OnlineSubsystemSteam/Private/OnlineExternalUIInterfaceSteam.h
@@ -79,6 +79,7 @@ public:
 	virtual bool ShowFriendsUI(int32 LocalUserNum) OVERRIDE;
 	virtual bool ShowInviteUI(int32 LocalUserNum) OVERRIDE;
 	virtual bool ShowAchievementsUI(int32 LocalUserNum) OVERRIDE;
+	virtual bool ShowLeaderboardUI(const FString& LeaderboardName) OVERRIDE;
 	virtual bool ShowWebURL(const FString& WebURL) OVERRIDE;
 	virtual bool ShowProfileUI(const FUniqueNetId& Requestor, const FUniqueNetId& Requestee, const FOnProfileUIClosedDelegate& Delegate) OVERRIDE;
 };
diff --git a/Engine/Source/Runtime/Online/OnlineSubsystemUtils/Private/LeaderboardBlueprintLibrary.cpp b/Engine/Source/Runtime/Online/OnlineSubsystemUtils/Private/LeaderboardBlueprintLibrary.cpp
--- a/Engine/Source/Runtime/Online/OnlineSubsystemUtils/Private/LeaderboardBlueprintLibrary.cpp
+++ b/Engine/Source/Runtime/Online/OnlineSubsystemUtils/Private/LeaderboardBlueprintLibrary.cpp
@@ -4,65 +4,6 @@
 #include "LeaderboardBlueprintLibrary.h"
 #include "Runtime/Online/OnlineSubsystem/Public/Interfaces/OnlineLeaderboardInterface.h"
 
-
-#if PLATFORM_ANDROID
-/**
- * Class to store the leaderboard config.
- * This really belongs in the proper OSS, move it there when we bring the whole OSS online.
- */
-class FLeaderboardsConfig : public FOnlineJsonSerializable
-{
-public:
-	FJsonSerializableKeyValueMap LeaderboardMap;
-
-	BEGIN_ONLINE_JSON_SERIALIZER
-	ONLINE_JSON_SERIALIZE_MAP( "LeaderboardMap", LeaderboardMap );
-	END_ONLINE_JSON_SERIALIZER
-};
-
-/**
- * Helper function to get the platform- and game-specific leaderboard ID from the JSON config file.
- * This really belongs in the proper OSS class, move it there when we bring the whole OSS online.
- */
-FString GetLeaderboardID(const FString& LeaderboardName)
-{
-	// Hacks to initialize the config one
-	static bool bHasReadConfig = false;
-	static FLeaderboardsConfig LeaderboardConfig;
-
-	if(!bHasReadConfig)
-	{
-		const FString JSonConfigFilename = FPaths::GameDir() + TEXT( "Config/OSS/GooglePlay/Leaderboards.json" );
-
-		FString JSonText;
-
-		if ( !FFileHelper::LoadFileToString( JSonText, *JSonConfigFilename ) )
-		{
-			UE_LOG(LogOnline, Warning, TEXT( "GetLeaderboardID: Failed to find json OSS leaderboard config: %s"), *JSonConfigFilename );
-			return LeaderboardName;
-		}
-
-		if ( !LeaderboardConfig.FromJson( JSonText ) )
-		{
-			UE_LOG(LogOnline, Warning, TEXT( "GetLeaderboardID: Failed to parse json OSS leaderboard config: %s"), *JSonConfigFilename );
-			return LeaderboardName;
-		}
-
-		bHasReadConfig = true;
-	}
-
-	FString * LeaderboardID = LeaderboardConfig.LeaderboardMap.Find( LeaderboardName );
-
-	if(LeaderboardID == NULL)
-	{
-		UE_LOG(LogOnline, Warning, TEXT( "GetLeaderboardID: No mapping for leaderboard %s"), *LeaderboardName );
-		return LeaderboardName;
-	}
-
-	return *LeaderboardID;
-}
-#endif
-
 //////////////////////////////////////////////////////////////////////////
 // ULeaderboardBlueprintLibrary
 
@@ -73,41 +14,6 @@ ULeaderboardBlueprintLibrary::ULeaderboardBlueprintLibrary(const FPostConstructI
 
 bool ULeaderboardBlueprintLibrary::WriteLeaderboardObject(APlayerController* PlayerController, class FOnlineLeaderboardWrite& WriteObject)
 {
-#if PLATFORM_ANDROID
-	// Hack in Android leaderboard write since we don't have the full OSS yet.
-	UE_LOG(LogOnline, Log, TEXT("Leaderboard stat name is %s"), *WriteObject.RatedStat.ToString());
-
-	for (FStatPropertyArray::TConstIterator It(WriteObject.Properties); It; ++It)
-	{
-		// Access the stat and the value.
-		const FVariantData& Stat = It.Value();
-
-		FString LeaderboardName(It.Key().ToString());
-		
-		int32 Value = 0;
-
-		switch (Stat.GetType())
-		{
-		case EOnlineKeyValuePairDataType::Int32:
-			{
-				Stat.GetValue(Value);
-				UE_LOG(LogOnline, Log, TEXT("Writing Int32 leaderboard stat %s with value %d."), *LeaderboardName, Value);
-				extern void AndroidThunkCpp_WriteLeaderboardValue(const FString&, int64_t);
-				AndroidThunkCpp_WriteLeaderboardValue(GetLeaderboardID(LeaderboardName), Value);
-				return true;
-			}
-
-		default:
-			{
-				UE_LOG(LogOnline, Warning, TEXT("ULeaderboardBlueprintLibrary::WriteLeaderboardObject(Leaderboard: %s) Invalid data type (only Int32 is currently supported)"), *LeaderboardName);
-				break;
-			}
-		}
-	}
-
-	return false;
-
-#else
 	if (APlayerState* PlayerState = (PlayerController != NULL) ? PlayerController->PlayerState : NULL)
 	{
 		TSharedPtr<FUniqueNetId> UserId = PlayerState->UniqueId.GetUniqueNetId();
@@ -147,7 +53,6 @@ bool ULeaderboardBlueprintLibrary::WriteLeaderboardObject(APlayerController* Pla
 	}
 
 	return false;
-#endif
 }
 
 bool ULeaderboardBlueprintLibrary::WriteLeaderboardInteger(APlayerController* PlayerController, FName StatName, int32 StatValue)
