Achievements: Added name, progress, and user tag outputs to the WriteAchievementProgress node delegates, simplifying management of local achievement state in blueprints

diff --git a/Engine/Source/Runtime/Online/OnlineSubsystemUtils/Classes/AchievementWriteCallbackProxy.h b/Engine/Source/Runtime/Online/OnlineSubsystemUtils/Classes/AchievementWriteCallbackProxy.h
--- a/Engine/Source/Runtime/Online/OnlineSubsystemUtils/Classes/AchievementWriteCallbackProxy.h
+++ b/Engine/Source/Runtime/Online/OnlineSubsystemUtils/Classes/AchievementWriteCallbackProxy.h
@@ -4,6 +4,8 @@
 #include "Runtime/Online/OnlineSubsystem/Public/Interfaces/OnlineAchievementsInterface.h"
 #include "AchievementWriteCallbackProxy.generated.h"
 
+DECLARE_DYNAMIC_MULTICAST_DELEGATE_ThreeParams(FAchievementWriteDelegate, FName, WrittenAchievementName, float, WrittenProgress, int32, WrittenUserTag);
+
 UCLASS(MinimalAPI)
 class UAchievementWriteCallbackProxy : public UOnlineBlueprintCallProxyBase
 {
@@ -11,15 +13,18 @@ class UAchievementWriteCallbackProxy : public UOnlineBlueprintCallProxyBase
 
 	// Called when there is a successful achievement write
 	UPROPERTY(BlueprintAssignable)
-	FEmptyOnlineDelegate OnSuccess;
+	FAchievementWriteDelegate OnSuccess;
 
 	// Called when there is an unsuccessful achievement write
 	UPROPERTY(BlueprintAssignable)
-	FEmptyOnlineDelegate OnFailure;
+	FAchievementWriteDelegate OnFailure;
 
 	// Writes progress about an achievement to the default online subsystem
+	//   AchievementName is the ID of the achievement to update progress on
+	//   Progress is the reported progress toward accomplishing the achievement
+	//   UserTag is not used internally, but it is returned on success or failure
 	UFUNCTION(BlueprintCallable, meta=(BlueprintInternalUseOnly = "true"), Category="Online|Achievements")
-	static UAchievementWriteCallbackProxy* WriteAchievementProgress(class APlayerController* PlayerController, FName AchievementName, float Progress = 100.0f);
+	static UAchievementWriteCallbackProxy* WriteAchievementProgress(class APlayerController* PlayerController, FName AchievementName, float Progress = 100.0f, int32 UserTag = 0);
 
 	// UOnlineBlueprintCallProxyBase interface
 	virtual void Activate() OVERRIDE;
@@ -39,4 +44,13 @@ private:
 
 	/** The achievements write object */
 	FOnlineAchievementsWritePtr WriteObject;
+
+	/** The achievement name */
+	FName AchievementName;
+
+	/** The amount of progress made towards the achievement */
+	float AchievementProgress;
+
+	/** The specified user tag */
+	int32 UserTag;
 };
diff --git a/Engine/Source/Runtime/Online/OnlineSubsystemUtils/Private/AchievementWriteCallbackProxy.cpp b/Engine/Source/Runtime/Online/OnlineSubsystemUtils/Private/AchievementWriteCallbackProxy.cpp
--- a/Engine/Source/Runtime/Online/OnlineSubsystemUtils/Private/AchievementWriteCallbackProxy.cpp
+++ b/Engine/Source/Runtime/Online/OnlineSubsystemUtils/Private/AchievementWriteCallbackProxy.cpp
@@ -10,13 +10,16 @@ UAchievementWriteCallbackProxy::UAchievementWriteCallbackProxy(const class FPost
 {
 }
 
-UAchievementWriteCallbackProxy* UAchievementWriteCallbackProxy::WriteAchievementProgress(class APlayerController* PlayerController, FName AchievementName, float Progress)
+UAchievementWriteCallbackProxy* UAchievementWriteCallbackProxy::WriteAchievementProgress(class APlayerController* PlayerController, FName AchievementName, float Progress, int32 UserTag)
 {
 	UAchievementWriteCallbackProxy* Proxy = NewObject<UAchievementWriteCallbackProxy>();
 
 	Proxy->WriteObject = MakeShareable(new FOnlineAchievementsWrite);
 	Proxy->WriteObject->SetFloatStat(AchievementName, Progress);
 	Proxy->PlayerControllerWeakPtr = PlayerController;
+	Proxy->AchievementName = AchievementName;
+	Proxy->AchievementProgress = Progress;
+	Proxy->UserTag = UserTag;
 
 	return Proxy;
 }
@@ -45,7 +48,7 @@ void UAchievementWriteCallbackProxy::Activate()
 	}
 
 	// Fail immediately
-	OnFailure.Broadcast();
+	OnFailure.Broadcast(AchievementName, AchievementProgress, UserTag);
 	WriteObject.Reset();
 }
 
@@ -53,11 +56,11 @@ void UAchievementWriteCallbackProxy::OnAchievementWritten(const FUniqueNetId& Us
 {
 	if (bSuccess)
 	{
-		OnSuccess.Broadcast();
+		OnSuccess.Broadcast(AchievementName, AchievementProgress, UserTag);
 	}
 	else
 	{
-		OnFailure.Broadcast();
+		OnFailure.Broadcast(AchievementName, AchievementProgress, UserTag);
 	}
 
 	WriteObject.Reset();
