Integrate Xbox One HTTP support from Lionhead

diff --git a/Engine/Source/Runtime/Online/HTTP/Private/HttpModule.cpp b/Engine/Source/Runtime/Online/HTTP/Private/HttpModule.cpp
--- a/Engine/Source/Runtime/Online/HTTP/Private/HttpModule.cpp
+++ b/Engine/Source/Runtime/Online/HTTP/Private/HttpModule.cpp
@@ -199,6 +199,8 @@ bool FHttpManager::IsValidRequest(class IHttpRequest* RequestPtr)
 
 void FHttpManager::DumpRequests(FOutputDevice& Ar)
 {
+	FScopeLock ScopeLock(&RequestLock);
+	
 	Ar.Logf(TEXT("------- (%d) Http Requests"), Requests.Num());
 	for (TArray<TSharedRef<class IHttpRequest> >::TIterator It(Requests); It; ++It)
 	{
