Changing NoRedist -> NotForLicensees. This will eventually be removed.

diff --git a/Engine/Source/Runtime/Online/OnlineSubsystemSteam/OnlineSubsystemSteam.Build.cs b/Engine/Source/Runtime/Online/OnlineSubsystemSteam/OnlineSubsystemSteam.Build.cs
--- a/Engine/Source/Runtime/Online/OnlineSubsystemSteam/OnlineSubsystemSteam.Build.cs
+++ b/Engine/Source/Runtime/Online/OnlineSubsystemSteam/OnlineSubsystemSteam.Build.cs
@@ -31,8 +31,8 @@ public class OnlineSubsystemSteam : ModuleRules
 		Definitions.Add("WITH_STEAMGC=0");
 
 #if false
-		if (Directory.Exists(UEBuildConfiguration.UEThirdPartyDirectory + "NoRedist") == true &&
-			Directory.Exists(UEBuildConfiguration.UEThirdPartyDirectory + "NoRedist/SteamGC") == true)
+		if (Directory.Exists(UEBuildConfiguration.UEThirdPartyDirectory + "NotForLicensees") == true &&
+			Directory.Exists(UEBuildConfiguration.UEThirdPartyDirectory + "NotForLicensees/SteamGC") == true)
 		{
 			AddThirdPartyPrivateStaticDependencies(Target, 
 				"ProtoBuf",
