Add a better way to handle client port selection on Xbox One

Add new virtual function to UIpNetDriver: GetClientPort so that platforms can customize the port used. Base class implementation returns 0 to maintain current functionality on non-XB1 platforms.
Add a UNetDriver subclass for XB1 that overrides GetClientPort to return the port specified in the config. Can't just do this for all platforms because it causes issues with multiple instances and play-in-editor.
ShooterGame now uses this new subclass on XB1.

diff --git a/Engine/Source/Runtime/Online/OnlineSubsystemUtils/Classes/IpNetDriver.h b/Engine/Source/Runtime/Online/OnlineSubsystemUtils/Classes/IpNetDriver.h
--- a/Engine/Source/Runtime/Online/OnlineSubsystemUtils/Classes/IpNetDriver.h
+++ b/Engine/Source/Runtime/Online/OnlineSubsystemUtils/Classes/IpNetDriver.h
@@ -48,6 +48,15 @@ class ONLINESUBSYSTEMUTILS_API UIpNetDriver : public UNetDriver
 
 	// Begin UIpNetDriver interface.
 	virtual FSocket * CreateSocket();
+
+	/**
+	 * Returns the port number to use when a client is creating a socket.
+	 * Platforms that can't use the default of 0 (system-selected port) may override
+	 * this function.
+	 *
+	 * @return The port number to use for client sockets. Base implementation returns 0.
+	 */
+	virtual int GetClientPort();
 	// End UIpNetDriver interface.
 
 	// Begin FExec Interface
@@ -60,5 +69,5 @@ class ONLINESUBSYSTEMUTILS_API UIpNetDriver : public UNetDriver
 	bool HandleSocketsCommand( const TCHAR* Cmd, FOutputDevice& Ar, UWorld* InWorld );
 
 	/** @return TCPIP connection to server */
-	UIpConnection* GetServerConnection();
+	class UIpConnection* GetServerConnection();
 };
diff --git a/Engine/Source/Runtime/Online/OnlineSubsystemUtils/Private/IpNetDriver.cpp b/Engine/Source/Runtime/Online/OnlineSubsystemUtils/Private/IpNetDriver.cpp
--- a/Engine/Source/Runtime/Online/OnlineSubsystemUtils/Private/IpNetDriver.cpp
+++ b/Engine/Source/Runtime/Online/OnlineSubsystemUtils/Private/IpNetDriver.cpp
@@ -49,6 +49,11 @@ FSocket * UIpNetDriver::CreateSocket()
 	return SocketSubsystem->CreateSocket( NAME_DGram, TEXT( "Unreal" ) );
 }
 
+int UIpNetDriver::GetClientPort()
+{
+	return 0;
+}
+
 bool UIpNetDriver::InitBase( bool bInitAsClient, FNetworkNotify* InNotify, const FURL& URL, bool bReuseAddressAndPort, FString& Error )
 {
 	if (!Super::InitBase(bInitAsClient, InNotify, URL, bReuseAddressAndPort, Error))
@@ -102,19 +107,10 @@ bool UIpNetDriver::InitBase( bool bInitAsClient, FNetworkNotify* InNotify, const
 	// Bind socket to our port.
 	LocalAddr = SocketSubsystem->GetLocalBindAddr(*GLog);
 	
-	// Try the specific port first, even for clients. Required for some platforms.
-	LocalAddr->SetPort(bInitAsClient ? 0 : URL.Port);
+	LocalAddr->SetPort(bInitAsClient ? GetClientPort() : URL.Port);
 	
 	int32 AttemptPort = LocalAddr->GetPort();
 	int32 BoundPort = SocketSubsystem->BindNextPort( Socket, *LocalAddr, MaxPortCountToTry + 1, 1 );
-	if (BoundPort == 0 && bInitAsClient)
-	{
-		// If initializing as a client, and the specific port failed, retry with port 0 (any port)
-		LocalAddr->SetPort(0);
-		AttemptPort = LocalAddr->GetPort();
-		BoundPort = SocketSubsystem->BindNextPort( Socket, *LocalAddr, 1, 1 );
-	}
-
 	if (BoundPort == 0)
 	{
 		Error = FString::Printf( TEXT("%s: binding to port %i failed (%i)"), SocketSubsystem->GetSocketAPIName(), AttemptPort,
