Updated hardcoded Steamworks version.

Fixes TTP #335162 (Shooter not initializing Steam OSS on Windows).

diff --git a/Engine/Source/Runtime/Online/OnlineSubsystemSteam/Private/OnlineSubsystemSteamPrivatePCH.h b/Engine/Source/Runtime/Online/OnlineSubsystemSteam/Private/OnlineSubsystemSteamPrivatePCH.h
--- a/Engine/Source/Runtime/Online/OnlineSubsystemSteam/Private/OnlineSubsystemSteamPrivatePCH.h
+++ b/Engine/Source/Runtime/Online/OnlineSubsystemSteam/Private/OnlineSubsystemSteamPrivatePCH.h
@@ -16,7 +16,7 @@
 #define INVALID_INDEX -1
 
 /** Name of the current Steam SDK version in use (matches directory name) */
-#define STEAM_SDK_VER TEXT("Steamv128")
+#define STEAM_SDK_VER TEXT("Steamv129")
 
 /** FName declaration of Steam subsystem */
 #define STEAM_SUBSYSTEM FName(TEXT("STEAM"))
