- fix for mcp automation tests crashing due to missing GWorld

diff --git a/Engine/Source/Runtime/Online/OnlineSubsystemUtils/Public/OnlineSubsystemUtils.h b/Engine/Source/Runtime/Online/OnlineSubsystemUtils/Public/OnlineSubsystemUtils.h
--- a/Engine/Source/Runtime/Online/OnlineSubsystemUtils/Public/OnlineSubsystemUtils.h
+++ b/Engine/Source/Runtime/Online/OnlineSubsystemUtils/Public/OnlineSubsystemUtils.h
@@ -27,9 +27,9 @@ namespace Online
 	static IOnlineSubsystem* GetSubsystem(UWorld* World, const FName& SubsystemName = NAME_None)
 	{
 #if UE_EDITOR // at present, multiple worlds are only possible in the editor
-		check(World);
 		FName Identifier = SubsystemName; 
-		if (GIsPlayInEditorWorld)
+		if (GIsPlayInEditorWorld &&
+			World != NULL)
 		{
 			FWorldContext& CurrentContext = GEngine->GetWorldContextFromWorldChecked(World);
 			if (CurrentContext.WorldType == EWorldType::PIE)
