Implement QueryAchievements for Google Play on Android

This fixes TPP 336286, which involves some achievement issues in TappyChicken. Somewhat hacky, probably worth cleaning up in the future :)

diff --git a/Engine/Build/Android/Java/src/com/epicgames/ue4/GameActivity.java b/Engine/Build/Android/Java/src/com/epicgames/ue4/GameActivity.java
--- a/Engine/Build/Android/Java/src/com/epicgames/ue4/GameActivity.java
+++ b/Engine/Build/Android/Java/src/com/epicgames/ue4/GameActivity.java
@@ -270,6 +270,56 @@ public class GameActivity extends NativeActivity implements GoogleApiClient.Conn
 		}
 	}
 
+	/** Class that holds achievement data to pass back to C++ through JNI */
+	public class JavaAchievement
+	{
+		public String ID;
+		public double Progress;
+	}
+
+	/** Callback that returns queried achievement data back to C++ */
+	private class QueryAchievementsResultCallback implements ResultCallback<Achievements.LoadAchievementsResult>
+	{
+		@Override
+		public void onResult(Achievements.LoadAchievementsResult result)
+		{
+			Log.debug("Google Play Services: queried achievements with status " + result.getStatus().toString());
+
+			AchievementBuffer Achievements = result.getAchievements();
+			
+			JavaAchievement[] SimpleAchievements = new JavaAchievement[Achievements.getCount()];
+
+			for(int i = 0; i < Achievements.getCount(); ++i)
+			{
+				Achievement CurrentAchievement = Achievements.get(i);
+
+				SimpleAchievements[i] = new JavaAchievement();
+				SimpleAchievements[i].ID = CurrentAchievement.getAchievementId();
+				
+				if(CurrentAchievement.getState() == Achievement.STATE_UNLOCKED)
+				{
+					SimpleAchievements[i].Progress = 100.0;
+					continue;
+				}
+
+				if(CurrentAchievement.getType() == Achievement.TYPE_INCREMENTAL)
+				{
+					double Fraction = (double)CurrentAchievement.getCurrentSteps() / (double)CurrentAchievement.getTotalSteps();
+					SimpleAchievements[i].Progress = Fraction * 100.0;
+				}
+				else
+				{
+					SimpleAchievements[i].Progress = 0.0;
+				}
+			}
+
+			nativeUpdateAchievements(SimpleAchievements);
+
+			Achievements.close();
+			result.release();
+		}
+	}
+
 	// Callbacks to handle connections with Google Play
 	 @Override
     public void onConnected(Bundle connectionHint)
@@ -482,6 +532,12 @@ public class GameActivity extends NativeActivity implements GoogleApiClient.Conn
 		}
 	}
 
+	public void AndroidThunkJava_QueryAchievements()
+	{
+		PendingResult<Achievements.LoadAchievementsResult> loadAchievementsResult = Games.Achievements.load(googleClient, false);
+		loadAchievementsResult.setResultCallback(new QueryAchievementsResultCallback());		
+	}
+
 	public void AndroidThunkJava_ShowAdBanner(String AdMobAdUnitID, boolean bShowOnBottonOfScreen)
 	{
 		Log.debug("In AndroidThunkJava_ShowAdBanner");
@@ -628,6 +684,7 @@ public class GameActivity extends NativeActivity implements GoogleApiClient.Conn
 	public native void nativeSetGlobalActivity();
 	public native void nativeSetWindowInfo(boolean bIsPortrait);
 	public native void nativeSetObbInfo(String PackageName, int Version, int PatchVersion);
+	public native void nativeUpdateAchievements(JavaAchievement[] Achievements);
 
 	public native void nativeConsoleCommand(String commandString);
 	
diff --git a/Engine/Source/Runtime/Launch/Private/Android/AndroidJNI.cpp b/Engine/Source/Runtime/Launch/Private/Android/AndroidJNI.cpp
--- a/Engine/Source/Runtime/Launch/Private/Android/AndroidJNI.cpp
+++ b/Engine/Source/Runtime/Launch/Private/Android/AndroidJNI.cpp
@@ -99,6 +99,7 @@ jmethodID JDef_GameActivity::AndroidThunkJava_ShowConsoleWindow;
 jmethodID JDef_GameActivity::AndroidThunkJava_LaunchURL;
 jmethodID JDef_GameActivity::AndroidThunkJava_ShowLeaderboard;
 jmethodID JDef_GameActivity::AndroidThunkJava_ShowAchievements;
+jmethodID JDef_GameActivity::AndroidThunkJava_QueryAchievements;
 jmethodID JDef_GameActivity::AndroidThunkJava_WriteLeaderboardValue;
 jmethodID JDef_GameActivity::AndroidThunkJava_GooglePlayConnect;
 jmethodID JDef_GameActivity::AndroidThunkJava_WriteAchievement;
@@ -107,6 +108,10 @@ jmethodID JDef_GameActivity::AndroidThunkJava_HideAdBanner;
 jmethodID JDef_GameActivity::AndroidThunkJava_CloseAdBanner;
 jmethodID JDef_GameActivity::AndroidThunkJava_GetAssetManager;
 
+jclass JDef_GameActivity::JavaAchievementClassID;
+jfieldID JDef_GameActivity::AchievementIDField;
+jfieldID JDef_GameActivity::AchievementProgressField;
+
 DEFINE_LOG_CATEGORY_STATIC(LogEngine, Log, All);
 
 //Game-specific crash reporter
@@ -219,6 +224,15 @@ void AndroidThunkCpp_WriteAchievement(const FString& AchievementID, float Percen
 	}
 }
 
+void AndroidThunkCpp_QueryAchievements()
+{
+	if (JNIEnv* Env = GetJavaEnv())
+	{
+		Env->CallVoidMethod(GJavaGlobalThis, JDef_GameActivity::AndroidThunkJava_QueryAchievements);
+	}
+}
+
+
 void AndroidThunkCpp_ShowAdBanner(const FString& AdUnitID, bool bShowOnBottomOfScreen)
 {
 	if (JNIEnv* Env = GetJavaEnv())
@@ -293,6 +307,7 @@ JNIEXPORT jint JNI_OnLoad(JavaVM* InJavaVM, void* InReserved)
 	JDef_GameActivity::AndroidThunkJava_LaunchURL = env->GetMethodID(JDef_GameActivity::ClassID, "AndroidThunkJava_LaunchURL", "(Ljava/lang/String;)V");
 	JDef_GameActivity::AndroidThunkJava_ShowLeaderboard = env->GetMethodID(JDef_GameActivity::ClassID, "AndroidThunkJava_ShowLeaderboard", "(Ljava/lang/String;)V");
 	JDef_GameActivity::AndroidThunkJava_ShowAchievements = env->GetMethodID(JDef_GameActivity::ClassID, "AndroidThunkJava_ShowAchievements", "()V");
+	JDef_GameActivity::AndroidThunkJava_QueryAchievements = env->GetMethodID(JDef_GameActivity::ClassID, "AndroidThunkJava_QueryAchievements", "()V");
 	JDef_GameActivity::AndroidThunkJava_WriteLeaderboardValue = env->GetMethodID(JDef_GameActivity::ClassID, "AndroidThunkJava_WriteLeaderboardValue", "(Ljava/lang/String;J)V");
 	JDef_GameActivity::AndroidThunkJava_GooglePlayConnect = env->GetMethodID(JDef_GameActivity::ClassID, "AndroidThunkJava_GooglePlayConnect", "()V");
 	JDef_GameActivity::AndroidThunkJava_WriteAchievement = env->GetMethodID(JDef_GameActivity::ClassID, "AndroidThunkJava_WriteAchievement", "(Ljava/lang/String;F)V");
@@ -301,6 +316,11 @@ JNIEXPORT jint JNI_OnLoad(JavaVM* InJavaVM, void* InReserved)
 	JDef_GameActivity::AndroidThunkJava_CloseAdBanner = env->GetMethodID(JDef_GameActivity::ClassID, "AndroidThunkJava_CloseAdBanner", "()V");
 	JDef_GameActivity::AndroidThunkJava_GetAssetManager = env->GetMethodID(JDef_GameActivity::ClassID, "AndroidThunkJava_GetAssetManager", "()Landroid/content/res/AssetManager;");
 
+	// Set up achievement query IDs
+	JDef_GameActivity::JavaAchievementClassID = env->FindClass("com/epicgames/ue4/GameActivity$JavaAchievement");
+	JDef_GameActivity::AchievementIDField = env->GetFieldID(JDef_GameActivity::JavaAchievementClassID, "ID", "Ljava/lang/String;");
+	JDef_GameActivity::AchievementProgressField = env->GetFieldID(JDef_GameActivity::JavaAchievementClassID, "Progress", "D");
+
 	// hook signals
 #if UE_BUILD_DEBUG
 	if( GAlwaysReportCrash )
diff --git a/Engine/Source/Runtime/Launch/Private/Android/AndroidJNI.h b/Engine/Source/Runtime/Launch/Private/Android/AndroidJNI.h
--- a/Engine/Source/Runtime/Launch/Private/Android/AndroidJNI.h
+++ b/Engine/Source/Runtime/Launch/Private/Android/AndroidJNI.h
@@ -23,6 +23,7 @@ public:
 	static jmethodID AndroidThunkJava_LaunchURL;
 	static jmethodID AndroidThunkJava_ShowLeaderboard;
 	static jmethodID AndroidThunkJava_ShowAchievements;
+	static jmethodID AndroidThunkJava_QueryAchievements;
 	static jmethodID AndroidThunkJava_WriteLeaderboardValue;
 	static jmethodID AndroidThunkJava_GooglePlayConnect;
 	static jmethodID AndroidThunkJava_WriteAchievement;
@@ -31,6 +32,11 @@ public:
 	static jmethodID AndroidThunkJava_CloseAdBanner;
 
 	static jmethodID AndroidThunkJava_GetAssetManager;
+
+	// IDs related to achievement query
+	static jclass JavaAchievementClassID;
+	static jfieldID AchievementIDField;
+	static jfieldID AchievementProgressField;
 };
 
 // Returns the java environment
diff --git a/Engine/Source/Runtime/Online/Android/OnlineSubsystemGooglePlay/OnlineSubsystemGooglePlay.Build.cs b/Engine/Source/Runtime/Online/Android/OnlineSubsystemGooglePlay/OnlineSubsystemGooglePlay.Build.cs
--- a/Engine/Source/Runtime/Online/Android/OnlineSubsystemGooglePlay/OnlineSubsystemGooglePlay.Build.cs
+++ b/Engine/Source/Runtime/Online/Android/OnlineSubsystemGooglePlay/OnlineSubsystemGooglePlay.Build.cs
@@ -23,7 +23,8 @@ public class OnlineSubsystemGooglePlay : ModuleRules
                 "Sockets",
 				"OnlineSubsystem", 
                 "Http",
-				"AndroidRuntimeSettings"
+				"AndroidRuntimeSettings",
+				"Launch"
             }
 			);
 	}
diff --git a/Engine/Source/Runtime/Online/Android/OnlineSubsystemGooglePlay/Private/OnlineAchievementsInterfaceGooglePlay.cpp b/Engine/Source/Runtime/Online/Android/OnlineSubsystemGooglePlay/Private/OnlineAchievementsInterfaceGooglePlay.cpp
--- a/Engine/Source/Runtime/Online/Android/OnlineSubsystemGooglePlay/Private/OnlineAchievementsInterfaceGooglePlay.cpp
+++ b/Engine/Source/Runtime/Online/Android/OnlineSubsystemGooglePlay/Private/OnlineAchievementsInterfaceGooglePlay.cpp
@@ -1,15 +1,94 @@
 // Copyright 1998-2014 Epic Games, Inc. All Rights Reserved.
 
 #include "OnlineSubsystemGooglePlayPrivatePCH.h"
+#include "OnlineAsyncTaskManagerGooglePlay.h"
+#include "Online.h"
+#include "Android/AndroidJNI.h"
+
+FOnlineAchievementsGooglePlay::FPendingAchievementQuery FOnlineAchievementsGooglePlay::PendingAchievementQuery;
+
+extern "C" void Java_com_epicgames_ue4_GameActivity_nativeUpdateAchievements(JNIEnv* LocalJNIEnv, jobject LocalThiz, jobjectArray Achievements)
+{
+	TArray<FOnlineAchievement> GooglePlayAchievements;
+
+	auto ArrayLength = LocalJNIEnv->GetArrayLength(Achievements);
+
+	for(jsize i = 0; i < ArrayLength; ++i)
+	{
+		jobject Entry = LocalJNIEnv->GetObjectArrayElement(Achievements, i);
+
+		jobject StringID = LocalJNIEnv->GetObjectField(Entry, JDef_GameActivity::AchievementIDField);
+		jdouble Progress = LocalJNIEnv->GetDoubleField(Entry, JDef_GameActivity::AchievementProgressField);
+
+		FOnlineAchievement NewAchievement;
+
+		const char* StringIDChars = LocalJNIEnv->GetStringUTFChars((jstring)StringID, nullptr);
+		NewAchievement.Id = FString(StringIDChars);
+		LocalJNIEnv->ReleaseStringUTFChars((jstring)StringID, StringIDChars);
+
+		NewAchievement.Progress = Progress;
+
+		GooglePlayAchievements.Add(NewAchievement);
+	}
+
+	// Do the final copy of new achievements on the game thread
+	auto AchievementsInterface = FOnlineAchievementsGooglePlay::PendingAchievementQuery.AchievementsInterface;
+	
+	if(!AchievementsInterface ||
+	   !AchievementsInterface->AndroidSubsystem ||
+	   !AchievementsInterface->AndroidSubsystem->GetAsyncTaskManager())
+	{
+		// We should call the delegate with a false parameter here, but if we don't have
+		// the async task manager we're not going to call it on the game thread.
+		return;
+	}
+
+	AchievementsInterface->AndroidSubsystem->GetAsyncTaskManager()->AddGenericToOutQueue([GooglePlayAchievements]()
+	{
+		auto& PendingQuery = FOnlineAchievementsGooglePlay::PendingAchievementQuery;
+
+		PendingQuery.AchievementsInterface->Achievements = GooglePlayAchievements;
+
+		// Must convert the Google Play IDs to the names specified by the game.
+		auto DefaultSettings = GetDefault<UAndroidRuntimeSettings>();
+		for(auto& Achievement : PendingQuery.AchievementsInterface->Achievements)
+		{
+			for(const auto& Mapping : DefaultSettings->AchievementMap)
+			{
+				if(Mapping.AchievementID == Achievement.Id)
+				{
+					Achievement.Id = Mapping.Name;
+					break;
+				}
+			}
+		}
+
+		PendingQuery.Delegate.ExecuteIfBound( PendingQuery.PlayerID, true );
+		PendingQuery.IsQueryPending = false;
+	});
+}
 
 FOnlineAchievementsGooglePlay::FOnlineAchievementsGooglePlay( FOnlineSubsystemGooglePlay* InSubsystem )
 	: AndroidSubsystem(InSubsystem)
 {
+	PendingAchievementQuery.AchievementsInterface = this;
 }
 
 void FOnlineAchievementsGooglePlay::QueryAchievements(const FUniqueNetId& PlayerId, const FOnQueryAchievementsCompleteDelegate& Delegate)
 {
-	Delegate.ExecuteIfBound( PlayerId, false );
+	if(PendingAchievementQuery.IsQueryPending)
+	{
+		UE_LOG(LogOnline, Warning, TEXT("FOnlineAchievementsGooglePlay::QueryAchievements: Cannot start a new query while one is in progress."));
+		Delegate.ExecuteIfBound(PlayerId, false);
+		return;
+	}
+
+	PendingAchievementQuery.Delegate = Delegate;
+	PendingAchievementQuery.PlayerID = FUniqueNetIdString(PlayerId);
+	PendingAchievementQuery.IsQueryPending = true;
+
+	extern void AndroidThunkCpp_QueryAchievements();
+	AndroidThunkCpp_QueryAchievements();
 }
 
 
diff --git a/Engine/Source/Runtime/Online/Android/OnlineSubsystemGooglePlay/Private/OnlineAchievementsInterfaceGooglePlay.h b/Engine/Source/Runtime/Online/Android/OnlineSubsystemGooglePlay/Private/OnlineAchievementsInterfaceGooglePlay.h
--- a/Engine/Source/Runtime/Online/Android/OnlineSubsystemGooglePlay/Private/OnlineAchievementsInterfaceGooglePlay.h
+++ b/Engine/Source/Runtime/Online/Android/OnlineSubsystemGooglePlay/Private/OnlineAchievementsInterfaceGooglePlay.h
@@ -5,6 +5,9 @@
 #include "OnlineAchievementsInterface.h"
 #include "OnlineSubsystemTypes.h"
 #include "AndroidRuntimeSettings.h"
+#include <jni.h>
+
+extern "C" void Java_com_epicgames_ue4_GameActivity_nativeUpdateAchievements(JNIEnv* LocalJNIEnv, jobject LocalThiz, jobjectArray Achievements);
 
 /**
  *	IOnlineAchievements - Interface class for Achievements
@@ -25,6 +28,31 @@ private:
 	/** Cached achievement descriptions for an Id */
 	TMap< FString, FOnlineAchievementDesc > AchievementDescriptions;
 
+	/**
+	 * Hack to store context needed by the QueryAchievement callback to avoid round-tripping
+	 * this data through JNI and back.
+	 */
+	struct FPendingAchievementQuery
+	{
+		FOnlineAchievementsGooglePlay* AchievementsInterface;
+		FOnQueryAchievementsCompleteDelegate Delegate;
+		FUniqueNetIdString PlayerID;
+		bool IsQueryPending;
+	};
+
+	/** Instance of the query context data */
+	static FPendingAchievementQuery PendingAchievementQuery;
+
+	/**
+	 * Native function called from Java when we get the achievement load result callback.
+	 * Converts the Google Play achievement IDs to the names specified by the game's mapping.
+	 *
+	 * @param LocalJNIEnv the current JNI environment
+	 * @param LocalThiz instance of the Java GameActivity class
+	 * @param Achievements an array of JavaAchievement objects holding the data just queried from the backend
+	 */
+	friend void Java_com_epicgames_ue4_GameActivity_nativeUpdateAchievements(JNIEnv* LocalJNIEnv, jobject LocalThiz, jobjectArray Achievements);
+
 public:
 
 	// Begin IOnlineAchievements interface
diff --git a/Engine/Source/Runtime/Online/Android/OnlineSubsystemGooglePlay/Private/OnlineAsyncTaskManagerGooglePlay.h b/Engine/Source/Runtime/Online/Android/OnlineSubsystemGooglePlay/Private/OnlineAsyncTaskManagerGooglePlay.h
new file mode 100644
--- /dev/null
+++ b/Engine/Source/Runtime/Online/Android/OnlineSubsystemGooglePlay/Private/OnlineAsyncTaskManagerGooglePlay.h
@@ -0,0 +1,21 @@
+// Copyright 1998-2014 Epic Games, Inc. All Rights Reserved.
+
+#pragma once
+
+#include "OnlineAsyncTaskManager.h"
+#include "OnlineSubsystemGooglePlayPackage.h"
+
+/**
+ *	Google Play version of the async task manager
+ */
+class FOnlineAsyncTaskManagerGooglePlay : public FOnlineAsyncTaskManager
+{
+public:
+
+	FOnlineAsyncTaskManagerGooglePlay() {}
+
+	// FOnlineAsyncTaskManager
+	virtual void OnlineTick() OVERRIDE {}
+};
+
+
diff --git a/Engine/Source/Runtime/Online/Android/OnlineSubsystemGooglePlay/Private/OnlineSubsystemGooglePlay.cpp b/Engine/Source/Runtime/Online/Android/OnlineSubsystemGooglePlay/Private/OnlineSubsystemGooglePlay.cpp
--- a/Engine/Source/Runtime/Online/Android/OnlineSubsystemGooglePlay/Private/OnlineSubsystemGooglePlay.cpp
+++ b/Engine/Source/Runtime/Online/Android/OnlineSubsystemGooglePlay/Private/OnlineSubsystemGooglePlay.cpp
@@ -4,10 +4,10 @@
 
 #include "OnlineSubsystemGooglePlayPrivatePCH.h"
 #include "AndroidApplication.h"
+#include "OnlineAsyncTaskManagerGooglePlay.h"
 
 FOnlineSubsystemGooglePlay::FOnlineSubsystemGooglePlay()
-	: OnlineAsyncTaskThread(nullptr)
-	, IdentityInterface(nullptr)
+	: IdentityInterface(nullptr)
 	, LeaderboardsInterface(nullptr)
 	, AchievementsInterface(nullptr)
 {
@@ -94,6 +94,8 @@ bool FOnlineSubsystemGooglePlay::Init()
 {
 	FPlatformMisc::LowLevelOutputDebugStringf(TEXT("FOnlineSubsystemAndroid::Init"));
 	
+	OnlineAsyncTaskThreadRunnable.Reset(new FOnlineAsyncTaskManagerGooglePlay);
+
 	IdentityInterface = MakeShareable(new FOnlineIdentityGooglePlay(this));
 	LeaderboardsInterface = MakeShareable(new FOnlineLeaderboardsGooglePlay(this));
 	AchievementsInterface = MakeShareable(new FOnlineAchievementsGooglePlay(this));
@@ -107,6 +109,11 @@ bool FOnlineSubsystemGooglePlay::Init()
 
 bool FOnlineSubsystemGooglePlay::Tick(float DeltaTime)
 {
+	if (OnlineAsyncTaskThreadRunnable)
+	{
+		OnlineAsyncTaskThreadRunnable->GameTick();
+	}
+
 	return true;
 }
 
@@ -129,6 +136,8 @@ bool FOnlineSubsystemGooglePlay::Shutdown()
 	DESTRUCT_INTERFACE(IdentityInterface);
 #undef DESTRUCT_INTERFACE
 
+	OnlineAsyncTaskThreadRunnable.Reset();
+
 	return true;
 }
 
diff --git a/Engine/Source/Runtime/Online/Android/OnlineSubsystemGooglePlay/Public/OnlineSubsystemGooglePlay.h b/Engine/Source/Runtime/Online/Android/OnlineSubsystemGooglePlay/Public/OnlineSubsystemGooglePlay.h
--- a/Engine/Source/Runtime/Online/Android/OnlineSubsystemGooglePlay/Public/OnlineSubsystemGooglePlay.h
+++ b/Engine/Source/Runtime/Online/Android/OnlineSubsystemGooglePlay/Public/OnlineSubsystemGooglePlay.h
@@ -7,6 +7,7 @@
 #include "OnlineAchievementsInterfaceGooglePlay.h"
 #include "OnlineLeaderboardInterfaceGooglePlay.h"
 #include "OnlineExternalUIInterfaceGooglePlay.h"
+#include "UniquePtr.h"
 
 /**
  * OnlineSubsystemGooglePlay - Implementation of the online subsystem for Google Play services
@@ -61,10 +62,13 @@ PACKAGE_SCOPE:
 	 */
 	bool IsEnabled();
 
+	/** Return the async task manager owned by this subsystem */
+	class FOnlineAsyncTaskManagerGooglePlay* GetAsyncTaskManager() { return OnlineAsyncTaskThreadRunnable.Get(); }
+
 private:
 
-	/** Online async task thread */
-	class FRunnableThread* OnlineAsyncTaskThread;
+	/** Online async task runnable */
+	TUniquePtr<class FOnlineAsyncTaskManagerGooglePlay> OnlineAsyncTaskThreadRunnable;
 
 	/** Interface to the online identity system */
 	FOnlineIdentityGooglePlayPtr IdentityInterface;
