Fix shadowed variables

diff --git a/Engine/Source/Developer/AssetTools/Private/AssetTypeActions/AssetTypeActions_SlateWidgetStyle.cpp b/Engine/Source/Developer/AssetTools/Private/AssetTypeActions/AssetTypeActions_SlateWidgetStyle.cpp
--- a/Engine/Source/Developer/AssetTools/Private/AssetTypeActions/AssetTypeActions_SlateWidgetStyle.cpp
+++ b/Engine/Source/Developer/AssetTools/Private/AssetTypeActions/AssetTypeActions_SlateWidgetStyle.cpp
@@ -22,7 +22,7 @@ void FAssetTypeActions_SlateWidgetStyle::GetActions( const TArray<UObject*>& InO
 	);
 }
 
-void FAssetTypeActions_SlateWidgetStyle::OpenAssetEditor( const TArray<UObject*>& InObjects, TSharedPtr<IToolkitHost> EditWithinLevelEditor )
+void FAssetTypeActions_SlateWidgetStyle::OpenAssetEditor( const TArray<UObject*>& Objects, TSharedPtr<IToolkitHost> EditWithinLevelEditor )
 {
 	struct Local
 	{
@@ -41,7 +41,7 @@ void FAssetTypeActions_SlateWidgetStyle::OpenAssetEditor( const TArray<UObject*>
 		}
 	};
 
-	FSimpleAssetEditor::CreateEditor(EToolkitMode::Standalone, EditWithinLevelEditor, InObjects, FSimpleAssetEditor::FGetDetailsViewObjects::CreateStatic(&Local::GetSubObjects));
+	FSimpleAssetEditor::CreateEditor(EToolkitMode::Standalone, EditWithinLevelEditor, Objects, FSimpleAssetEditor::FGetDetailsViewObjects::CreateStatic(&Local::GetSubObjects));
 }
 
 void FAssetTypeActions_SlateWidgetStyle::ExecuteEdit(TArray<TWeakObjectPtr< USlateWidgetStyleAsset >> Styles)
diff --git a/Engine/Source/Editor/BehaviorTreeEditor/Private/BehaviorTreeGraph.cpp b/Engine/Source/Editor/BehaviorTreeEditor/Private/BehaviorTreeGraph.cpp
--- a/Engine/Source/Editor/BehaviorTreeEditor/Private/BehaviorTreeGraph.cpp
+++ b/Engine/Source/Editor/BehaviorTreeEditor/Private/BehaviorTreeGraph.cpp
@@ -161,8 +161,8 @@ void UBehaviorTreeGraph::ReplaceNodeConnections(UEdGraphNode* OldNode, UEdGraphN
 
 				if (LinkedNode == OldNode)
 				{
-					int32 PinIndex = LinkedNode->Pins.IndexOfByKey(LinkedPin);
-					Pin->LinkedTo[LinkedIndex] = NewNode->Pins[PinIndex];
+					const int32 LinkedPinIndex = LinkedNode->Pins.IndexOfByKey(LinkedPin);
+					Pin->LinkedTo[LinkedIndex] = NewNode->Pins[LinkedPinIndex];
 				}
 			}
 		}
diff --git a/Engine/Source/Editor/BlueprintGraph/Private/K2Node_MakeStruct.cpp b/Engine/Source/Editor/BlueprintGraph/Private/K2Node_MakeStruct.cpp
--- a/Engine/Source/Editor/BlueprintGraph/Private/K2Node_MakeStruct.cpp
+++ b/Engine/Source/Editor/BlueprintGraph/Private/K2Node_MakeStruct.cpp
@@ -128,16 +128,16 @@ void UK2Node_MakeStruct::FMakeStructPinManager::CustomizePinData(UEdGraphPin* Pi
 {
 	struct FPinDefaultValueHelper
 	{
-		static void Set(UEdGraphPin& Pin, const FString& Value, bool bIsText)
+		static void Set(UEdGraphPin& InPin, const FString& Value, bool bIsText)
 		{
-			Pin.AutogeneratedDefaultValue = Value;
+			InPin.AutogeneratedDefaultValue = Value;
 			if (bIsText)
 			{
-				Pin.DefaultTextValue = FText::FromString(Value);
+				InPin.DefaultTextValue = FText::FromString(Value);
 			}
 			else
 			{
-				Pin.DefaultValue = Value;
+				InPin.DefaultValue = Value;
 			}
 		}
 	};
diff --git a/Engine/Source/Editor/GraphEditor/Private/KismetNodes/SGraphNodeK2Var.cpp b/Engine/Source/Editor/GraphEditor/Private/KismetNodes/SGraphNodeK2Var.cpp
--- a/Engine/Source/Editor/GraphEditor/Private/KismetNodes/SGraphNodeK2Var.cpp
+++ b/Engine/Source/Editor/GraphEditor/Private/KismetNodes/SGraphNodeK2Var.cpp
@@ -62,8 +62,6 @@ void SGraphNodeK2Var::UpdateGraphNode()
 		}
 		else if (GraphNode->IsA(UK2Node_StructMemberSet::StaticClass()))
 		{
-			FFormatNamedArguments Args;
-			Args.Add(TEXT("VariableName"), StructOp->GetVarNameText());
 			TitleText = FText::Format(NSLOCTEXT("GraphEditor", "StructMemberSet", "Set in {VariableName}"), Args);
 		}
 		else if (GraphNode->IsA(UK2Node_MakeStruct::StaticClass()))
diff --git a/Engine/Source/Editor/Kismet/Private/FindInBlueprints.cpp b/Engine/Source/Editor/Kismet/Private/FindInBlueprints.cpp
--- a/Engine/Source/Editor/Kismet/Private/FindInBlueprints.cpp
+++ b/Engine/Source/Editor/Kismet/Private/FindInBlueprints.cpp
@@ -505,8 +505,8 @@ void SFindInBlueprints::MatchTokensWithinBlueprint(const TArray<FString>& Tokens
 		const FBPInterfaceDescription& InterfaceDesc = Blueprint->ImplementedInterfaces[i];
 		
 		// First check the interface name
- 		FString NodeSearchString = InterfaceDesc.Interface->GetName();
-		const bool bInterfaceMatchesSearch = StringMatchesSearchTokens(Tokens, NodeSearchString);
+ 		FString InterfaceNodeSearchString = InterfaceDesc.Interface->GetName();
+		const bool bInterfaceMatchesSearch = StringMatchesSearchTokens(Tokens, InterfaceNodeSearchString);
 
 		for (int32 FuncIdx = 0; FuncIdx < InterfaceDesc.Graphs.Num(); FuncIdx++)
 		{
@@ -539,7 +539,7 @@ void SFindInBlueprints::MatchTokensWithinBlueprint(const TArray<FString>& Tokens
 		// Did we match interface name but no functions ?
 		if( ( bInterfaceMatchesSearch == true ) && ( RootInterfaceResult->Children.Num() == 0 ) )
 		{
-			const FString InterfaceNameResult = FString::Printf(*LOCTEXT("BlueprintResultInterfaceImplementation", "Interface implementation: %s").ToString(), *NodeSearchString);
+			const FString InterfaceNameResult = FString::Printf(*LOCTEXT("BlueprintResultInterfaceImplementation", "Interface implementation: %s").ToString(), *InterfaceNodeSearchString);
 			const FSearchResult NodeResult(new FFindInBlueprintsResult(InterfaceNameResult, RootInterfaceResult, InterfaceDesc.Interface->GetClass(), 0));
 			RootInterfaceResult->Children.Add(NodeResult);
 		}
diff --git a/Engine/Source/Editor/Persona/Private/Persona.cpp b/Engine/Source/Editor/Persona/Private/Persona.cpp
--- a/Engine/Source/Editor/Persona/Private/Persona.cpp
+++ b/Engine/Source/Editor/Persona/Private/Persona.cpp
@@ -2363,9 +2363,9 @@ void FPersona::RemoveUnusedBones()
 		FAssetRegistryModule& AssetRegistryModule = FModuleManager::LoadModuleChecked<FAssetRegistryModule>(TEXT("AssetRegistry"));
 		AssetRegistryModule.Get().GetAssets(Filter, SkeletalMeshes);
 
-		FText Message = FText::Format( LOCTEXT("TimeTakenWarning", "In order to verify bone use all Skeletal Meshes that use this skeleton will be loaded, this may take some time.\n\nProceed?\n\nNumber of Meshes: {0}"), FText::AsNumber(SkeletalMeshes.Num()) );
+		FText TimeTakenMessage = FText::Format( LOCTEXT("TimeTakenWarning", "In order to verify bone use all Skeletal Meshes that use this skeleton will be loaded, this may take some time.\n\nProceed?\n\nNumber of Meshes: {0}"), FText::AsNumber(SkeletalMeshes.Num()) );
 		
-		if(FMessageDialog::Open( EAppMsgType::YesNo, Message ) == EAppReturnType::Yes)
+		if(FMessageDialog::Open( EAppMsgType::YesNo, TimeTakenMessage ) == EAppReturnType::Yes)
 		{
 			const FText StatusUpdate = FText::Format(LOCTEXT("RemoveUnusedBones_ProcessingAssets", "Processing Skeletal Meshes for {0}"), FText::FromString(TargetSkeleton->GetName()) );
 			GWarn->BeginSlowTask(StatusUpdate, true );
@@ -2412,10 +2412,10 @@ void FPersona::RemoveUnusedBones()
 			// If we have any bones left they are unused
 			if(SkeletonBones.Num() > 0)
 			{
-				const FText Message = FText::Format(LOCTEXT("RemoveBoneWarning", "Continuing will remove the following bones from the skeleton '{0}'. These bones are not being used by any of the SkeletalMeshes assigned to this skeleton\n\nOnce the bones have been removed all loaded animations for this skeleton will be recompressed (any that aren't loaded will be recompressed the next time they are loaded)."), FText::FromString(TargetSkeleton->GetName()) );
+				const FText RemoveBoneMessage = FText::Format(LOCTEXT("RemoveBoneWarning", "Continuing will remove the following bones from the skeleton '{0}'. These bones are not being used by any of the SkeletalMeshes assigned to this skeleton\n\nOnce the bones have been removed all loaded animations for this skeleton will be recompressed (any that aren't loaded will be recompressed the next time they are loaded)."), FText::FromString(TargetSkeleton->GetName()) );
 
 				// Ask User whether they would like to remove the bones from the skeleton
-				if (SSkeletonBoneRemoval::ShowModal(SkeletonBones, Message))
+				if (SSkeletonBoneRemoval::ShowModal(SkeletonBones, RemoveBoneMessage))
 				{
 					//Remove these bones from the skeleton
 					TargetSkeleton->RemoveBonesFromSkeleton(SkeletonBones, true);
diff --git a/Engine/Source/Editor/UMGEditor/Private/DetailCustomizations.h b/Engine/Source/Editor/UMGEditor/Private/DetailCustomizations.h
--- a/Engine/Source/Editor/UMGEditor/Private/DetailCustomizations.h
+++ b/Engine/Source/Editor/UMGEditor/Private/DetailCustomizations.h
@@ -11,8 +11,8 @@ public:
 		return MakeShareable(new FBlueprintWidgetCustomization(Blueprint));
 	}
 
-	FBlueprintWidgetCustomization(UBlueprint* Blueprint)
-		: Blueprint(CastChecked<UWidgetBlueprint>(Blueprint))
+	FBlueprintWidgetCustomization(UBlueprint* InBlueprint)
+		: Blueprint(CastChecked<UWidgetBlueprint>(InBlueprint))
 	{
 	}
 	
diff --git a/Engine/Source/Editor/UnrealEd/Private/Commandlets/CookCommandlet.cpp b/Engine/Source/Editor/UnrealEd/Private/Commandlets/CookCommandlet.cpp
--- a/Engine/Source/Editor/UnrealEd/Private/Commandlets/CookCommandlet.cpp
+++ b/Engine/Source/Editor/UnrealEd/Private/Commandlets/CookCommandlet.cpp
@@ -345,17 +345,12 @@ bool UCookCommandlet::SaveCookedPackage( UPackage* Package, uint32 SaveFlags, bo
 
 		if ( TargetPlatformNames.Num() )
 		{
-			for ( int Index = 0; Index < TargetPlatformNames.Num(); ++Index )
-			{
-				const FString &TargetPlatformName = TargetPlatformNames[Index];
-
-
-				const TArray<ITargetPlatform*>& TargetPlatforms = TPM.GetTargetPlatforms();	
+			const TArray<ITargetPlatform*>& TargetPlatforms = TPM.GetTargetPlatforms();
 
-
-				for ( int Index = 0; Index < TargetPlatforms.Num(); ++Index )
+			for (const FString &TargetPlatformName : TargetPlatformNames)
+			{
+				for (ITargetPlatform *TargetPlatform : TargetPlatforms)
 				{
-					ITargetPlatform *TargetPlatform = TargetPlatforms[ Index ];
 					if ( TargetPlatform->PlatformName() == TargetPlatformName )
 					{
 						Platforms.Add( TargetPlatform );
diff --git a/Engine/Source/Editor/UnrealEd/Private/CookOnTheFlyServer.cpp b/Engine/Source/Editor/UnrealEd/Private/CookOnTheFlyServer.cpp
--- a/Engine/Source/Editor/UnrealEd/Private/CookOnTheFlyServer.cpp
+++ b/Engine/Source/Editor/UnrealEd/Private/CookOnTheFlyServer.cpp
@@ -311,9 +311,9 @@ uint32 UCookOnTheFlyServer::TickCookOnTheSide( const float TimeSlice, uint32 &Co
 					const FString FullPath = FPaths::ConvertRelativePathToFull(PackageFilename);
 
 					bool AlreadyCooked = true;
-					for ( int Index = 0; Index < AllTargetPlatformNames.Num(); ++Index )
+					for ( const FName TargetPlatformName : AllTargetPlatformNames )
 					{
-						if ( !ThreadSafeFilenameSet.Exists(FFilePlatformRequest(PackageFilename, AllTargetPlatformNames[Index])) )
+						if ( !ThreadSafeFilenameSet.Exists(FFilePlatformRequest(PackageFilename, TargetPlatformName)) )
 						{
 							AlreadyCooked = false;
 							break;
@@ -321,9 +321,9 @@ uint32 UCookOnTheFlyServer::TickCookOnTheSide( const float TimeSlice, uint32 &Co
 					}
 					if ( !AlreadyCooked )
 					{
-						for ( int Index = 0; Index < AllTargetPlatformNames.Num(); ++Index )
+						for (const FName TargetPlatformName : AllTargetPlatformNames)
 						{
-							const FFilePlatformRequest Request( PackageFilename, AllTargetPlatformNames[Index] );
+							const FFilePlatformRequest Request( PackageFilename, TargetPlatformName );
 							UnsolicitedFileRequests.EnqueueUnique(Request);
 						}
 					}
@@ -547,17 +547,14 @@ bool UCookOnTheFlyServer::SaveCookedPackage( UPackage* Package, uint32 SaveFlags
 
 		if ( TargetPlatformNames.Num() )
 		{
-			for ( int Index = 0; Index < TargetPlatformNames.Num(); ++Index )
-			{
-				const FString TargetPlatformName = TargetPlatformNames[Index].ToString();
-
-
-				const TArray<ITargetPlatform*>& TargetPlatforms = TPM.GetTargetPlatforms();	
+			const TArray<ITargetPlatform*>& TargetPlatforms = TPM.GetTargetPlatforms();
 
+			for (const FName TargetPlatformFName : TargetPlatformNames)
+			{
+				const FString TargetPlatformName = TargetPlatformFName.ToString();
 
-				for ( int Index = 0; Index < TargetPlatforms.Num(); ++Index )
+				for (ITargetPlatform *TargetPlatform  : TargetPlatforms)
 				{
-					ITargetPlatform *TargetPlatform = TargetPlatforms[ Index ];
 					if ( TargetPlatform->PlatformName() == TargetPlatformName )
 					{
 						Platforms.Add( TargetPlatform );
@@ -569,16 +566,15 @@ bool UCookOnTheFlyServer::SaveCookedPackage( UPackage* Package, uint32 SaveFlags
 		{
 			Platforms = ActiveStartupPlatforms;
 
-			for ( int Index = 0; Index < Platforms.Num(); ++Index )
+			for (ITargetPlatform *Platform : Platforms)
 			{
-				TargetPlatformNames.Add( FName(*Platforms[Index]->PlatformName()) );
+				TargetPlatformNames.Add(FName(*Platform->PlatformName()));
 			}
 		}
 		
 
-		for (int32 Index = 0; Index < Platforms.Num(); Index++)
+		for (ITargetPlatform* Target : Platforms)
 		{
-			ITargetPlatform* Target = Platforms[Index];
 			FString PlatFilename = Filename.Replace(TEXT("[Platform]"), *Target->PlatformName());
 
 			// If we are not iterative cooking, then cook the package
diff --git a/Engine/Source/Editor/UnrealEd/Private/Fbx/FbxMainExport.cpp b/Engine/Source/Editor/UnrealEd/Private/Fbx/FbxMainExport.cpp
--- a/Engine/Source/Editor/UnrealEd/Private/Fbx/FbxMainExport.cpp
+++ b/Engine/Source/Editor/UnrealEd/Private/Fbx/FbxMainExport.cpp
@@ -1294,8 +1294,8 @@ void ConvertInterpToFBX(uint8 UnrealInterpMode, FbxAnimCurveDef::EInterpolationT
 		Tangent = (FbxAnimCurveDef::ETangentMode) (FbxAnimCurveDef::eTangentAuto | FbxAnimCurveDef::eTangentGenericClamp);
 		break;
 	case CIM_Unknown:  // ???
-		FbxAnimCurveDef::EInterpolationType Interpolation = FbxAnimCurveDef::eInterpolationConstant;
-		FbxAnimCurveDef::ETangentMode Tangent = FbxAnimCurveDef::eTangentAuto;
+		Interpolation = FbxAnimCurveDef::eInterpolationConstant;
+		Tangent = FbxAnimCurveDef::eTangentAuto;
 		break;
 	}
 }
diff --git a/Engine/Source/Editor/UnrealEd/Private/UnrealEdSrv.cpp b/Engine/Source/Editor/UnrealEd/Private/UnrealEdSrv.cpp
--- a/Engine/Source/Editor/UnrealEd/Private/UnrealEdSrv.cpp
+++ b/Engine/Source/Editor/UnrealEd/Private/UnrealEdSrv.cpp
@@ -1082,9 +1082,9 @@ bool UUnrealEdEngine::Exec( UWorld* InWorld, const TCHAR* Stream, FOutputDevice&
 				const FString& SourceFilePath = Data->SourceFilePath;
 				if( !SourceFilePath.IsEmpty() )
 				{
-					for (const FString& Str : SearchTerms)
+					for (const FString& SearchTerm : SearchTerms)
 					{
-						if (SourceFilePath.Contains(Str))
+						if (SourceFilePath.Contains(SearchTerm))
 						{
 							Data->Modify();
 							UE_LOG(LogUnrealEdSrv, Log, TEXT("Removing Path: %s"), *SourceFilePath);
diff --git a/Engine/Source/Runtime/Engine/Classes/Animation/AnimNodeBase.h b/Engine/Source/Runtime/Engine/Classes/Animation/AnimNodeBase.h
--- a/Engine/Source/Runtime/Engine/Classes/Animation/AnimNodeBase.h
+++ b/Engine/Source/Runtime/Engine/Classes/Animation/AnimNodeBase.h
@@ -180,7 +180,7 @@ struct FNodeDebugData
 private:
 	struct DebugItem
 	{
-		DebugItem(FString Data, bool bPoseSource) : DebugData(Data), bPoseSource(bPoseSource) {}
+		DebugItem(FString Data, bool bInPoseSource) : DebugData(Data), bPoseSource(bInPoseSource) {}
 		
 		// This node items debug text to display
 		FString DebugData;
@@ -201,7 +201,7 @@ private:
 public:
 	struct FFlattenedDebugData
 	{
-		FFlattenedDebugData(FString Line, float AbsWeight, int32 Indent, int32 ChainID, bool bPoseSource) : DebugLine(Line), AbsoluteWeight(AbsWeight), Indent(Indent), ChainID(ChainID), bPoseSource(bPoseSource){}
+		FFlattenedDebugData(FString Line, float AbsWeight, int32 InIndent, int32 InChainID, bool bInPoseSource) : DebugLine(Line), AbsoluteWeight(AbsWeight), Indent(InIndent), ChainID(InChainID), bPoseSource(bInPoseSource){}
 		FString DebugLine;
 		float	AbsoluteWeight;
 		int32	Indent;
@@ -211,8 +211,8 @@ public:
 		bool IsOnActiveBranch() { return AbsoluteWeight > ZERO_ANIMWEIGHT_THRESH; }
 	};
 
-	FNodeDebugData(const class UAnimInstance* AnimInstance) : AbsoluteWeight(1.f), AnimInstance(AnimInstance) {}
-	FNodeDebugData(const class UAnimInstance* AnimInstance, float AbsWeight) : AbsoluteWeight(AbsWeight), AnimInstance(AnimInstance) {}
+	FNodeDebugData(const class UAnimInstance* InAnimInstance) : AbsoluteWeight(1.f), AnimInstance(InAnimInstance) {}
+	FNodeDebugData(const class UAnimInstance* InAnimInstance, const float AbsWeight) : AbsoluteWeight(AbsWeight), AnimInstance(InAnimInstance) {}
 
 	void			AddDebugItem(FString DebugData, bool bPoseSource = false);
 	FNodeDebugData&	BranchFlow(float BranchWeight);
diff --git a/Engine/Source/Runtime/Engine/Private/PhysicsEngine/PhysXSupport.h b/Engine/Source/Runtime/Engine/Private/PhysicsEngine/PhysXSupport.h
--- a/Engine/Source/Runtime/Engine/Private/PhysicsEngine/PhysXSupport.h
+++ b/Engine/Source/Runtime/Engine/Private/PhysicsEngine/PhysXSupport.h
@@ -73,9 +73,11 @@ private:
 	PxScene* PScene;
 };
 
+#define SCOPED_SCENE_TOKENPASTE_INNER(x,y) x##y
+#define SCOPED_SCENE_TOKENPASTE(x,y) SCOPED_SCENE_TOKENPASTE_INNER(x,y)
+#define SCOPED_SCENE_READ_LOCK( _scene ) FPhysXSceneReadLock SCOPED_SCENE_TOKENPASTE(_rlock,__LINE__)(_scene)
+#define SCOPED_SCENE_WRITE_LOCK( _scene ) FPhysXSceneWriteLock SCOPED_SCENE_TOKENPASTE(_wlock,__LINE__)(_scene)
 
-#define SCOPED_SCENE_READ_LOCK( _scene ) FPhysXSceneReadLock _rlock##__LINE__(_scene)
-#define SCOPED_SCENE_WRITE_LOCK( _scene ) FPhysXSceneWriteLock _wlock##__LINE__(_scene)
 #define SCENE_LOCK_READ( _scene ) if((_scene) != NULL) { (_scene)->lockRead(); }
 #define SCENE_UNLOCK_READ( _scene ) if((_scene) != NULL) { (_scene)->unlockRead(); }
 #define SCENE_LOCK_WRITE( _scene ) if((_scene) != NULL) { (_scene)->lockWrite(); }
diff --git a/Engine/Source/Runtime/Engine/Public/EngineUtils.h b/Engine/Source/Runtime/Engine/Public/EngineUtils.h
--- a/Engine/Source/Runtime/Engine/Public/EngineUtils.h
+++ b/Engine/Source/Runtime/Engine/Public/EngineUtils.h
@@ -25,10 +25,10 @@ struct HActor : public HHitProxy
 		, MaterialIndex( -1 )
 		{}
 
-	HActor(AActor* InActor, const UPrimitiveComponent* InPrimComponent, int32 MaterialIndex)
+	HActor(AActor* InActor, const UPrimitiveComponent* InPrimComponent, int32 InMaterialIndex)
 		: Actor( InActor )
 		, PrimComponent( InPrimComponent )
-		, MaterialIndex( MaterialIndex )
+		, MaterialIndex( InMaterialIndex )
 		{}
 
 	HActor( AActor* InActor, const UPrimitiveComponent* InPrimComponent, EHitProxyPriority InPriority) 
@@ -38,11 +38,11 @@ struct HActor : public HHitProxy
 		, MaterialIndex( -1 )
 		{}
 
-	HActor(AActor* InActor, const UPrimitiveComponent* InPrimComponent, EHitProxyPriority InPriority, int32 MaterialIndex)
+	HActor(AActor* InActor, const UPrimitiveComponent* InPrimComponent, EHitProxyPriority InPriority, int32 InMaterialIndex)
 		: HHitProxy(InPriority)
 		, Actor(InActor)
 		, PrimComponent(InPrimComponent)
-		, MaterialIndex(MaterialIndex)
+		, MaterialIndex(InMaterialIndex)
 		{}
 
 	virtual void AddReferencedObjects( FReferenceCollector& Collector ) OVERRIDE
diff --git a/Engine/Source/Runtime/Engine/Public/SceneUtils.h b/Engine/Source/Runtime/Engine/Public/SceneUtils.h
--- a/Engine/Source/Runtime/Engine/Public/SceneUtils.h
+++ b/Engine/Source/Runtime/Engine/Public/SceneUtils.h
@@ -79,10 +79,12 @@ inline void appSetCounterValue(const TCHAR* CounterName, float Value) {}
 	};
 
 	// Macros to allow for scoping of draw events
-	#define SCOPED_DRAW_EVENT(Name, Color) FDrawEvent Event_##Name; if(GEmitDrawEvents) Event_##Name.Start(Color, TEXT(#Name));
-	#define SCOPED_DRAW_EVENTF(Name, Color, Format, ...) FDrawEvent Event_##Name; if(GEmitDrawEvents) Event_##Name.Start(Color, Format, ##__VA_ARGS__);
-	#define SCOPED_CONDITIONAL_DRAW_EVENT(Name, Condition, Color) FDrawEvent Event_##Name; if(GEmitDrawEvents && (Condition)) Event_##Name.Start(Color, TEXT(#Name));
-	#define SCOPED_CONDITIONAL_DRAW_EVENTF(Name, Condition, Color, Format, ...) FDrawEvent Event_##Name; if(GEmitDrawEvents && (Condition)) Event_##Name.Start(Color, Format, ##__VA_ARGS__);
+	#define SCOPED_DRAW_TOKENPASTE_INNER(x,y) x##y
+	#define SCOPED_DRAW_TOKENPASTE(x,y) SCOPED_DRAW_TOKENPASTE_INNER(x,y)
+	#define SCOPED_DRAW_EVENT(Name, Color) FDrawEvent SCOPED_DRAW_TOKENPASTE(Event_##Name,__LINE__); if(GEmitDrawEvents) SCOPED_DRAW_TOKENPASTE(Event_##Name,__LINE__).Start(Color, TEXT(#Name));
+	#define SCOPED_DRAW_EVENTF(Name, Color, Format, ...) FDrawEvent SCOPED_DRAW_TOKENPASTE(Event_##Name,__LINE__); if(GEmitDrawEvents) SCOPED_DRAW_TOKENPASTE(Event_##Name,__LINE__).Start(Color, Format, ##__VA_ARGS__);
+	#define SCOPED_CONDITIONAL_DRAW_EVENT(Name, Condition, Color) FDrawEvent SCOPED_DRAW_TOKENPASTE(Event_##Name,__LINE__); if(GEmitDrawEvents && (Condition)) SCOPED_DRAW_TOKENPASTE(Event_##Name,__LINE__).Start(Color, TEXT(#Name));
+	#define SCOPED_CONDITIONAL_DRAW_EVENTF(Name, Condition, Color, Format, ...) FDrawEvent SCOPED_DRAW_TOKENPASTE(Event_##Name,__LINE__); if(GEmitDrawEvents && (Condition)) SCOPED_DRAW_TOKENPASTE(Event_##Name,__LINE__).Start(Color, Format, ##__VA_ARGS__);
 
 #else
 
diff --git a/Engine/Source/Runtime/Online/OnlineSubsystemUtils/Private/IpNetDriver.cpp b/Engine/Source/Runtime/Online/OnlineSubsystemUtils/Private/IpNetDriver.cpp
--- a/Engine/Source/Runtime/Online/OnlineSubsystemUtils/Private/IpNetDriver.cpp
+++ b/Engine/Source/Runtime/Online/OnlineSubsystemUtils/Private/IpNetDriver.cpp
@@ -372,9 +372,9 @@ bool UIpNetDriver::HandleSocketsCommand( const TCHAR* Cmd, FOutputDevice& Ar, UW
 	Ar.Logf(TEXT(""));
 	if (Socket != NULL)
 	{
-		TSharedRef<FInternetAddr> LocalAddr = GetSocketSubsystem()->CreateInternetAddr();
-		Socket->GetAddress(*LocalAddr);
-		Ar.Logf(TEXT("%s Socket: %s"), *GetDescription(), *LocalAddr->ToString(true));
+		TSharedRef<FInternetAddr> LocalInternetAddr = GetSocketSubsystem()->CreateInternetAddr();
+		Socket->GetAddress(*LocalInternetAddr);
+		Ar.Logf(TEXT("%s Socket: %s"), *GetDescription(), *LocalInternetAddr->ToString(true));
 	}		
 	else
 	{
diff --git a/Engine/Source/Runtime/Online/OnlineSubsystemUtils/Private/OnlineBeacon.cpp b/Engine/Source/Runtime/Online/OnlineSubsystemUtils/Private/OnlineBeacon.cpp
--- a/Engine/Source/Runtime/Online/OnlineSubsystemUtils/Private/OnlineBeacon.cpp
+++ b/Engine/Source/Runtime/Online/OnlineSubsystemUtils/Private/OnlineBeacon.cpp
@@ -33,9 +33,9 @@ void AOnlineBeacon::DestroyBeacon()
 	Destroy();
 }
 
-void AOnlineBeacon::HandleNetworkFailure(UWorld *World, UNetDriver *NetDriver, ENetworkFailure::Type FailureType, const FString& ErrorString)
+void AOnlineBeacon::HandleNetworkFailure(UWorld *World, UNetDriver *InNetDriver, ENetworkFailure::Type FailureType, const FString& ErrorString)
 {
-	if (NetDriver && NetDriver->NetDriverName == NetDriverName)
+	if (InNetDriver && InNetDriver->NetDriverName == NetDriverName)
 	{
 		OnFailure();
 	}
diff --git a/Engine/Source/Runtime/Online/OnlineSubsystemUtils/Private/OnlineSessionClient.cpp b/Engine/Source/Runtime/Online/OnlineSubsystemUtils/Private/OnlineSessionClient.cpp
--- a/Engine/Source/Runtime/Online/OnlineSubsystemUtils/Private/OnlineSessionClient.cpp
+++ b/Engine/Source/Runtime/Online/OnlineSubsystemUtils/Private/OnlineSessionClient.cpp
@@ -80,12 +80,12 @@ void UOnlineSessionClient::ClearOnlineDelegates(UWorld* InWorld)
 	IOnlineSubsystem* OnlineSub = Online::GetSubsystem(InWorld);
 	if (OnlineSub)
 	{
-		IOnlineSessionPtr SessionInt = OnlineSub->GetSessionInterface();
+		IOnlineSessionPtr OnlineSessionInt = OnlineSub->GetSessionInterface();
 
 		int32 ControllerId = GetControllerId();
 		if (ControllerId != INVALID_CONTROLLERID)
 		{
-			SessionInt->ClearOnSessionInviteAcceptedDelegate(ControllerId, OnSessionInviteAcceptedDelegate);
+			OnlineSessionInt->ClearOnSessionInviteAcceptedDelegate(ControllerId, OnSessionInviteAcceptedDelegate);
 		}
 	}
 }
diff --git a/Engine/Source/Runtime/Online/OnlineSubsystemUtils/Private/Tests/TestEntitlementsInterface.cpp b/Engine/Source/Runtime/Online/OnlineSubsystemUtils/Private/Tests/TestEntitlementsInterface.cpp
--- a/Engine/Source/Runtime/Online/OnlineSubsystemUtils/Private/Tests/TestEntitlementsInterface.cpp
+++ b/Engine/Source/Runtime/Online/OnlineSubsystemUtils/Private/Tests/TestEntitlementsInterface.cpp
@@ -105,7 +105,7 @@ void FTestEntitlementsInterface::FinishTest()
 	delete this;
 }
 
-void FTestEntitlementsInterface::OnCreateEntitlementsGroupComplete(bool bWasSuccessful, const FString& Error, const FUniqueNetId& UserId)
+void FTestEntitlementsInterface::OnCreateEntitlementsGroupComplete(bool bWasSuccessful, const FString& Error, const FUniqueNetId& )
 {
 	UE_LOG(LogOnline, Log, TEXT("created entitlements group. Result=%s Error=[%s]"), 
 		bWasSuccessful ? TEXT("true") : TEXT("false"), *Error);
@@ -124,7 +124,7 @@ void FTestEntitlementsInterface::OnCreateEntitlementsGroupComplete(bool bWasSucc
 	}
 }
 
-void FTestEntitlementsInterface::OnGrantEntitlementsComplete(bool bWasSuccessful, const FString& Error, const FUniqueNetId& UserId, TArray<FEntitlement>& NewEntitlements, TArray<FEntitlement>& ModifiedEntitlements)
+void FTestEntitlementsInterface::OnGrantEntitlementsComplete(bool bWasSuccessful, const FString& Error, const FUniqueNetId& InUserId, TArray<FEntitlement>& NewEntitlements, TArray<FEntitlement>& ModifiedEntitlements)
 {
 	UE_LOG(LogOnline, Log, TEXT("granted entitlements. Result=%s Error=[%s]"), 
 		bWasSuccessful ? TEXT("true") : TEXT("false"), *Error);
@@ -135,7 +135,7 @@ void FTestEntitlementsInterface::OnGrantEntitlementsComplete(bool bWasSuccessful
 		// @todo check the returned new/modified in addition to getting the full cache to check
 		// Check that the cached entitlements are correct after granting 1
 		TMap<FString, FEntitlement> Entitlements;
-		EntitlementsOSS->GetEntitlements(UserId, Entitlements);
+		EntitlementsOSS->GetEntitlements(InUserId, Entitlements);
 		if (Entitlements.Num() != 1)
 		{
 			UE_LOG(LogOnline, Error, TEXT("Incorrect number of entitlement1 found after granting 1 entitlement: %d"), Entitlements.Num());
@@ -166,7 +166,7 @@ void FTestEntitlementsInterface::OnGrantEntitlementsComplete(bool bWasSuccessful
 	{
 		// Check on the plural entitlement granting results in local cache
 		TMap<FString, FEntitlement> Entitlements;
-		EntitlementsOSS->GetEntitlements(UserId, Entitlements);
+		EntitlementsOSS->GetEntitlements(InUserId, Entitlements);
 		if (Entitlements.Num() != 4)
 		{
 			UE_LOG(LogOnline, Error, TEXT("Incorrect number of entitlements found after granting multiple entitlement: %d"), Entitlements.Num());
@@ -205,7 +205,7 @@ void FTestEntitlementsInterface::OnGrantEntitlementsComplete(bool bWasSuccessful
 	}
 }
 
-void FTestEntitlementsInterface::OnConsumeEntitlementComplete(bool bWasSuccessful, const FString& Error, const FUniqueNetId& UserId, TArray<FEntitlement>& NewEntitlements, TArray<FEntitlement>& ModifiedEntitlements)
+void FTestEntitlementsInterface::OnConsumeEntitlementComplete(bool bWasSuccessful, const FString& Error, const FUniqueNetId& InUserId, TArray<FEntitlement>& NewEntitlements, TArray<FEntitlement>& ModifiedEntitlements)
 {
 	UE_LOG(LogOnline, Log, TEXT("consumed entitlement. Result=%s Error=[%s]"), 
 		bWasSuccessful ? TEXT("true") : TEXT("false"), *Error);
@@ -213,7 +213,7 @@ void FTestEntitlementsInterface::OnConsumeEntitlementComplete(bool bWasSuccessfu
 	// @todo check the returned new/modified in addition to getting the full cache to check
 	// Check on the entitlement consumption results in local cache
 	TMap<FString, FEntitlement> Entitlements;
-	EntitlementsOSS->GetEntitlements(UserId, Entitlements);
+	EntitlementsOSS->GetEntitlements(InUserId, Entitlements);
 	if (Entitlements.Num() != 4)
 	{
 		UE_LOG(LogOnline, Error, TEXT("Incorrect number of entitlements found after granting multiple entitlement: %d"), Entitlements.Num());
@@ -252,14 +252,14 @@ void FTestEntitlementsInterface::OnConsumeEntitlementComplete(bool bWasSuccessfu
 	}
 }
 
-void FTestEntitlementsInterface::OnEnumerateEntitlementsComplete(bool bWasSuccessful, const FString& Error, const FUniqueNetId& UserId)
+void FTestEntitlementsInterface::OnEnumerateEntitlementsComplete(bool bWasSuccessful, const FString& Error, const FUniqueNetId& InUserId)
 {
 	UE_LOG(LogOnline, Log, TEXT("enumerated entitlements. Result=%s Error=[%s]"), 
 		bWasSuccessful ? TEXT("true") : TEXT("false"), *Error);
 
 	// Now check em out
 	TMap<FString, FEntitlement> Entitlements;
-	EntitlementsOSS->GetEntitlements(UserId, Entitlements);
+	EntitlementsOSS->GetEntitlements(InUserId, Entitlements);
 	if (Entitlements.Num() != 4)
 	{
 		UE_LOG(LogOnline, Error, TEXT("Incorrect number of entitlements found: %d"), Entitlements.Num());
@@ -292,7 +292,7 @@ void FTestEntitlementsInterface::OnEnumerateEntitlementsComplete(bool bWasSucces
 	}
 }
 
-void FTestEntitlementsInterface::OnDeleteEntitlementsGroupComplete(bool bWasSuccessful, const FString& Error, const FUniqueNetId& UserId)
+void FTestEntitlementsInterface::OnDeleteEntitlementsGroupComplete(bool bWasSuccessful, const FString& Error, const FUniqueNetId& )
 {
 	UE_LOG(LogOnline, Log, TEXT("deleting entitlements group. Result=%s Error=[%s]"), 
 		bWasSuccessful ? TEXT("true") : TEXT("false"), *Error);
diff --git a/Engine/Source/Runtime/Online/OnlineSubsystemUtils/Private/Tests/TestVoice.cpp b/Engine/Source/Runtime/Online/OnlineSubsystemUtils/Private/Tests/TestVoice.cpp
--- a/Engine/Source/Runtime/Online/OnlineSubsystemUtils/Private/Tests/TestVoice.cpp
+++ b/Engine/Source/Runtime/Online/OnlineSubsystemUtils/Private/Tests/TestVoice.cpp
@@ -147,7 +147,7 @@ bool FTestVoice::Tick(float DeltaTime)
 						FMemory::Memcpy(RawCaptureData.GetTypedData(), Remainder.GetTypedData(), LastRemainderSize);
 					}
 
-					EVoiceCaptureState::Type MicState = VoiceCapture->GetVoiceData(RawCaptureData.GetTypedData() + LastRemainderSize, NewVoiceDataBytes, NewVoiceDataBytes);
+					MicState = VoiceCapture->GetVoiceData(RawCaptureData.GetTypedData() + LastRemainderSize, NewVoiceDataBytes, NewVoiceDataBytes);
 					TotalVoiceBytes = NewVoiceDataBytes + LastRemainderSize;
 					bDoWork = MicState == EVoiceCaptureState::Ok;
 				}
diff --git a/Engine/Source/Runtime/Renderer/Private/PostProcess/PostProcessTemporalAA.cpp b/Engine/Source/Runtime/Renderer/Private/PostProcess/PostProcessTemporalAA.cpp
--- a/Engine/Source/Runtime/Renderer/Private/PostProcess/PostProcessTemporalAA.cpp
+++ b/Engine/Source/Runtime/Renderer/Private/PostProcess/PostProcessTemporalAA.cpp
@@ -453,13 +453,14 @@ void FRCPassPostProcessTemporalAA::Process(FRenderingCompositePassContext& Conte
 	uint32 Quality = FMath::Clamp(CVar->GetValueOnRenderThread(), 1, 6);
 	bool bUseFast = Quality == 3;
 
-	TShaderMapRef< FPostProcessTonemapVS >			VertexShader( GetGlobalShaderMap() );
 
-	{	// Normal temporal feedback
+	{	
+		// Normal temporal feedback
 		// Draw to pixels where stencil == 0
 		RHISetDepthStencilState(TStaticDepthStencilState<false,CF_Always,true,CF_Equal,SO_Keep,SO_Keep,SO_Keep>::GetRHI(), 0);
 
-		if(bUseFast)
+		TShaderMapRef< FPostProcessTonemapVS >			VertexShader(GetGlobalShaderMap());
+		if (bUseFast)
 		{
 			TShaderMapRef< FPostProcessTemporalAAPS<4,0> >	PixelShader( GetGlobalShaderMap() );
 
