Add matchmaking functions to IOnlineSession and new OSS interface, IOnlineParty.
Matchmaking and IOnlineParty are only implemented for Xbox One currently.
Based on Lionhead's contributions.

diff --git a/Engine/Source/Runtime/Online/Android/OnlineSubsystemGooglePlay/Private/OnlineSubsystemGooglePlay.cpp b/Engine/Source/Runtime/Online/Android/OnlineSubsystemGooglePlay/Private/OnlineSubsystemGooglePlay.cpp
--- a/Engine/Source/Runtime/Online/Android/OnlineSubsystemGooglePlay/Private/OnlineSubsystemGooglePlay.cpp
+++ b/Engine/Source/Runtime/Online/Android/OnlineSubsystemGooglePlay/Private/OnlineSubsystemGooglePlay.cpp
@@ -72,14 +72,16 @@ IOnlineTimePtr FOnlineSubsystemGooglePlay::GetTimeInterface() const
 	return nullptr;
 }
 
-
+IOnlinePartyPtr FOnlineSubsystemGooglePlay::GetPartyInterface() const
+{
+	return nullptr;
+}
 
 IOnlineTitleFilePtr FOnlineSubsystemGooglePlay::GetTitleFileInterface() const
 {
 	return nullptr;
 }
 
-
 IOnlineEntitlementsPtr FOnlineSubsystemGooglePlay::GetEntitlementsInterface() const
 {
 	return nullptr;
diff --git a/Engine/Source/Runtime/Online/Android/OnlineSubsystemGooglePlay/Public/OnlineSubsystemGooglePlay.h b/Engine/Source/Runtime/Online/Android/OnlineSubsystemGooglePlay/Public/OnlineSubsystemGooglePlay.h
--- a/Engine/Source/Runtime/Online/Android/OnlineSubsystemGooglePlay/Public/OnlineSubsystemGooglePlay.h
+++ b/Engine/Source/Runtime/Online/Android/OnlineSubsystemGooglePlay/Public/OnlineSubsystemGooglePlay.h
@@ -30,6 +30,7 @@ public:
 	virtual IOnlineExternalUIPtr GetExternalUIInterface() const OVERRIDE;
 	virtual IOnlineTimePtr GetTimeInterface() const OVERRIDE;
 	virtual IOnlineIdentityPtr GetIdentityInterface() const OVERRIDE;
+	virtual IOnlinePartyPtr GetPartyInterface() const OVERRIDE;
 	virtual IOnlineTitleFilePtr GetTitleFileInterface() const OVERRIDE;
 	virtual IOnlineEntitlementsPtr GetEntitlementsInterface() const OVERRIDE;
 	virtual IOnlineStorePtr GetStoreInterface() const OVERRIDE;
diff --git a/Engine/Source/Runtime/Online/IOS/OnlineSubsystemIOS/Private/OnlineSessionInterfaceIOS.cpp b/Engine/Source/Runtime/Online/IOS/OnlineSubsystemIOS/Private/OnlineSessionInterfaceIOS.cpp
--- a/Engine/Source/Runtime/Online/IOS/OnlineSubsystemIOS/Private/OnlineSessionInterfaceIOS.cpp
+++ b/Engine/Source/Runtime/Online/IOS/OnlineSubsystemIOS/Private/OnlineSessionInterfaceIOS.cpp
@@ -312,6 +312,20 @@ bool FOnlineSessionIOS::IsPlayerInSession(FName SessionName, const FUniqueNetId&
 	return IsPlayerInSessionImpl(this, SessionName, UniqueId);
 }
 
+bool FOnlineSessionIOS::StartMatchmaking(int32 SearchingPlayerNum, FName SessionName, const FOnlineSessionSettings& NewSessionSettings, TSharedRef<FOnlineSessionSearch>& SearchSettings)
+{
+	UE_LOG(LogOnline, Warning, TEXT("Matchmaking is not supported on this platform."));
+	TriggerOnMatchmakingCompleteDelegates(SessionName, false);
+	return false;
+}
+
+bool FOnlineSessionIOS::CancelMatchmaking(int32 SearchingPlayerNum, FName SessionName)
+{
+	UE_LOG(LogOnline, Warning, TEXT("Matchmaking is not supported on this platform."));
+	TriggerOnCancelMatchmakingCompleteDelegates(SessionName, false);
+	return false;
+}
+
 bool FOnlineSessionIOS::FindSessions(int32 SearchingPlayerNum, const TSharedRef<FOnlineSessionSearch>& SearchSettings)
 {
 	bool bSuccessfullyFoundSessions = false;
diff --git a/Engine/Source/Runtime/Online/IOS/OnlineSubsystemIOS/Private/OnlineSessionInterfaceIOS.h b/Engine/Source/Runtime/Online/IOS/OnlineSubsystemIOS/Private/OnlineSessionInterfaceIOS.h
--- a/Engine/Source/Runtime/Online/IOS/OnlineSubsystemIOS/Private/OnlineSessionInterfaceIOS.h
+++ b/Engine/Source/Runtime/Online/IOS/OnlineSubsystemIOS/Private/OnlineSessionInterfaceIOS.h
@@ -90,6 +90,10 @@ public:
 
 	virtual bool IsPlayerInSession(FName SessionName, const FUniqueNetId& UniqueId) OVERRIDE;
 
+	virtual bool StartMatchmaking(int32 SearchingPlayerNum, FName SessionName, const FOnlineSessionSettings& NewSessionSettings, TSharedRef<FOnlineSessionSearch>& SearchSettings) OVERRIDE;
+
+	virtual bool CancelMatchmaking(int32 SearchingPlayerNum, FName SessionName) OVERRIDE;
+
 	virtual bool FindSessions(int32 SearchingPlayerNum, const TSharedRef<FOnlineSessionSearch>& SearchSettings) OVERRIDE;
 
 	virtual bool CancelFindSessions() OVERRIDE;
diff --git a/Engine/Source/Runtime/Online/IOS/OnlineSubsystemIOS/Private/OnlineSubsystemIOS.cpp b/Engine/Source/Runtime/Online/IOS/OnlineSubsystemIOS/Private/OnlineSubsystemIOS.cpp
--- a/Engine/Source/Runtime/Online/IOS/OnlineSubsystemIOS/Private/OnlineSubsystemIOS.cpp
+++ b/Engine/Source/Runtime/Online/IOS/OnlineSubsystemIOS/Private/OnlineSubsystemIOS.cpp
@@ -62,6 +62,12 @@ IOnlineIdentityPtr FOnlineSubsystemIOS::GetIdentityInterface() const
 }
 
 
+IOnlinePartyPtr FOnlineSubsystemIOS::GetPartyInterface() const
+{
+	return NULL;
+}
+
+
 IOnlineTitleFilePtr FOnlineSubsystemIOS::GetTitleFileInterface() const
 {
 	return NULL;
diff --git a/Engine/Source/Runtime/Online/IOS/OnlineSubsystemIOS/Public/OnlineSubsystemIOS.h b/Engine/Source/Runtime/Online/IOS/OnlineSubsystemIOS/Public/OnlineSubsystemIOS.h
--- a/Engine/Source/Runtime/Online/IOS/OnlineSubsystemIOS/Public/OnlineSubsystemIOS.h
+++ b/Engine/Source/Runtime/Online/IOS/OnlineSubsystemIOS/Public/OnlineSubsystemIOS.h
@@ -29,6 +29,7 @@ public:
 	virtual IOnlineExternalUIPtr GetExternalUIInterface() const OVERRIDE;
 	virtual IOnlineTimePtr GetTimeInterface() const OVERRIDE;
 	virtual IOnlineIdentityPtr GetIdentityInterface() const OVERRIDE;
+	virtual IOnlinePartyPtr GetPartyInterface() const OVERRIDE;
 	virtual IOnlineTitleFilePtr GetTitleFileInterface() const OVERRIDE;
 	virtual IOnlineEntitlementsPtr GetEntitlementsInterface() const OVERRIDE;
 	virtual IOnlineStorePtr GetStoreInterface() const OVERRIDE;
diff --git a/Engine/Source/Runtime/Online/OnlineSubsystem/Private/OnlineAsyncTaskManager.cpp b/Engine/Source/Runtime/Online/OnlineSubsystem/Private/OnlineAsyncTaskManager.cpp
--- a/Engine/Source/Runtime/Online/OnlineSubsystem/Private/OnlineAsyncTaskManager.cpp
+++ b/Engine/Source/Runtime/Online/OnlineSubsystem/Private/OnlineAsyncTaskManager.cpp
@@ -156,8 +156,6 @@ void FOnlineAsyncTaskManager::Exit(void)
 
 void FOnlineAsyncTaskManager::AddToInQueue(FOnlineAsyncTask* NewTask)
 {
-	// assert if not game thread	
-	check(IsInGameThread());
 	FScopeLock Lock(&InQueueLock);
 	InQueue.Add(NewTask);
 	WorkEvent->Trigger();
diff --git a/Engine/Source/Runtime/Online/OnlineSubsystem/Public/Interfaces/OnlinePartyInterface.h b/Engine/Source/Runtime/Online/OnlineSubsystem/Public/Interfaces/OnlinePartyInterface.h
new file mode 100644
--- /dev/null
+++ b/Engine/Source/Runtime/Online/OnlineSubsystem/Public/Interfaces/OnlinePartyInterface.h
@@ -0,0 +1,51 @@
+// Copyright 1998-2014 Epic Games, Inc. All Rights Reserved.
+
+#pragma once
+
+#include "OnlineSessionSettings.h"
+#include "OnlineSubsystemTypes.h"
+
+/**
+* Delegate used in friends list change notifications
+*/
+DECLARE_MULTICAST_DELEGATE(FOnPartyChange);
+typedef FOnPartyChange::FDelegate FOnPartyChangeDelegate;
+
+/**
+ * Interface definition for the online party services
+ */
+class IOnlineParty
+{
+protected:
+
+	/** Hidden on purpose */
+	IOnlineParty() {};
+
+public:
+	/**
+	 * Fill in an array with the current party members
+	 * 
+	 * @param LocalUserNum the index of the local user making the request
+	 * @param PartyMembers the array to receive the party member data
+	 */
+	virtual EOnlineAsyncTaskState::Type GetMembers(int32 LocalUserNum, TArray<FOnlinePartyMember>& PartyMembers) = 0;
+
+	/** Returns the number of members currently in the party */
+	virtual int32 GetNumberOfMembers() = 0;
+
+	/**
+     * Delegate used in party change notifications
+     */
+	DEFINE_ONLINE_PLAYER_DELEGATE(MAX_LOCAL_PLAYERS, OnPartyChange);
+
+	/**
+	 * Sends an invite to a list of players to join the party.
+	 *
+	 * @param LocalUserNum the index of the local user sending the invitation
+	 * @param Players the players to invite
+	 */
+	virtual void InvitePlayersToParty(int32 LocalUserNum, const TArray<TSharedRef<FUniqueNetId>>& Players) = 0;
+
+	/** Returns true if the user represented by LocalUserNum is currently in the party */
+	virtual bool IsInParty(int32 LocalUserNum) = 0;
+};
\ No newline at end of file
diff --git a/Engine/Source/Runtime/Online/OnlineSubsystem/Public/Interfaces/OnlineSessionInterface.h b/Engine/Source/Runtime/Online/OnlineSubsystem/Public/Interfaces/OnlineSessionInterface.h
--- a/Engine/Source/Runtime/Online/OnlineSubsystem/Public/Interfaces/OnlineSessionInterface.h
+++ b/Engine/Source/Runtime/Online/OnlineSubsystem/Public/Interfaces/OnlineSessionInterface.h
@@ -56,6 +56,24 @@ typedef FOnEndSessionComplete::FDelegate FOnEndSessionCompleteDelegate;
 DECLARE_MULTICAST_DELEGATE_TwoParams(FOnDestroySessionComplete, FName, bool);
 typedef FOnDestroySessionComplete::FDelegate FOnDestroySessionCompleteDelegate;
 
+/**
+ * Delegate fired when the Matchmaking for an online session has completed
+ *
+ * @param SessionName the name of the session the that can now be joined (on success)
+ * @param bWasSuccessful true if the async action completed without error, false if there was an error
+ */
+DECLARE_MULTICAST_DELEGATE_TwoParams(FOnMatchmakingComplete, FName, bool);
+typedef FOnMatchmakingComplete::FDelegate FOnMatchmakingCompleteDelegate;
+
+/**
+ * Delegate fired when the Matchmaking request has been canceled
+ *
+ * @param SessionName the name of the session that was passed to StartMatchmaking
+ * @param bWasSuccessful true if the async action completed without error, false if there was an error
+ */
+DECLARE_MULTICAST_DELEGATE_TwoParams(FOnCancelMatchmakingComplete, FName, bool);
+typedef FOnCancelMatchmakingComplete::FDelegate FOnCancelMatchmakingCompleteDelegate;
+
 /**
  * Delegate fired when the search for an online session has completed
  *
@@ -298,6 +316,42 @@ public:
 	 */
 	virtual bool IsPlayerInSession(FName SessionName, const FUniqueNetId& UniqueId) = 0;
 
+	/**
+	 * Begins cloud based matchmaking for a session
+	 *
+	 * @param SearchingPlayerNum the index of the player searching for a match
+	 * @param SessionName the name of the session to use, usually will be GAME_SESSION_NAME
+	 * @param NewSessionSettings the desired settings to match against or create with when forming new sessions
+	 * @param SearchSettings the desired settings that the matched session will have
+	 *
+	 * @return true if successful searching for sessions, false otherwise
+	 */
+	virtual bool StartMatchmaking(int32 SearchingPlayerNum, FName SessionName, const FOnlineSessionSettings& NewSessionSettings, TSharedRef<FOnlineSessionSearch>& SearchSettings) = 0;
+
+	/**
+	 * Delegate fired when the cloud matchmaking has completed
+	 *
+	 * @param SessionName The name of the session that was found via matchmaking
+	 * @param bWasSuccessful true if the async action completed without error, false if there was an error
+	 */
+	DEFINE_ONLINE_DELEGATE_TWO_PARAM(OnMatchmakingComplete, FName, bool);
+
+	/**
+	 * Cancel a Matchmaking request for a given session name
+	 *
+	 * @param SearchingPlayerNum the index of the player canceling the search
+	 * @param SessionName the name of the session that was passed to StartMatchmaking (or CreateSession)
+	 */ 
+	virtual bool CancelMatchmaking(int32 SearchingPlayerNum, FName SessionName) = 0;
+
+	/**
+	 * Delegate fired when the cloud matchmaking has been canceled
+	 *
+	 * @param SessionName the name of the session that was canceled
+	 * @param bWasSuccessful true if the async action completed without error, false if there was an error
+	 */
+	DEFINE_ONLINE_DELEGATE_TWO_PARAM(OnCancelMatchmakingComplete, FName, bool);
+
 	/**
 	 * Searches for sessions matching the settings specified
 	 *
diff --git a/Engine/Source/Runtime/Online/OnlineSubsystem/Public/Online.h b/Engine/Source/Runtime/Online/OnlineSubsystem/Public/Online.h
--- a/Engine/Source/Runtime/Online/OnlineSubsystem/Public/Online.h
+++ b/Engine/Source/Runtime/Online/OnlineSubsystem/Public/Online.h
@@ -16,6 +16,7 @@
 #include "OnlinePresenceInterface.h"
 #include "OnlineUserCloudInterface.h"
 #include "OnlineUserInterface.h"
+#include "OnlinePartyInterface.h"
 
 /** Macro to handle the boilerplate of accessing the proper online subsystem and getting the requested interface */
 #define IMPLEMENT_GET_INTERFACE(InterfaceType) \
@@ -104,6 +105,13 @@ namespace Online
 	 */
 	IMPLEMENT_GET_INTERFACE(Identity);
 
+	/** 
+	 * Get the interface for accessing party services
+	 * @param SubsystemName - Name of the requested online service
+	 * @return Interface pointer for the appropriate party service
+	 */
+	IMPLEMENT_GET_INTERFACE(Party);
+
 	/** 
 	 * Get the interface for accessing title file online services
 	 * @param SubsystemName - Name of the requested online service
diff --git a/Engine/Source/Runtime/Online/OnlineSubsystem/Public/OnlineSessionSettings.h b/Engine/Source/Runtime/Online/OnlineSubsystem/Public/OnlineSessionSettings.h
--- a/Engine/Source/Runtime/Online/OnlineSubsystem/Public/OnlineSessionSettings.h
+++ b/Engine/Source/Runtime/Online/OnlineSubsystem/Public/OnlineSessionSettings.h
@@ -437,6 +437,10 @@ public:
 #define SEARCH_USER FName(TEXT("SEARCHUSER"))
 /** Keywords to match in session search */
 #define SEARCH_KEYWORDS FName(TEXT("SEARCHKEYWORDS"))
+/** If set, use the named Xbox Live hopper to find a session via matchmaking (value is a string) */
+#define SEARCH_XBOX_LIVE_HOPPER_NAME FName(TEXT("LIVEHOPPERNAME"))
+/** Which session template from the service configuration to use */
+#define SEARCH_XBOX_LIVE_SESSION_TEMPLATE_NAME FName(TEXT("LIVESESSIONTEMPLATE"))
 
 /** 
  * Encapsulation of a search for sessions request.
@@ -464,6 +468,9 @@ public:
 	 */
 	int32 PingBucketSize;
 
+	/** Amount of time to wait for the search results. May not apply to all platforms. */
+	float TimeoutInSeconds;
+
 	/** Constructor */
 	FOnlineSessionSearch() :
 		SearchState(EOnlineAsyncTaskState::NotStarted)
diff --git a/Engine/Source/Runtime/Online/OnlineSubsystem/Public/OnlineSubsystem.h b/Engine/Source/Runtime/Online/OnlineSubsystem/Public/OnlineSubsystem.h
--- a/Engine/Source/Runtime/Online/OnlineSubsystem/Public/OnlineSubsystem.h
+++ b/Engine/Source/Runtime/Online/OnlineSubsystem/Public/OnlineSubsystem.h
@@ -49,6 +49,7 @@ typedef TSharedPtr<class IOnlineSharing, ESPMode::ThreadSafe> IOnlineSharingPtr;
 typedef TSharedPtr<class IOnlineUser, ESPMode::ThreadSafe> IOnlineUserPtr;
 typedef TSharedPtr<class IOnlineMessage, ESPMode::ThreadSafe> IOnlineMessagePtr;
 typedef TSharedPtr<class IOnlinePresence, ESPMode::ThreadSafe> IOnlinePresencePtr;
+typedef TSharedPtr<class IOnlineParty, ESPMode::ThreadSafe> IOnlinePartyPtr;
 
 /**
  * Called when the connection state as reported by the online platform changes
@@ -169,6 +170,12 @@ public:
 	 */
 	virtual IOnlineIdentityPtr GetIdentityInterface() const = 0;
 
+	/** 
+	 * Get the interface for accessing party online services
+	 * @return Interface pointer for the appropriate party service
+	 */
+	virtual IOnlinePartyPtr GetPartyInterface() const = 0;
+
 	/** 
 	 * Get the interface for accessing title file online services
 	 * @return Interface pointer for the appropriate title file service
diff --git a/Engine/Source/Runtime/Online/OnlineSubsystem/Public/OnlineSubsystemTypes.h b/Engine/Source/Runtime/Online/OnlineSubsystem/Public/OnlineSubsystemTypes.h
--- a/Engine/Source/Runtime/Online/OnlineSubsystem/Public/OnlineSubsystemTypes.h
+++ b/Engine/Source/Runtime/Online/OnlineSubsystem/Public/OnlineSubsystemTypes.h
@@ -1072,6 +1072,37 @@ public:
 	virtual const class FOnlineUserPresence& GetPresence() const = 0;
 };
 
+/**
+ * Party user info returned via IOnlineParty interface
+ */
+struct FOnlinePartyMember
+{
+	/** The unique ID of this party member */
+	TSharedPtr<FUniqueNetId>	UniqueId;
+
+	/** True if this party member is using the local console */
+	bool						IsLocal;
+
+	/** The time at which the player joined the party */
+	FDateTime					JoinTime;
+
+	/** If the platform supports grouping party members by device, this will be set */
+	int32						DeviceGroup;
+
+	/** Constructor */
+	FOnlinePartyMember(
+		const TSharedPtr<FUniqueNetId>& PartyMemberId,
+		const bool IsLocalPartyMember,
+		const FDateTime& PartyJoinTime,
+		const int32 InDeviceGroup
+		)
+		: UniqueId(PartyMemberId)
+		, IsLocal(IsLocalPartyMember)
+		, JoinTime(PartyJoinTime)
+		, DeviceGroup(InDeviceGroup)
+	{}
+};
+
 /** The possible permission categories we can choose from to read from the server */
 namespace EOnlineSharingReadCategory
 {
diff --git a/Engine/Source/Runtime/Online/OnlineSubsystemAmazon/Public/OnlineSubsystemAmazon.h b/Engine/Source/Runtime/Online/OnlineSubsystemAmazon/Public/OnlineSubsystemAmazon.h
--- a/Engine/Source/Runtime/Online/OnlineSubsystemAmazon/Public/OnlineSubsystemAmazon.h
+++ b/Engine/Source/Runtime/Online/OnlineSubsystemAmazon/Public/OnlineSubsystemAmazon.h
@@ -66,7 +66,10 @@ public:
 	{
 		return NULL;
 	}
-
+	virtual IOnlinePartyPtr GetPartyInterface() const OVERRIDE
+	{
+		return NULL;
+	}
 	virtual IOnlineTitleFilePtr GetTitleFileInterface() const OVERRIDE
 	{
 		return NULL;
@@ -114,6 +117,11 @@ public:
 		return NULL;
 	}
 
+	virtual IOnlinePartyPtr GetPartyInterface() const OVERRIDE
+	{
+		return NULL;
+	}
+
 	virtual bool Init() OVERRIDE;
 	virtual bool Shutdown() OVERRIDE;
 	virtual FString GetAppId() const OVERRIDE;
diff --git a/Engine/Source/Runtime/Online/OnlineSubsystemFacebook/Private/IOS/OnlineSubsystemFacebook.cpp b/Engine/Source/Runtime/Online/OnlineSubsystemFacebook/Private/IOS/OnlineSubsystemFacebook.cpp
--- a/Engine/Source/Runtime/Online/OnlineSubsystemFacebook/Private/IOS/OnlineSubsystemFacebook.cpp
+++ b/Engine/Source/Runtime/Online/OnlineSubsystemFacebook/Private/IOS/OnlineSubsystemFacebook.cpp
@@ -111,6 +111,11 @@ IOnlinePresencePtr FOnlineSubsystemFacebook::GetPresenceInterface() const
 	return NULL;
 }
 
+IOnlinePartyPtr FOnlineSubsystemFacebook::GetPartyInterface() const
+{
+	return NULL;
+}
+
 bool FOnlineSubsystemFacebook::Init() 
 {
 	bool bSuccessfullyStartedUp = true;
diff --git a/Engine/Source/Runtime/Online/OnlineSubsystemFacebook/Private/Windows/OnlineSubsystemFacebook.cpp b/Engine/Source/Runtime/Online/OnlineSubsystemFacebook/Private/Windows/OnlineSubsystemFacebook.cpp
--- a/Engine/Source/Runtime/Online/OnlineSubsystemFacebook/Private/Windows/OnlineSubsystemFacebook.cpp
+++ b/Engine/Source/Runtime/Online/OnlineSubsystemFacebook/Private/Windows/OnlineSubsystemFacebook.cpp
@@ -109,6 +109,11 @@ IOnlinePresencePtr FOnlineSubsystemFacebook::GetPresenceInterface() const
 	return NULL;
 }
 
+IOnlinePartyPtr FOnlineSubsystemFacebook::GetPartyInterface() const
+{
+	return NULL;
+}
+
 bool FOnlineSubsystemFacebook::Tick(float DeltaTime)
 {
 	if (FacebookIdentity.IsValid())
diff --git a/Engine/Source/Runtime/Online/OnlineSubsystemFacebook/Public/OnlineSubsystemFacebook.h b/Engine/Source/Runtime/Online/OnlineSubsystemFacebook/Public/OnlineSubsystemFacebook.h
--- a/Engine/Source/Runtime/Online/OnlineSubsystemFacebook/Public/OnlineSubsystemFacebook.h
+++ b/Engine/Source/Runtime/Online/OnlineSubsystemFacebook/Public/OnlineSubsystemFacebook.h
@@ -31,6 +31,7 @@ public:
 	virtual IOnlineExternalUIPtr GetExternalUIInterface() const OVERRIDE;
 	virtual IOnlineTimePtr GetTimeInterface() const OVERRIDE;
 	virtual IOnlineIdentityPtr GetIdentityInterface() const OVERRIDE;
+	virtual IOnlinePartyPtr GetPartyInterface() const OVERRIDE;
 	virtual IOnlineTitleFilePtr GetTitleFileInterface() const OVERRIDE;
 	virtual IOnlineEntitlementsPtr GetEntitlementsInterface() const OVERRIDE;
 	virtual IOnlineStorePtr GetStoreInterface() const OVERRIDE;
diff --git a/Engine/Source/Runtime/Online/OnlineSubsystemNull/Private/OnlineSessionInterfaceNull.cpp b/Engine/Source/Runtime/Online/OnlineSubsystemNull/Private/OnlineSessionInterfaceNull.cpp
--- a/Engine/Source/Runtime/Online/OnlineSubsystemNull/Private/OnlineSessionInterfaceNull.cpp
+++ b/Engine/Source/Runtime/Online/OnlineSubsystemNull/Private/OnlineSessionInterfaceNull.cpp
@@ -435,6 +435,20 @@ bool FOnlineSessionNull::IsPlayerInSession(FName SessionName, const FUniqueNetId
 	return IsPlayerInSessionImpl(this, SessionName, UniqueId);
 }
 
+bool FOnlineSessionNull::StartMatchmaking(int32 SearchingPlayerNum, FName SessionName, const FOnlineSessionSettings& NewSessionSettings, TSharedRef<FOnlineSessionSearch>& SearchSettings)
+{
+	UE_LOG(LogOnline, Warning, TEXT("Matchmaking is not supported on this platform."));
+	TriggerOnMatchmakingCompleteDelegates(SessionName, false);
+	return false;
+}
+
+bool FOnlineSessionNull::CancelMatchmaking(int32 SearchingPlayerNum, FName SessionName)
+{
+	UE_LOG(LogOnline, Warning, TEXT("Matchmaking is not supported on this platform."));
+	TriggerOnCancelMatchmakingCompleteDelegates(SessionName, false);
+	return false;
+}
+
 bool FOnlineSessionNull::FindSessions(int32 SearchingPlayerNum, const TSharedRef<FOnlineSessionSearch>& SearchSettings)
 {
 	uint32 Return = E_FAIL;
diff --git a/Engine/Source/Runtime/Online/OnlineSubsystemNull/Private/OnlineSessionInterfaceNull.h b/Engine/Source/Runtime/Online/OnlineSubsystemNull/Private/OnlineSessionInterfaceNull.h
--- a/Engine/Source/Runtime/Online/OnlineSubsystemNull/Private/OnlineSessionInterfaceNull.h
+++ b/Engine/Source/Runtime/Online/OnlineSubsystemNull/Private/OnlineSessionInterfaceNull.h
@@ -269,6 +269,8 @@ public:
 	virtual bool EndSession(FName SessionName) OVERRIDE;
 	virtual bool DestroySession(FName SessionName) OVERRIDE;
 	virtual bool IsPlayerInSession(FName SessionName, const FUniqueNetId& UniqueId) OVERRIDE;
+	virtual bool StartMatchmaking(int32 SearchingPlayerNum, FName SessionName, const FOnlineSessionSettings& NewSessionSettings, TSharedRef<FOnlineSessionSearch>& SearchSettings) OVERRIDE;
+	virtual bool CancelMatchmaking(int32 SearchingPlayerNum, FName SessionName) OVERRIDE;
 	virtual bool FindSessions(int32 SearchingPlayerNum, const TSharedRef<FOnlineSessionSearch>& SearchSettings) OVERRIDE;
 	virtual bool CancelFindSessions() OVERRIDE;
 	virtual bool PingSearchResults(const FOnlineSessionSearchResult& SearchResult) OVERRIDE;
diff --git a/Engine/Source/Runtime/Online/OnlineSubsystemNull/Private/OnlineSubsystemNull.cpp b/Engine/Source/Runtime/Online/OnlineSubsystemNull/Private/OnlineSubsystemNull.cpp
--- a/Engine/Source/Runtime/Online/OnlineSubsystemNull/Private/OnlineSubsystemNull.cpp
+++ b/Engine/Source/Runtime/Online/OnlineSubsystemNull/Private/OnlineSubsystemNull.cpp
@@ -63,6 +63,11 @@ IOnlineIdentityPtr FOnlineSubsystemNull::GetIdentityInterface() const
 	return IdentityInterface;
 }
 
+IOnlinePartyPtr FOnlineSubsystemNull::GetPartyInterface() const
+{
+	return NULL;
+}
+
 IOnlineTitleFilePtr FOnlineSubsystemNull::GetTitleFileInterface() const
 {
 	return NULL;
diff --git a/Engine/Source/Runtime/Online/OnlineSubsystemNull/Public/OnlineSubsystemNull.h b/Engine/Source/Runtime/Online/OnlineSubsystemNull/Public/OnlineSubsystemNull.h
--- a/Engine/Source/Runtime/Online/OnlineSubsystemNull/Public/OnlineSubsystemNull.h
+++ b/Engine/Source/Runtime/Online/OnlineSubsystemNull/Public/OnlineSubsystemNull.h
@@ -43,6 +43,7 @@ public:
 	virtual IOnlineExternalUIPtr GetExternalUIInterface() const OVERRIDE;	
 	virtual IOnlineTimePtr GetTimeInterface() const OVERRIDE;
 	virtual IOnlineIdentityPtr GetIdentityInterface() const OVERRIDE;
+	virtual IOnlinePartyPtr GetPartyInterface() const OVERRIDE;
 	virtual IOnlineTitleFilePtr GetTitleFileInterface() const OVERRIDE;
 	virtual IOnlineStorePtr GetStoreInterface() const OVERRIDE;
 	virtual IOnlineEventsPtr GetEventsInterface() const OVERRIDE;
diff --git a/Engine/Source/Runtime/Online/OnlineSubsystemSteam/Private/OnlineSessionInterfaceSteam.cpp b/Engine/Source/Runtime/Online/OnlineSubsystemSteam/Private/OnlineSessionInterfaceSteam.cpp
--- a/Engine/Source/Runtime/Online/OnlineSubsystemSteam/Private/OnlineSessionInterfaceSteam.cpp
+++ b/Engine/Source/Runtime/Online/OnlineSubsystemSteam/Private/OnlineSessionInterfaceSteam.cpp
@@ -634,6 +634,20 @@ bool FOnlineSessionSteam::IsPlayerInSession(FName SessionName, const FUniqueNetI
 	return IsPlayerInSessionImpl(this, SessionName, UniqueId);
 }
 
+bool FOnlineSessionSteam::StartMatchmaking(int32 SearchingPlayerNum, FName SessionName, const FOnlineSessionSettings& NewSessionSettings, TSharedRef<FOnlineSessionSearch>& SearchSettings)
+{
+	UE_LOG(LogOnline, Warning, TEXT("Matchmaking is not supported on this platform."));
+	TriggerOnMatchmakingCompleteDelegates(SessionName, false);
+	return false;
+}
+
+bool FOnlineSessionSteam::CancelMatchmaking(int32 SearchingPlayerNum, FName SessionName)
+{
+	UE_LOG(LogOnline, Warning, TEXT("Matchmaking is not supported on this platform."));
+	TriggerOnCancelMatchmakingCompleteDelegates(SessionName, false);
+	return false;
+}
+
 bool FOnlineSessionSteam::FindSessions(int32 SearchingPlayerNum, const TSharedRef<FOnlineSessionSearch>& SearchSettings)
 {
 	uint32 Return = E_FAIL;
diff --git a/Engine/Source/Runtime/Online/OnlineSubsystemSteam/Private/OnlineSessionInterfaceSteam.h b/Engine/Source/Runtime/Online/OnlineSubsystemSteam/Private/OnlineSessionInterfaceSteam.h
--- a/Engine/Source/Runtime/Online/OnlineSubsystemSteam/Private/OnlineSessionInterfaceSteam.h
+++ b/Engine/Source/Runtime/Online/OnlineSubsystemSteam/Private/OnlineSessionInterfaceSteam.h
@@ -495,6 +495,8 @@ public:
 	virtual bool EndSession(FName SessionName) OVERRIDE;
 	virtual bool DestroySession(FName SessionName) OVERRIDE;
 	virtual bool IsPlayerInSession(FName SessionName, const FUniqueNetId& UniqueId) OVERRIDE;
+	virtual bool StartMatchmaking(int32 SearchingPlayerNum, FName SessionName, const FOnlineSessionSettings& NewSessionSettings, TSharedRef<FOnlineSessionSearch>& SearchSettings) OVERRIDE;
+	virtual bool CancelMatchmaking(int32 SearchingPlayerNum, FName SessionName) OVERRIDE;
 	virtual bool FindSessions(int32 SearchingPlayerNum, const TSharedRef<FOnlineSessionSearch>& SearchSettings) OVERRIDE;
 	virtual bool CancelFindSessions() OVERRIDE;
 	virtual bool PingSearchResults(const FOnlineSessionSearchResult& SearchResult) OVERRIDE;
diff --git a/Engine/Source/Runtime/Online/OnlineSubsystemSteam/Private/OnlineSubsystemSteam.cpp b/Engine/Source/Runtime/Online/OnlineSubsystemSteam/Private/OnlineSubsystemSteam.cpp
--- a/Engine/Source/Runtime/Online/OnlineSubsystemSteam/Private/OnlineSubsystemSteam.cpp
+++ b/Engine/Source/Runtime/Online/OnlineSubsystemSteam/Private/OnlineSubsystemSteam.cpp
@@ -199,6 +199,11 @@ IOnlineIdentityPtr FOnlineSubsystemSteam::GetIdentityInterface() const
 	return IdentityInterface;
 }
 
+IOnlinePartyPtr FOnlineSubsystemSteam::GetPartyInterface() const
+{
+	return NULL;
+}
+
 IOnlineTitleFilePtr FOnlineSubsystemSteam::GetTitleFileInterface() const
 {
 	return NULL;
diff --git a/Engine/Source/Runtime/Online/OnlineSubsystemSteam/Public/OnlineSubsystemSteam.h b/Engine/Source/Runtime/Online/OnlineSubsystemSteam/Public/OnlineSubsystemSteam.h
--- a/Engine/Source/Runtime/Online/OnlineSubsystemSteam/Public/OnlineSubsystemSteam.h
+++ b/Engine/Source/Runtime/Online/OnlineSubsystemSteam/Public/OnlineSubsystemSteam.h
@@ -220,6 +220,7 @@ public:
 	virtual IOnlineExternalUIPtr GetExternalUIInterface() const OVERRIDE;
 	virtual IOnlineTimePtr GetTimeInterface() const OVERRIDE;
 	virtual IOnlineIdentityPtr GetIdentityInterface() const OVERRIDE;
+	virtual IOnlinePartyPtr GetPartyInterface() const OVERRIDE;
 	virtual IOnlineTitleFilePtr GetTitleFileInterface() const OVERRIDE;
 	virtual IOnlineEntitlementsPtr GetEntitlementsInterface() const OVERRIDE;
 	virtual IOnlineStorePtr GetStoreInterface() const OVERRIDE;
