Remove duplicated function - whoops

diff --git a/Engine/Source/Runtime/Online/OnlineSubsystemAmazon/Public/OnlineSubsystemAmazon.h b/Engine/Source/Runtime/Online/OnlineSubsystemAmazon/Public/OnlineSubsystemAmazon.h
--- a/Engine/Source/Runtime/Online/OnlineSubsystemAmazon/Public/OnlineSubsystemAmazon.h
+++ b/Engine/Source/Runtime/Online/OnlineSubsystemAmazon/Public/OnlineSubsystemAmazon.h
@@ -117,11 +117,6 @@ public:
 		return NULL;
 	}
 
-	virtual IOnlinePartyPtr GetPartyInterface() const OVERRIDE
-	{
-		return NULL;
-	}
-
 	virtual bool Init() OVERRIDE;
 	virtual bool Shutdown() OVERRIDE;
 	virtual FString GetAppId() const OVERRIDE;
