#UE4
        conform to coding standard -  move comments to right place.

diff --git a/Engine/Source/Runtime/Online/OnlineSubsystem/Private/OnlineAsyncTaskManager.cpp b/Engine/Source/Runtime/Online/OnlineSubsystem/Private/OnlineAsyncTaskManager.cpp
--- a/Engine/Source/Runtime/Online/OnlineSubsystem/Private/OnlineAsyncTaskManager.cpp
+++ b/Engine/Source/Runtime/Online/OnlineSubsystem/Private/OnlineAsyncTaskManager.cpp
@@ -223,8 +223,6 @@ void FOnlineAsyncTaskManager::GameTick()
 	while (Item != NULL);
 }
 
-
-// Non blocking single threaded tick for Online Async Tasks. 
 void FOnlineAsyncTaskManager::Tick()
 {
 
diff --git a/Engine/Source/Runtime/Online/OnlineSubsystem/Public/OnlineAsyncTaskManager.h b/Engine/Source/Runtime/Online/OnlineSubsystem/Public/OnlineAsyncTaskManager.h
--- a/Engine/Source/Runtime/Online/OnlineSubsystem/Public/OnlineAsyncTaskManager.h
+++ b/Engine/Source/Runtime/Online/OnlineSubsystem/Public/OnlineAsyncTaskManager.h
@@ -284,6 +284,12 @@ public:
 	 */
 	virtual void Exit();
 
+	/**
+	 *  FSingleThreadRunnable accessor for ticking this FRunnable when multi-threading is disabled. 
+	 *  @return FSingleThreadRunnable Interface for this FRunnable object.
+	 */
+	virtual class FSingleThreadRunnable* GetSingleThreadInterface() { return this; }
+
 	/**
 	 * Add online async tasks that need processing onto the incoming queue
 	 * @param NewTask - some request of the online services
@@ -333,10 +339,13 @@ public:
 	 */
 	virtual void OnlineTick() = 0;
 
+	// FSingleThreadRunnable interface
+	/**
+	 *  Non blocking version of FOnlineAsyncTaskManager::Run.
+	 *  
+	 */
 	virtual void Tick();
 
-	virtual class FSingleThreadRunnable* GetSingleThreadInterface() { return this; }
-
 
 };
 
