Integrate fixes for Xbox One from Black Tusk
Write out the port in IPv6 address strings
Add checks to guard against problems when users sign out
Set the port when converting secure addresses

diff --git a/Engine/Source/Runtime/Online/OnlineSubsystemUtils/Private/Tests/TestSessionInterface.cpp b/Engine/Source/Runtime/Online/OnlineSubsystemUtils/Private/Tests/TestSessionInterface.cpp
--- a/Engine/Source/Runtime/Online/OnlineSubsystemUtils/Private/Tests/TestSessionInterface.cpp
+++ b/Engine/Source/Runtime/Online/OnlineSubsystemUtils/Private/Tests/TestSessionInterface.cpp
@@ -6,6 +6,7 @@
 #include "OnlineFriendsInterface.h"
 #include "OnlineExternalUIInterface.h"
 #include "ModuleManager.h"
+#include "Kismet/GameplayStatics.h"
 
 #include "Engine.h"
 
diff --git a/Engine/Source/Runtime/Sockets/Private/BSDIPv6Sockets/IPAddressBSDIPv6.h b/Engine/Source/Runtime/Sockets/Private/BSDIPv6Sockets/IPAddressBSDIPv6.h
--- a/Engine/Source/Runtime/Sockets/Private/BSDIPv6Sockets/IPAddressBSDIPv6.h
+++ b/Engine/Source/Runtime/Sockets/Private/BSDIPv6Sockets/IPAddressBSDIPv6.h
@@ -244,6 +244,12 @@ public:
 		FString Result("[");
 		Result += IPStr;
 		Result += "]";
+		
+		if (bAppendPort)
+		{
+			Result += ":";
+			Result += FString::Printf(TEXT("%d"), GetPort());
+		}
 
 		return Result;
 	}
diff --git a/Engine/Source/Runtime/Sockets/Private/BSDIPv6Sockets/SocketsBSDIPv6.h b/Engine/Source/Runtime/Sockets/Private/BSDIPv6Sockets/SocketsBSDIPv6.h
--- a/Engine/Source/Runtime/Sockets/Private/BSDIPv6Sockets/SocketsBSDIPv6.h
+++ b/Engine/Source/Runtime/Sockets/Private/BSDIPv6Sockets/SocketsBSDIPv6.h
@@ -54,7 +54,7 @@ public:
 	 */
 	virtual ~FSocketBSDIPv6()
 	{
-		Close();
+		FSocketBSDIPv6::Close();
 	}
 
 
