Merging using UE4-To-UE4-LauncherReleases

support serializing Json struct member as double

diff --git a/Engine/Source/Runtime/Online/OnlineSubsystem/Public/OnlineJsonSerializer.h b/Engine/Source/Runtime/Online/OnlineSubsystem/Public/OnlineJsonSerializer.h
--- a/Engine/Source/Runtime/Online/OnlineSubsystem/Public/OnlineJsonSerializer.h
+++ b/Engine/Source/Runtime/Online/OnlineSubsystem/Public/OnlineJsonSerializer.h
@@ -53,6 +53,7 @@ struct FOnlineJsonSerializerBase
 	virtual void Serialize(const TCHAR* Name, bool& Value) = 0;
 	virtual void Serialize(const TCHAR* Name, FString& Value) = 0;
 	virtual void Serialize(const TCHAR* Name, float& Value) = 0;
+	virtual void Serialize(const TCHAR* Name, double& Value) = 0;
 	virtual void SerializeArray(FJsonSerializableArray& Array) = 0;
 	virtual void SerializeArray(const TCHAR* Name, FJsonSerializableArray& Value) = 0;
 	virtual void SerializeMap(const TCHAR* Name, FJsonSerializableKeyValueMap& Map) = 0;
@@ -172,6 +173,16 @@ public:
 	{
 		JsonWriter->WriteValue(Name, (const float)Value);
 	}
+	/**
+	* Writes the field name and the corresponding value to the JSON data
+	*
+	* @param Name the field name to write out
+	* @param Value the value to write out
+	*/
+	virtual void Serialize(const TCHAR* Name, double& Value) OVERRIDE
+	{
+		JsonWriter->WriteValue(Name, (const double)Value);
+	}
 	/**
 	 * Serializes an array of values
 	 *
@@ -358,6 +369,19 @@ public:
 			Value = JsonObject->GetNumberField(Name);
 		}
 	}
+	/**
+	* If the underlying json object has the field, it is read into the value
+	*
+	* @param Name the name of the field to read
+	* @param Value the out value to read the data into
+	*/
+	virtual void Serialize(const TCHAR* Name, double& Value) OVERRIDE
+	{
+		if (JsonObject->HasField(Name))
+		{
+			Value = JsonObject->GetNumberField(Name);
+		}
+	}
 	/**
 	 * Serializes an array of values
 	 *
