UE4 Refactoring. Changed OVERRIDE and FINAL macros to keywords override and final.

diff --git a/Engine/Source/Runtime/Online/IOS/OnlineSubsystemIOS/Private/OnlineSessionInterfaceIOS.h b/Engine/Source/Runtime/Online/IOS/OnlineSubsystemIOS/Private/OnlineSessionInterfaceIOS.h
--- a/Engine/Source/Runtime/Online/IOS/OnlineSubsystemIOS/Private/OnlineSessionInterfaceIOS.h
+++ b/Engine/Source/Runtime/Online/IOS/OnlineSubsystemIOS/Private/OnlineSessionInterfaceIOS.h
@@ -62,69 +62,69 @@ PACKAGE_SCOPE:
 
 
 	// Begin IOnlineSession interface
-	class FNamedOnlineSession* AddNamedSession(FName SessionName, const FOnlineSessionSettings& SessionSettings) OVERRIDE;
+	class FNamedOnlineSession* AddNamedSession(FName SessionName, const FOnlineSessionSettings& SessionSettings) override;
 
-	class FNamedOnlineSession* AddNamedSession(FName SessionName, const FOnlineSession& Session) OVERRIDE;
+	class FNamedOnlineSession* AddNamedSession(FName SessionName, const FOnlineSession& Session) override;
 
 public:
 
 	virtual ~FOnlineSessionIOS();
 
-	FNamedOnlineSession* GetNamedSession(FName SessionName) OVERRIDE;
+	FNamedOnlineSession* GetNamedSession(FName SessionName) override;
 
-	virtual void RemoveNamedSession(FName SessionName) OVERRIDE;
+	virtual void RemoveNamedSession(FName SessionName) override;
 
-	virtual EOnlineSessionState::Type GetSessionState(FName SessionName) const OVERRIDE;
+	virtual EOnlineSessionState::Type GetSessionState(FName SessionName) const override;
 
-	virtual bool HasPresenceSession() OVERRIDE;
+	virtual bool HasPresenceSession() override;
 
-	virtual bool CreateSession(int32 HostingPlayerNum, FName SessionName, const FOnlineSessionSettings& NewSessionSettings) OVERRIDE;
+	virtual bool CreateSession(int32 HostingPlayerNum, FName SessionName, const FOnlineSessionSettings& NewSessionSettings) override;
 
-	virtual bool StartSession(FName SessionName) OVERRIDE;
+	virtual bool StartSession(FName SessionName) override;
 
-	virtual bool UpdateSession(FName SessionName, FOnlineSessionSettings& UpdatedSessionSettings, bool bShouldRefreshOnlineData = false) OVERRIDE;
+	virtual bool UpdateSession(FName SessionName, FOnlineSessionSettings& UpdatedSessionSettings, bool bShouldRefreshOnlineData = false) override;
 
-	virtual bool EndSession(FName SessionName) OVERRIDE;
+	virtual bool EndSession(FName SessionName) override;
 
-	virtual bool DestroySession(FName SessionName) OVERRIDE;
+	virtual bool DestroySession(FName SessionName) override;
 
-	virtual bool IsPlayerInSession(FName SessionName, const FUniqueNetId& UniqueId) OVERRIDE;
+	virtual bool IsPlayerInSession(FName SessionName, const FUniqueNetId& UniqueId) override;
 
-	virtual bool StartMatchmaking(int32 SearchingPlayerNum, FName SessionName, const FOnlineSessionSettings& NewSessionSettings, TSharedRef<FOnlineSessionSearch>& SearchSettings) OVERRIDE;
+	virtual bool StartMatchmaking(int32 SearchingPlayerNum, FName SessionName, const FOnlineSessionSettings& NewSessionSettings, TSharedRef<FOnlineSessionSearch>& SearchSettings) override;
 
-	virtual bool CancelMatchmaking(int32 SearchingPlayerNum, FName SessionName) OVERRIDE;
+	virtual bool CancelMatchmaking(int32 SearchingPlayerNum, FName SessionName) override;
 
-	virtual bool FindSessions(int32 SearchingPlayerNum, const TSharedRef<FOnlineSessionSearch>& SearchSettings) OVERRIDE;
+	virtual bool FindSessions(int32 SearchingPlayerNum, const TSharedRef<FOnlineSessionSearch>& SearchSettings) override;
 
-	virtual bool CancelFindSessions() OVERRIDE;
+	virtual bool CancelFindSessions() override;
 
-	virtual bool PingSearchResults(const FOnlineSessionSearchResult& SearchResult) OVERRIDE;
+	virtual bool PingSearchResults(const FOnlineSessionSearchResult& SearchResult) override;
 
-	virtual bool JoinSession(int32 PlayerNum, FName SessionName, const FOnlineSessionSearchResult& DesiredSession) OVERRIDE;
+	virtual bool JoinSession(int32 PlayerNum, FName SessionName, const FOnlineSessionSearchResult& DesiredSession) override;
 
-	virtual bool FindFriendSession(int32 LocalUserNum, const FUniqueNetId& Friend) OVERRIDE;
+	virtual bool FindFriendSession(int32 LocalUserNum, const FUniqueNetId& Friend) override;
 
-	virtual bool SendSessionInviteToFriend(int32 LocalUserNum, FName SessionName, const FUniqueNetId& Friend) OVERRIDE;
+	virtual bool SendSessionInviteToFriend(int32 LocalUserNum, FName SessionName, const FUniqueNetId& Friend) override;
 
-	virtual bool SendSessionInviteToFriends(int32 LocalUserNum, FName SessionName, const TArray< TSharedRef<FUniqueNetId> >& Friends) OVERRIDE;
+	virtual bool SendSessionInviteToFriends(int32 LocalUserNum, FName SessionName, const TArray< TSharedRef<FUniqueNetId> >& Friends) override;
 
-	virtual bool GetResolvedConnectString(FName SessionName, FString& ConnectInfo) OVERRIDE;
+	virtual bool GetResolvedConnectString(FName SessionName, FString& ConnectInfo) override;
 
-	virtual bool GetResolvedConnectString(const class FOnlineSessionSearchResult& SearchResult, FName PortType, FString& ConnectInfo) OVERRIDE;
+	virtual bool GetResolvedConnectString(const class FOnlineSessionSearchResult& SearchResult, FName PortType, FString& ConnectInfo) override;
 
-	virtual FOnlineSessionSettings* GetSessionSettings(FName SessionName) OVERRIDE;
+	virtual FOnlineSessionSettings* GetSessionSettings(FName SessionName) override;
 
-	virtual bool RegisterPlayer(FName SessionName, const FUniqueNetId& PlayerId, bool bWasInvited) OVERRIDE;
+	virtual bool RegisterPlayer(FName SessionName, const FUniqueNetId& PlayerId, bool bWasInvited) override;
 
-	virtual bool RegisterPlayers(FName SessionName, const TArray< TSharedRef<FUniqueNetId> >& Players, bool bWasInvited = false) OVERRIDE;
+	virtual bool RegisterPlayers(FName SessionName, const TArray< TSharedRef<FUniqueNetId> >& Players, bool bWasInvited = false) override;
 
-	virtual bool UnregisterPlayer(FName SessionName, const FUniqueNetId& PlayerId) OVERRIDE;
+	virtual bool UnregisterPlayer(FName SessionName, const FUniqueNetId& PlayerId) override;
 
-	virtual bool UnregisterPlayers(FName SessionName, const TArray< TSharedRef<FUniqueNetId> >& Players) OVERRIDE;
+	virtual bool UnregisterPlayers(FName SessionName, const TArray< TSharedRef<FUniqueNetId> >& Players) override;
 
-	virtual int32 GetNumSessions() OVERRIDE;
+	virtual int32 GetNumSessions() override;
 
-	virtual void DumpSessionState() OVERRIDE;
+	virtual void DumpSessionState() override;
 	// End IOnlineSession interface
 };
 
diff --git a/Engine/Source/Runtime/Online/IOS/OnlineSubsystemIOS/Private/OnlineStoreInterfaceIOS.h b/Engine/Source/Runtime/Online/IOS/OnlineSubsystemIOS/Private/OnlineStoreInterfaceIOS.h
--- a/Engine/Source/Runtime/Online/IOS/OnlineSubsystemIOS/Private/OnlineStoreInterfaceIOS.h
+++ b/Engine/Source/Runtime/Online/IOS/OnlineSubsystemIOS/Private/OnlineStoreInterfaceIOS.h
@@ -50,15 +50,15 @@ public:
 	virtual ~FOnlineStoreInterfaceIOS();
 
 	// Begin IOnlineStore interface
-	virtual bool QueryForAvailablePurchases() OVERRIDE;
+	virtual bool QueryForAvailablePurchases() override;
 
-	virtual bool IsAllowedToMakePurchases() OVERRIDE;
+	virtual bool IsAllowedToMakePurchases() override;
 
-	virtual bool BeginPurchase(int Index) OVERRIDE;
+	virtual bool BeginPurchase(int Index) override;
 	// End IOnlineStore interface
 
 	// Begin FExec Interface
-	virtual bool Exec( UWorld* InWorld, const TCHAR* Cmd, FOutputDevice& Ar=*GLog ) OVERRIDE;
+	virtual bool Exec( UWorld* InWorld, const TCHAR* Cmd, FOutputDevice& Ar=*GLog ) override;
 	// End FExec Interface
 
 	/**
diff --git a/Engine/Source/Runtime/Online/IOS/OnlineSubsystemIOS/Private/OnlineSubsystemIOSTypes.h b/Engine/Source/Runtime/Online/IOS/OnlineSubsystemIOS/Private/OnlineSubsystemIOSTypes.h
--- a/Engine/Source/Runtime/Online/IOS/OnlineSubsystemIOS/Private/OnlineSubsystemIOSTypes.h
+++ b/Engine/Source/Runtime/Online/IOS/OnlineSubsystemIOS/Private/OnlineSubsystemIOSTypes.h
@@ -48,7 +48,7 @@ public:
 	 *
 	 * @return byte array of size GetSize()
 	 */
-	virtual const uint8* GetBytes() const OVERRIDE
+	virtual const uint8* GetBytes() const override
 	{
 		return (uint8*)&UniqueNetId;
 	}
@@ -58,7 +58,7 @@ public:
 	 *
 	 * @return size in bytes of the id representation
 	 */
-	virtual int32 GetSize() const OVERRIDE
+	virtual int32 GetSize() const override
 	{
 		return sizeof(uint64);
 	}
@@ -68,7 +68,7 @@ public:
 	 *
 	 * @return true if this is a well formed ID, false otherwise
 	 */
-	virtual bool IsValid() const OVERRIDE
+	virtual bool IsValid() const override
 	{
 		return UniqueNetId != 0;
 	}
@@ -78,7 +78,7 @@ public:
 	 *
 	 * @return data in string form 
 	 */
-	virtual FString ToString() const OVERRIDE
+	virtual FString ToString() const override
 	{
 		return FString::Printf(TEXT("%I64d"), UniqueNetId);
 	}
@@ -89,7 +89,7 @@ public:
 	 *
 	 * @return id in string form 
 	 */
-	virtual FString ToDebugString() const OVERRIDE
+	virtual FString ToDebugString() const override
 	{
 		return FString::Printf(TEXT("0%I64X"), UniqueNetId);
 	}
diff --git a/Engine/Source/Runtime/Online/IOS/OnlineSubsystemIOS/Public/OnlineSubsystemIOS.h b/Engine/Source/Runtime/Online/IOS/OnlineSubsystemIOS/Public/OnlineSubsystemIOS.h
--- a/Engine/Source/Runtime/Online/IOS/OnlineSubsystemIOS/Public/OnlineSubsystemIOS.h
+++ b/Engine/Source/Runtime/Online/IOS/OnlineSubsystemIOS/Public/OnlineSubsystemIOS.h
@@ -20,30 +20,30 @@ public:
 	virtual ~FOnlineSubsystemIOS() {};
 
 	// Begin IOnlineSubsystem Interface
-	virtual IOnlineSessionPtr GetSessionInterface() const OVERRIDE;
-	virtual IOnlineFriendsPtr GetFriendsInterface() const OVERRIDE;
-	virtual IOnlineSharedCloudPtr GetSharedCloudInterface() const OVERRIDE;
-	virtual IOnlineUserCloudPtr GetUserCloudInterface() const OVERRIDE;
-	virtual IOnlineLeaderboardsPtr GetLeaderboardsInterface() const OVERRIDE;
-	virtual IOnlineVoicePtr GetVoiceInterface() const  OVERRIDE;
-	virtual IOnlineExternalUIPtr GetExternalUIInterface() const OVERRIDE;
-	virtual IOnlineTimePtr GetTimeInterface() const OVERRIDE;
-	virtual IOnlineIdentityPtr GetIdentityInterface() const OVERRIDE;
-	virtual IOnlinePartyPtr GetPartyInterface() const OVERRIDE;
-	virtual IOnlineTitleFilePtr GetTitleFileInterface() const OVERRIDE;
-	virtual IOnlineEntitlementsPtr GetEntitlementsInterface() const OVERRIDE;
-	virtual IOnlineStorePtr GetStoreInterface() const OVERRIDE;
-	virtual IOnlineEventsPtr GetEventsInterface() const OVERRIDE;
-	virtual IOnlineAchievementsPtr GetAchievementsInterface() const OVERRIDE;
-	virtual IOnlineSharingPtr GetSharingInterface() const OVERRIDE;
-	virtual IOnlineUserPtr GetUserInterface() const OVERRIDE;
-	virtual IOnlineMessagePtr GetMessageInterface() const OVERRIDE;
-	virtual IOnlinePresencePtr GetPresenceInterface() const OVERRIDE;
-	virtual bool Init() OVERRIDE;
-	virtual bool Shutdown() OVERRIDE;
-	virtual FString GetAppId() const OVERRIDE;
-	virtual bool Exec(class UWorld* InWorld, const TCHAR* Cmd, FOutputDevice& Ar) OVERRIDE;
-	virtual bool Tick(float DeltaTime) OVERRIDE;
+	virtual IOnlineSessionPtr GetSessionInterface() const override;
+	virtual IOnlineFriendsPtr GetFriendsInterface() const override;
+	virtual IOnlineSharedCloudPtr GetSharedCloudInterface() const override;
+	virtual IOnlineUserCloudPtr GetUserCloudInterface() const override;
+	virtual IOnlineLeaderboardsPtr GetLeaderboardsInterface() const override;
+	virtual IOnlineVoicePtr GetVoiceInterface() const  override;
+	virtual IOnlineExternalUIPtr GetExternalUIInterface() const override;
+	virtual IOnlineTimePtr GetTimeInterface() const override;
+	virtual IOnlineIdentityPtr GetIdentityInterface() const override;
+	virtual IOnlinePartyPtr GetPartyInterface() const override;
+	virtual IOnlineTitleFilePtr GetTitleFileInterface() const override;
+	virtual IOnlineEntitlementsPtr GetEntitlementsInterface() const override;
+	virtual IOnlineStorePtr GetStoreInterface() const override;
+	virtual IOnlineEventsPtr GetEventsInterface() const override;
+	virtual IOnlineAchievementsPtr GetAchievementsInterface() const override;
+	virtual IOnlineSharingPtr GetSharingInterface() const override;
+	virtual IOnlineUserPtr GetUserInterface() const override;
+	virtual IOnlineMessagePtr GetMessageInterface() const override;
+	virtual IOnlinePresencePtr GetPresenceInterface() const override;
+	virtual bool Init() override;
+	virtual bool Shutdown() override;
+	virtual FString GetAppId() const override;
+	virtual bool Exec(class UWorld* InWorld, const TCHAR* Cmd, FOutputDevice& Ar) override;
+	virtual bool Tick(float DeltaTime) override;
 	// End IOnlineSubsystem Interface
 
 PACKAGE_SCOPE:
diff --git a/Engine/Source/Runtime/Online/IOS/OnlineSubsystemIOS/Public/OnlineSubsystemIOSModule.h b/Engine/Source/Runtime/Online/IOS/OnlineSubsystemIOS/Public/OnlineSubsystemIOSModule.h
--- a/Engine/Source/Runtime/Online/IOS/OnlineSubsystemIOS/Public/OnlineSubsystemIOSModule.h
+++ b/Engine/Source/Runtime/Online/IOS/OnlineSubsystemIOS/Public/OnlineSubsystemIOSModule.h
@@ -26,15 +26,15 @@ public:
 	virtual ~FOnlineSubsystemIOSModule(){}
 
 	// Begin IModuleInterface interface
-	virtual void StartupModule() OVERRIDE;
-	virtual void ShutdownModule() OVERRIDE;
+	virtual void StartupModule() override;
+	virtual void ShutdownModule() override;
 
-	virtual bool SupportsDynamicReloading() OVERRIDE
+	virtual bool SupportsDynamicReloading() override
 	{
 		return false;
 	}
 
-	virtual bool SupportsAutomaticShutdown() OVERRIDE
+	virtual bool SupportsAutomaticShutdown() override
 	{
 		return false;
 	}
diff --git a/Engine/Source/Runtime/Online/OnlineSubsystem/Public/OnlineAsyncTaskManager.h b/Engine/Source/Runtime/Online/OnlineSubsystem/Public/OnlineAsyncTaskManager.h
--- a/Engine/Source/Runtime/Online/OnlineSubsystem/Public/OnlineAsyncTaskManager.h
+++ b/Engine/Source/Runtime/Online/OnlineSubsystem/Public/OnlineAsyncTaskManager.h
@@ -77,12 +77,12 @@ public:
 	explicit FOnlineAsyncItemGenericCallable(const CallableType& InCallable)
 		: CallableObject(InCallable) {}
 
-	virtual void Finalize() OVERRIDE
+	virtual void Finalize() override
 	{
 		CallableObject();
 	}
 
-	virtual FString ToString() const OVERRIDE { return FString("FOnlineAsyncItemGenericCallable"); }
+	virtual FString ToString() const override { return FString("FOnlineAsyncItemGenericCallable"); }
 
 private:
 	/** Stored copy of the object to invoke on the game thread. */
@@ -194,7 +194,7 @@ public:
 	 * Check the state of the async task
 	 * @return true if complete, false otherwise
 	 */
-	virtual bool IsDone() OVERRIDE
+	virtual bool IsDone() override
 	{
 		return bIsComplete;
 	}
@@ -203,7 +203,7 @@ public:
 	 * Check the success of the async task
 	 * @return true if successful, false otherwise
 	 */
-	virtual bool WasSuccessful() OVERRIDE
+	virtual bool WasSuccessful() override
 	{
 		return bWasSuccessful;
 	}
diff --git a/Engine/Source/Runtime/Online/OnlineSubsystem/Public/OnlineJsonSerializer.h b/Engine/Source/Runtime/Online/OnlineSubsystem/Public/OnlineJsonSerializer.h
--- a/Engine/Source/Runtime/Online/OnlineSubsystem/Public/OnlineJsonSerializer.h
+++ b/Engine/Source/Runtime/Online/OnlineSubsystem/Public/OnlineJsonSerializer.h
@@ -8,7 +8,7 @@
  * Macros used to generate a serialization function for a class derived from FJsonSerializable
  */
  #define BEGIN_ONLINE_JSON_SERIALIZER \
-	virtual void Serialize(FOnlineJsonSerializerBase& Serializer) OVERRIDE \
+	virtual void Serialize(FOnlineJsonSerializerBase& Serializer) override \
 	{ \
 		Serializer.StartObject();
 
@@ -82,14 +82,14 @@ public:
 	}
 
     /** Is the JSON being read from */
-	virtual bool IsLoading() const OVERRIDE { return false; }
+	virtual bool IsLoading() const override { return false; }
     /** Is the JSON being written to */
-	virtual bool IsSaving() const OVERRIDE { return true; }
+	virtual bool IsSaving() const override { return true; }
 
 	/**
 	 * Starts a new object "{"
 	 */
-	virtual void StartObject() OVERRIDE
+	virtual void StartObject() override
 	{
 		JsonWriter->WriteObjectStart();
 	}
@@ -97,29 +97,29 @@ public:
 	/**
 	 * Starts a new object "{"
 	 */
-	virtual void StartObject(const FString& Name) OVERRIDE
+	virtual void StartObject(const FString& Name) override
 	{
 		JsonWriter->WriteObjectStart(Name);
 	}
 	/**
 	 * Completes the definition of an object "}"
 	 */
-	virtual void EndObject() OVERRIDE
+	virtual void EndObject() override
 	{
 		JsonWriter->WriteObjectEnd();
 	}
 
-	virtual void StartArray() OVERRIDE
+	virtual void StartArray() override
 	{
 		JsonWriter->WriteArrayStart();
 	}
 
-	virtual void StartArray(const FString& Name) OVERRIDE
+	virtual void StartArray(const FString& Name) override
 	{
 		JsonWriter->WriteArrayStart(Name);
 	}
 
-	virtual void EndArray() OVERRIDE
+	virtual void EndArray() override
 	{
 		JsonWriter->WriteArrayEnd();
 	}
@@ -129,7 +129,7 @@ public:
 	 * @param Name the field name to write out
 	 * @param Value the value to write out
 	 */
-	virtual void Serialize(const TCHAR* Name, int32& Value) OVERRIDE
+	virtual void Serialize(const TCHAR* Name, int32& Value) override
 	{
 		JsonWriter->WriteValue(Name, (const float)Value);
 	}
@@ -139,7 +139,7 @@ public:
 	 * @param Name the field name to write out
 	 * @param Value the value to write out
 	 */
-	virtual void Serialize(const TCHAR* Name, uint32& Value) OVERRIDE
+	virtual void Serialize(const TCHAR* Name, uint32& Value) override
 	{
 		JsonWriter->WriteValue(Name, (const float)Value);
 	}
@@ -149,7 +149,7 @@ public:
 	 * @param Name the field name to write out
 	 * @param Value the value to write out
 	 */
-	virtual void Serialize(const TCHAR* Name, bool& Value) OVERRIDE
+	virtual void Serialize(const TCHAR* Name, bool& Value) override
 	{
 		JsonWriter->WriteValue(Name, (const bool)Value);
 	}
@@ -159,7 +159,7 @@ public:
 	 * @param Name the field name to write out
 	 * @param Value the value to write out
 	 */
-	virtual void Serialize(const TCHAR* Name, FString& Value) OVERRIDE
+	virtual void Serialize(const TCHAR* Name, FString& Value) override
 	{
 		JsonWriter->WriteValue(Name, Value);
 	}
@@ -169,7 +169,7 @@ public:
 	 * @param Name the field name to write out
 	 * @param Value the value to write out
 	 */
-	virtual void Serialize(const TCHAR* Name, float& Value) OVERRIDE
+	virtual void Serialize(const TCHAR* Name, float& Value) override
 	{
 		JsonWriter->WriteValue(Name, (const float)Value);
 	}
@@ -179,7 +179,7 @@ public:
 	* @param Name the field name to write out
 	* @param Value the value to write out
 	*/
-	virtual void Serialize(const TCHAR* Name, double& Value) OVERRIDE
+	virtual void Serialize(const TCHAR* Name, double& Value) override
 	{
 		JsonWriter->WriteValue(Name, (const double)Value);
 	}
@@ -189,7 +189,7 @@ public:
 	 * @param Name the name of the property to serialize
 	 * @param Array the array to serialize
 	 */
-	virtual void SerializeArray(FJsonSerializableArray& Array) OVERRIDE
+	virtual void SerializeArray(FJsonSerializableArray& Array) override
 	{
 		JsonWriter->WriteArrayStart();
 		// Iterate all of values
@@ -205,7 +205,7 @@ public:
 	 * @param Name the name of the property to serialize
 	 * @param Array the array to serialize
 	 */
-	virtual void SerializeArray(const TCHAR* Name, FJsonSerializableArray& Array) OVERRIDE
+	virtual void SerializeArray(const TCHAR* Name, FJsonSerializableArray& Array) override
 	{
 		JsonWriter->WriteArrayStart(Name);
 		// Iterate all of values
@@ -221,7 +221,7 @@ public:
 	 * @param Name the name of the property to serialize
 	 * @param Map the map to serialize
 	 */
-	virtual void SerializeMap(const TCHAR* Name, FJsonSerializableKeyValueMap& Map) OVERRIDE
+	virtual void SerializeMap(const TCHAR* Name, FJsonSerializableKeyValueMap& Map) override
 	{
 		JsonWriter->WriteObjectStart(Name);
 		// Iterate all of the keys and their values
@@ -237,7 +237,7 @@ public:
 	 * @param Name the name of the property to serialize
 	 * @param Map the map to serialize
 	 */
-	virtual void SerializeMap(const TCHAR* Name, FJsonSerializableKeyValueMapInt& Map) OVERRIDE
+	virtual void SerializeMap(const TCHAR* Name, FJsonSerializableKeyValueMapInt& Map) override
 	{
 		JsonWriter->WriteObjectStart(Name);
 		// Iterate all of the keys and their values
@@ -270,37 +270,37 @@ public:
 	}
 
     /** Is the JSON being read from */
-	virtual bool IsLoading() const OVERRIDE { return true; }
+	virtual bool IsLoading() const override { return true; }
     /** Is the JSON being written to */
-	virtual bool IsSaving() const OVERRIDE { return false; }
+	virtual bool IsSaving() const override { return false; }
 
 	/** Ignored */
-	virtual void StartObject() OVERRIDE
+	virtual void StartObject() override
 	{
 		// Empty on purpose
 	}
 	/** Ignored */
-	virtual void StartObject(const FString& Name) OVERRIDE
+	virtual void StartObject(const FString& Name) override
 	{
 		// Empty on purpose
 	}
 	/** Ignored */
-	virtual void EndObject() OVERRIDE
+	virtual void EndObject() override
 	{
 		// Empty on purpose
 	}
 	/** Ignored */
-	virtual void StartArray() OVERRIDE
+	virtual void StartArray() override
 	{
 		// Empty on purpose
 	}
 	/** Ignored */
-	virtual void StartArray(const FString& Name) OVERRIDE
+	virtual void StartArray(const FString& Name) override
 	{
 		// Empty on purpose
 	}
 	/** Ignored */
-	virtual void EndArray() OVERRIDE
+	virtual void EndArray() override
 	{
 		// Empty on purpose
 	}
@@ -310,7 +310,7 @@ public:
 	 * @param Name the name of the field to read
 	 * @param Value the out value to read the data into
 	 */
-	virtual void Serialize(const TCHAR* Name, int32& Value) OVERRIDE
+	virtual void Serialize(const TCHAR* Name, int32& Value) override
 	{
 		if (JsonObject->HasField(Name))
 		{
@@ -323,7 +323,7 @@ public:
 	 * @param Name the name of the field to read
 	 * @param Value the out value to read the data into
 	 */
-	virtual void Serialize(const TCHAR* Name, uint32& Value) OVERRIDE
+	virtual void Serialize(const TCHAR* Name, uint32& Value) override
 	{
 		if (JsonObject->HasField(Name))
 		{
@@ -336,7 +336,7 @@ public:
 	 * @param Name the name of the field to read
 	 * @param Value the out value to read the data into
 	 */
-	virtual void Serialize(const TCHAR* Name, bool& Value) OVERRIDE
+	virtual void Serialize(const TCHAR* Name, bool& Value) override
 	{
 		if (JsonObject->HasField(Name))
 		{
@@ -349,7 +349,7 @@ public:
 	 * @param Name the name of the field to read
 	 * @param Value the out value to read the data into
 	 */
-	virtual void Serialize(const TCHAR* Name, FString& Value) OVERRIDE
+	virtual void Serialize(const TCHAR* Name, FString& Value) override
 	{
 		if (JsonObject->HasField(Name))
 		{
@@ -362,7 +362,7 @@ public:
 	 * @param Name the name of the field to read
 	 * @param Value the out value to read the data into
 	 */
-	virtual void Serialize(const TCHAR* Name, float& Value) OVERRIDE
+	virtual void Serialize(const TCHAR* Name, float& Value) override
 	{
 		if (JsonObject->HasField(Name))
 		{
@@ -375,7 +375,7 @@ public:
 	* @param Name the name of the field to read
 	* @param Value the out value to read the data into
 	*/
-	virtual void Serialize(const TCHAR* Name, double& Value) OVERRIDE
+	virtual void Serialize(const TCHAR* Name, double& Value) override
 	{
 		if (JsonObject->HasField(Name))
 		{
@@ -388,7 +388,7 @@ public:
 	 * @param Name the name of the property to serialize
 	 * @param Array the array to serialize
 	 */
-	virtual void SerializeArray(FJsonSerializableArray& Array) OVERRIDE
+	virtual void SerializeArray(FJsonSerializableArray& Array) override
 	{
 		// @todo - higher level serialization is expecting a Json Object
 		check(0 && TEXT("Not implemented"));
@@ -399,7 +399,7 @@ public:
 	 * @param Name the name of the property to serialize
 	 * @param Array the array to serialize
 	 */
-	virtual void SerializeArray(const TCHAR* Name, FJsonSerializableArray& Array) OVERRIDE
+	virtual void SerializeArray(const TCHAR* Name, FJsonSerializableArray& Array) override
 	{
 		if (JsonObject->HasField(Name))
 		{
@@ -417,7 +417,7 @@ public:
 	 * @param Name the name of the property to serialize
 	 * @param Map the map to serialize
 	 */
-	virtual void SerializeMap(const TCHAR* Name, FJsonSerializableKeyValueMap& Map) OVERRIDE
+	virtual void SerializeMap(const TCHAR* Name, FJsonSerializableKeyValueMap& Map) override
 	{
 		if (JsonObject->HasField(Name))
 		{
@@ -437,7 +437,7 @@ public:
 	 * @param Name the name of the property to serialize
 	 * @param Map the map to serialize
 	 */
-	virtual void SerializeMap(const TCHAR* Name, FJsonSerializableKeyValueMapInt& Map) OVERRIDE
+	virtual void SerializeMap(const TCHAR* Name, FJsonSerializableKeyValueMapInt& Map) override
 	{
 		if (JsonObject->HasField(Name))
 		{
diff --git a/Engine/Source/Runtime/Online/OnlineSubsystem/Public/OnlineSubsystemImpl.h b/Engine/Source/Runtime/Online/OnlineSubsystem/Public/OnlineSubsystemImpl.h
--- a/Engine/Source/Runtime/Online/OnlineSubsystem/Public/OnlineSubsystemImpl.h
+++ b/Engine/Source/Runtime/Online/OnlineSubsystem/Public/OnlineSubsystemImpl.h
@@ -33,14 +33,14 @@ public:
 	virtual ~FOnlineSubsystemImpl();
 
 	// IOnlineSubsystem
-	virtual bool IsServer() const OVERRIDE;
-	virtual bool IsDedicated() const OVERRIDE{ return bForceDedicated || IsRunningDedicatedServer(); }
-	virtual void SetForceDedicated(bool bForce) OVERRIDE { bForceDedicated = bForce; }
+	virtual bool IsServer() const override;
+	virtual bool IsDedicated() const override{ return bForceDedicated || IsRunningDedicatedServer(); }
+	virtual void SetForceDedicated(bool bForce) override { bForceDedicated = bForce; }
 	
-	virtual class UObject* GetNamedInterface(FName InterfaceName) OVERRIDE;
-	virtual void SetNamedInterface(FName InterfaceName, class UObject* NewInterface) OVERRIDE;
+	virtual class UObject* GetNamedInterface(FName InterfaceName) override;
+	virtual void SetNamedInterface(FName InterfaceName, class UObject* NewInterface) override;
 
-	virtual bool IsLocalPlayer(const FUniqueNetId& UniqueId) const OVERRIDE;
+	virtual bool IsLocalPlayer(const FUniqueNetId& UniqueId) const override;
 
 	// FOnlineSubsystemImpl
 
diff --git a/Engine/Source/Runtime/Online/OnlineSubsystem/Public/OnlineSubsystemModule.h b/Engine/Source/Runtime/Online/OnlineSubsystem/Public/OnlineSubsystemModule.h
--- a/Engine/Source/Runtime/Online/OnlineSubsystem/Public/OnlineSubsystemModule.h
+++ b/Engine/Source/Runtime/Online/OnlineSubsystem/Public/OnlineSubsystemModule.h
@@ -113,20 +113,20 @@ public:
 	 * Called right after the module DLL has been loaded and the module object has been created
 	 * Overloaded to allow the default subsystem a chance to load
 	 */
-	virtual void StartupModule() OVERRIDE;
+	virtual void StartupModule() override;
 
 	/**
 	 * Called before the module is unloaded, right before the module object is destroyed.
 	 * Overloaded to shut down all loaded online subsystems
 	 */
-	virtual void ShutdownModule() OVERRIDE;
+	virtual void ShutdownModule() override;
 
 	/**
 	 * Override this to set whether your module is allowed to be unloaded on the fly
 	 *
 	 * @return	Whether the module supports shutdown separate from the rest of the engine.
 	 */
-	virtual bool SupportsDynamicReloading() OVERRIDE
+	virtual bool SupportsDynamicReloading() override
 	{
 		return false;
 	}
@@ -136,7 +136,7 @@ public:
 	 *
 	 * @return	Whether the module supports shutdown on application exit
 	 */
-	virtual bool SupportsAutomaticShutdown() OVERRIDE
+	virtual bool SupportsAutomaticShutdown() override
 	{
 		return false;
 	}
diff --git a/Engine/Source/Runtime/Online/OnlineSubsystem/Public/OnlineSubsystemTypes.h b/Engine/Source/Runtime/Online/OnlineSubsystem/Public/OnlineSubsystemTypes.h
--- a/Engine/Source/Runtime/Online/OnlineSubsystem/Public/OnlineSubsystemTypes.h
+++ b/Engine/Source/Runtime/Online/OnlineSubsystem/Public/OnlineSubsystemTypes.h
@@ -825,27 +825,27 @@ public:
 
 	// IOnlinePlatformData
 
-	virtual const uint8* GetBytes() const OVERRIDE
+	virtual const uint8* GetBytes() const override
 	{
 		return (const uint8*)UniqueNetIdStr.GetCharArray().GetData();
 	}
 
-	virtual int32 GetSize() const OVERRIDE
+	virtual int32 GetSize() const override
 	{
 		return UniqueNetIdStr.GetCharArray().GetTypeSize() * UniqueNetIdStr.GetCharArray().Num();
 	}
 
-	virtual bool IsValid() const OVERRIDE
+	virtual bool IsValid() const override
 	{
 		return !UniqueNetIdStr.IsEmpty();
 	}
 
-	virtual FString ToString() const OVERRIDE
+	virtual FString ToString() const override
 	{
 		return UniqueNetIdStr;
 	}
 
-	virtual FString ToDebugString() const OVERRIDE
+	virtual FString ToDebugString() const override
 	{
 		return UniqueNetIdStr;
 	}
diff --git a/Engine/Source/Runtime/Online/OnlineSubsystemAmazon/Private/OnlineIdentityAmazon.h b/Engine/Source/Runtime/Online/OnlineSubsystemAmazon/Private/OnlineIdentityAmazon.h
--- a/Engine/Source/Runtime/Online/OnlineSubsystemAmazon/Private/OnlineIdentityAmazon.h
+++ b/Engine/Source/Runtime/Online/OnlineSubsystemAmazon/Private/OnlineIdentityAmazon.h
@@ -16,15 +16,15 @@ public:
 
 	// FOnlineUser
 	
-	virtual TSharedRef<FUniqueNetId> GetUserId() const OVERRIDE;
-	virtual FString GetRealName() const OVERRIDE;
-	virtual FString GetDisplayName() const OVERRIDE;
-	virtual bool GetUserAttribute(const FString& AttrName, FString& OutAttrValue) const OVERRIDE;
+	virtual TSharedRef<FUniqueNetId> GetUserId() const override;
+	virtual FString GetRealName() const override;
+	virtual FString GetDisplayName() const override;
+	virtual bool GetUserAttribute(const FString& AttrName, FString& OutAttrValue) const override;
 
 	// FUserOnlineAccount
 
-	virtual FString GetAccessToken() const OVERRIDE;
-	virtual bool GetAuthAttribute(const FString& AttrName, FString& OutAttrValue) const OVERRIDE;
+	virtual FString GetAccessToken() const override;
+	virtual bool GetAuthAttribute(const FString& AttrName, FString& OutAttrValue) const override;
 
 	// FUserOnlineAccountAmazon
 
@@ -98,17 +98,17 @@ class FOnlineIdentityAmazon :
 public:
 	// IOnlineIdentity
 	
-	virtual bool Login(int32 LocalUserNum, const FOnlineAccountCredentials& AccountCredentials) OVERRIDE;
-	virtual bool Logout(int32 LocalUserNum) OVERRIDE;
-	virtual bool AutoLogin(int32 LocalUserNum) OVERRIDE;
-	virtual TSharedPtr<FUserOnlineAccount> GetUserAccount(const FUniqueNetId& UserId) const OVERRIDE;
-	virtual TArray<TSharedPtr<FUserOnlineAccount> > GetAllUserAccounts() const OVERRIDE;
-	virtual TSharedPtr<FUniqueNetId> GetUniquePlayerId(int32 LocalUserNum) const OVERRIDE;
-	virtual TSharedPtr<FUniqueNetId> CreateUniquePlayerId(uint8* Bytes, int32 Size) OVERRIDE;
-	virtual TSharedPtr<FUniqueNetId> CreateUniquePlayerId(const FString& Str) OVERRIDE;
-	virtual ELoginStatus::Type GetLoginStatus(int32 LocalUserNum) const OVERRIDE;
-	virtual FString GetPlayerNickname(int32 LocalUserNum) const OVERRIDE;
-	virtual FString GetAuthToken(int32 LocalUserNum) const OVERRIDE;
+	virtual bool Login(int32 LocalUserNum, const FOnlineAccountCredentials& AccountCredentials) override;
+	virtual bool Logout(int32 LocalUserNum) override;
+	virtual bool AutoLogin(int32 LocalUserNum) override;
+	virtual TSharedPtr<FUserOnlineAccount> GetUserAccount(const FUniqueNetId& UserId) const override;
+	virtual TArray<TSharedPtr<FUserOnlineAccount> > GetAllUserAccounts() const override;
+	virtual TSharedPtr<FUniqueNetId> GetUniquePlayerId(int32 LocalUserNum) const override;
+	virtual TSharedPtr<FUniqueNetId> CreateUniquePlayerId(uint8* Bytes, int32 Size) override;
+	virtual TSharedPtr<FUniqueNetId> CreateUniquePlayerId(const FString& Str) override;
+	virtual ELoginStatus::Type GetLoginStatus(int32 LocalUserNum) const override;
+	virtual FString GetPlayerNickname(int32 LocalUserNum) const override;
+	virtual FString GetAuthToken(int32 LocalUserNum) const override;
 
 	// FOnlineIdentityAmazon
 
diff --git a/Engine/Source/Runtime/Online/OnlineSubsystemAmazon/Public/OnlineSubsystemAmazon.h b/Engine/Source/Runtime/Online/OnlineSubsystemAmazon/Public/OnlineSubsystemAmazon.h
--- a/Engine/Source/Runtime/Online/OnlineSubsystemAmazon/Public/OnlineSubsystemAmazon.h
+++ b/Engine/Source/Runtime/Online/OnlineSubsystemAmazon/Public/OnlineSubsystemAmazon.h
@@ -34,97 +34,97 @@ PACKAGE_SCOPE:
 public:
 	// IOnlineSubsystem
 
-	virtual IOnlineSessionPtr GetSessionInterface() const OVERRIDE
+	virtual IOnlineSessionPtr GetSessionInterface() const override
 	{
 		return NULL;
 	}
-	virtual IOnlineFriendsPtr GetFriendsInterface() const OVERRIDE
+	virtual IOnlineFriendsPtr GetFriendsInterface() const override
 	{
 		return NULL;
 	}
-	virtual IOnlineSharedCloudPtr GetSharedCloudInterface() const OVERRIDE
+	virtual IOnlineSharedCloudPtr GetSharedCloudInterface() const override
 	{
 		return NULL;
 	}
-	virtual IOnlineUserCloudPtr GetUserCloudInterface() const OVERRIDE
+	virtual IOnlineUserCloudPtr GetUserCloudInterface() const override
 	{
 		return NULL;
 	}
-	virtual IOnlineLeaderboardsPtr GetLeaderboardsInterface() const OVERRIDE
+	virtual IOnlineLeaderboardsPtr GetLeaderboardsInterface() const override
 	{
 		return NULL;
 	}
-	virtual IOnlineVoicePtr GetVoiceInterface() const OVERRIDE
+	virtual IOnlineVoicePtr GetVoiceInterface() const override
 	{
 		return NULL;
 	}
-	virtual IOnlineExternalUIPtr GetExternalUIInterface() const OVERRIDE	
+	virtual IOnlineExternalUIPtr GetExternalUIInterface() const override	
 	{
 		return NULL;
 	}
-	virtual IOnlineTimePtr GetTimeInterface() const OVERRIDE
+	virtual IOnlineTimePtr GetTimeInterface() const override
 	{
 		return NULL;
 	}
-	virtual IOnlinePartyPtr GetPartyInterface() const OVERRIDE
+	virtual IOnlinePartyPtr GetPartyInterface() const override
 	{
 		return NULL;
 	}
-	virtual IOnlineTitleFilePtr GetTitleFileInterface() const OVERRIDE
+	virtual IOnlineTitleFilePtr GetTitleFileInterface() const override
 	{
 		return NULL;
 	}
 
-	virtual IOnlineEntitlementsPtr GetEntitlementsInterface() const OVERRIDE
+	virtual IOnlineEntitlementsPtr GetEntitlementsInterface() const override
 	{
 		return NULL;
 	}
 
-	virtual IOnlineIdentityPtr GetIdentityInterface() const OVERRIDE;
+	virtual IOnlineIdentityPtr GetIdentityInterface() const override;
 
-	virtual IOnlineStorePtr GetStoreInterface() const OVERRIDE
+	virtual IOnlineStorePtr GetStoreInterface() const override
 	{
 		return NULL;
 	}
 
-	virtual IOnlineEventsPtr GetEventsInterface() const OVERRIDE
+	virtual IOnlineEventsPtr GetEventsInterface() const override
 	{
 		return NULL;
 	}
 
-	virtual IOnlineAchievementsPtr GetAchievementsInterface() const OVERRIDE
+	virtual IOnlineAchievementsPtr GetAchievementsInterface() const override
 	{
 		return NULL;
 	}
 
-	virtual IOnlineSharingPtr GetSharingInterface() const OVERRIDE
+	virtual IOnlineSharingPtr GetSharingInterface() const override
 	{
 		return NULL;
 	}
 
-	virtual IOnlineUserPtr GetUserInterface() const OVERRIDE
+	virtual IOnlineUserPtr GetUserInterface() const override
 	{
 		return NULL;
 	}
 
-	virtual IOnlineMessagePtr GetMessageInterface() const OVERRIDE
+	virtual IOnlineMessagePtr GetMessageInterface() const override
 	{
 		return NULL;
 	}
 
-	virtual IOnlinePresencePtr GetPresenceInterface() const OVERRIDE
+	virtual IOnlinePresencePtr GetPresenceInterface() const override
 	{
 		return NULL;
 	}
 
-	virtual bool Init() OVERRIDE;
-	virtual bool Shutdown() OVERRIDE;
-	virtual FString GetAppId() const OVERRIDE;
-	virtual bool Exec(class UWorld* InWorld, const TCHAR* Cmd, FOutputDevice& Ar) OVERRIDE;
+	virtual bool Init() override;
+	virtual bool Shutdown() override;
+	virtual FString GetAppId() const override;
+	virtual bool Exec(class UWorld* InWorld, const TCHAR* Cmd, FOutputDevice& Ar) override;
 
 	// FTickerBaseObject
 
-	virtual bool Tick(float DeltaTime) OVERRIDE;
+	virtual bool Tick(float DeltaTime) override;
 
 	// FOnlineSubsystemAmazon
 
diff --git a/Engine/Source/Runtime/Online/OnlineSubsystemAmazon/Public/OnlineSubsystemAmazonModule.h b/Engine/Source/Runtime/Online/OnlineSubsystemAmazon/Public/OnlineSubsystemAmazonModule.h
--- a/Engine/Source/Runtime/Online/OnlineSubsystemAmazon/Public/OnlineSubsystemAmazonModule.h
+++ b/Engine/Source/Runtime/Online/OnlineSubsystemAmazon/Public/OnlineSubsystemAmazonModule.h
@@ -35,14 +35,14 @@ public:
 
 	// IModuleInterface
 
-	virtual void StartupModule() OVERRIDE;
-	virtual void ShutdownModule() OVERRIDE;
-	virtual bool SupportsDynamicReloading() OVERRIDE
+	virtual void StartupModule() override;
+	virtual void ShutdownModule() override;
+	virtual bool SupportsDynamicReloading() override
 	{
 		return false;
 	}
 
-	virtual bool SupportsAutomaticShutdown() OVERRIDE
+	virtual bool SupportsAutomaticShutdown() override
 	{
 		return false;
 	}
diff --git a/Engine/Source/Runtime/Online/OnlineSubsystemFacebook/Private/IOS/OnlineFriendsFacebook.h b/Engine/Source/Runtime/Online/OnlineSubsystemFacebook/Private/IOS/OnlineFriendsFacebook.h
--- a/Engine/Source/Runtime/Online/OnlineSubsystemFacebook/Private/IOS/OnlineFriendsFacebook.h
+++ b/Engine/Source/Runtime/Online/OnlineSubsystemFacebook/Private/IOS/OnlineFriendsFacebook.h
@@ -17,15 +17,15 @@ public:
 
 	// FOnlineUser
 
-	virtual TSharedRef<FUniqueNetId> GetUserId() const OVERRIDE;
-	virtual FString GetRealName() const OVERRIDE;
-	virtual FString GetDisplayName() const OVERRIDE;
-	virtual bool GetUserAttribute(const FString& AttrName, FString& OutAttrValue) const OVERRIDE;
+	virtual TSharedRef<FUniqueNetId> GetUserId() const override;
+	virtual FString GetRealName() const override;
+	virtual FString GetDisplayName() const override;
+	virtual bool GetUserAttribute(const FString& AttrName, FString& OutAttrValue) const override;
 
 	// FOnlineFriend
 	
-	virtual EInviteStatus::Type GetInviteStatus() const OVERRIDE;
-	virtual const FOnlineUserPresence& GetPresence() const OVERRIDE;
+	virtual EInviteStatus::Type GetInviteStatus() const override;
+	virtual const FOnlineUserPresence& GetPresence() const override;
 
 	// FOnlineFriendMcp
 
@@ -83,15 +83,15 @@ public:
 
 	// IOnlineFriends
 
-	virtual bool ReadFriendsList(int32 LocalUserNum, const FString& ListName) OVERRIDE;
-	virtual bool DeleteFriendsList(int32 LocalUserNum, const FString& ListName) OVERRIDE;
-	virtual bool SendInvite(int32 LocalUserNum, const FUniqueNetId& FriendId, const FString& ListName) OVERRIDE;
-	virtual bool AcceptInvite(int32 LocalUserNum, const FUniqueNetId& FriendId, const FString& ListName) OVERRIDE;
- 	virtual bool RejectInvite(int32 LocalUserNum, const FUniqueNetId& FriendId, const FString& ListName) OVERRIDE;
- 	virtual bool DeleteFriend(int32 LocalUserNum, const FUniqueNetId& FriendId, const FString& ListName) OVERRIDE;
-	virtual bool GetFriendsList(int32 LocalUserNum, const FString& ListName, TArray< TSharedRef<FOnlineFriend> >& OutFriends) OVERRIDE;
-	virtual TSharedPtr<FOnlineFriend> GetFriend(int32 LocalUserNum, const FUniqueNetId& FriendId, const FString& ListName) OVERRIDE;
-	virtual bool IsFriend(int32 LocalUserNum, const FUniqueNetId& FriendId, const FString& ListName) OVERRIDE;
+	virtual bool ReadFriendsList(int32 LocalUserNum, const FString& ListName) override;
+	virtual bool DeleteFriendsList(int32 LocalUserNum, const FString& ListName) override;
+	virtual bool SendInvite(int32 LocalUserNum, const FUniqueNetId& FriendId, const FString& ListName) override;
+	virtual bool AcceptInvite(int32 LocalUserNum, const FUniqueNetId& FriendId, const FString& ListName) override;
+ 	virtual bool RejectInvite(int32 LocalUserNum, const FUniqueNetId& FriendId, const FString& ListName) override;
+ 	virtual bool DeleteFriend(int32 LocalUserNum, const FUniqueNetId& FriendId, const FString& ListName) override;
+	virtual bool GetFriendsList(int32 LocalUserNum, const FString& ListName, TArray< TSharedRef<FOnlineFriend> >& OutFriends) override;
+	virtual TSharedPtr<FOnlineFriend> GetFriend(int32 LocalUserNum, const FUniqueNetId& FriendId, const FString& ListName) override;
+	virtual bool IsFriend(int32 LocalUserNum, const FUniqueNetId& FriendId, const FString& ListName) override;
 
 	// FOnlineFriendsFacebook
 
diff --git a/Engine/Source/Runtime/Online/OnlineSubsystemFacebook/Private/IOS/OnlineIdentityFacebook.h b/Engine/Source/Runtime/Online/OnlineSubsystemFacebook/Private/IOS/OnlineIdentityFacebook.h
--- a/Engine/Source/Runtime/Online/OnlineSubsystemFacebook/Private/IOS/OnlineIdentityFacebook.h
+++ b/Engine/Source/Runtime/Online/OnlineSubsystemFacebook/Private/IOS/OnlineIdentityFacebook.h
@@ -25,15 +25,15 @@ public:
 
 	// FOnlineUser
 	
-	virtual TSharedRef<FUniqueNetId> GetUserId() const OVERRIDE;
-	virtual FString GetRealName() const OVERRIDE;
-	virtual FString GetDisplayName() const OVERRIDE;
-	virtual bool GetUserAttribute(const FString& AttrName, FString& OutAttrValue) const OVERRIDE;
+	virtual TSharedRef<FUniqueNetId> GetUserId() const override;
+	virtual FString GetRealName() const override;
+	virtual FString GetDisplayName() const override;
+	virtual bool GetUserAttribute(const FString& AttrName, FString& OutAttrValue) const override;
 
 	// FUserOnlineAccount
 
-	virtual FString GetAccessToken() const OVERRIDE;
-	virtual bool GetAuthAttribute(const FString& AttrName, FString& OutAttrValue) const OVERRIDE;
+	virtual FString GetAccessToken() const override;
+	virtual bool GetAuthAttribute(const FString& AttrName, FString& OutAttrValue) const override;
 
 public:
 
@@ -73,17 +73,17 @@ class FOnlineIdentityFacebook : public IOnlineIdentity
 public:
 
 	// Begin IOnlineIdentity interface	
-	virtual bool Login(int32 LocalUserNum, const FOnlineAccountCredentials& AccountCredentials) OVERRIDE;
-	virtual bool Logout(int32 LocalUserNum) OVERRIDE;
-	virtual bool AutoLogin(int32 LocalUserNum) OVERRIDE;
-	virtual TSharedPtr<FUserOnlineAccount> GetUserAccount(const FUniqueNetId& UserId) const OVERRIDE;
-	virtual TArray<TSharedPtr<FUserOnlineAccount> > GetAllUserAccounts() const OVERRIDE;
-	virtual TSharedPtr<FUniqueNetId> GetUniquePlayerId(int32 LocalUserNum) const OVERRIDE;
-	virtual TSharedPtr<FUniqueNetId> CreateUniquePlayerId(uint8* Bytes, int32 Size) OVERRIDE;
-	virtual TSharedPtr<FUniqueNetId> CreateUniquePlayerId(const FString& Str) OVERRIDE;
-	virtual ELoginStatus::Type GetLoginStatus(int32 LocalUserNum) const OVERRIDE;
-	virtual FString GetPlayerNickname(int32 LocalUserNum) const OVERRIDE;
-	virtual FString GetAuthToken(int32 LocalUserNum) const OVERRIDE;
+	virtual bool Login(int32 LocalUserNum, const FOnlineAccountCredentials& AccountCredentials) override;
+	virtual bool Logout(int32 LocalUserNum) override;
+	virtual bool AutoLogin(int32 LocalUserNum) override;
+	virtual TSharedPtr<FUserOnlineAccount> GetUserAccount(const FUniqueNetId& UserId) const override;
+	virtual TArray<TSharedPtr<FUserOnlineAccount> > GetAllUserAccounts() const override;
+	virtual TSharedPtr<FUniqueNetId> GetUniquePlayerId(int32 LocalUserNum) const override;
+	virtual TSharedPtr<FUniqueNetId> CreateUniquePlayerId(uint8* Bytes, int32 Size) override;
+	virtual TSharedPtr<FUniqueNetId> CreateUniquePlayerId(const FString& Str) override;
+	virtual ELoginStatus::Type GetLoginStatus(int32 LocalUserNum) const override;
+	virtual FString GetPlayerNickname(int32 LocalUserNum) const override;
+	virtual FString GetAuthToken(int32 LocalUserNum) const override;
 	// End IOnlineIdentity interface
 
 
diff --git a/Engine/Source/Runtime/Online/OnlineSubsystemFacebook/Private/IOS/OnlineSharingFacebook.h b/Engine/Source/Runtime/Online/OnlineSubsystemFacebook/Private/IOS/OnlineSharingFacebook.h
--- a/Engine/Source/Runtime/Online/OnlineSubsystemFacebook/Private/IOS/OnlineSharingFacebook.h
+++ b/Engine/Source/Runtime/Online/OnlineSubsystemFacebook/Private/IOS/OnlineSharingFacebook.h
@@ -17,12 +17,12 @@ class FOnlineSharingFacebook : public IOnlineSharing
 public:
 
 	// Begin IOnlineSharing Interface
-	virtual bool ReadNewsFeed(int32 LocalUserNum, int32 NumPostsToRead) OVERRIDE;
-	virtual bool RequestNewReadPermissions(int32 LocalUserNum, EOnlineSharingReadCategory::Type NewPermissions) OVERRIDE;
-	virtual bool ShareStatusUpdate(int32 LocalUserNum, const FOnlineStatusUpdate& StatusUpdate) OVERRIDE;
-	virtual bool RequestNewPublishPermissions(int32 LocalUserNum, EOnlineSharingPublishingCategory::Type NewPermissions, EOnlineStatusUpdatePrivacy::Type Privacy) OVERRIDE;
-	virtual EOnlineCachedResult::Type GetCachedNewsFeed(int32 LocalUserNum, int32 NewsFeedIdx, FOnlineStatusUpdate& OutNewsFeed) OVERRIDE;
-	virtual EOnlineCachedResult::Type GetCachedNewsFeeds(int32 LocalUserNum, TArray<FOnlineStatusUpdate>& OutNewsFeeds) OVERRIDE;
+	virtual bool ReadNewsFeed(int32 LocalUserNum, int32 NumPostsToRead) override;
+	virtual bool RequestNewReadPermissions(int32 LocalUserNum, EOnlineSharingReadCategory::Type NewPermissions) override;
+	virtual bool ShareStatusUpdate(int32 LocalUserNum, const FOnlineStatusUpdate& StatusUpdate) override;
+	virtual bool RequestNewPublishPermissions(int32 LocalUserNum, EOnlineSharingPublishingCategory::Type NewPermissions, EOnlineStatusUpdatePrivacy::Type Privacy) override;
+	virtual EOnlineCachedResult::Type GetCachedNewsFeed(int32 LocalUserNum, int32 NewsFeedIdx, FOnlineStatusUpdate& OutNewsFeed) override;
+	virtual EOnlineCachedResult::Type GetCachedNewsFeeds(int32 LocalUserNum, TArray<FOnlineStatusUpdate>& OutNewsFeeds) override;
 	// End IOnlineSharing Interface
 
 	
diff --git a/Engine/Source/Runtime/Online/OnlineSubsystemFacebook/Private/IOS/OnlineUserFacebook.h b/Engine/Source/Runtime/Online/OnlineSubsystemFacebook/Private/IOS/OnlineUserFacebook.h
--- a/Engine/Source/Runtime/Online/OnlineSubsystemFacebook/Private/IOS/OnlineUserFacebook.h
+++ b/Engine/Source/Runtime/Online/OnlineSubsystemFacebook/Private/IOS/OnlineUserFacebook.h
@@ -17,10 +17,10 @@ public:
 
 	// FOnlineUser
 
-	virtual TSharedRef<FUniqueNetId> GetUserId() const OVERRIDE;
-	virtual FString GetRealName() const OVERRIDE;
-	virtual FString GetDisplayName() const OVERRIDE;
-	virtual bool GetUserAttribute(const FString& AttrName, FString& OutAttrValue) const OVERRIDE;
+	virtual TSharedRef<FUniqueNetId> GetUserId() const override;
+	virtual FString GetRealName() const override;
+	virtual FString GetDisplayName() const override;
+	virtual bool GetUserAttribute(const FString& AttrName, FString& OutAttrValue) const override;
 
 	// FOnlineUserInfoFacebook
 
@@ -75,9 +75,9 @@ public:
 	
 	// IOnlineUser
 
-	virtual bool QueryUserInfo(int32 LocalUserNum, const TArray<TSharedRef<class FUniqueNetId> >& UserIds) OVERRIDE;
-	virtual bool GetAllUserInfo(int32 LocalUserNum, TArray< TSharedRef<class FOnlineUser> >& OutUsers) OVERRIDE;
-	virtual TSharedPtr<FOnlineUser> GetUserInfo(int32 LocalUserNum, const class FUniqueNetId& UserId) OVERRIDE;
+	virtual bool QueryUserInfo(int32 LocalUserNum, const TArray<TSharedRef<class FUniqueNetId> >& UserIds) override;
+	virtual bool GetAllUserInfo(int32 LocalUserNum, TArray< TSharedRef<class FOnlineUser> >& OutUsers) override;
+	virtual TSharedPtr<FOnlineUser> GetUserInfo(int32 LocalUserNum, const class FUniqueNetId& UserId) override;
 
 	// FOnlineUserFacebook
 
diff --git a/Engine/Source/Runtime/Online/OnlineSubsystemFacebook/Public/OnlineSubsystemFacebook.h b/Engine/Source/Runtime/Online/OnlineSubsystemFacebook/Public/OnlineSubsystemFacebook.h
--- a/Engine/Source/Runtime/Online/OnlineSubsystemFacebook/Public/OnlineSubsystemFacebook.h
+++ b/Engine/Source/Runtime/Online/OnlineSubsystemFacebook/Public/OnlineSubsystemFacebook.h
@@ -22,33 +22,33 @@ class ONLINESUBSYSTEMFACEBOOK_API FOnlineSubsystemFacebook
 public:
 
 	// IOnlineSubsystem Interface
-	virtual IOnlineSessionPtr GetSessionInterface() const OVERRIDE;
-	virtual IOnlineFriendsPtr GetFriendsInterface() const OVERRIDE;
-	virtual IOnlineSharedCloudPtr GetSharedCloudInterface() const OVERRIDE;
-	virtual IOnlineUserCloudPtr GetUserCloudInterface() const OVERRIDE;
-	virtual IOnlineLeaderboardsPtr GetLeaderboardsInterface() const OVERRIDE;
-	virtual IOnlineVoicePtr GetVoiceInterface() const OVERRIDE;
-	virtual IOnlineExternalUIPtr GetExternalUIInterface() const OVERRIDE;
-	virtual IOnlineTimePtr GetTimeInterface() const OVERRIDE;
-	virtual IOnlineIdentityPtr GetIdentityInterface() const OVERRIDE;
-	virtual IOnlinePartyPtr GetPartyInterface() const OVERRIDE;
-	virtual IOnlineTitleFilePtr GetTitleFileInterface() const OVERRIDE;
-	virtual IOnlineEntitlementsPtr GetEntitlementsInterface() const OVERRIDE;
-	virtual IOnlineStorePtr GetStoreInterface() const OVERRIDE;
-	virtual IOnlineEventsPtr GetEventsInterface() const OVERRIDE;
-	virtual IOnlineAchievementsPtr GetAchievementsInterface() const OVERRIDE;
-	virtual IOnlineSharingPtr GetSharingInterface() const OVERRIDE;
-	virtual IOnlineUserPtr GetUserInterface() const OVERRIDE;
-	virtual IOnlineMessagePtr GetMessageInterface() const OVERRIDE;
-	virtual IOnlinePresencePtr GetPresenceInterface() const OVERRIDE;
-	virtual bool Init() OVERRIDE;
-	virtual bool Shutdown() OVERRIDE;
-	virtual FString GetAppId() const OVERRIDE;
-	virtual bool Exec(class UWorld* InWorld, const TCHAR* Cmd, FOutputDevice& Ar) OVERRIDE;
+	virtual IOnlineSessionPtr GetSessionInterface() const override;
+	virtual IOnlineFriendsPtr GetFriendsInterface() const override;
+	virtual IOnlineSharedCloudPtr GetSharedCloudInterface() const override;
+	virtual IOnlineUserCloudPtr GetUserCloudInterface() const override;
+	virtual IOnlineLeaderboardsPtr GetLeaderboardsInterface() const override;
+	virtual IOnlineVoicePtr GetVoiceInterface() const override;
+	virtual IOnlineExternalUIPtr GetExternalUIInterface() const override;
+	virtual IOnlineTimePtr GetTimeInterface() const override;
+	virtual IOnlineIdentityPtr GetIdentityInterface() const override;
+	virtual IOnlinePartyPtr GetPartyInterface() const override;
+	virtual IOnlineTitleFilePtr GetTitleFileInterface() const override;
+	virtual IOnlineEntitlementsPtr GetEntitlementsInterface() const override;
+	virtual IOnlineStorePtr GetStoreInterface() const override;
+	virtual IOnlineEventsPtr GetEventsInterface() const override;
+	virtual IOnlineAchievementsPtr GetAchievementsInterface() const override;
+	virtual IOnlineSharingPtr GetSharingInterface() const override;
+	virtual IOnlineUserPtr GetUserInterface() const override;
+	virtual IOnlineMessagePtr GetMessageInterface() const override;
+	virtual IOnlinePresencePtr GetPresenceInterface() const override;
+	virtual bool Init() override;
+	virtual bool Shutdown() override;
+	virtual FString GetAppId() const override;
+	virtual bool Exec(class UWorld* InWorld, const TCHAR* Cmd, FOutputDevice& Ar) override;
 
 	// FTickerBaseObject
 
-	virtual bool Tick(float DeltaTime) OVERRIDE;
+	virtual bool Tick(float DeltaTime) override;
 
 	// FOnlineSubsystemFacebook
 
diff --git a/Engine/Source/Runtime/Online/OnlineSubsystemFacebook/Public/OnlineSubsystemFacebookModule.h b/Engine/Source/Runtime/Online/OnlineSubsystemFacebook/Public/OnlineSubsystemFacebookModule.h
--- a/Engine/Source/Runtime/Online/OnlineSubsystemFacebook/Public/OnlineSubsystemFacebookModule.h
+++ b/Engine/Source/Runtime/Online/OnlineSubsystemFacebook/Public/OnlineSubsystemFacebookModule.h
@@ -21,13 +21,13 @@ public:
 
 	// IModuleInterface 
 
-	virtual void StartupModule() OVERRIDE;
-	virtual void ShutdownModule() OVERRIDE;
-	virtual bool SupportsDynamicReloading() OVERRIDE
+	virtual void StartupModule() override;
+	virtual void ShutdownModule() override;
+	virtual bool SupportsDynamicReloading() override
 	{
 		return false;
 	}
-	virtual bool SupportsAutomaticShutdown() OVERRIDE
+	virtual bool SupportsAutomaticShutdown() override
 	{
 		return false;
 	}
diff --git a/Engine/Source/Runtime/Online/OnlineSubsystemNull/Private/OnlineAchievementsInterfaceNull.h b/Engine/Source/Runtime/Online/OnlineSubsystemNull/Private/OnlineAchievementsInterfaceNull.h
--- a/Engine/Source/Runtime/Online/OnlineSubsystemNull/Private/OnlineAchievementsInterfaceNull.h
+++ b/Engine/Source/Runtime/Online/OnlineSubsystemNull/Private/OnlineAchievementsInterfaceNull.h
@@ -112,14 +112,14 @@ public:
 	virtual ~FOnlineAchievementsNull(){}
 
 	// Begin IOnlineAchievements interface
-	virtual void WriteAchievements(const FUniqueNetId& PlayerId, FOnlineAchievementsWriteRef& WriteObject, const FOnAchievementsWrittenDelegate& Delegate = FOnAchievementsWrittenDelegate()) OVERRIDE;
-	virtual void QueryAchievements(const FUniqueNetId & PlayerId, const FOnQueryAchievementsCompleteDelegate & Delegate = FOnQueryAchievementsCompleteDelegate()) OVERRIDE;
-	virtual void QueryAchievementDescriptions( const FUniqueNetId& PlayerId, const FOnQueryAchievementsCompleteDelegate & Delegate = FOnQueryAchievementsCompleteDelegate() ) OVERRIDE;
-	virtual EOnlineCachedResult::Type GetCachedAchievement(const FUniqueNetId& PlayerId, const FString& AchievementId, FOnlineAchievement& OutAchievement) OVERRIDE;
-	virtual EOnlineCachedResult::Type GetCachedAchievements(const FUniqueNetId& PlayerId, TArray<FOnlineAchievement> & OutAchievements) OVERRIDE;
-	virtual EOnlineCachedResult::Type GetCachedAchievementDescription(const FString& AchievementId, FOnlineAchievementDesc& OutAchievementDesc) OVERRIDE;
+	virtual void WriteAchievements(const FUniqueNetId& PlayerId, FOnlineAchievementsWriteRef& WriteObject, const FOnAchievementsWrittenDelegate& Delegate = FOnAchievementsWrittenDelegate()) override;
+	virtual void QueryAchievements(const FUniqueNetId & PlayerId, const FOnQueryAchievementsCompleteDelegate & Delegate = FOnQueryAchievementsCompleteDelegate()) override;
+	virtual void QueryAchievementDescriptions( const FUniqueNetId& PlayerId, const FOnQueryAchievementsCompleteDelegate & Delegate = FOnQueryAchievementsCompleteDelegate() ) override;
+	virtual EOnlineCachedResult::Type GetCachedAchievement(const FUniqueNetId& PlayerId, const FString& AchievementId, FOnlineAchievement& OutAchievement) override;
+	virtual EOnlineCachedResult::Type GetCachedAchievements(const FUniqueNetId& PlayerId, TArray<FOnlineAchievement> & OutAchievements) override;
+	virtual EOnlineCachedResult::Type GetCachedAchievementDescription(const FString& AchievementId, FOnlineAchievementDesc& OutAchievementDesc) override;
 #if !UE_BUILD_SHIPPING
-	virtual bool ResetAchievements( const FUniqueNetId& PlayerId ) OVERRIDE;
+	virtual bool ResetAchievements( const FUniqueNetId& PlayerId ) override;
 #endif // !UE_BUILD_SHIPPING
 	// End IOnlineAchievements interface
 };
diff --git a/Engine/Source/Runtime/Online/OnlineSubsystemNull/Private/OnlineAsyncTaskManagerNull.h b/Engine/Source/Runtime/Online/OnlineSubsystemNull/Private/OnlineAsyncTaskManagerNull.h
--- a/Engine/Source/Runtime/Online/OnlineSubsystemNull/Private/OnlineAsyncTaskManagerNull.h
+++ b/Engine/Source/Runtime/Online/OnlineSubsystemNull/Private/OnlineAsyncTaskManagerNull.h
@@ -26,5 +26,5 @@ public:
 	}
 
 	// FOnlineAsyncTaskManager
-	virtual void OnlineTick() OVERRIDE;
+	virtual void OnlineTick() override;
 };
diff --git a/Engine/Source/Runtime/Online/OnlineSubsystemNull/Private/OnlineIdentityNull.h b/Engine/Source/Runtime/Online/OnlineSubsystemNull/Private/OnlineIdentityNull.h
--- a/Engine/Source/Runtime/Online/OnlineSubsystemNull/Private/OnlineIdentityNull.h
+++ b/Engine/Source/Runtime/Online/OnlineSubsystemNull/Private/OnlineIdentityNull.h
@@ -15,15 +15,15 @@ public:
 
 	// FOnlineUser
 	
-	virtual TSharedRef<FUniqueNetId> GetUserId() const OVERRIDE { return UserIdPtr; }
-	virtual FString GetRealName() const OVERRIDE { return TEXT("DummyRealName"); }
-	virtual FString GetDisplayName() const OVERRIDE  { return TEXT("DummyDisplayName"); }
-	virtual bool GetUserAttribute(const FString& AttrName, FString& OutAttrValue) const OVERRIDE;
+	virtual TSharedRef<FUniqueNetId> GetUserId() const override { return UserIdPtr; }
+	virtual FString GetRealName() const override { return TEXT("DummyRealName"); }
+	virtual FString GetDisplayName() const override  { return TEXT("DummyDisplayName"); }
+	virtual bool GetUserAttribute(const FString& AttrName, FString& OutAttrValue) const override;
 
 	// FUserOnlineAccount
 
-	virtual FString GetAccessToken() const OVERRIDE { return TEXT("DummyAuthTicket"); }
-	virtual bool GetAuthAttribute(const FString& AttrName, FString& OutAttrValue) const OVERRIDE;
+	virtual FString GetAccessToken() const override { return TEXT("DummyAuthTicket"); }
+	virtual bool GetAuthAttribute(const FString& AttrName, FString& OutAttrValue) const override;
 
 	// FUserOnlineAccountNull
 
@@ -53,17 +53,17 @@ public:
 
 	// IOnlineIdentity
 
-	virtual bool Login(int32 LocalUserNum, const FOnlineAccountCredentials& AccountCredentials) OVERRIDE;
-	virtual bool Logout(int32 LocalUserNum) OVERRIDE;
-	virtual bool AutoLogin(int32 LocalUserNum) OVERRIDE;
-	virtual TSharedPtr<FUserOnlineAccount> GetUserAccount(const FUniqueNetId& UserId) const OVERRIDE;
-	virtual TArray<TSharedPtr<FUserOnlineAccount> > GetAllUserAccounts() const OVERRIDE;
-	virtual TSharedPtr<FUniqueNetId> GetUniquePlayerId(int32 LocalUserNum) const OVERRIDE;
-	virtual TSharedPtr<FUniqueNetId> CreateUniquePlayerId(uint8* Bytes, int32 Size) OVERRIDE;
-	virtual TSharedPtr<FUniqueNetId> CreateUniquePlayerId(const FString& Str) OVERRIDE;
-	virtual ELoginStatus::Type GetLoginStatus(int32 LocalUserNum) const OVERRIDE;
-	virtual FString GetPlayerNickname(int32 LocalUserNum) const OVERRIDE;
-	virtual FString GetAuthToken(int32 LocalUserNum) const OVERRIDE;
+	virtual bool Login(int32 LocalUserNum, const FOnlineAccountCredentials& AccountCredentials) override;
+	virtual bool Logout(int32 LocalUserNum) override;
+	virtual bool AutoLogin(int32 LocalUserNum) override;
+	virtual TSharedPtr<FUserOnlineAccount> GetUserAccount(const FUniqueNetId& UserId) const override;
+	virtual TArray<TSharedPtr<FUserOnlineAccount> > GetAllUserAccounts() const override;
+	virtual TSharedPtr<FUniqueNetId> GetUniquePlayerId(int32 LocalUserNum) const override;
+	virtual TSharedPtr<FUniqueNetId> CreateUniquePlayerId(uint8* Bytes, int32 Size) override;
+	virtual TSharedPtr<FUniqueNetId> CreateUniquePlayerId(const FString& Str) override;
+	virtual ELoginStatus::Type GetLoginStatus(int32 LocalUserNum) const override;
+	virtual FString GetPlayerNickname(int32 LocalUserNum) const override;
+	virtual FString GetAuthToken(int32 LocalUserNum) const override;
 
 	// FOnlineIdentityNull
 
diff --git a/Engine/Source/Runtime/Online/OnlineSubsystemNull/Private/OnlineSessionInterfaceNull.cpp b/Engine/Source/Runtime/Online/OnlineSubsystemNull/Private/OnlineSessionInterfaceNull.cpp
--- a/Engine/Source/Runtime/Online/OnlineSubsystemNull/Private/OnlineSessionInterfaceNull.cpp
+++ b/Engine/Source/Runtime/Online/OnlineSubsystemNull/Private/OnlineSessionInterfaceNull.cpp
@@ -54,7 +54,7 @@ public:
 	/**
 	 *	Get a human readable description of task
 	 */
-	virtual FString ToString() const OVERRIDE
+	virtual FString ToString() const override
 	{
 		return FString::Printf(TEXT("FOnlineAsyncTaskNullEndSession bWasSuccessful: %d SessionName: %s"), bWasSuccessful, *SessionName.ToString());
 	}
@@ -63,7 +63,7 @@ public:
 	 * Give the async task time to do its work
 	 * Can only be called on the async task manager thread
 	 */
-	virtual void Tick() OVERRIDE
+	virtual void Tick() override
 	{
 		bIsComplete = true;
 		bWasSuccessful = true;
@@ -73,7 +73,7 @@ public:
 	 * Give the async task a chance to marshal its data back to the game thread
 	 * Can only be called on the game thread by the async task manager
 	 */
-	virtual void Finalize() OVERRIDE
+	virtual void Finalize() override
 	{
 		IOnlineSessionPtr SessionInt = Subsystem->GetSessionInterface();
 		FNamedOnlineSession* Session = SessionInt->GetNamedSession(SessionName);
@@ -86,7 +86,7 @@ public:
 	/**
 	 *	Async task is given a chance to trigger it's delegates
 	 */
-	virtual void TriggerDelegates() OVERRIDE
+	virtual void TriggerDelegates() override
 	{
 		IOnlineSessionPtr SessionInt = Subsystem->GetSessionInterface();
 		if (SessionInt.IsValid())
@@ -119,7 +119,7 @@ public:
 	/**
 	 *	Get a human readable description of task
 	 */
-	virtual FString ToString() const OVERRIDE
+	virtual FString ToString() const override
 	{
 		return FString::Printf(TEXT("FOnlineAsyncTaskNullDestroySession bWasSuccessful: %d SessionName: %s"), bWasSuccessful, *SessionName.ToString());
 	}
@@ -128,7 +128,7 @@ public:
 	 * Give the async task time to do its work
 	 * Can only be called on the async task manager thread
 	 */
-	virtual void Tick() OVERRIDE
+	virtual void Tick() override
 	{
 		bIsComplete = true;
 		bWasSuccessful = true;
@@ -138,7 +138,7 @@ public:
 	 * Give the async task a chance to marshal its data back to the game thread
 	 * Can only be called on the game thread by the async task manager
 	 */
-	virtual void Finalize() OVERRIDE
+	virtual void Finalize() override
 	{
 		IOnlineSessionPtr SessionInt = Subsystem->GetSessionInterface();
 		if (SessionInt.IsValid())
@@ -154,7 +154,7 @@ public:
 	/**
 	 *	Async task is given a chance to trigger it's delegates
 	 */
-	virtual void TriggerDelegates() OVERRIDE
+	virtual void TriggerDelegates() override
 	{
 		IOnlineSessionPtr SessionInt = Subsystem->GetSessionInterface();
 		if (SessionInt.IsValid())
diff --git a/Engine/Source/Runtime/Online/OnlineSubsystemNull/Private/OnlineSessionInterfaceNull.h b/Engine/Source/Runtime/Online/OnlineSubsystemNull/Private/OnlineSessionInterfaceNull.h
--- a/Engine/Source/Runtime/Online/OnlineSubsystemNull/Private/OnlineSessionInterfaceNull.h
+++ b/Engine/Source/Runtime/Online/OnlineSubsystemNull/Private/OnlineSessionInterfaceNull.h
@@ -166,13 +166,13 @@ PACKAGE_SCOPE:
 	void Tick(float DeltaTime);
 
 	// IOnlineSession
-	class FNamedOnlineSession* AddNamedSession(FName SessionName, const FOnlineSessionSettings& SessionSettings) OVERRIDE
+	class FNamedOnlineSession* AddNamedSession(FName SessionName, const FOnlineSessionSettings& SessionSettings) override
 	{
 		FScopeLock ScopeLock(&SessionLock);
 		return new (Sessions) FNamedOnlineSession(SessionName, SessionSettings);
 	}
 
-	class FNamedOnlineSession* AddNamedSession(FName SessionName, const FOnlineSession& Session) OVERRIDE
+	class FNamedOnlineSession* AddNamedSession(FName SessionName, const FOnlineSession& Session) override
 	{
 		FScopeLock ScopeLock(&SessionLock);
 		return new (Sessions) FNamedOnlineSession(SessionName, Session);
@@ -208,7 +208,7 @@ public:
 
 	virtual ~FOnlineSessionNull() {}
 
-	FNamedOnlineSession* GetNamedSession(FName SessionName) OVERRIDE
+	FNamedOnlineSession* GetNamedSession(FName SessionName) override
 	{
 		FScopeLock ScopeLock(&SessionLock);
 		for (int32 SearchIndex = 0; SearchIndex < Sessions.Num(); SearchIndex++)
@@ -221,7 +221,7 @@ public:
 		return NULL;
 	}
 
-	virtual void RemoveNamedSession(FName SessionName) OVERRIDE
+	virtual void RemoveNamedSession(FName SessionName) override
 	{
 		FScopeLock ScopeLock(&SessionLock);
 		for (int32 SearchIndex = 0; SearchIndex < Sessions.Num(); SearchIndex++)
@@ -234,7 +234,7 @@ public:
 		}
 	}
 
-	virtual EOnlineSessionState::Type GetSessionState(FName SessionName) const OVERRIDE
+	virtual EOnlineSessionState::Type GetSessionState(FName SessionName) const override
 	{
 		FScopeLock ScopeLock(&SessionLock);
 		for (int32 SearchIndex = 0; SearchIndex < Sessions.Num(); SearchIndex++)
@@ -248,7 +248,7 @@ public:
 		return EOnlineSessionState::NoSession;
 	}
 
-	virtual bool HasPresenceSession() OVERRIDE
+	virtual bool HasPresenceSession() override
 	{
 		FScopeLock ScopeLock(&SessionLock);
 		for (int32 SearchIndex = 0; SearchIndex < Sessions.Num(); SearchIndex++)
@@ -263,30 +263,30 @@ public:
 	}
 
 	// IOnlineSession
-	virtual bool CreateSession(int32 HostingPlayerNum, FName SessionName, const FOnlineSessionSettings& NewSessionSettings) OVERRIDE;
-	virtual bool StartSession(FName SessionName) OVERRIDE;
-	virtual bool UpdateSession(FName SessionName, FOnlineSessionSettings& UpdatedSessionSettings, bool bShouldRefreshOnlineData = true) OVERRIDE;
-	virtual bool EndSession(FName SessionName) OVERRIDE;
-	virtual bool DestroySession(FName SessionName) OVERRIDE;
-	virtual bool IsPlayerInSession(FName SessionName, const FUniqueNetId& UniqueId) OVERRIDE;
-	virtual bool StartMatchmaking(int32 SearchingPlayerNum, FName SessionName, const FOnlineSessionSettings& NewSessionSettings, TSharedRef<FOnlineSessionSearch>& SearchSettings) OVERRIDE;
-	virtual bool CancelMatchmaking(int32 SearchingPlayerNum, FName SessionName) OVERRIDE;
-	virtual bool FindSessions(int32 SearchingPlayerNum, const TSharedRef<FOnlineSessionSearch>& SearchSettings) OVERRIDE;
-	virtual bool CancelFindSessions() OVERRIDE;
-	virtual bool PingSearchResults(const FOnlineSessionSearchResult& SearchResult) OVERRIDE;
-	virtual bool JoinSession(int32 PlayerNum, FName SessionName, const FOnlineSessionSearchResult& DesiredSession) OVERRIDE;
-	virtual bool FindFriendSession(int32 LocalUserNum, const FUniqueNetId& Friend) OVERRIDE;
-	virtual bool SendSessionInviteToFriend(int32 LocalUserNum, FName SessionName, const FUniqueNetId& Friend) OVERRIDE;
-	virtual bool SendSessionInviteToFriends(int32 LocalUserNum, FName SessionName, const TArray< TSharedRef<FUniqueNetId> >& Friends) OVERRIDE;
-	virtual bool GetResolvedConnectString(FName SessionName, FString& ConnectInfo) OVERRIDE;
-	virtual bool GetResolvedConnectString(const class FOnlineSessionSearchResult& SearchResult, FName PortType, FString& ConnectInfo) OVERRIDE;
-	virtual FOnlineSessionSettings* GetSessionSettings(FName SessionName) OVERRIDE;
-	virtual bool RegisterPlayer(FName SessionName, const FUniqueNetId& PlayerId, bool bWasInvited) OVERRIDE;
-	virtual bool RegisterPlayers(FName SessionName, const TArray< TSharedRef<FUniqueNetId> >& Players, bool bWasInvited = false) OVERRIDE;
-	virtual bool UnregisterPlayer(FName SessionName, const FUniqueNetId& PlayerId) OVERRIDE;
-	virtual bool UnregisterPlayers(FName SessionName, const TArray< TSharedRef<FUniqueNetId> >& Players) OVERRIDE;
-	virtual int32 GetNumSessions() OVERRIDE;
-	virtual void DumpSessionState() OVERRIDE;
+	virtual bool CreateSession(int32 HostingPlayerNum, FName SessionName, const FOnlineSessionSettings& NewSessionSettings) override;
+	virtual bool StartSession(FName SessionName) override;
+	virtual bool UpdateSession(FName SessionName, FOnlineSessionSettings& UpdatedSessionSettings, bool bShouldRefreshOnlineData = true) override;
+	virtual bool EndSession(FName SessionName) override;
+	virtual bool DestroySession(FName SessionName) override;
+	virtual bool IsPlayerInSession(FName SessionName, const FUniqueNetId& UniqueId) override;
+	virtual bool StartMatchmaking(int32 SearchingPlayerNum, FName SessionName, const FOnlineSessionSettings& NewSessionSettings, TSharedRef<FOnlineSessionSearch>& SearchSettings) override;
+	virtual bool CancelMatchmaking(int32 SearchingPlayerNum, FName SessionName) override;
+	virtual bool FindSessions(int32 SearchingPlayerNum, const TSharedRef<FOnlineSessionSearch>& SearchSettings) override;
+	virtual bool CancelFindSessions() override;
+	virtual bool PingSearchResults(const FOnlineSessionSearchResult& SearchResult) override;
+	virtual bool JoinSession(int32 PlayerNum, FName SessionName, const FOnlineSessionSearchResult& DesiredSession) override;
+	virtual bool FindFriendSession(int32 LocalUserNum, const FUniqueNetId& Friend) override;
+	virtual bool SendSessionInviteToFriend(int32 LocalUserNum, FName SessionName, const FUniqueNetId& Friend) override;
+	virtual bool SendSessionInviteToFriends(int32 LocalUserNum, FName SessionName, const TArray< TSharedRef<FUniqueNetId> >& Friends) override;
+	virtual bool GetResolvedConnectString(FName SessionName, FString& ConnectInfo) override;
+	virtual bool GetResolvedConnectString(const class FOnlineSessionSearchResult& SearchResult, FName PortType, FString& ConnectInfo) override;
+	virtual FOnlineSessionSettings* GetSessionSettings(FName SessionName) override;
+	virtual bool RegisterPlayer(FName SessionName, const FUniqueNetId& PlayerId, bool bWasInvited) override;
+	virtual bool RegisterPlayers(FName SessionName, const TArray< TSharedRef<FUniqueNetId> >& Players, bool bWasInvited = false) override;
+	virtual bool UnregisterPlayer(FName SessionName, const FUniqueNetId& PlayerId) override;
+	virtual bool UnregisterPlayers(FName SessionName, const TArray< TSharedRef<FUniqueNetId> >& Players) override;
+	virtual int32 GetNumSessions() override;
+	virtual void DumpSessionState() override;
 };
 
 typedef TSharedPtr<FOnlineSessionNull, ESPMode::ThreadSafe> FOnlineSessionNullPtr;
diff --git a/Engine/Source/Runtime/Online/OnlineSubsystemNull/Private/OnlineSubsystemNullTypes.h b/Engine/Source/Runtime/Online/OnlineSubsystemNull/Private/OnlineSubsystemNullTypes.h
--- a/Engine/Source/Runtime/Online/OnlineSubsystemNull/Private/OnlineSubsystemNullTypes.h
+++ b/Engine/Source/Runtime/Online/OnlineSubsystemNull/Private/OnlineSubsystemNullTypes.h
@@ -49,35 +49,35 @@ public:
  		return false;
  	}
 
-	virtual const uint8* GetBytes() const OVERRIDE
+	virtual const uint8* GetBytes() const override
 	{
 		return NULL;
 	}
 
-	virtual int32 GetSize() const OVERRIDE
+	virtual int32 GetSize() const override
 	{
 		return sizeof(uint64) + sizeof(TSharedPtr<class FInternetAddr>);
 	}
 
-	virtual bool IsValid() const OVERRIDE
+	virtual bool IsValid() const override
 	{
 		// LAN case
 		return HostAddr.IsValid() && HostAddr->IsValid();
 	}
 
-	virtual FString ToString() const OVERRIDE
+	virtual FString ToString() const override
 	{
 		return SessionId.ToString();
 	}
 
-	virtual FString ToDebugString() const OVERRIDE
+	virtual FString ToDebugString() const override
 	{
 		return FString::Printf(TEXT("HostIP: %s SessionId: %s"), 
 			HostAddr.IsValid() ? *HostAddr->ToString(true) : TEXT("INVALID"), 
 			*SessionId.ToDebugString());
 	}
 
-	virtual const FUniqueNetId& GetSessionId() const OVERRIDE
+	virtual const FUniqueNetId& GetSessionId() const override
 	{
 		return SessionId;
 	}
diff --git a/Engine/Source/Runtime/Online/OnlineSubsystemNull/Public/OnlineSubsystemNull.h b/Engine/Source/Runtime/Online/OnlineSubsystemNull/Public/OnlineSubsystemNull.h
--- a/Engine/Source/Runtime/Online/OnlineSubsystemNull/Public/OnlineSubsystemNull.h
+++ b/Engine/Source/Runtime/Online/OnlineSubsystemNull/Public/OnlineSubsystemNull.h
@@ -33,34 +33,34 @@ public:
 
 	// IOnlineSubsystem
 
-	virtual IOnlineSessionPtr GetSessionInterface() const OVERRIDE;
-	virtual IOnlineFriendsPtr GetFriendsInterface() const OVERRIDE;
-	virtual IOnlineSharedCloudPtr GetSharedCloudInterface() const OVERRIDE;
-	virtual IOnlineUserCloudPtr GetUserCloudInterface() const OVERRIDE;
-	virtual IOnlineEntitlementsPtr GetEntitlementsInterface() const OVERRIDE;
-	virtual IOnlineLeaderboardsPtr GetLeaderboardsInterface() const OVERRIDE;
-	virtual IOnlineVoicePtr GetVoiceInterface() const OVERRIDE;
-	virtual IOnlineExternalUIPtr GetExternalUIInterface() const OVERRIDE;	
-	virtual IOnlineTimePtr GetTimeInterface() const OVERRIDE;
-	virtual IOnlineIdentityPtr GetIdentityInterface() const OVERRIDE;
-	virtual IOnlinePartyPtr GetPartyInterface() const OVERRIDE;
-	virtual IOnlineTitleFilePtr GetTitleFileInterface() const OVERRIDE;
-	virtual IOnlineStorePtr GetStoreInterface() const OVERRIDE;
-	virtual IOnlineEventsPtr GetEventsInterface() const OVERRIDE;
-	virtual IOnlineAchievementsPtr GetAchievementsInterface() const OVERRIDE;
-	virtual IOnlineSharingPtr GetSharingInterface() const OVERRIDE;
-	virtual IOnlineUserPtr GetUserInterface() const OVERRIDE;
-	virtual IOnlineMessagePtr GetMessageInterface() const OVERRIDE;
-	virtual IOnlinePresencePtr GetPresenceInterface() const OVERRIDE;
-
-	virtual bool Init() OVERRIDE;
-	virtual bool Shutdown() OVERRIDE;
-	virtual FString GetAppId() const OVERRIDE;
-	virtual bool Exec(class UWorld* InWorld, const TCHAR* Cmd, FOutputDevice& Ar) OVERRIDE;
+	virtual IOnlineSessionPtr GetSessionInterface() const override;
+	virtual IOnlineFriendsPtr GetFriendsInterface() const override;
+	virtual IOnlineSharedCloudPtr GetSharedCloudInterface() const override;
+	virtual IOnlineUserCloudPtr GetUserCloudInterface() const override;
+	virtual IOnlineEntitlementsPtr GetEntitlementsInterface() const override;
+	virtual IOnlineLeaderboardsPtr GetLeaderboardsInterface() const override;
+	virtual IOnlineVoicePtr GetVoiceInterface() const override;
+	virtual IOnlineExternalUIPtr GetExternalUIInterface() const override;	
+	virtual IOnlineTimePtr GetTimeInterface() const override;
+	virtual IOnlineIdentityPtr GetIdentityInterface() const override;
+	virtual IOnlinePartyPtr GetPartyInterface() const override;
+	virtual IOnlineTitleFilePtr GetTitleFileInterface() const override;
+	virtual IOnlineStorePtr GetStoreInterface() const override;
+	virtual IOnlineEventsPtr GetEventsInterface() const override;
+	virtual IOnlineAchievementsPtr GetAchievementsInterface() const override;
+	virtual IOnlineSharingPtr GetSharingInterface() const override;
+	virtual IOnlineUserPtr GetUserInterface() const override;
+	virtual IOnlineMessagePtr GetMessageInterface() const override;
+	virtual IOnlinePresencePtr GetPresenceInterface() const override;
+
+	virtual bool Init() override;
+	virtual bool Shutdown() override;
+	virtual FString GetAppId() const override;
+	virtual bool Exec(class UWorld* InWorld, const TCHAR* Cmd, FOutputDevice& Ar) override;
 
 	// FTickerObjectBase
 	
-	virtual bool Tick(float DeltaTime) OVERRIDE;
+	virtual bool Tick(float DeltaTime) override;
 
 	// FOnlineSubsystemNull
 
diff --git a/Engine/Source/Runtime/Online/OnlineSubsystemSteam/Classes/SteamMessagingNetDriver.h b/Engine/Source/Runtime/Online/OnlineSubsystemSteam/Classes/SteamMessagingNetDriver.h
--- a/Engine/Source/Runtime/Online/OnlineSubsystemSteam/Classes/SteamMessagingNetDriver.h
+++ b/Engine/Source/Runtime/Online/OnlineSubsystemSteam/Classes/SteamMessagingNetDriver.h
@@ -25,7 +25,7 @@ class USteamMessagingNetDriver : public UIpNetDriver
 	 *
 	 * @return true if successful, false otherwise (check Error parameter)
 	 */
-	virtual bool InitBase(bool bInitAsClient, FNetworkNotify* InNotify, const FURL& URL, bool bReuseAddressAndPort, FString& Error) OVERRIDE;
+	virtual bool InitBase(bool bInitAsClient, FNetworkNotify* InNotify, const FURL& URL, bool bReuseAddressAndPort, FString& Error) override;
 	/**
 	 * Initialize the net driver in client mode
 	 *
@@ -35,7 +35,7 @@ class USteamMessagingNetDriver : public UIpNetDriver
 	 * 
 	 * @return true if successful, false otherwise (check Error parameter)
 	 */
-	virtual bool InitConnect(FNetworkNotify* InNotify, const FURL& ConnectURL, FString& Error) OVERRIDE;	
+	virtual bool InitConnect(FNetworkNotify* InNotify, const FURL& ConnectURL, FString& Error) override;	
 	/**
 	 * Initialize the network driver in server mode (listener)
 	 *
@@ -46,10 +46,10 @@ class USteamMessagingNetDriver : public UIpNetDriver
 	 *
 	 * @return true if successful, false otherwise (check Error parameter)
 	 */
-	virtual bool InitListen(FNetworkNotify* InNotify, FURL& ListenURL, bool bReuseAddressAndPort, FString& Error) OVERRIDE;
+	virtual bool InitListen(FNetworkNotify* InNotify, FURL& ListenURL, bool bReuseAddressAndPort, FString& Error) override;
 
 	/** @return true if the net resource is valid or false if it should not be used */
-	virtual bool IsNetResourceValid() OVERRIDE;
+	virtual bool IsNetResourceValid() override;
 
 	/**
 	 * Process a remote function call on some actor destined for a remote location
@@ -59,7 +59,7 @@ class USteamMessagingNetDriver : public UIpNetDriver
 	 * @param Params parameters in a UObject memory layout
 	 * @param Stack stack frame the UFunction is called in
 	 */
-	virtual void ProcessRemoteFunction(class AActor* Actor, class UFunction* Function, void* Parameters, struct FOutParmRec* OutParms, struct FFrame* Stack, UObject * SubObject = NULL) OVERRIDE;
+	virtual void ProcessRemoteFunction(class AActor* Actor, class UFunction* Function, void* Parameters, struct FOutParmRec* OutParms, struct FFrame* Stack, UObject * SubObject = NULL) override;
 
 	// End UIpNetDriver Interface
 
diff --git a/Engine/Source/Runtime/Online/OnlineSubsystemSteam/Classes/SteamNetConnection.h b/Engine/Source/Runtime/Online/OnlineSubsystemSteam/Classes/SteamNetConnection.h
--- a/Engine/Source/Runtime/Online/OnlineSubsystemSteam/Classes/SteamNetConnection.h
+++ b/Engine/Source/Runtime/Online/OnlineSubsystemSteam/Classes/SteamNetConnection.h
@@ -13,9 +13,9 @@ class USteamNetConnection : public UIpConnection
     GENERATED_UCLASS_BODY()
 
 	// Begin UIpConnection Interface
-	virtual void InitRemoteConnection(class UNetDriver* InDriver, class FSocket* InSocket, const FURL& InURL, const class FInternetAddr& InRemoteAddr, EConnectionState InState, int32 InMaxPacket = 0, int32 InPacketOverhead = 0) OVERRIDE;
-	virtual void InitLocalConnection(class UNetDriver* InDriver, class FSocket* InSocket, const FURL& InURL, EConnectionState InState, int32 InMaxPacket = 0, int32 InPacketOverhead = 0) OVERRIDE;
-	virtual void CleanUp() OVERRIDE;
+	virtual void InitRemoteConnection(class UNetDriver* InDriver, class FSocket* InSocket, const FURL& InURL, const class FInternetAddr& InRemoteAddr, EConnectionState InState, int32 InMaxPacket = 0, int32 InPacketOverhead = 0) override;
+	virtual void InitLocalConnection(class UNetDriver* InDriver, class FSocket* InSocket, const FURL& InURL, EConnectionState InState, int32 InMaxPacket = 0, int32 InPacketOverhead = 0) override;
+	virtual void CleanUp() override;
 	
 	// End UIpConnection Interface
 
diff --git a/Engine/Source/Runtime/Online/OnlineSubsystemSteam/Classes/SteamNetDriver.h b/Engine/Source/Runtime/Online/OnlineSubsystemSteam/Classes/SteamNetDriver.h
--- a/Engine/Source/Runtime/Online/OnlineSubsystemSteam/Classes/SteamNetDriver.h
+++ b/Engine/Source/Runtime/Online/OnlineSubsystemSteam/Classes/SteamNetDriver.h
@@ -20,19 +20,19 @@ class USteamNetDriver : public UIpNetDriver
 	bool bIsPassthrough;
 
 	// Begin UObject Interface
-	virtual void PostInitProperties() OVERRIDE;
+	virtual void PostInitProperties() override;
 	// End UObject Interface
 
 	// Begin UIpNetDriver Interface
 
-	virtual class ISocketSubsystem* GetSocketSubsystem() OVERRIDE;
-	virtual bool IsAvailable() const OVERRIDE;
-	virtual bool InitBase(bool bInitAsClient, FNetworkNotify* InNotify, const FURL& URL, bool bReuseAddressAndPort, FString& Error) OVERRIDE;
-	virtual bool InitConnect(FNetworkNotify* InNotify, const FURL& ConnectURL, FString& Error) OVERRIDE;	
-	virtual bool InitListen(FNetworkNotify* InNotify, FURL& ListenURL, bool bReuseAddressAndPort, FString& Error) OVERRIDE;
-	virtual void Shutdown() OVERRIDE;
-	virtual void TickFlush(float DeltaSeconds) OVERRIDE;
-	virtual bool IsNetResourceValid() OVERRIDE;
+	virtual class ISocketSubsystem* GetSocketSubsystem() override;
+	virtual bool IsAvailable() const override;
+	virtual bool InitBase(bool bInitAsClient, FNetworkNotify* InNotify, const FURL& URL, bool bReuseAddressAndPort, FString& Error) override;
+	virtual bool InitConnect(FNetworkNotify* InNotify, const FURL& ConnectURL, FString& Error) override;	
+	virtual bool InitListen(FNetworkNotify* InNotify, FURL& ListenURL, bool bReuseAddressAndPort, FString& Error) override;
+	virtual void Shutdown() override;
+	virtual void TickFlush(float DeltaSeconds) override;
+	virtual bool IsNetResourceValid() override;
 
 	// End UIpNetDriver Interface
 
diff --git a/Engine/Source/Runtime/Online/OnlineSubsystemSteam/Private/IPAddressSteam.h b/Engine/Source/Runtime/Online/OnlineSubsystemSteam/Private/IPAddressSteam.h
--- a/Engine/Source/Runtime/Online/OnlineSubsystemSteam/Private/IPAddressSteam.h
+++ b/Engine/Source/Runtime/Online/OnlineSubsystemSteam/Private/IPAddressSteam.h
@@ -50,7 +50,7 @@ public:
 	 *
 	 * @param InAddr the new address to use (must convert to network byte order)
 	 */
-	void SetIp(uint32 InAddr) OVERRIDE
+	void SetIp(uint32 InAddr) override
 	{
 		check("Invalid interface for Steam P2P addresses" && 0);
 	}
@@ -60,14 +60,14 @@ public:
 	 *
 	 * @param InAddr the string containing the new ip address to use
 	 */
-	void SetIp(const TCHAR* InAddr, bool& bIsValid) OVERRIDE;
+	void SetIp(const TCHAR* InAddr, bool& bIsValid) override;
 
 	/**
 	 * Copies the network byte order ip address to a host byte order dword
 	 *
 	 * @param OutAddr the out param receiving the ip address
 	 */
-	void GetIp(uint32& OutAddr) const OVERRIDE
+	void GetIp(uint32& OutAddr) const override
 	{ 
 		check("Invalid interface for Steam P2P addresses" && 0);
 	}
@@ -77,7 +77,7 @@ public:
 	 *
 	 * @param InPort the new port to use (must convert to network byte order)
 	 */
-	void SetPort(int32 InPort) OVERRIDE
+	void SetPort(int32 InPort) override
 	{
 		SteamChannel = InPort;
 	}
@@ -87,7 +87,7 @@ public:
 	 *
 	 * @param OutPort the host byte order int that receives the port
 	 */
-	void GetPort(int32& OutPort) const OVERRIDE
+	void GetPort(int32& OutPort) const override
 	{
 		OutPort = SteamChannel;
 	}
@@ -95,18 +95,18 @@ public:
 	/**
 	 * Returns the port number from this address in host byte order
 	 */
-	int32 GetPort() const OVERRIDE
+	int32 GetPort() const override
 	{
 		return SteamChannel;
 	}
 
 	/** Sets the address to be any address */
-	void SetAnyAddress() OVERRIDE
+	void SetAnyAddress() override
 	{
 	}
 
 	/** Sets the address to broadcast */
-	void SetBroadcastAddress() OVERRIDE
+	void SetBroadcastAddress() override
 	{
 	}
 
@@ -115,7 +115,7 @@ public:
 	 *
 	 * @param bAppendPort whether to append the port information or not
 	 */
-	FString ToString(bool bAppendPort) const OVERRIDE
+	FString ToString(bool bAppendPort) const override
 	{
 		if (bAppendPort)
 		{
@@ -132,7 +132,7 @@ public:
 	 *
 	 * @param Other the address to compare against
 	 */
-	virtual bool operator==(const FInternetAddr& Other) const OVERRIDE
+	virtual bool operator==(const FInternetAddr& Other) const override
 	{
 		FInternetAddrSteam& SteamOther = (FInternetAddrSteam&)Other;
 		return SteamId == SteamOther.SteamId && SteamChannel == SteamOther.SteamChannel;
@@ -148,7 +148,7 @@ public:
 	 *
 	 * @return true if a valid IP, false otherwise
 	 */
-	virtual bool IsValid() const OVERRIDE
+	virtual bool IsValid() const override
 	{
 		return SteamId.IsValid();
 	}
diff --git a/Engine/Source/Runtime/Online/OnlineSubsystemSteam/Private/OnlineAchievementsInterfaceSteam.h b/Engine/Source/Runtime/Online/OnlineSubsystemSteam/Private/OnlineAchievementsInterfaceSteam.h
--- a/Engine/Source/Runtime/Online/OnlineSubsystemSteam/Private/OnlineAchievementsInterfaceSteam.h
+++ b/Engine/Source/Runtime/Online/OnlineSubsystemSteam/Private/OnlineAchievementsInterfaceSteam.h
@@ -136,14 +136,14 @@ PACKAGE_SCOPE:
 public:
 
 	// Begin IOnlineAchievements interface
-	virtual void WriteAchievements(const FUniqueNetId& PlayerId, FOnlineAchievementsWriteRef& WriteObject, const FOnAchievementsWrittenDelegate& Delegate = FOnAchievementsWrittenDelegate()) OVERRIDE;
-	virtual void QueryAchievements(const FUniqueNetId& PlayerId, const FOnQueryAchievementsCompleteDelegate& Delegate = FOnQueryAchievementsCompleteDelegate()) OVERRIDE;
-	virtual void QueryAchievementDescriptions( const FUniqueNetId& PlayerId, const FOnQueryAchievementsCompleteDelegate & Delegate = FOnQueryAchievementsCompleteDelegate() ) OVERRIDE;
-	virtual EOnlineCachedResult::Type GetCachedAchievement(const FUniqueNetId& PlayerId, const FString& AchievementId, FOnlineAchievement& OutAchievement) OVERRIDE;
-	virtual EOnlineCachedResult::Type GetCachedAchievements(const FUniqueNetId& PlayerId, TArray<FOnlineAchievement> & OutAchievements) OVERRIDE;
-	virtual EOnlineCachedResult::Type GetCachedAchievementDescription(const FString& AchievementId, FOnlineAchievementDesc& OutAchievementDesc) OVERRIDE;
+	virtual void WriteAchievements(const FUniqueNetId& PlayerId, FOnlineAchievementsWriteRef& WriteObject, const FOnAchievementsWrittenDelegate& Delegate = FOnAchievementsWrittenDelegate()) override;
+	virtual void QueryAchievements(const FUniqueNetId& PlayerId, const FOnQueryAchievementsCompleteDelegate& Delegate = FOnQueryAchievementsCompleteDelegate()) override;
+	virtual void QueryAchievementDescriptions( const FUniqueNetId& PlayerId, const FOnQueryAchievementsCompleteDelegate & Delegate = FOnQueryAchievementsCompleteDelegate() ) override;
+	virtual EOnlineCachedResult::Type GetCachedAchievement(const FUniqueNetId& PlayerId, const FString& AchievementId, FOnlineAchievement& OutAchievement) override;
+	virtual EOnlineCachedResult::Type GetCachedAchievements(const FUniqueNetId& PlayerId, TArray<FOnlineAchievement> & OutAchievements) override;
+	virtual EOnlineCachedResult::Type GetCachedAchievementDescription(const FString& AchievementId, FOnlineAchievementDesc& OutAchievementDesc) override;
 #if !UE_BUILD_SHIPPING
-	virtual bool ResetAchievements( const FUniqueNetId& PlayerId ) OVERRIDE;
+	virtual bool ResetAchievements( const FUniqueNetId& PlayerId ) override;
 #endif // !UE_BUILD_SHIPPING
 	// End IOnlineAchievements interface
 
diff --git a/Engine/Source/Runtime/Online/OnlineSubsystemSteam/Private/OnlineAsyncTaskManagerSteam.cpp b/Engine/Source/Runtime/Online/OnlineSubsystemSteam/Private/OnlineAsyncTaskManagerSteam.cpp
--- a/Engine/Source/Runtime/Online/OnlineSubsystemSteam/Private/OnlineAsyncTaskManagerSteam.cpp
+++ b/Engine/Source/Runtime/Online/OnlineSubsystemSteam/Private/OnlineAsyncTaskManagerSteam.cpp
@@ -210,7 +210,7 @@ public:
 	/**
 	 *	Get a human readable description of task
 	 */
-	virtual FString ToString() const OVERRIDE
+	virtual FString ToString() const override
 	{
 		return FString::Printf(TEXT("FOnlineAsyncEventSteamLobbyEnter LobbyId: %s Result: %s"), 
 			*FUniqueNetIdSteam(CallbackResults.m_ulSteamIDLobby).ToDebugString(),
@@ -284,7 +284,7 @@ public:
 	/**
 	 *	Get a human readable description of task
 	 */
-	virtual FString ToString() const OVERRIDE
+	virtual FString ToString() const override
 	{
 		return FString::Printf(TEXT("FOnlineAsyncEventSteamLobbyChatUpdate User: %s Instigator: %s Result: %s"), 
 			*FUniqueNetIdSteam(CallbackResults.m_ulSteamIDUserChanged).ToDebugString(),
@@ -359,7 +359,7 @@ public:
 	/**
 	 *	Get a human readable description of task
 	 */
-	virtual FString ToString() const OVERRIDE
+	virtual FString ToString() const override
 	{
 		return FString::Printf(TEXT("FOnlineAsyncEventSteamLobbyUpdate LobbyId: %s"), *LobbyId.ToDebugString());	
 	}
@@ -368,7 +368,7 @@ public:
 	 * Give the async task a chance to marshal its data back to the game thread
 	 * Can only be called on the game thread by the async task manager
 	 */
-	virtual void Finalize() OVERRIDE
+	virtual void Finalize() override
 	{
 		FOnlineAsyncEvent::Finalize();
 
@@ -474,7 +474,7 @@ public:
 	/**
 	 *	Get a human readable description of task
 	 */
-	virtual FString ToString() const OVERRIDE
+	virtual FString ToString() const override
 	{
 		return FString::Printf(TEXT("FOnlineAsyncEventSteamStatsReceived bWasSuccessful: %d User: %s Result: %s"), (StatsReceivedResult == k_EResultOK) ? 1 : 0, *UserId.ToDebugString(), *SteamResultString(StatsReceivedResult));	
 	}
@@ -483,7 +483,7 @@ public:
 	 * Give the async task a chance to marshal its data back to the game thread
 	 * Can only be called on the game thread by the async task manager
 	 */
-	virtual void Finalize() OVERRIDE
+	virtual void Finalize() override
 	{
 		FOnlineAsyncEvent::Finalize();
 
@@ -565,7 +565,7 @@ public:
 	/**
 	 *	Get a human readable description of task
 	 */
-	virtual FString ToString() const OVERRIDE
+	virtual FString ToString() const override
 	{
 		return FString::Printf(TEXT("FOnlineAsyncEventSteamStatsStored bWasSuccessful: %d User: %s Result: %s"), (StatsStoredResult == k_EResultOK) ? 1 : 0, *UserId.ToDebugString(), *SteamResultString(StatsStoredResult));	
 	}
@@ -650,7 +650,7 @@ public:
 	/**
 	 *	Get a human readable description of task
 	 */
-	virtual FString ToString() const OVERRIDE
+	virtual FString ToString() const override
 	{
 		return FString::Printf(TEXT("FOnlineAsyncEventSteamStatsUnloaded UserId: %s"), *UserId.ToDebugString());	
 	}
@@ -659,7 +659,7 @@ public:
 	 * Give the async task a chance to marshal its data back to the game thread
 	 * Can only be called on the game thread by the async task manager
 	 */
-	virtual void Finalize() OVERRIDE
+	virtual void Finalize() override
 	{
 		FOnlineAsyncEvent::Finalize();
 		// Remove given user's stats from the cached list
@@ -722,7 +722,7 @@ public:
 	/**
 	 *	Get a human readable description of task
 	 */
-	virtual FString ToString() const OVERRIDE
+	virtual FString ToString() const override
 	{
 		return FString::Printf(TEXT("FOnlineAsyncEventSteamServerConnectionState StateChange: %s"), EOnlineServerConnectionStatus::ToString(ConnectionState));	
 	}
@@ -731,7 +731,7 @@ public:
 	 * Give the async task a chance to marshal its data back to the game thread
 	 * Can only be called on the game thread by the async task manager
 	 */
-	virtual void Finalize() OVERRIDE
+	virtual void Finalize() override
 	{
 		Subsystem->TriggerOnConnectionStatusChangedDelegates(ConnectionState);
 	}
@@ -786,7 +786,7 @@ public:
 	/**
 	 *	Get a human readable description of task
 	 */
-	virtual FString ToString() const OVERRIDE
+	virtual FString ToString() const override
 	{
 		return FString::Printf(TEXT("FOnlineAsyncEventSteamServerConnectedGS ServerId: %s"), *ServerId.ToDebugString());	
 	}
@@ -795,7 +795,7 @@ public:
 	 * Give the async task a chance to marshal its data back to the game thread
 	 * Can only be called on the game thread by the async task manager
 	 */
-	virtual void Finalize() OVERRIDE
+	virtual void Finalize() override
 	{
 		FOnlineSessionSteamPtr SessionInt = StaticCastSharedPtr<FOnlineSessionSteam>(Subsystem->GetSessionInterface());
 		if (SessionInt.IsValid())
@@ -855,7 +855,7 @@ public:
 	/**
 	 *	Get a human readable description of task
 	 */
-	virtual FString ToString() const OVERRIDE
+	virtual FString ToString() const override
 	{
 		return FString::Printf(TEXT("FOnlineAsyncEventSteamServerDisconnectedGS Result: %s"), *SteamResultString(CallbackResults.m_eResult));	
 	}
@@ -864,7 +864,7 @@ public:
 	 * Give the async task a chance to marshal its data back to the game thread
 	 * Can only be called on the game thread by the async task manager
 	 */
-	virtual void Finalize() OVERRIDE
+	virtual void Finalize() override
 	{
 		bool bTriggerConnectionStatusUpdate = true;
 		FOnlineSessionSteamPtr SessionInt = StaticCastSharedPtr<FOnlineSessionSteam>(Subsystem->GetSessionInterface());
@@ -937,7 +937,7 @@ public:
 	/**
 	 *	Get a human readable description of task
 	 */
-	virtual FString ToString() const OVERRIDE
+	virtual FString ToString() const override
 	{
 		return FString::Printf(TEXT("FOnlineAsyncEventSteamServerPolicyResponseGS Secure: %d"), CallbackResults.m_bSecure);	
 	}
@@ -946,7 +946,7 @@ public:
 	 * Give the async task a chance to marshal its data back to the game thread
 	 * Can only be called on the game thread by the async task manager
 	 */
-	virtual void Finalize() OVERRIDE
+	virtual void Finalize() override
 	{
 		FOnlineSessionSteamPtr SessionInt = StaticCastSharedPtr<FOnlineSessionSteam>(Subsystem->GetSessionInterface());
 		if (SessionInt.IsValid())
@@ -1003,7 +1003,7 @@ public:
 	/**
 	 *	Get a human readable description of task
 	 */
-	virtual FString ToString() const OVERRIDE
+	virtual FString ToString() const override
 	{
 		return FString::Printf(TEXT("FOnlineAsyncEventSteamConnectionRequest RemoteId: %s"), *RemoteId.ToDebugString());	
 	}
@@ -1012,7 +1012,7 @@ public:
 	 * Give the async task a chance to marshal its data back to the game thread
 	 * Can only be called on the game thread by the async task manager
 	 */
-	virtual void Finalize() OVERRIDE
+	virtual void Finalize() override
 	{
 		FSocketSubsystemSteam* SocketSubsystem = (FSocketSubsystemSteam*)ISocketSubsystem::Get(STEAM_SUBSYSTEM);
 		if (SocketSubsystem)
@@ -1058,7 +1058,7 @@ public:
 	/**
 	 *	Get a human readable description of task
 	 */
-	virtual FString ToString() const OVERRIDE
+	virtual FString ToString() const override
 	{
 		return FString::Printf(TEXT("FOnlineAsyncEventSteamConnectionFailed RemoteId: %s Reason: %s"), *RemoteId.ToDebugString(), *SteamP2PConnectError(ErrorCode));	
 	}
@@ -1067,7 +1067,7 @@ public:
 	 * Give the async task a chance to marshal its data back to the game thread
 	 * Can only be called on the game thread by the async task manager
 	 */
-	virtual void Finalize() OVERRIDE
+	virtual void Finalize() override
 	{
 		// Mark the relevant sockets with this failure so they can properly notify higher level engine code
 		FSocketSubsystemSteam* SocketSubsystem = (FSocketSubsystemSteam*)ISocketSubsystem::Get(STEAM_SUBSYSTEM);
@@ -1159,7 +1159,7 @@ public:
 	/**
 	 *	Get a human readable description of task
 	 */
-	virtual FString ToString() const OVERRIDE
+	virtual FString ToString() const override
 	{
 		return FString::Printf(TEXT("FOnlineAsyncEventSteamShutdown shutdown received."));	
 	}
@@ -1168,7 +1168,7 @@ public:
 	 * Give the async task a chance to marshal its data back to the game thread
 	 * Can only be called on the game thread by the async task manager
 	 */
-	virtual void Finalize() OVERRIDE
+	virtual void Finalize() override
 	{
 		FPlatformMisc::RequestExit(false);
 	}
diff --git a/Engine/Source/Runtime/Online/OnlineSubsystemSteam/Private/OnlineAsyncTaskManagerSteam.h b/Engine/Source/Runtime/Online/OnlineSubsystemSteam/Private/OnlineAsyncTaskManagerSteam.h
--- a/Engine/Source/Runtime/Online/OnlineSubsystemSteam/Private/OnlineAsyncTaskManagerSteam.h
+++ b/Engine/Source/Runtime/Online/OnlineSubsystemSteam/Private/OnlineAsyncTaskManagerSteam.h
@@ -148,7 +148,7 @@ public:
 	}
 
 	// FOnlineAsyncTaskManager
-	virtual void OnlineTick() OVERRIDE;
+	virtual void OnlineTick() override;
 
 	// FOnlineAsyncTaskManagerSteam
 
diff --git a/Engine/Source/Runtime/Online/OnlineSubsystemSteam/Private/OnlineExternalUIInterfaceSteam.h b/Engine/Source/Runtime/Online/OnlineSubsystemSteam/Private/OnlineExternalUIInterfaceSteam.h
--- a/Engine/Source/Runtime/Online/OnlineSubsystemSteam/Private/OnlineExternalUIInterfaceSteam.h
+++ b/Engine/Source/Runtime/Online/OnlineSubsystemSteam/Private/OnlineExternalUIInterfaceSteam.h
@@ -37,12 +37,12 @@ public:
 	/**
 	 *	Get a human readable description of task
 	 */
-	virtual FString ToString() const OVERRIDE;
+	virtual FString ToString() const override;
 
 	/**
 	 *	Async task is given a chance to trigger it's delegates
 	 */
-	virtual void TriggerDelegates() OVERRIDE;
+	virtual void TriggerDelegates() override;
 };
 
 /** 
@@ -75,13 +75,13 @@ public:
 	}
 
 	// IOnlineExternalUI
-	virtual bool ShowLoginUI(const int ControllerIndex, bool bShowOnlineOnly, const FOnLoginUIClosedDelegate& Delegate) OVERRIDE;
-	virtual bool ShowFriendsUI(int32 LocalUserNum) OVERRIDE;
-	virtual bool ShowInviteUI(int32 LocalUserNum) OVERRIDE;
-	virtual bool ShowAchievementsUI(int32 LocalUserNum) OVERRIDE;
-	virtual bool ShowLeaderboardUI(const FString& LeaderboardName) OVERRIDE;
-	virtual bool ShowWebURL(const FString& WebURL) OVERRIDE;
-	virtual bool ShowProfileUI(const FUniqueNetId& Requestor, const FUniqueNetId& Requestee, const FOnProfileUIClosedDelegate& Delegate) OVERRIDE;
+	virtual bool ShowLoginUI(const int ControllerIndex, bool bShowOnlineOnly, const FOnLoginUIClosedDelegate& Delegate) override;
+	virtual bool ShowFriendsUI(int32 LocalUserNum) override;
+	virtual bool ShowInviteUI(int32 LocalUserNum) override;
+	virtual bool ShowAchievementsUI(int32 LocalUserNum) override;
+	virtual bool ShowLeaderboardUI(const FString& LeaderboardName) override;
+	virtual bool ShowWebURL(const FString& WebURL) override;
+	virtual bool ShowProfileUI(const FUniqueNetId& Requestor, const FUniqueNetId& Requestee, const FOnProfileUIClosedDelegate& Delegate) override;
 };
 
 typedef TSharedPtr<FOnlineExternalUISteam, ESPMode::ThreadSafe> FOnlineExternalUISteamPtr;
diff --git a/Engine/Source/Runtime/Online/OnlineSubsystemSteam/Private/OnlineFriendsInterfaceSteam.h b/Engine/Source/Runtime/Online/OnlineSubsystemSteam/Private/OnlineFriendsInterfaceSteam.h
--- a/Engine/Source/Runtime/Online/OnlineSubsystemSteam/Private/OnlineFriendsInterfaceSteam.h
+++ b/Engine/Source/Runtime/Online/OnlineSubsystemSteam/Private/OnlineFriendsInterfaceSteam.h
@@ -17,15 +17,15 @@ public:
 
 	// FOnlineUser
 
-	virtual TSharedRef<FUniqueNetId> GetUserId() const OVERRIDE;
-	virtual FString GetRealName() const OVERRIDE;
-	virtual FString GetDisplayName() const OVERRIDE;
-	virtual bool GetUserAttribute(const FString& AttrName, FString& OutAttrValue) const OVERRIDE;
+	virtual TSharedRef<FUniqueNetId> GetUserId() const override;
+	virtual FString GetRealName() const override;
+	virtual FString GetDisplayName() const override;
+	virtual bool GetUserAttribute(const FString& AttrName, FString& OutAttrValue) const override;
 
 	// FOnlineFriend
 	
-	virtual EInviteStatus::Type GetInviteStatus() const OVERRIDE;
-	virtual const FOnlineUserPresence& GetPresence() const OVERRIDE;
+	virtual EInviteStatus::Type GetInviteStatus() const override;
+	virtual const FOnlineUserPresence& GetPresence() const override;
 
 	// FOnlineFriendMcp
 
@@ -103,15 +103,15 @@ public:
 
 	// IOnlineFriends
 
-	virtual bool ReadFriendsList(int32 LocalUserNum, const FString& ListName) OVERRIDE;
-	virtual bool DeleteFriendsList(int32 LocalUserNum, const FString& ListName) OVERRIDE;
-	virtual bool SendInvite(int32 LocalUserNum, const FUniqueNetId& FriendId, const FString& ListName) OVERRIDE;
-	virtual bool AcceptInvite(int32 LocalUserNum, const FUniqueNetId& FriendId, const FString& ListName) OVERRIDE;
- 	virtual bool RejectInvite(int32 LocalUserNum, const FUniqueNetId& FriendId, const FString& ListName) OVERRIDE;
- 	virtual bool DeleteFriend(int32 LocalUserNum, const FUniqueNetId& FriendId, const FString& ListName) OVERRIDE;
-	virtual bool GetFriendsList(int32 LocalUserNum, const FString& ListName, TArray< TSharedRef<FOnlineFriend> >& OutFriends) OVERRIDE;
-	virtual TSharedPtr<FOnlineFriend> GetFriend(int32 LocalUserNum, const FUniqueNetId& FriendId, const FString& ListName) OVERRIDE;
-	virtual bool IsFriend(int32 LocalUserNum, const FUniqueNetId& FriendId, const FString& ListName) OVERRIDE;
+	virtual bool ReadFriendsList(int32 LocalUserNum, const FString& ListName) override;
+	virtual bool DeleteFriendsList(int32 LocalUserNum, const FString& ListName) override;
+	virtual bool SendInvite(int32 LocalUserNum, const FUniqueNetId& FriendId, const FString& ListName) override;
+	virtual bool AcceptInvite(int32 LocalUserNum, const FUniqueNetId& FriendId, const FString& ListName) override;
+ 	virtual bool RejectInvite(int32 LocalUserNum, const FUniqueNetId& FriendId, const FString& ListName) override;
+ 	virtual bool DeleteFriend(int32 LocalUserNum, const FUniqueNetId& FriendId, const FString& ListName) override;
+	virtual bool GetFriendsList(int32 LocalUserNum, const FString& ListName, TArray< TSharedRef<FOnlineFriend> >& OutFriends) override;
+	virtual TSharedPtr<FOnlineFriend> GetFriend(int32 LocalUserNum, const FUniqueNetId& FriendId, const FString& ListName) override;
+	virtual bool IsFriend(int32 LocalUserNum, const FUniqueNetId& FriendId, const FString& ListName) override;
 };
 
 typedef TSharedPtr<FOnlineFriendsSteam, ESPMode::ThreadSafe> FOnlineFriendsSteamPtr;
@@ -143,21 +143,21 @@ public:
 
 	// FOnlineAsyncTask
 
-	virtual FString ToString(void) const OVERRIDE
+	virtual FString ToString(void) const override
 	{
 		return TEXT("FOnlineFriendsSteam::ReadFriendsList() async task completed successfully");
 	}
 
-	virtual bool IsDone(void) OVERRIDE
+	virtual bool IsDone(void) override
 	{
 		return true;
 	}
 
-	virtual bool WasSuccessful(void) OVERRIDE
+	virtual bool WasSuccessful(void) override
 	{
 		return true;
 	}
 
-	virtual void Finalize() OVERRIDE;
-	virtual void TriggerDelegates(void) OVERRIDE;
+	virtual void Finalize() override;
+	virtual void TriggerDelegates(void) override;
 };
