UE4 Refactoring. Changed OVERRIDE and FINAL macros to keywords override and final.

diff --git a/Engine/Source/Runtime/Online/OnlineSubsystemSteam/Private/OnlineIdentityInterfaceSteam.h b/Engine/Source/Runtime/Online/OnlineSubsystemSteam/Private/OnlineIdentityInterfaceSteam.h
--- a/Engine/Source/Runtime/Online/OnlineSubsystemSteam/Private/OnlineIdentityInterfaceSteam.h
+++ b/Engine/Source/Runtime/Online/OnlineSubsystemSteam/Private/OnlineIdentityInterfaceSteam.h
@@ -24,17 +24,17 @@ public:
 
 	// IOnlineIdentity
 
-	virtual bool Login(int32 LocalUserNum, const FOnlineAccountCredentials& AccountCredentials) OVERRIDE;
-	virtual bool Logout(int32 LocalUserNum) OVERRIDE;
-	virtual bool AutoLogin(int32 LocalUserNum) OVERRIDE;
-	virtual TSharedPtr<FUserOnlineAccount> GetUserAccount(const FUniqueNetId& UserId) const OVERRIDE;
-	virtual TArray<TSharedPtr<FUserOnlineAccount> > GetAllUserAccounts() const OVERRIDE;
-	virtual TSharedPtr<FUniqueNetId> GetUniquePlayerId(int32 LocalUserNum) const OVERRIDE;
-	virtual TSharedPtr<FUniqueNetId> CreateUniquePlayerId(uint8* Bytes, int32 Size) OVERRIDE;
-	virtual TSharedPtr<FUniqueNetId> CreateUniquePlayerId(const FString& Str) OVERRIDE;
-	virtual ELoginStatus::Type GetLoginStatus(int32 LocalUserNum) const OVERRIDE;
-	virtual FString GetPlayerNickname(int32 LocalUserNum) const OVERRIDE;
-	virtual FString GetAuthToken(int32 LocalUserNum) const OVERRIDE;
+	virtual bool Login(int32 LocalUserNum, const FOnlineAccountCredentials& AccountCredentials) override;
+	virtual bool Logout(int32 LocalUserNum) override;
+	virtual bool AutoLogin(int32 LocalUserNum) override;
+	virtual TSharedPtr<FUserOnlineAccount> GetUserAccount(const FUniqueNetId& UserId) const override;
+	virtual TArray<TSharedPtr<FUserOnlineAccount> > GetAllUserAccounts() const override;
+	virtual TSharedPtr<FUniqueNetId> GetUniquePlayerId(int32 LocalUserNum) const override;
+	virtual TSharedPtr<FUniqueNetId> CreateUniquePlayerId(uint8* Bytes, int32 Size) override;
+	virtual TSharedPtr<FUniqueNetId> CreateUniquePlayerId(const FString& Str) override;
+	virtual ELoginStatus::Type GetLoginStatus(int32 LocalUserNum) const override;
+	virtual FString GetPlayerNickname(int32 LocalUserNum) const override;
+	virtual FString GetAuthToken(int32 LocalUserNum) const override;
 };
 
 typedef TSharedPtr<FOnlineIdentitySteam, ESPMode::ThreadSafe> FOnlineIdentitySteamPtr;
\ No newline at end of file
diff --git a/Engine/Source/Runtime/Online/OnlineSubsystemSteam/Private/OnlineLeaderboardInterfaceSteam.cpp b/Engine/Source/Runtime/Online/OnlineSubsystemSteam/Private/OnlineLeaderboardInterfaceSteam.cpp
--- a/Engine/Source/Runtime/Online/OnlineSubsystemSteam/Private/OnlineLeaderboardInterfaceSteam.cpp
+++ b/Engine/Source/Runtime/Online/OnlineSubsystemSteam/Private/OnlineLeaderboardInterfaceSteam.cpp
@@ -112,7 +112,7 @@ public:
 	/**
 	 *	Get a human readable description of task
 	 */
-	FString ToString() const OVERRIDE
+	FString ToString() const override
 	{
 		return FString::Printf(TEXT("FOnlineAsyncTaskSteamRequestUserStats bWasSuccessful: %d UserId: %s"), bWasSuccessful, *UserId.ToDebugString());
 	}
@@ -121,7 +121,7 @@ public:
 	 * Give the async task time to do its work
 	 * Can only be called on the async task manager thread
 	 */
-	virtual void Tick() OVERRIDE
+	virtual void Tick() override
 	{
 		if (!bInit)
 		{
@@ -160,7 +160,7 @@ public:
 	 * Give the async task a chance to marshal its data back to the game thread
 	 * Can only be called on the game thread by the async task manager
 	 */
-	virtual void Finalize() OVERRIDE
+	virtual void Finalize() override
 	{
 		FOnlineAsyncTaskSteam::Finalize();
 
@@ -237,7 +237,7 @@ public:
 	/**
 	 *	Get a human readable description of task
 	 */
-	FString ToString() const OVERRIDE
+	FString ToString() const override
 	{
 		return FString::Printf(TEXT("FOnlineAsyncTaskSteamUpdateStats bWasSuccessful: %d User: %s"), bWasSuccessful, *UserId.ToDebugString());
 	}
@@ -246,7 +246,7 @@ public:
 	 * Give the async task time to do its work
 	 * Can only be called on the async task manager thread
 	 */
-	virtual void Tick() OVERRIDE
+	virtual void Tick() override
 	{
 		if (!bInit)
 		{
@@ -373,7 +373,7 @@ public:
 	/**
 	 *	Get a human readable description of task
 	 */
-	FString ToString() const OVERRIDE
+	FString ToString() const override
 	{
 		return FString::Printf(TEXT("FOnlineAsyncTaskSteamRetrieveStats bWasSuccessful: %d UserId: %s"), bWasSuccessful, *UserId.ToDebugString());
 	}
@@ -382,7 +382,7 @@ public:
 	 * Give the async task time to do its work
 	 * Can only be called on the async task manager thread
 	 */
-	virtual void Tick() OVERRIDE
+	virtual void Tick() override
 	{
 		if (!bInit)
 		{
@@ -419,7 +419,7 @@ public:
 	 * Give the async task a chance to marshal its data back to the game thread
 	 * Can only be called on the game thread by the async task manager
 	 */
-	virtual void Finalize() OVERRIDE
+	virtual void Finalize() override
 	{
 		FOnlineAsyncTaskSteam::Finalize();
 
@@ -498,7 +498,7 @@ public:
 	/**
 	 *	Async task is given a chance to trigger it's delegates
 	 */
-	virtual void TriggerDelegates() OVERRIDE 
+	virtual void TriggerDelegates() override 
 	{
 		FOnlineAsyncTaskSteam::TriggerDelegates();
 		
@@ -553,7 +553,7 @@ public:
 	/**
 	 *	Get a human readable description of task
 	 */
-	FString ToString() const OVERRIDE
+	FString ToString() const override
 	{
 		return FString::Printf(TEXT("FOnlineAsyncTaskSteamGetAchievements bWasSuccessful: %d UserId: %s"), bWasSuccessful, *UserId.ToDebugString());
 	}
@@ -562,7 +562,7 @@ public:
 	 * Give the async task time to do its work
 	 * Can only be called on the async task manager thread
 	 */
-	virtual void Tick() OVERRIDE
+	virtual void Tick() override
 	{
 		if (!bInit)
 		{
@@ -600,7 +600,7 @@ public:
 	 * Give the async task a chance to marshal its data back to the game thread
 	 * Can only be called on the game thread by the async task manager
 	 */
-	virtual void Finalize() OVERRIDE
+	virtual void Finalize() override
 	{
 		FOnlineAsyncTaskSteam::Finalize();
 
@@ -611,7 +611,7 @@ public:
 	/**
 	 *	Async task is given a chance to trigger its delegates
 	 */
-	virtual void TriggerDelegates() OVERRIDE 
+	virtual void TriggerDelegates() override 
 	{
 		FOnlineAsyncTaskSteam::TriggerDelegates();
 		
@@ -699,7 +699,7 @@ public:
 	/**
 	 *	Get a human readable description of task
 	 */
-	FString ToString() const OVERRIDE
+	FString ToString() const override
 	{
 		return FString::Printf(TEXT("FOnlineAsyncTaskSteamRetrieveLeaderboard bWasSuccessful: %d"), bWasSuccessful);
 	}
@@ -708,7 +708,7 @@ public:
 	 * Give the async task time to do its work
 	 * Can only be called on the async task manager thread
 	 */
-	virtual void Tick() OVERRIDE
+	virtual void Tick() override
 	{
 		ISteamUtils* SteamUtilsPtr = SteamUtils();
 		check(SteamUtilsPtr);
@@ -748,7 +748,7 @@ public:
 	 * Give the async task a chance to marshal its data back to the game thread
 	 * Can only be called on the game thread by the async task manager
 	 */
-	virtual void Finalize() OVERRIDE
+	virtual void Finalize() override
 	{
 		FOnlineAsyncTaskSteam::Finalize();
 
@@ -808,7 +808,7 @@ public:
 	/**
 	 *	Get a human readable description of task
 	 */
-	FString ToString() const OVERRIDE
+	FString ToString() const override
 	{
 		return FString::Printf(TEXT("FOnlineAsyncTaskSteamRetrieveLeaderboardEntries bWasSuccessful: %d"), bWasSuccessful);
 	}
@@ -817,7 +817,7 @@ public:
 	 * Give the async task time to do its work
 	 * Can only be called on the async task manager thread
 	 */
-	virtual void Tick() OVERRIDE
+	virtual void Tick() override
 	{
 		ISteamUtils* SteamUtilsPtr = SteamUtils();
 		check(SteamUtilsPtr);
@@ -882,7 +882,7 @@ public:
 	 * Give the async task a chance to marshal its data back to the game thread
 	 * Can only be called on the game thread by the async task manager
 	 */
-	virtual void Finalize() OVERRIDE
+	virtual void Finalize() override
 	{
 		FOnlineAsyncTaskSteam::Finalize();
 
@@ -980,7 +980,7 @@ public:
 	/**
 	 *	Get a human readable description of task
 	 */
-	virtual FString ToString() const OVERRIDE
+	virtual FString ToString() const override
 	{
 		return FString::Printf(TEXT("FOnlineAsyncTaskSteamUpdateLeaderboard bWasSuccessful: %d Leaderboard: %s Score: %d"), bWasSuccessful, *LeaderboardName.ToString(), NewScore);
 	}
@@ -989,7 +989,7 @@ public:
 	 * Give the async task time to do its work
 	 * Can only be called on the async task manager thread
 	 */
-	virtual void Tick() OVERRIDE
+	virtual void Tick() override
 	{
 		ISteamUtils* SteamUtilsPtr = SteamUtils();
 		check(SteamUtilsPtr);
@@ -1063,7 +1063,7 @@ public:
 	/**
 	 *	Async task is given a chance to trigger it's delegates
 	 */
-	virtual void TriggerDelegates() OVERRIDE
+	virtual void TriggerDelegates() override
 	{
 		FOnlineAsyncTaskSteam::TriggerDelegates();
 		if (bShouldTriggerDelegates)
@@ -1127,7 +1127,7 @@ public:
 	/**
 	 *	Get a human readable description of task
 	 */
-	FString ToString() const OVERRIDE
+	FString ToString() const override
 	{
 		return FString::Printf(TEXT("FOnlineAsyncTaskSteamStoreStats SessionName: %s bWasSuccessful: %d"), *SessionName.ToString(), bWasSuccessful);
 	}
@@ -1136,7 +1136,7 @@ public:
 	 * Give the async task time to do its work
 	 * Can only be called on the async task manager thread
 	 */
-	virtual void Tick() OVERRIDE
+	virtual void Tick() override
 	{
 		FOnlineLeaderboardsSteamPtr Leaderboards = StaticCastSharedPtr<FOnlineLeaderboardsSteam>(Subsystem->GetLeaderboardsInterface());
 
@@ -1198,7 +1198,7 @@ public:
 	/**
 	 *	Async task is given a chance to trigger it's delegates
 	 */
-	virtual void TriggerDelegates() OVERRIDE 
+	virtual void TriggerDelegates() override 
 	{
 		FOnlineAsyncTaskSteam::TriggerDelegates();
 		
@@ -1229,21 +1229,21 @@ private:
 	}
 
 	/** Internal function to allow write state tracking */
-	virtual void OperationStarted() OVERRIDE
+	virtual void OperationStarted() override
 	{
 		check(WriteObject.IsValid());
 		WriteObject->WriteState = EOnlineAsyncTaskState::InProgress;
 	}
 
 	/** Internal function to allow write state tracking */
-	virtual void OperationFailed() OVERRIDE
+	virtual void OperationFailed() override
 	{
 		check(WriteObject.IsValid());
 		WriteObject->WriteState = EOnlineAsyncTaskState::Failed;
 	}
 
 	/** Internal function to allow write state tracking */
-	virtual void OperationSucceeded() OVERRIDE
+	virtual void OperationSucceeded() override
 	{
 		check(WriteObject.IsValid());
 		WriteObject->WriteState = EOnlineAsyncTaskState::Done;
@@ -1260,7 +1260,7 @@ public:
 	/**
 	 *	Async task is given a chance to trigger it's delegates
 	 */
-	virtual void TriggerDelegates() OVERRIDE 
+	virtual void TriggerDelegates() override 
 	{
 		FOnlineAsyncTaskSteam::TriggerDelegates();
 		
diff --git a/Engine/Source/Runtime/Online/OnlineSubsystemSteam/Private/OnlineMsgSteam.h b/Engine/Source/Runtime/Online/OnlineSubsystemSteam/Private/OnlineMsgSteam.h
--- a/Engine/Source/Runtime/Online/OnlineSubsystemSteam/Private/OnlineMsgSteam.h
+++ b/Engine/Source/Runtime/Online/OnlineSubsystemSteam/Private/OnlineMsgSteam.h
@@ -115,12 +115,12 @@ public:
 	/**
 	 *	Get a human readable description of task
 	 */
-	virtual FString ToString() const OVERRIDE;
+	virtual FString ToString() const override;
 
 	/**
 	 *	Async task is given a chance to trigger it's delegates
 	 */
-	virtual void TriggerDelegates() OVERRIDE;
+	virtual void TriggerDelegates() override;
 
 	/** Get Unique Id defining the message type sent to the service */
 	int32 GetMsgType() const
diff --git a/Engine/Source/Runtime/Online/OnlineSubsystemSteam/Private/OnlineSessionAsyncLobbySteam.h b/Engine/Source/Runtime/Online/OnlineSubsystemSteam/Private/OnlineSessionAsyncLobbySteam.h
--- a/Engine/Source/Runtime/Online/OnlineSubsystemSteam/Private/OnlineSessionAsyncLobbySteam.h
+++ b/Engine/Source/Runtime/Online/OnlineSubsystemSteam/Private/OnlineSessionAsyncLobbySteam.h
@@ -51,24 +51,24 @@ public:
 	/**
 	 *	Get a human readable description of task
 	 */
-	virtual FString ToString() const OVERRIDE;
+	virtual FString ToString() const override;
 
 	/**
 	 * Give the async task time to do its work
 	 * Can only be called on the async task manager thread
 	 */
-	virtual void Tick() OVERRIDE;
+	virtual void Tick() override;
 
 	/**
 	 * Give the async task a chance to marshal its data back to the game thread
 	 * Can only be called on the game thread by the async task manager
 	 */
-	virtual void Finalize() OVERRIDE;
+	virtual void Finalize() override;
 	
 	/**
 	 *	Async task is given a chance to trigger it's delegates
 	 */
-	virtual void TriggerDelegates() OVERRIDE;
+	virtual void TriggerDelegates() override;
 };
 
 /** 
@@ -106,18 +106,18 @@ public:
 	/**
 	 *	Get a human readable description of task
 	 */
-	virtual FString ToString() const OVERRIDE;
+	virtual FString ToString() const override;
 
 	/**
 	* Give the async task time to do its work
 	* Can only be called on the async task manager thread
 	*/
-	virtual void Tick() OVERRIDE;
+	virtual void Tick() override;
 	
 	/**
 	 *	Async task is given a chance to trigger it's delegates
 	 */
-	virtual void TriggerDelegates() OVERRIDE;
+	virtual void TriggerDelegates() override;
 };
 
 /** 
@@ -161,24 +161,24 @@ public:
 	/**
 	 *	Get a human readable description of task
 	 */
-	virtual FString ToString() const OVERRIDE;
+	virtual FString ToString() const override;
 
 	/**
 	* Give the async task time to do its work
 	* Can only be called on the async task manager thread
 	*/
-	virtual void Tick() OVERRIDE;
+	virtual void Tick() override;
 
 	/**
 	 * Give the async task a chance to marshal its data back to the game thread
 	 * Can only be called on the game thread by the async task manager
 	 */
-	virtual void Finalize() OVERRIDE;
+	virtual void Finalize() override;
 	
 	/**
 	 *	Async task is given a chance to trigger it's delegates
 	 */
-	virtual void TriggerDelegates() OVERRIDE;
+	virtual void TriggerDelegates() override;
 };
 
 /** 
@@ -212,13 +212,13 @@ public:
 	/**
 	 *	Get a human readable description of task
 	 */
-	virtual FString ToString() const OVERRIDE;
+	virtual FString ToString() const override;
 
 	/**
 	 * Give the async task time to do its work
 	 * Can only be called on the async task manager thread
 	 */
-	virtual void Tick() OVERRIDE;
+	virtual void Tick() override;
 };
 
 /** 
@@ -277,24 +277,24 @@ public:
 	/**
 	 *	Get a human readable description of task
 	 */
-	virtual FString ToString() const OVERRIDE;
+	virtual FString ToString() const override;
 
 	/**
 	 * Give the async task time to do its work
 	 * Can only be called on the async task manager thread
 	 */
-	virtual void Tick() OVERRIDE;
+	virtual void Tick() override;
 
 	/**
 	 * Give the async task a chance to marshal its data back to the game thread
 	 * Can only be called on the game thread by the async task manager
 	 */
-	virtual void Finalize() OVERRIDE;
+	virtual void Finalize() override;
 	
 	/**
 	 *	Async task is given a chance to trigger it's delegates
 	 */
-	virtual void TriggerDelegates() OVERRIDE;
+	virtual void TriggerDelegates() override;
 };
 
 
@@ -353,7 +353,7 @@ public:
 	/**
 	 *	Get a human readable description of task
 	 */
-	virtual FString ToString() const OVERRIDE
+	virtual FString ToString() const override
 	{
 		return FString::Printf(TEXT("FOnlineAsyncTaskSteamFindLobby bWasSuccessful: %d LobbyId: %s"), bWasSuccessful, *LobbyId.ToDebugString());
 	}
@@ -362,18 +362,18 @@ public:
 	 * Give the async task time to do its work
 	 * Can only be called on the async task manager thread
 	 */
-	virtual void Tick() OVERRIDE;
+	virtual void Tick() override;
 
 	/**
 	 * Give the async task a chance to marshal its data back to the game thread
 	 * Can only be called on the game thread by the async task manager
 	 */
-	virtual void Finalize() OVERRIDE;
+	virtual void Finalize() override;
 	
 	/**
 	 *	Async task is given a chance to trigger it's delegates
 	 */
-	virtual void TriggerDelegates() OVERRIDE;
+	virtual void TriggerDelegates() override;
 };
 
 /**
@@ -412,13 +412,13 @@ public:
 	/**
 	 *	Get a human readable description of task
 	 */
-	virtual FString ToString() const OVERRIDE;
+	virtual FString ToString() const override;
 
 	/**
 	 * Give the async task a chance to marshal its data back to the game thread
 	 * Can only be called on the game thread by the async task manager
 	 */
-	virtual void Finalize() OVERRIDE;
+	virtual void Finalize() override;
 };
 
 /**
diff --git a/Engine/Source/Runtime/Online/OnlineSubsystemSteam/Private/OnlineSessionAsyncServerSteam.h b/Engine/Source/Runtime/Online/OnlineSubsystemSteam/Private/OnlineSessionAsyncServerSteam.h
--- a/Engine/Source/Runtime/Online/OnlineSubsystemSteam/Private/OnlineSessionAsyncServerSteam.h
+++ b/Engine/Source/Runtime/Online/OnlineSubsystemSteam/Private/OnlineSessionAsyncServerSteam.h
@@ -41,24 +41,24 @@ public:
 	/**
 	 *	Get a human readable description of task
 	 */
-	virtual FString ToString() const OVERRIDE;
+	virtual FString ToString() const override;
 
 	/**
 	 * Give the async task time to do its work
 	 * Can only be called on the async task manager thread
 	 */
-	virtual void Tick() OVERRIDE;
+	virtual void Tick() override;
 
 	/**
 	 * Give the async task a chance to marshal its data back to the game thread
 	 * Can only be called on the game thread by the async task manager
 	 */
-	virtual void Finalize() OVERRIDE;
+	virtual void Finalize() override;
 	
 	/**
 	 *	Async task is given a chance to trigger it's delegates
 	 */
-	virtual void TriggerDelegates() OVERRIDE;
+	virtual void TriggerDelegates() override;
 };
 
 /** 
@@ -96,18 +96,18 @@ public:
 	/**
 	 *	Get a human readable description of task
 	 */
-	virtual FString ToString() const OVERRIDE;
+	virtual FString ToString() const override;
 
 	/**
 	* Give the async task time to do its work
 	* Can only be called on the async task manager thread
 	*/
-	virtual void Tick() OVERRIDE;
+	virtual void Tick() override;
 	
 	/**
 	 *	Async task is given a chance to trigger it's delegates
 	 */
-	virtual void TriggerDelegates() OVERRIDE;
+	virtual void TriggerDelegates() override;
 };
 
 /** 
@@ -142,13 +142,13 @@ public:
 	/**
 	 *	Get a human readable description of task
 	 */
-	virtual FString ToString() const OVERRIDE;
+	virtual FString ToString() const override;
 
 	/**
 	 * Give the async task time to do its work
 	 * Can only be called on the async task manager thread
 	 */
-	virtual void Tick() OVERRIDE;
+	virtual void Tick() override;
 };
 
 /**
@@ -298,7 +298,7 @@ public:
 	 * Give the async task time to do its work
 	 * Can only be called on the async task manager thread
 	 */
-	virtual void Tick() OVERRIDE;
+	virtual void Tick() override;
 
 	/**
 	 * Called by the SteamAPI when a server has successfully responded
@@ -319,7 +319,7 @@ public:
 	 * Give the async task a chance to marshal its data back to the game thread
 	 * Can only be called on the game thread by the async task manager
 	 */
-	virtual void Finalize() OVERRIDE;
+	virtual void Finalize() override;
 };
 
 /**
@@ -355,12 +355,12 @@ public:
 	/**
 	 *	Get a human readable description of task
 	 */
-	virtual FString ToString() const OVERRIDE;
+	virtual FString ToString() const override;
 
 	/**
 	 *	Async task is given a chance to trigger it's delegates
 	 */
-	virtual void TriggerDelegates() OVERRIDE;
+	virtual void TriggerDelegates() override;
 };
 
 /**
@@ -392,12 +392,12 @@ public:
 	/**
 	 *	Get a human readable description of task
 	 */
-	virtual FString ToString() const OVERRIDE;
+	virtual FString ToString() const override;
 
 	/**
 	 *	Async task is given a chance to trigger it's delegates
 	 */
-	virtual void TriggerDelegates() OVERRIDE;
+	virtual void TriggerDelegates() override;
 };
 
 /**
@@ -436,7 +436,7 @@ public:
 	/**
 	 *	Get a human readable description of task
 	 */
-	virtual FString ToString() const OVERRIDE
+	virtual FString ToString() const override
 	{
 		return FString::Printf(TEXT("FOnlineAsyncEventSteamInviteAccepted Friend: %s URL: %s"), *FriendId.ToDebugString(), *ConnectionURL);
 	}
@@ -445,5 +445,5 @@ public:
 	 * Give the async task a chance to marshal its data back to the game thread
 	 * Can only be called on the game thread by the async task manager
 	 */
-	virtual void Finalize() OVERRIDE;
+	virtual void Finalize() override;
 };
diff --git a/Engine/Source/Runtime/Online/OnlineSubsystemSteam/Private/OnlineSessionInterfaceSteam.cpp b/Engine/Source/Runtime/Online/OnlineSubsystemSteam/Private/OnlineSessionInterfaceSteam.cpp
--- a/Engine/Source/Runtime/Online/OnlineSubsystemSteam/Private/OnlineSessionInterfaceSteam.cpp
+++ b/Engine/Source/Runtime/Online/OnlineSubsystemSteam/Private/OnlineSessionInterfaceSteam.cpp
@@ -77,7 +77,7 @@ public:
 	/**
 	 *	Get a human readable description of task
 	 */
-	virtual FString ToString() const OVERRIDE
+	virtual FString ToString() const override
 	{
 		return FString::Printf(TEXT("FOnlineAsyncTaskSteamEndSession bWasSuccessful: %d SessionName: %s"), bWasSuccessful, *SessionName.ToString());
 	}
@@ -86,7 +86,7 @@ public:
 	 * Give the async task time to do its work
 	 * Can only be called on the async task manager thread
 	 */
-	virtual void Tick() OVERRIDE
+	virtual void Tick() override
 	{
 		bIsComplete = true;
 		bWasSuccessful = true;
@@ -96,7 +96,7 @@ public:
 	 * Give the async task a chance to marshal its data back to the game thread
 	 * Can only be called on the game thread by the async task manager
 	 */
-	virtual void Finalize() OVERRIDE
+	virtual void Finalize() override
 	{
 		IOnlineSessionPtr SessionInt = Subsystem->GetSessionInterface();
 		FNamedOnlineSession* Session = SessionInt->GetNamedSession(SessionName);
@@ -109,7 +109,7 @@ public:
 	/**
 	 *	Async task is given a chance to trigger it's delegates
 	 */
-	virtual void TriggerDelegates() OVERRIDE
+	virtual void TriggerDelegates() override
 	{
 		IOnlineSessionPtr SessionInt = Subsystem->GetSessionInterface();
 		if (SessionInt.IsValid())
@@ -142,7 +142,7 @@ public:
 	/**
 	 *	Get a human readable description of task
 	 */
-	virtual FString ToString() const OVERRIDE
+	virtual FString ToString() const override
 	{
 		return FString::Printf(TEXT("FOnlineAsyncTaskSteamDestroySession bWasSuccessful: %d SessionName: %s"), bWasSuccessful, *SessionName.ToString());
 	}
@@ -151,7 +151,7 @@ public:
 	 * Give the async task time to do its work
 	 * Can only be called on the async task manager thread
 	 */
-	virtual void Tick() OVERRIDE
+	virtual void Tick() override
 	{
 		bIsComplete = true;
 		bWasSuccessful = true;
@@ -161,7 +161,7 @@ public:
 	 * Give the async task a chance to marshal its data back to the game thread
 	 * Can only be called on the game thread by the async task manager
 	 */
-	virtual void Finalize() OVERRIDE
+	virtual void Finalize() override
 	{
 		IOnlineSessionPtr SessionInt = Subsystem->GetSessionInterface();
 		if (SessionInt.IsValid())
@@ -192,7 +192,7 @@ public:
 	/**
 	 *	Async task is given a chance to trigger it's delegates
 	 */
-	virtual void TriggerDelegates() OVERRIDE
+	virtual void TriggerDelegates() override
 	{
 		IOnlineSessionPtr SessionInt = Subsystem->GetSessionInterface();
 		if (SessionInt.IsValid())
diff --git a/Engine/Source/Runtime/Online/OnlineSubsystemSteam/Private/OnlineSessionInterfaceSteam.h b/Engine/Source/Runtime/Online/OnlineSubsystemSteam/Private/OnlineSessionInterfaceSteam.h
--- a/Engine/Source/Runtime/Online/OnlineSubsystemSteam/Private/OnlineSessionInterfaceSteam.h
+++ b/Engine/Source/Runtime/Online/OnlineSubsystemSteam/Private/OnlineSessionInterfaceSteam.h
@@ -299,7 +299,7 @@ PACKAGE_SCOPE:
 	 *
 	 * @return a pointer to the struct that was added
 	 */
-	class FNamedOnlineSession* AddNamedSession(FName SessionName, const FOnlineSessionSettings& SessionSettings) OVERRIDE
+	class FNamedOnlineSession* AddNamedSession(FName SessionName, const FOnlineSessionSettings& SessionSettings) override
 	{
 		FScopeLock ScopeLock(&SessionLock);
 		return new (Sessions) FNamedOnlineSession(SessionName, SessionSettings);
@@ -313,7 +313,7 @@ PACKAGE_SCOPE:
 	 *
 	 * @return a pointer to the struct that was added
 	 */
-	class FNamedOnlineSession* AddNamedSession(FName SessionName, const FOnlineSession& Session) OVERRIDE
+	class FNamedOnlineSession* AddNamedSession(FName SessionName, const FOnlineSession& Session) override
 	{
 		FScopeLock ScopeLock(&SessionLock);
 		return new (Sessions) FNamedOnlineSession(SessionName, Session);
@@ -434,7 +434,7 @@ public:
 
 	virtual ~FOnlineSessionSteam() {}
 
-	FNamedOnlineSession* GetNamedSession(FName SessionName) OVERRIDE
+	FNamedOnlineSession* GetNamedSession(FName SessionName) override
 	{
 		FScopeLock ScopeLock(&SessionLock);
 		for (int32 SearchIndex = 0; SearchIndex < Sessions.Num(); SearchIndex++)
@@ -447,7 +447,7 @@ public:
 		return NULL;
 	}
 
-	virtual void RemoveNamedSession(FName SessionName) OVERRIDE
+	virtual void RemoveNamedSession(FName SessionName) override
 	{
 		FScopeLock ScopeLock(&SessionLock);
 		for (int32 SearchIndex = 0; SearchIndex < Sessions.Num(); SearchIndex++)
@@ -460,7 +460,7 @@ public:
 		}
 	}
 
-	virtual EOnlineSessionState::Type GetSessionState(FName SessionName) const OVERRIDE
+	virtual EOnlineSessionState::Type GetSessionState(FName SessionName) const override
 	{
 		FScopeLock ScopeLock(&SessionLock);
 		for (int32 SearchIndex = 0; SearchIndex < Sessions.Num(); SearchIndex++)
@@ -474,7 +474,7 @@ public:
 		return EOnlineSessionState::NoSession;
 	}
 
-	virtual bool HasPresenceSession() OVERRIDE
+	virtual bool HasPresenceSession() override
 	{
 		FScopeLock ScopeLock(&SessionLock);
 		for (int32 SearchIndex = 0; SearchIndex < Sessions.Num(); SearchIndex++)
@@ -489,30 +489,30 @@ public:
 	}
 
 	// IOnlineSession
-	virtual bool CreateSession(int32 HostingPlayerNum, FName SessionName, const FOnlineSessionSettings& NewSessionSettings) OVERRIDE;
-	virtual bool StartSession(FName SessionName) OVERRIDE;
-	virtual bool UpdateSession(FName SessionName, FOnlineSessionSettings& UpdatedSessionSettings, bool bShouldRefreshOnlineData = true) OVERRIDE;
-	virtual bool EndSession(FName SessionName) OVERRIDE;
-	virtual bool DestroySession(FName SessionName) OVERRIDE;
-	virtual bool IsPlayerInSession(FName SessionName, const FUniqueNetId& UniqueId) OVERRIDE;
-	virtual bool StartMatchmaking(int32 SearchingPlayerNum, FName SessionName, const FOnlineSessionSettings& NewSessionSettings, TSharedRef<FOnlineSessionSearch>& SearchSettings) OVERRIDE;
-	virtual bool CancelMatchmaking(int32 SearchingPlayerNum, FName SessionName) OVERRIDE;
-	virtual bool FindSessions(int32 SearchingPlayerNum, const TSharedRef<FOnlineSessionSearch>& SearchSettings) OVERRIDE;
-	virtual bool CancelFindSessions() OVERRIDE;
-	virtual bool PingSearchResults(const FOnlineSessionSearchResult& SearchResult) OVERRIDE;
-	virtual bool JoinSession(int32 PlayerNum, FName SessionName, const FOnlineSessionSearchResult& DesiredSession) OVERRIDE;
-	virtual bool FindFriendSession(int32 LocalUserNum, const FUniqueNetId& Friend) OVERRIDE;
-	virtual bool SendSessionInviteToFriend(int32 LocalUserNum, FName SessionName, const FUniqueNetId& Friend) OVERRIDE;
-	virtual bool SendSessionInviteToFriends(int32 LocalUserNum, FName SessionName, const TArray< TSharedRef<FUniqueNetId> >& Friends) OVERRIDE;
-	virtual bool GetResolvedConnectString(FName SessionName, FString& ConnectInfo) OVERRIDE;
-	virtual bool GetResolvedConnectString(const class FOnlineSessionSearchResult& SearchResult, FName PortType, FString& ConnectInfo) OVERRIDE;
-	virtual FOnlineSessionSettings* GetSessionSettings(FName SessionName) OVERRIDE;
-	virtual bool RegisterPlayer(FName SessionName, const FUniqueNetId& PlayerId, bool bWasInvited) OVERRIDE;
-	virtual bool RegisterPlayers(FName SessionName, const TArray< TSharedRef<FUniqueNetId> >& Players, bool bWasInvited = false) OVERRIDE;
-	virtual bool UnregisterPlayer(FName SessionName, const FUniqueNetId& PlayerId) OVERRIDE;
-	virtual bool UnregisterPlayers(FName SessionName, const TArray< TSharedRef<FUniqueNetId> >& Players) OVERRIDE;
-	virtual int32 GetNumSessions() OVERRIDE;
-	virtual void DumpSessionState() OVERRIDE;
+	virtual bool CreateSession(int32 HostingPlayerNum, FName SessionName, const FOnlineSessionSettings& NewSessionSettings) override;
+	virtual bool StartSession(FName SessionName) override;
+	virtual bool UpdateSession(FName SessionName, FOnlineSessionSettings& UpdatedSessionSettings, bool bShouldRefreshOnlineData = true) override;
+	virtual bool EndSession(FName SessionName) override;
+	virtual bool DestroySession(FName SessionName) override;
+	virtual bool IsPlayerInSession(FName SessionName, const FUniqueNetId& UniqueId) override;
+	virtual bool StartMatchmaking(int32 SearchingPlayerNum, FName SessionName, const FOnlineSessionSettings& NewSessionSettings, TSharedRef<FOnlineSessionSearch>& SearchSettings) override;
+	virtual bool CancelMatchmaking(int32 SearchingPlayerNum, FName SessionName) override;
+	virtual bool FindSessions(int32 SearchingPlayerNum, const TSharedRef<FOnlineSessionSearch>& SearchSettings) override;
+	virtual bool CancelFindSessions() override;
+	virtual bool PingSearchResults(const FOnlineSessionSearchResult& SearchResult) override;
+	virtual bool JoinSession(int32 PlayerNum, FName SessionName, const FOnlineSessionSearchResult& DesiredSession) override;
+	virtual bool FindFriendSession(int32 LocalUserNum, const FUniqueNetId& Friend) override;
+	virtual bool SendSessionInviteToFriend(int32 LocalUserNum, FName SessionName, const FUniqueNetId& Friend) override;
+	virtual bool SendSessionInviteToFriends(int32 LocalUserNum, FName SessionName, const TArray< TSharedRef<FUniqueNetId> >& Friends) override;
+	virtual bool GetResolvedConnectString(FName SessionName, FString& ConnectInfo) override;
+	virtual bool GetResolvedConnectString(const class FOnlineSessionSearchResult& SearchResult, FName PortType, FString& ConnectInfo) override;
+	virtual FOnlineSessionSettings* GetSessionSettings(FName SessionName) override;
+	virtual bool RegisterPlayer(FName SessionName, const FUniqueNetId& PlayerId, bool bWasInvited) override;
+	virtual bool RegisterPlayers(FName SessionName, const TArray< TSharedRef<FUniqueNetId> >& Players, bool bWasInvited = false) override;
+	virtual bool UnregisterPlayer(FName SessionName, const FUniqueNetId& PlayerId) override;
+	virtual bool UnregisterPlayers(FName SessionName, const TArray< TSharedRef<FUniqueNetId> >& Players) override;
+	virtual int32 GetNumSessions() override;
+	virtual void DumpSessionState() override;
 };
 
 typedef TSharedPtr<FOnlineSessionSteam, ESPMode::ThreadSafe> FOnlineSessionSteamPtr;
diff --git a/Engine/Source/Runtime/Online/OnlineSubsystemSteam/Private/OnlineSharedCloudInterfaceSteam.h b/Engine/Source/Runtime/Online/OnlineSubsystemSteam/Private/OnlineSharedCloudInterfaceSteam.h
--- a/Engine/Source/Runtime/Online/OnlineSubsystemSteam/Private/OnlineSharedCloudInterfaceSteam.h
+++ b/Engine/Source/Runtime/Online/OnlineSubsystemSteam/Private/OnlineSharedCloudInterfaceSteam.h
@@ -44,24 +44,24 @@ public:
 	/**
 	 *	Get a human readable description of task
 	 */
-	virtual FString ToString() const OVERRIDE;
+	virtual FString ToString() const override;
 
 	/**
 	 * Give the async task time to do its work
 	 * Can only be called on the async task manager thread
 	 */
-	virtual void Tick() OVERRIDE;
+	virtual void Tick() override;
 
 	/**
 	 * Give the async task a chance to marshal its data back to the game thread
 	 * Can only be called on the game thread by the async task manager
 	 */
-	virtual void Finalize() OVERRIDE;
+	virtual void Finalize() override;
 	
 	/**
 	 *	Async task is given a chance to trigger it's delegates
 	 */
-	virtual void TriggerDelegates() OVERRIDE;
+	virtual void TriggerDelegates() override;
 };
 
 /** 
@@ -96,23 +96,23 @@ public:
 	/**
 	 *	Get a human readable description of task
 	 */
-	virtual FString ToString() const OVERRIDE;
+	virtual FString ToString() const override;
 
 	/**
 	 * Give the async task time to do its work
 	 * Can only be called on the async task manager thread
 	 */
-	virtual void Tick() OVERRIDE;
+	virtual void Tick() override;
 
 	/**
 	 * Write out the state of the request
 	 */
-	virtual void Finalize() OVERRIDE;
+	virtual void Finalize() override;
 	
 	/**
 	 *	Async task is given a chance to trigger it's delegates
 	 */
-	virtual void TriggerDelegates() OVERRIDE;
+	virtual void TriggerDelegates() override;
 };
 
 /**
@@ -157,12 +157,12 @@ public:
 	}
 
 	// IOnlineSharedCloud
-	virtual bool GetSharedFileContents(const FSharedContentHandle& SharedHandle, TArray<uint8>& FileContents) OVERRIDE;
-	virtual bool ClearSharedFiles() OVERRIDE;
-	virtual bool ClearSharedFile(const FSharedContentHandle& SharedHandle) OVERRIDE;
-	virtual bool ReadSharedFile(const FSharedContentHandle& SharedHandle) OVERRIDE;
-	virtual bool WriteSharedFile(const FUniqueNetId& UserId, const FString& Filename, TArray<uint8>& FileContents) OVERRIDE;
-	virtual void GetDummySharedHandlesForTest(TArray< TSharedRef<FSharedContentHandle> > & OutHandles) OVERRIDE;
+	virtual bool GetSharedFileContents(const FSharedContentHandle& SharedHandle, TArray<uint8>& FileContents) override;
+	virtual bool ClearSharedFiles() override;
+	virtual bool ClearSharedFile(const FSharedContentHandle& SharedHandle) override;
+	virtual bool ReadSharedFile(const FSharedContentHandle& SharedHandle) override;
+	virtual bool WriteSharedFile(const FUniqueNetId& UserId, const FString& Filename, TArray<uint8>& FileContents) override;
+	virtual void GetDummySharedHandlesForTest(TArray< TSharedRef<FSharedContentHandle> > & OutHandles) override;
 };
 
 typedef TSharedPtr<FOnlineSharedCloudSteam, ESPMode::ThreadSafe> FOnlineSharedCloudSteamPtr;
diff --git a/Engine/Source/Runtime/Online/OnlineSubsystemSteam/Private/OnlineSubsystemSteamTypes.h b/Engine/Source/Runtime/Online/OnlineSubsystemSteam/Private/OnlineSubsystemSteamTypes.h
--- a/Engine/Source/Runtime/Online/OnlineSubsystemSteam/Private/OnlineSubsystemSteamTypes.h
+++ b/Engine/Source/Runtime/Online/OnlineSubsystemSteam/Private/OnlineSubsystemSteamTypes.h
@@ -127,7 +127,7 @@ public:
 	 *
 	 * @return byte array of size GetSize()
 	 */
-	virtual const uint8* GetBytes() const OVERRIDE
+	virtual const uint8* GetBytes() const override
 	{
 		return (uint8*)&UniqueNetId;
 	}
@@ -137,7 +137,7 @@ public:
 	 *
 	 * @return size in bytes of the id representation
 	 */
-	virtual int32 GetSize() const OVERRIDE
+	virtual int32 GetSize() const override
 	{
 		return sizeof(uint64);
 	}
@@ -147,7 +147,7 @@ public:
 	 *
 	 * @return true if this is a well formed ID, false otherwise
 	 */
-	virtual bool IsValid() const OVERRIDE
+	virtual bool IsValid() const override
 	{
 		return UniqueNetId != 0 && CSteamID(UniqueNetId).IsValid();
 	}
@@ -157,7 +157,7 @@ public:
 	 *
 	 * @return data in string form 
 	 */
-	virtual FString ToString() const OVERRIDE
+	virtual FString ToString() const override
 	{
 		return FString::Printf(TEXT("%llu"), UniqueNetId);
 	}
@@ -168,7 +168,7 @@ public:
 	 *
 	 * @return id in string form 
 	 */
-	virtual FString ToDebugString() const OVERRIDE
+	virtual FString ToDebugString() const override
 	{
 		CSteamID SteamID(UniqueNetId);
 		if (SteamID.IsLobby())
@@ -278,12 +278,12 @@ public:
  		return false;
  	}
 
-	virtual const uint8* GetBytes() const OVERRIDE
+	virtual const uint8* GetBytes() const override
 	{
 		return NULL;
 	}
 
-	virtual int32 GetSize() const OVERRIDE
+	virtual int32 GetSize() const override
 	{
 		return sizeof(uint64) + 
 			sizeof(ESteamSession::Type) +
@@ -292,7 +292,7 @@ public:
 			sizeof(FUniqueNetIdSteam);
 	}
 
-	virtual bool IsValid() const OVERRIDE
+	virtual bool IsValid() const override
 	{
 		switch (SessionType)
 		{
@@ -309,12 +309,12 @@ public:
 		}
 	}
 
-	virtual FString ToString() const OVERRIDE
+	virtual FString ToString() const override
 	{
 		return SessionId.ToString();
 	}
 
-	virtual FString ToDebugString() const OVERRIDE
+	virtual FString ToDebugString() const override
 	{
 		return FString::Printf(TEXT("HostIP: %s SteamP2P: %s Type: %s SessionId: %s"), 
 			HostAddr.IsValid() ? *HostAddr->ToString(true) : TEXT("INVALID"), 
@@ -322,7 +322,7 @@ public:
 			ESteamSession::ToString(SessionType), *SessionId.ToDebugString());
 	}
 
-	virtual const FUniqueNetId& GetSessionId() const OVERRIDE
+	virtual const FUniqueNetId& GetSessionId() const override
 	{
 		return SessionId;
 	}
@@ -370,7 +370,7 @@ public:
 	 *
 	 * @return byte array of size GetSize()
 	 */
-	virtual const uint8* GetBytes() const OVERRIDE
+	virtual const uint8* GetBytes() const override
 	{
 		return (uint8*)&SharedContentHandle;
 	}
@@ -380,7 +380,7 @@ public:
 	 *
 	 * @return size in bytes of the id representation
 	 */
-	virtual int32 GetSize() const OVERRIDE
+	virtual int32 GetSize() const override
 	{
 		return sizeof(UGCHandle_t);
 	}
@@ -390,7 +390,7 @@ public:
 	 *
 	 * @return true if this is a well formed ID, false otherwise
 	 */
-	virtual bool IsValid() const OVERRIDE
+	virtual bool IsValid() const override
 	{
 		return SharedContentHandle != k_UGCHandleInvalid;
 	}
@@ -400,7 +400,7 @@ public:
 	 *
 	 * @return data in string form 
 	 */
-	virtual FString ToString() const OVERRIDE
+	virtual FString ToString() const override
 	{
 		return FString::Printf(TEXT("%lld"), SharedContentHandle);
 	}
@@ -411,7 +411,7 @@ public:
 	 *
 	 * @return handle in string form 
 	 */
-	virtual FString ToDebugString() const OVERRIDE
+	virtual FString ToDebugString() const override
 	{
 		return FString::Printf(TEXT("0x%llX"), SharedContentHandle);
 	}
diff --git a/Engine/Source/Runtime/Online/OnlineSubsystemSteam/Private/OnlineUserCloudInterfaceSteam.h b/Engine/Source/Runtime/Online/OnlineSubsystemSteam/Private/OnlineUserCloudInterfaceSteam.h
--- a/Engine/Source/Runtime/Online/OnlineSubsystemSteam/Private/OnlineUserCloudInterfaceSteam.h
+++ b/Engine/Source/Runtime/Online/OnlineSubsystemSteam/Private/OnlineUserCloudInterfaceSteam.h
@@ -32,18 +32,18 @@ public:
 	/**
 	 *	Get a human readable description of task
 	 */
-	virtual FString ToString() const OVERRIDE;
+	virtual FString ToString() const override;
 
 	/**
 	 * Give the async task time to do its work
 	 * Can only be called on the async task manager thread
 	 */
-	virtual void Tick() OVERRIDE;
+	virtual void Tick() override;
 	
 	/**
 	 *	Async task is given a chance to trigger it's delegates
 	 */
-	virtual void TriggerDelegates() OVERRIDE;
+	virtual void TriggerDelegates() override;
 };
 
 /** 
@@ -76,18 +76,18 @@ public:
 	/**
 	 *	Get a human readable description of task
 	 */
-	virtual FString ToString() const OVERRIDE;
+	virtual FString ToString() const override;
 
 	/**
 	 * Give the async task time to do its work
 	 * Can only be called on the async task manager thread
 	 */
-	virtual void Tick() OVERRIDE;
+	virtual void Tick() override;
 	
 	/**
 	 *	Async task is given a chance to trigger it's delegates
 	 */
-	virtual void TriggerDelegates() OVERRIDE;
+	virtual void TriggerDelegates() override;
 };
 
 /** 
@@ -134,18 +134,18 @@ public:
 	/**
 	 *	Get a human readable description of task
 	 */
-	virtual FString ToString() const OVERRIDE;
+	virtual FString ToString() const override;
 
 	/**
 	 * Give the async task time to do its work
 	 * Can only be called on the async task manager thread
 	 */
-	virtual void Tick() OVERRIDE;
+	virtual void Tick() override;
 	
 	/**
 	 *	Async task is given a chance to trigger it's delegates
 	 */
-	virtual void TriggerDelegates() OVERRIDE;
+	virtual void TriggerDelegates() override;
 };
 
 /** 
@@ -184,18 +184,18 @@ public:
 	/**
 	 *	Get a human readable description of task
 	 */
-	virtual FString ToString() const OVERRIDE;
+	virtual FString ToString() const override;
 
 	/**
 	 * Give the async task time to do its work
 	 * Can only be called on the async task manager thread
 	 */
-	virtual void Tick() OVERRIDE;
+	virtual void Tick() override;
 	
 	/**
 	 *	Async task is given a chance to trigger it's delegates
 	 */
-	virtual void TriggerDelegates() OVERRIDE;
+	virtual void TriggerDelegates() override;
 };
 
 /**
@@ -224,16 +224,16 @@ public:
 	virtual ~FOnlineUserCloudSteam();
 
 	// IOnlineUserCloud
-	virtual bool GetFileContents(const FUniqueNetId& UserId, const FString& FileName, TArray<uint8>& FileContents) OVERRIDE;
-	virtual bool ClearFiles(const FUniqueNetId& UserId) OVERRIDE;
-	virtual bool ClearFile(const FUniqueNetId& UserId, const FString& FileName) OVERRIDE;
-	virtual void EnumerateUserFiles(const FUniqueNetId& UserId) OVERRIDE;
-	virtual void GetUserFileList(const FUniqueNetId& UserId, TArray<FCloudFileHeader>& UserFiles) OVERRIDE;
-	virtual bool ReadUserFile(const FUniqueNetId& UserId, const FString& FileName) OVERRIDE;
-	virtual bool WriteUserFile(const FUniqueNetId& UserId, const FString& FileName, TArray<uint8>& FileContents) OVERRIDE;
-	virtual bool DeleteUserFile(const FUniqueNetId& UserId, const FString& FileName, bool bShouldCloudDelete, bool bShouldLocallyDelete) OVERRIDE;
-	virtual void DumpCloudState(const FUniqueNetId& UserId) OVERRIDE;
-	virtual void DumpCloudFileState(const FUniqueNetId& UserId, const FString& FileName) OVERRIDE;
+	virtual bool GetFileContents(const FUniqueNetId& UserId, const FString& FileName, TArray<uint8>& FileContents) override;
+	virtual bool ClearFiles(const FUniqueNetId& UserId) override;
+	virtual bool ClearFile(const FUniqueNetId& UserId, const FString& FileName) override;
+	virtual void EnumerateUserFiles(const FUniqueNetId& UserId) override;
+	virtual void GetUserFileList(const FUniqueNetId& UserId, TArray<FCloudFileHeader>& UserFiles) override;
+	virtual bool ReadUserFile(const FUniqueNetId& UserId, const FString& FileName) override;
+	virtual bool WriteUserFile(const FUniqueNetId& UserId, const FString& FileName, TArray<uint8>& FileContents) override;
+	virtual bool DeleteUserFile(const FUniqueNetId& UserId, const FString& FileName, bool bShouldCloudDelete, bool bShouldLocallyDelete) override;
+	virtual void DumpCloudState(const FUniqueNetId& UserId) override;
+	virtual void DumpCloudFileState(const FUniqueNetId& UserId, const FString& FileName) override;
 
 };
 
diff --git a/Engine/Source/Runtime/Online/OnlineSubsystemSteam/Private/SocketSubsystemSteam.h b/Engine/Source/Runtime/Online/OnlineSubsystemSteam/Private/SocketSubsystemSteam.h
--- a/Engine/Source/Runtime/Online/OnlineSubsystemSteam/Private/SocketSubsystemSteam.h
+++ b/Engine/Source/Runtime/Online/OnlineSubsystemSteam/Private/SocketSubsystemSteam.h
@@ -183,12 +183,12 @@ public:
 	 *
 	 * @return true if initialized ok, false otherwise
 	 */
-	virtual bool Init(FString& Error) OVERRIDE;
+	virtual bool Init(FString& Error) override;
 
 	/**
 	 * Performs platform specific socket clean up
 	 */
-	virtual void Shutdown() OVERRIDE;
+	virtual void Shutdown() override;
 
 	/**
 	 * Creates a socket
@@ -199,14 +199,14 @@ public:
 	 *
 	 * @return the new socket or NULL if failed
 	 */
-	virtual class FSocket* CreateSocket(const FName& SocketType, const FString& SocketDescription, bool bForceUDP = false) OVERRIDE;
+	virtual class FSocket* CreateSocket(const FName& SocketType, const FString& SocketDescription, bool bForceUDP = false) override;
 
 	/**
 	 * Cleans up a socket class
 	 *
 	 * @param Socket the socket object to destroy
 	 */
-	virtual void DestroySocket(class FSocket* Socket) OVERRIDE;
+	virtual void DestroySocket(class FSocket* Socket) override;
 
 	/**
 	 * Does a DNS look up of a host name
@@ -214,14 +214,14 @@ public:
 	 * @param HostName the name of the host to look up
 	 * @param OutAddr the address to copy the IP address to
 	 */
-	virtual ESocketErrors GetHostByName(const ANSICHAR* HostName, FInternetAddr& OutAddr) OVERRIDE;
+	virtual ESocketErrors GetHostByName(const ANSICHAR* HostName, FInternetAddr& OutAddr) override;
 
 	/**
 	 * Some platforms require chat data (voice, text, etc.) to be placed into
 	 * packets in a special way. This function tells the net connection
 	 * whether this is required for this platform
 	 */
-	virtual bool RequiresChatDataBeSeparate() OVERRIDE
+	virtual bool RequiresChatDataBeSeparate() override
 	{
 		return false;
 	}
@@ -230,7 +230,7 @@ public:
 	 * Some platforms require packets be encrypted. This function tells the
 	 * net connection whether this is required for this platform
 	 */
-	virtual bool RequiresEncryptedPackets() OVERRIDE
+	virtual bool RequiresEncryptedPackets() override
 	{
 		return false;
 	}
@@ -242,35 +242,35 @@ public:
 	 *
 	 * @return true if successful, false otherwise
 	 */
-	virtual bool GetHostName(FString& HostName) OVERRIDE;
+	virtual bool GetHostName(FString& HostName) override;
 
 	/**
 	 *	Create a proper FInternetAddr representation
 	 * @param Address host address
 	 * @param Port host port
 	 */
-	virtual TSharedRef<FInternetAddr> CreateInternetAddr(uint32 Address=0, uint32 Port=0) OVERRIDE;
+	virtual TSharedRef<FInternetAddr> CreateInternetAddr(uint32 Address=0, uint32 Port=0) override;
 
 	/**
 	 * @return Whether the machine has a properly configured network device or not
 	 */
-	virtual bool HasNetworkDevice() OVERRIDE;
+	virtual bool HasNetworkDevice() override;
 
 	/**
 	 *	Get the name of the socket subsystem
 	 * @return a string naming this subsystem
 	 */
-	virtual const TCHAR* GetSocketAPIName() const OVERRIDE;
+	virtual const TCHAR* GetSocketAPIName() const override;
 
 	/**
 	 * Returns the last error that has happened
 	 */
-	virtual ESocketErrors GetLastErrorCode() OVERRIDE;
+	virtual ESocketErrors GetLastErrorCode() override;
 
 	/**
 	 * Translates the platform error code to a ESocketErrors enum
 	 */
-	virtual ESocketErrors TranslateErrorCode(int32 Code) OVERRIDE;
+	virtual ESocketErrors TranslateErrorCode(int32 Code) override;
 
 	/**
 	 * Gets the list of addresses associated with the adapters on the local computer.
@@ -279,7 +279,7 @@ public:
 	 *
 	 * @return true on success, false otherwise.
 	 */
-	virtual bool GetLocalAdapterAddresses( TArray<TSharedPtr<FInternetAddr> >& OutAdresses ) OVERRIDE
+	virtual bool GetLocalAdapterAddresses( TArray<TSharedPtr<FInternetAddr> >& OutAdresses ) override
 	{
 		bool bCanBindAll;
 
@@ -291,7 +291,7 @@ public:
 	/**
 	 *	Get local IP to bind to
 	 */
-	virtual TSharedRef<FInternetAddr> GetLocalBindAddr(FOutputDevice& Out) OVERRIDE;
+	virtual TSharedRef<FInternetAddr> GetLocalBindAddr(FOutputDevice& Out) override;
 
 	/**
 	 * Dumps the Steam P2P networking information for a given session state
@@ -305,7 +305,7 @@ public:
 	 *
 	 * @param DeltaTime time since last tick
 	 */
-	virtual bool Tick(float DeltaTime) OVERRIDE;
+	virtual bool Tick(float DeltaTime) override;
 };
 
 /**
diff --git a/Engine/Source/Runtime/Online/OnlineSubsystemSteam/Private/SocketsSteam.h b/Engine/Source/Runtime/Online/OnlineSubsystemSteam/Private/SocketsSteam.h
--- a/Engine/Source/Runtime/Online/OnlineSubsystemSteam/Private/SocketsSteam.h
+++ b/Engine/Source/Runtime/Online/OnlineSubsystemSteam/Private/SocketsSteam.h
@@ -70,7 +70,7 @@ public:
 	 *
 	 * @param true if it closes without errors, false otherwise
 	 */
-	virtual bool Close() OVERRIDE;
+	virtual bool Close() override;
 
 	/**
 	 * Binds a socket to a network byte ordered address
@@ -79,7 +79,7 @@ public:
 	 *
 	 * @return true if successful, false otherwise
 	 */
-	virtual bool Bind(const FInternetAddr& Addr) OVERRIDE;
+	virtual bool Bind(const FInternetAddr& Addr) override;
 
 	/**
 	 * Connects a socket to a network byte ordered address
@@ -88,7 +88,7 @@ public:
 	 *
 	 * @return true if successful, false otherwise
 	 */
-	virtual bool Connect(const FInternetAddr& Addr) OVERRIDE;
+	virtual bool Connect(const FInternetAddr& Addr) override;
 
 	/**
 	 * Places the socket into a state to listen for incoming connections
@@ -97,7 +97,7 @@ public:
 	 *
 	 * @return true if successful, false otherwise
 	 */
-	virtual bool Listen(int32 MaxBacklog) OVERRIDE;
+	virtual bool Listen(int32 MaxBacklog) override;
 
 	/**
 	 * Queries the socket to determine if there is a pending connection
@@ -106,7 +106,7 @@ public:
 	 *
 	 * @return true if successful, false otherwise
 	 */
-	virtual bool HasPendingConnection(bool& bHasPendingConnection) OVERRIDE;
+	virtual bool HasPendingConnection(bool& bHasPendingConnection) override;
 
 	/**
 	* Queries the socket to determine if there is pending data on the queue
@@ -115,7 +115,7 @@ public:
 	*
 	* @return true if the socket has data, false otherwise
 	*/
-	virtual bool HasPendingData(uint32& PendingDataSize) OVERRIDE;
+	virtual bool HasPendingData(uint32& PendingDataSize) override;
 
 	/**
 	 * Accepts a connection that is pending
@@ -123,7 +123,7 @@ public:
 	 * @param		SocketDescription debug description of socket
 	 * @return		The new (heap-allocated) socket, or NULL if unsuccessful.
 	 */
-	virtual class FSocket* Accept(const FString& SocketDescription) OVERRIDE;
+	virtual class FSocket* Accept(const FString& SocketDescription) override;
 
 	/**
 	 * Accepts a connection that is pending
@@ -133,7 +133,7 @@ public:
 	 *
 	 * @return		The new (heap-allocated) socket, or NULL if unsuccessful.
 	 */
-	virtual class FSocket* Accept(FInternetAddr& OutAddr, const FString& SocketDescription) OVERRIDE;
+	virtual class FSocket* Accept(FInternetAddr& OutAddr, const FString& SocketDescription) override;
 
 	/**
 	 * Sends a buffer to a network byte ordered address
@@ -143,7 +143,7 @@ public:
 	 * @param BytesSent out param indicating how much was sent
 	 * @param Destination the network byte ordered address to send to
 	 */
-	virtual bool SendTo(const uint8* Data, int32 Count, int32& BytesSent, const FInternetAddr& Destination) OVERRIDE;
+	virtual bool SendTo(const uint8* Data, int32 Count, int32& BytesSent, const FInternetAddr& Destination) override;
 
 	/**
 	 * Sends a buffer on a connected socket
@@ -152,7 +152,7 @@ public:
 	 * @param Count the size of the data to send
 	 * @param BytesSent out param indicating how much was sent
 	 */
-	virtual bool Send(const uint8* Data, int32 Count, int32& BytesSent) OVERRIDE;
+	virtual bool Send(const uint8* Data, int32 Count, int32& BytesSent) override;
 
 	/**
 	 * Reads a chunk of data from the socket. Gathers the source address too
@@ -163,7 +163,7 @@ public:
 	 * @param Source out param receiving the address of the sender of the data
 	 * @param Flags the receive flags (must be ESocketReceiveFlags::None)
 	 */
-	virtual bool RecvFrom(uint8* Data, int32 BufferSize, int32& BytesRead, FInternetAddr& Source, ESocketReceiveFlags::Type Flags = ESocketReceiveFlags::None) OVERRIDE;
+	virtual bool RecvFrom(uint8* Data, int32 BufferSize, int32& BytesRead, FInternetAddr& Source, ESocketReceiveFlags::Type Flags = ESocketReceiveFlags::None) override;
 
 	/**
 	 * Reads a chunk of data from a connected socket
@@ -173,21 +173,21 @@ public:
 	 * @param BytesRead out param indicating how many bytes were read from the socket
 	 * @param Flags the receive flags
 	 */
-	virtual bool Recv(uint8* Data,int32 BufferSize,int32& BytesRead, ESocketReceiveFlags::Type Flags = ESocketReceiveFlags::None) OVERRIDE;
+	virtual bool Recv(uint8* Data,int32 BufferSize,int32& BytesRead, ESocketReceiveFlags::Type Flags = ESocketReceiveFlags::None) override;
 
-	virtual bool Wait(ESocketWaitConditions::Type Condition, FTimespan WaitTime) OVERRIDE;
+	virtual bool Wait(ESocketWaitConditions::Type Condition, FTimespan WaitTime) override;
 
 	/**
 	 * Determines the connection state of the socket
 	 */
-	virtual ESocketConnectionState GetConnectionState() OVERRIDE;
+	virtual ESocketConnectionState GetConnectionState() override;
 
 	/**
 	 * Reads the address the socket is bound to and returns it
 	 * 
 	 * @param OutAddr address the socket is bound to
 	 */
-	virtual void GetAddress(FInternetAddr& OutAddr) OVERRIDE;
+	virtual void GetAddress(FInternetAddr& OutAddr) override;
 
 	/**
 	 * Sets this socket into non-blocking mode
@@ -196,7 +196,7 @@ public:
 	 *
 	 * @return true if successful, false otherwise
 	 */
-	virtual bool SetNonBlocking(bool bIsNonBlocking = true) OVERRIDE;
+	virtual bool SetNonBlocking(bool bIsNonBlocking = true) override;
 
 	/**
 	 * Sets a socket into broadcast mode (UDP only)
@@ -205,15 +205,15 @@ public:
 	 *
 	 * @return true if successful, false otherwise
 	 */
-	virtual bool SetBroadcast(bool bAllowBroadcast = true) OVERRIDE;
+	virtual bool SetBroadcast(bool bAllowBroadcast = true) override;
 
-	virtual bool JoinMulticastGroup (const FInternetAddr& GroupAddress) OVERRIDE;
+	virtual bool JoinMulticastGroup (const FInternetAddr& GroupAddress) override;
 
-	virtual bool LeaveMulticastGroup (const FInternetAddr& GroupAddress) OVERRIDE;
+	virtual bool LeaveMulticastGroup (const FInternetAddr& GroupAddress) override;
 
-	virtual bool SetMulticastLoopback (bool bLoopback) OVERRIDE;
+	virtual bool SetMulticastLoopback (bool bLoopback) override;
 
-	virtual bool SetMulticastTtl (uint8 TimeToLive) OVERRIDE;
+	virtual bool SetMulticastTtl (uint8 TimeToLive) override;
 
 	/**
 	 * Sets whether a socket can be bound to an address in use
@@ -222,7 +222,7 @@ public:
 	 *
 	 * @return true if the call succeeded, false otherwise
 	 */
-	virtual bool SetReuseAddr(bool bAllowReuse = true) OVERRIDE;
+	virtual bool SetReuseAddr(bool bAllowReuse = true) override;
 
 	/**
 	 * Sets whether and how long a socket will linger after closing
@@ -232,7 +232,7 @@ public:
 	 *
 	 * @return true if the call succeeded, false otherwise
 	 */
-	virtual bool SetLinger(bool bShouldLinger = true, int32 Timeout = 0) OVERRIDE;
+	virtual bool SetLinger(bool bShouldLinger = true, int32 Timeout = 0) override;
 
 	/**
 	 * Enables error queue support for the socket
@@ -241,7 +241,7 @@ public:
 	 *
 	 * @return true if the call succeeded, false otherwise
 	 */
-	virtual bool SetRecvErr(bool bUseErrorQueue = true) OVERRIDE;
+	virtual bool SetRecvErr(bool bUseErrorQueue = true) override;
 
 	/**
 	 * Sets the size of the send buffer to use
@@ -251,7 +251,7 @@ public:
 	 *
 	 * @return true if the call succeeded, false otherwise
 	 */
-	virtual bool SetSendBufferSize(int32 Size,int32& NewSize) OVERRIDE;
+	virtual bool SetSendBufferSize(int32 Size,int32& NewSize) override;
 
 	/**
 	 * Sets the size of the receive buffer to use
@@ -261,10 +261,10 @@ public:
 	 *
 	 * @return true if the call succeeded, false otherwise
 	 */
-	virtual bool SetReceiveBufferSize(int32 Size,int32& NewSize) OVERRIDE;
+	virtual bool SetReceiveBufferSize(int32 Size,int32& NewSize) override;
 
 	/**
 	 * Reads the port this socket is bound to.
 	 */ 
-	virtual int32 GetPortNo() OVERRIDE;
+	virtual int32 GetPortNo() override;
 };
\ No newline at end of file
diff --git a/Engine/Source/Runtime/Online/OnlineSubsystemSteam/Private/VoiceEngineSteam.h b/Engine/Source/Runtime/Online/OnlineSubsystemSteam/Private/VoiceEngineSteam.h
--- a/Engine/Source/Runtime/Online/OnlineSubsystemSteam/Private/VoiceEngineSteam.h
+++ b/Engine/Source/Runtime/Online/OnlineSubsystemSteam/Private/VoiceEngineSteam.h
@@ -47,7 +47,7 @@ class FVoiceEngineSteam : public IVoiceEngine
  		~FVoiceSerializeHelper() {}
  		
  		/** FGCObject interface */
- 		virtual void AddReferencedObjects(FReferenceCollector& Collector) OVERRIDE
+ 		virtual void AddReferencedObjects(FReferenceCollector& Collector) override
  		{
  			// Prevent garbage collection of audio components
  			for (FRemoteTalkerData::TIterator It(VoiceEngine->RemoteTalkerBuffers); It; ++It)
@@ -138,7 +138,7 @@ PACKAGE_SCOPE:
 	{};
 
 	// IVoiceEngine
-	virtual bool Init(int32 MaxLocalTalkers, int32 MaxRemoteTalkers) OVERRIDE;
+	virtual bool Init(int32 MaxLocalTalkers, int32 MaxRemoteTalkers) override;
 
 public:
 
@@ -146,20 +146,20 @@ public:
 	virtual ~FVoiceEngineSteam();
 
 	// IVoiceEngine
-    virtual uint32 StartLocalVoiceProcessing(uint32 LocalUserNum) OVERRIDE;
-    virtual uint32 StopLocalVoiceProcessing(uint32 LocalUserNum) OVERRIDE;
-    virtual uint32 StartRemoteVoiceProcessing(const FUniqueNetId& UniqueId) OVERRIDE
+    virtual uint32 StartLocalVoiceProcessing(uint32 LocalUserNum) override;
+    virtual uint32 StopLocalVoiceProcessing(uint32 LocalUserNum) override;
+    virtual uint32 StartRemoteVoiceProcessing(const FUniqueNetId& UniqueId) override
 	{
 		// Not needed in Steam
 		return S_OK;
 	}
-    virtual uint32 StopRemoteVoiceProcessing(const FUniqueNetId& UniqueId) OVERRIDE
+    virtual uint32 StopRemoteVoiceProcessing(const FUniqueNetId& UniqueId) override
 	{
 		// Not needed in Steam
 		return S_OK;
 	}
 
-    virtual uint32 RegisterLocalTalker(uint32 LocalUserNum) OVERRIDE
+    virtual uint32 RegisterLocalTalker(uint32 LocalUserNum) override
 	{
 		if (OwningUserIndex == INVALID_INDEX)
 		{
@@ -170,7 +170,7 @@ public:
 		return E_FAIL;
 	}
 
-    virtual uint32 UnregisterLocalTalker(uint32 LocalUserNum) OVERRIDE
+    virtual uint32 UnregisterLocalTalker(uint32 LocalUserNum) override
 	{
 		if (IsOwningUser(LocalUserNum))
 		{
@@ -181,44 +181,44 @@ public:
 		return E_FAIL;
 	}
 
-    virtual uint32 RegisterRemoteTalker(const FUniqueNetId& UniqueId) OVERRIDE
+    virtual uint32 RegisterRemoteTalker(const FUniqueNetId& UniqueId) override
 	{
 		// Not needed in Steam
 		return S_OK;
 	}
 
-    virtual uint32 UnregisterRemoteTalker(const FUniqueNetId& UniqueId) OVERRIDE
+    virtual uint32 UnregisterRemoteTalker(const FUniqueNetId& UniqueId) override
 	{
 		// Not needed in Steam
 		return S_OK;
 	}
 
-    virtual bool IsHeadsetPresent(uint32 LocalUserNum) OVERRIDE
+    virtual bool IsHeadsetPresent(uint32 LocalUserNum) override
 	{
 		return IsOwningUser(LocalUserNum) ? true : false;
 	}
 
-    virtual bool IsLocalPlayerTalking(uint32 LocalUserNum) OVERRIDE
+    virtual bool IsLocalPlayerTalking(uint32 LocalUserNum) override
 	{
 		return (GetVoiceDataReadyFlags() & (LocalUserNum << 1)) != 0;
 	}
 
-	virtual bool IsRemotePlayerTalking(const FUniqueNetId& UniqueId) OVERRIDE
+	virtual bool IsRemotePlayerTalking(const FUniqueNetId& UniqueId) override
 	{
 		return RemoteTalkerBuffers.Find((const FUniqueNetIdSteam&)UniqueId) != NULL;
 	}
 
-	virtual uint32 GetVoiceDataReadyFlags() const OVERRIDE;
-	virtual uint32 SetPlaybackPriority(uint32 LocalUserNum, const FUniqueNetId& RemoteTalkerId, uint32 Priority) OVERRIDE
+	virtual uint32 GetVoiceDataReadyFlags() const override;
+	virtual uint32 SetPlaybackPriority(uint32 LocalUserNum, const FUniqueNetId& RemoteTalkerId, uint32 Priority) override
 	{
 		// Not supported
 		return S_OK;
 	}
 
-	virtual uint32 ReadLocalVoiceData(uint32 LocalUserNum, uint8* Data, uint32* Size) OVERRIDE;
-	virtual uint32 SubmitRemoteVoiceData(const FUniqueNetId& RemoteTalkerId, uint8* Data, uint32* Size) OVERRIDE;
-	virtual void Tick(float DeltaTime) OVERRIDE;
-	FString GetVoiceDebugState() const OVERRIDE;
+	virtual uint32 ReadLocalVoiceData(uint32 LocalUserNum, uint8* Data, uint32* Size) override;
+	virtual uint32 SubmitRemoteVoiceData(const FUniqueNetId& RemoteTalkerId, uint8* Data, uint32* Size) override;
+	virtual void Tick(float DeltaTime) override;
+	FString GetVoiceDebugState() const override;
 
 	/**
 	 * Update the state of all remote talkers, possibly dropping data or the talker entirely
diff --git a/Engine/Source/Runtime/Online/OnlineSubsystemSteam/Private/VoiceInterfaceSteam.h b/Engine/Source/Runtime/Online/OnlineSubsystemSteam/Private/VoiceInterfaceSteam.h
--- a/Engine/Source/Runtime/Online/OnlineSubsystemSteam/Private/VoiceInterfaceSteam.h
+++ b/Engine/Source/Runtime/Online/OnlineSubsystemSteam/Private/VoiceInterfaceSteam.h
@@ -67,8 +67,8 @@ PACKAGE_SCOPE:
 	{};
 
 	// IOnlineVoice
-	virtual bool Init() OVERRIDE;
-	void ProcessMuteChangeNotification() OVERRIDE;
+	virtual bool Init() override;
+	void ProcessMuteChangeNotification() override;
 
 	/**
 	 * Processes any talking delegates that need to be fired off
@@ -104,27 +104,27 @@ public:
 	virtual ~FOnlineVoiceSteam();
 
 	// IOnlineVoice
-	virtual void StartNetworkedVoice(uint8 LocalUserNum) OVERRIDE;
-	virtual void StopNetworkedVoice(uint8 LocalUserNum) OVERRIDE;
-    virtual bool RegisterLocalTalker(uint32 LocalUserNum) OVERRIDE;
-	virtual void RegisterLocalTalkers() OVERRIDE;
-    virtual bool UnregisterLocalTalker(uint32 LocalUserNum) OVERRIDE;
-	virtual void UnregisterLocalTalkers() OVERRIDE;
-    virtual bool RegisterRemoteTalker(const FUniqueNetId& UniqueId) OVERRIDE;
-    virtual bool UnregisterRemoteTalker(const FUniqueNetId& UniqueId) OVERRIDE;
-	virtual void RemoveAllRemoteTalkers() OVERRIDE;
-    virtual bool IsHeadsetPresent(uint32 LocalUserNum) OVERRIDE;
-    virtual bool IsLocalPlayerTalking(uint32 LocalUserNum) OVERRIDE;
-	virtual bool IsRemotePlayerTalking(const FUniqueNetId& UniqueId) OVERRIDE;
-	bool IsMuted(uint32 LocalUserNum, const FUniqueNetId& UniqueId) const OVERRIDE;
-	bool MuteRemoteTalker(uint8 LocalUserNum, const FUniqueNetId& PlayerId, bool bIsSystemWide) OVERRIDE;
-	bool UnmuteRemoteTalker(uint8 LocalUserNum, const FUniqueNetId& PlayerId, bool bIsSystemWide) OVERRIDE;
-	virtual TSharedPtr<class FVoicePacket> SerializeRemotePacket(FArchive& Ar) OVERRIDE;
-	virtual TSharedPtr<class FVoicePacket> GetLocalPacket(uint32 LocalUserNum) OVERRIDE;
-	virtual int32 GetNumLocalTalkers() OVERRIDE { return LocalTalkers.Num(); };
-	virtual void ClearVoicePackets() OVERRIDE;
-	virtual void Tick(float DeltaTime) OVERRIDE;
-	virtual FString GetVoiceDebugState() const OVERRIDE;
+	virtual void StartNetworkedVoice(uint8 LocalUserNum) override;
+	virtual void StopNetworkedVoice(uint8 LocalUserNum) override;
+    virtual bool RegisterLocalTalker(uint32 LocalUserNum) override;
+	virtual void RegisterLocalTalkers() override;
+    virtual bool UnregisterLocalTalker(uint32 LocalUserNum) override;
+	virtual void UnregisterLocalTalkers() override;
+    virtual bool RegisterRemoteTalker(const FUniqueNetId& UniqueId) override;
+    virtual bool UnregisterRemoteTalker(const FUniqueNetId& UniqueId) override;
+	virtual void RemoveAllRemoteTalkers() override;
+    virtual bool IsHeadsetPresent(uint32 LocalUserNum) override;
+    virtual bool IsLocalPlayerTalking(uint32 LocalUserNum) override;
+	virtual bool IsRemotePlayerTalking(const FUniqueNetId& UniqueId) override;
+	bool IsMuted(uint32 LocalUserNum, const FUniqueNetId& UniqueId) const override;
+	bool MuteRemoteTalker(uint8 LocalUserNum, const FUniqueNetId& PlayerId, bool bIsSystemWide) override;
+	bool UnmuteRemoteTalker(uint8 LocalUserNum, const FUniqueNetId& PlayerId, bool bIsSystemWide) override;
+	virtual TSharedPtr<class FVoicePacket> SerializeRemotePacket(FArchive& Ar) override;
+	virtual TSharedPtr<class FVoicePacket> GetLocalPacket(uint32 LocalUserNum) override;
+	virtual int32 GetNumLocalTalkers() override { return LocalTalkers.Num(); };
+	virtual void ClearVoicePackets() override;
+	virtual void Tick(float DeltaTime) override;
+	virtual FString GetVoiceDebugState() const override;
 };
 
 typedef TSharedPtr<FOnlineVoiceSteam, ESPMode::ThreadSafe> FOnlineVoiceSteamPtr;
diff --git a/Engine/Source/Runtime/Online/OnlineSubsystemSteam/Private/VoicePacketSteam.h b/Engine/Source/Runtime/Online/OnlineSubsystemSteam/Private/VoicePacketSteam.h
--- a/Engine/Source/Runtime/Online/OnlineSubsystemSteam/Private/VoicePacketSteam.h
+++ b/Engine/Source/Runtime/Online/OnlineSubsystemSteam/Private/VoicePacketSteam.h
@@ -41,22 +41,22 @@ public:
 	FVoicePacketSteam(const FVoicePacketSteam& Other);
 
 	/** Returns the amount of space this packet will consume in a buffer */
-	virtual uint16 GetTotalPacketSize() OVERRIDE;
+	virtual uint16 GetTotalPacketSize() override;
 
 	/** @return the amount of space used by the internal voice buffer */
-	virtual uint16 GetBufferSize() OVERRIDE;
+	virtual uint16 GetBufferSize() override;
 
 	/** @return the sender of this voice packet */
-	virtual TSharedPtr<class FUniqueNetId> GetSender() OVERRIDE;
+	virtual TSharedPtr<class FUniqueNetId> GetSender() override;
 
-	virtual bool IsReliable() OVERRIDE { return false; }
+	virtual bool IsReliable() override { return false; }
 
 	/** 
 	 * Serialize the voice packet data to a buffer 
 	 *
 	 * @param Ar buffer to write into
 	 */
-	virtual void Serialize(class FArchive& Ar) OVERRIDE;
+	virtual void Serialize(class FArchive& Ar) override;
 };
 
 /** Holds the current voice packet data state */
diff --git a/Engine/Source/Runtime/Online/OnlineSubsystemSteam/Public/OnlineSubsystemSteam.h b/Engine/Source/Runtime/Online/OnlineSubsystemSteam/Public/OnlineSubsystemSteam.h
--- a/Engine/Source/Runtime/Online/OnlineSubsystemSteam/Public/OnlineSubsystemSteam.h
+++ b/Engine/Source/Runtime/Online/OnlineSubsystemSteam/Public/OnlineSubsystemSteam.h
@@ -211,34 +211,34 @@ public:
 
 	// IOnlineSubsystem
 
-	virtual IOnlineSessionPtr GetSessionInterface() const OVERRIDE;
-	virtual IOnlineFriendsPtr GetFriendsInterface() const OVERRIDE;
-	virtual IOnlineSharedCloudPtr GetSharedCloudInterface() const OVERRIDE;
-	virtual IOnlineUserCloudPtr GetUserCloudInterface() const OVERRIDE;
-	virtual IOnlineLeaderboardsPtr GetLeaderboardsInterface() const OVERRIDE;
-	virtual IOnlineVoicePtr GetVoiceInterface() const OVERRIDE;
-	virtual IOnlineExternalUIPtr GetExternalUIInterface() const OVERRIDE;
-	virtual IOnlineTimePtr GetTimeInterface() const OVERRIDE;
-	virtual IOnlineIdentityPtr GetIdentityInterface() const OVERRIDE;
-	virtual IOnlinePartyPtr GetPartyInterface() const OVERRIDE;
-	virtual IOnlineTitleFilePtr GetTitleFileInterface() const OVERRIDE;
-	virtual IOnlineEntitlementsPtr GetEntitlementsInterface() const OVERRIDE;
-	virtual IOnlineStorePtr GetStoreInterface() const OVERRIDE;
-	virtual IOnlineEventsPtr GetEventsInterface() const OVERRIDE;
-	virtual IOnlineAchievementsPtr GetAchievementsInterface() const OVERRIDE;
-	virtual IOnlineSharingPtr GetSharingInterface() const OVERRIDE;
-	virtual IOnlineUserPtr GetUserInterface() const OVERRIDE;
-	virtual IOnlineMessagePtr GetMessageInterface() const OVERRIDE;
-	virtual IOnlinePresencePtr GetPresenceInterface() const OVERRIDE;
-	virtual bool IsLocalPlayer(const FUniqueNetId& UniqueId) const OVERRIDE;
-	virtual bool Init() OVERRIDE;
-	virtual bool Shutdown() OVERRIDE;
-	virtual bool Exec(class UWorld* InWorld, const TCHAR* Cmd, FOutputDevice& Ar) OVERRIDE;
-	FString GetAppId() const OVERRIDE;
+	virtual IOnlineSessionPtr GetSessionInterface() const override;
+	virtual IOnlineFriendsPtr GetFriendsInterface() const override;
+	virtual IOnlineSharedCloudPtr GetSharedCloudInterface() const override;
+	virtual IOnlineUserCloudPtr GetUserCloudInterface() const override;
+	virtual IOnlineLeaderboardsPtr GetLeaderboardsInterface() const override;
+	virtual IOnlineVoicePtr GetVoiceInterface() const override;
+	virtual IOnlineExternalUIPtr GetExternalUIInterface() const override;
+	virtual IOnlineTimePtr GetTimeInterface() const override;
+	virtual IOnlineIdentityPtr GetIdentityInterface() const override;
+	virtual IOnlinePartyPtr GetPartyInterface() const override;
+	virtual IOnlineTitleFilePtr GetTitleFileInterface() const override;
+	virtual IOnlineEntitlementsPtr GetEntitlementsInterface() const override;
+	virtual IOnlineStorePtr GetStoreInterface() const override;
+	virtual IOnlineEventsPtr GetEventsInterface() const override;
+	virtual IOnlineAchievementsPtr GetAchievementsInterface() const override;
+	virtual IOnlineSharingPtr GetSharingInterface() const override;
+	virtual IOnlineUserPtr GetUserInterface() const override;
+	virtual IOnlineMessagePtr GetMessageInterface() const override;
+	virtual IOnlinePresencePtr GetPresenceInterface() const override;
+	virtual bool IsLocalPlayer(const FUniqueNetId& UniqueId) const override;
+	virtual bool Init() override;
+	virtual bool Shutdown() override;
+	virtual bool Exec(class UWorld* InWorld, const TCHAR* Cmd, FOutputDevice& Ar) override;
+	FString GetAppId() const override;
 
 	// FTickerObjectBase
 
-	virtual bool Tick(float DeltaTime) OVERRIDE;
+	virtual bool Tick(float DeltaTime) override;
 
 	// FOnlineSubsystemSteam
 
diff --git a/Engine/Source/Runtime/Online/OnlineSubsystemUtils/Classes/AchievementWriteCallbackProxy.h b/Engine/Source/Runtime/Online/OnlineSubsystemUtils/Classes/AchievementWriteCallbackProxy.h
--- a/Engine/Source/Runtime/Online/OnlineSubsystemUtils/Classes/AchievementWriteCallbackProxy.h
+++ b/Engine/Source/Runtime/Online/OnlineSubsystemUtils/Classes/AchievementWriteCallbackProxy.h
@@ -27,11 +27,11 @@ class UAchievementWriteCallbackProxy : public UOnlineBlueprintCallProxyBase
 	static UAchievementWriteCallbackProxy* WriteAchievementProgress(class APlayerController* PlayerController, FName AchievementName, float Progress = 100.0f, int32 UserTag = 0);
 
 	// UOnlineBlueprintCallProxyBase interface
-	virtual void Activate() OVERRIDE;
+	virtual void Activate() override;
 	// End of UOnlineBlueprintCallProxyBase interface
 
 	// UObject interface
-	virtual void BeginDestroy() OVERRIDE;
+	virtual void BeginDestroy() override;
 	// End of UObject interface
 
 private:
diff --git a/Engine/Source/Runtime/Online/OnlineSubsystemUtils/Classes/IpConnection.h b/Engine/Source/Runtime/Online/OnlineSubsystemUtils/Classes/IpConnection.h
--- a/Engine/Source/Runtime/Online/OnlineSubsystemUtils/Classes/IpConnection.h
+++ b/Engine/Source/Runtime/Online/OnlineSubsystemUtils/Classes/IpConnection.h
@@ -17,13 +17,13 @@ class ONLINESUBSYSTEMUTILS_API UIpConnection : public UNetConnection
 	class FResolveInfo*			ResolveInfo;
 
 	// Begin NetConnection Interface
-	virtual void InitBase(UNetDriver* InDriver, class FSocket* InSocket, const FURL& InURL, EConnectionState InState, int32 InMaxPacket = 0, int32 InPacketOverhead = 0) OVERRIDE;
-	virtual void InitRemoteConnection(UNetDriver* InDriver, class FSocket* InSocket, const FURL& InURL, const class FInternetAddr& InRemoteAddr, EConnectionState InState, int32 InMaxPacket = 0, int32 InPacketOverhead = 0) OVERRIDE;
-	virtual void InitLocalConnection(UNetDriver* InDriver, class FSocket* InSocket, const FURL& InURL, EConnectionState InState, int32 InMaxPacket = 0, int32 InPacketOverhead = 0) OVERRIDE;
-	virtual void LowLevelSend(void* Data,int32 Count) OVERRIDE;
-	FString LowLevelGetRemoteAddress(bool bAppendPort=false) OVERRIDE;
-	FString LowLevelDescribe() OVERRIDE;
-	virtual int32 GetAddrAsInt(void) OVERRIDE;
-	virtual int32 GetAddrPort(void) OVERRIDE;
+	virtual void InitBase(UNetDriver* InDriver, class FSocket* InSocket, const FURL& InURL, EConnectionState InState, int32 InMaxPacket = 0, int32 InPacketOverhead = 0) override;
+	virtual void InitRemoteConnection(UNetDriver* InDriver, class FSocket* InSocket, const FURL& InURL, const class FInternetAddr& InRemoteAddr, EConnectionState InState, int32 InMaxPacket = 0, int32 InPacketOverhead = 0) override;
+	virtual void InitLocalConnection(UNetDriver* InDriver, class FSocket* InSocket, const FURL& InURL, EConnectionState InState, int32 InMaxPacket = 0, int32 InPacketOverhead = 0) override;
+	virtual void LowLevelSend(void* Data,int32 Count) override;
+	FString LowLevelGetRemoteAddress(bool bAppendPort=false) override;
+	FString LowLevelDescribe() override;
+	virtual int32 GetAddrAsInt(void) override;
+	virtual int32 GetAddrPort(void) override;
 	// End NetConnection Interface
 };
diff --git a/Engine/Source/Runtime/Online/OnlineSubsystemUtils/Classes/IpNetDriver.h b/Engine/Source/Runtime/Online/OnlineSubsystemUtils/Classes/IpNetDriver.h
--- a/Engine/Source/Runtime/Online/OnlineSubsystemUtils/Classes/IpNetDriver.h
+++ b/Engine/Source/Runtime/Online/OnlineSubsystemUtils/Classes/IpNetDriver.h
@@ -31,16 +31,16 @@ class ONLINESUBSYSTEMUTILS_API UIpNetDriver : public UNetDriver
 	FSocket* Socket;
 
 	// Begin UNetDriver interface.
-	virtual bool IsAvailable() const OVERRIDE;
-	virtual bool InitBase(bool bInitAsClient, FNetworkNotify* InNotify, const FURL& URL, bool bReuseAddressAndPort, FString& Error) OVERRIDE;
-	virtual bool InitConnect( FNetworkNotify* InNotify, const FURL& ConnectURL, FString& Error ) OVERRIDE;
-	virtual bool InitListen( FNetworkNotify* InNotify, FURL& LocalURL, bool bReuseAddressAndPort, FString& Error ) OVERRIDE;
-	virtual void ProcessRemoteFunction(class AActor* Actor, class UFunction* Function, void* Parameters, struct FOutParmRec* OutParms, struct FFrame* Stack, class UObject * SubObject = NULL) OVERRIDE;
-	virtual void TickDispatch( float DeltaTime ) OVERRIDE;
-	virtual FString LowLevelGetNetworkNumber() OVERRIDE;
-	virtual void LowLevelDestroy() OVERRIDE;
-	virtual class ISocketSubsystem* GetSocketSubsystem() OVERRIDE;
-	virtual bool IsNetResourceValid(void) OVERRIDE
+	virtual bool IsAvailable() const override;
+	virtual bool InitBase(bool bInitAsClient, FNetworkNotify* InNotify, const FURL& URL, bool bReuseAddressAndPort, FString& Error) override;
+	virtual bool InitConnect( FNetworkNotify* InNotify, const FURL& ConnectURL, FString& Error ) override;
+	virtual bool InitListen( FNetworkNotify* InNotify, FURL& LocalURL, bool bReuseAddressAndPort, FString& Error ) override;
+	virtual void ProcessRemoteFunction(class AActor* Actor, class UFunction* Function, void* Parameters, struct FOutParmRec* OutParms, struct FFrame* Stack, class UObject * SubObject = NULL) override;
+	virtual void TickDispatch( float DeltaTime ) override;
+	virtual FString LowLevelGetNetworkNumber() override;
+	virtual void LowLevelDestroy() override;
+	virtual class ISocketSubsystem* GetSocketSubsystem() override;
+	virtual bool IsNetResourceValid(void) override
 	{
 		return Socket != NULL;
 	}
@@ -60,7 +60,7 @@ class ONLINESUBSYSTEMUTILS_API UIpNetDriver : public UNetDriver
 	// End UIpNetDriver interface.
 
 	// Begin FExec Interface
-	virtual bool Exec( UWorld* InWorld, const TCHAR* Cmd, FOutputDevice& Ar=*GLog ) OVERRIDE;
+	virtual bool Exec( UWorld* InWorld, const TCHAR* Cmd, FOutputDevice& Ar=*GLog ) override;
 	// End FExec Interface
 
 	/**
diff --git a/Engine/Source/Runtime/Online/OnlineSubsystemUtils/Classes/OnlineBeacon.h b/Engine/Source/Runtime/Online/OnlineSubsystemUtils/Classes/OnlineBeacon.h
--- a/Engine/Source/Runtime/Online/OnlineSubsystemUtils/Classes/OnlineBeacon.h
+++ b/Engine/Source/Runtime/Online/OnlineSubsystemUtils/Classes/OnlineBeacon.h
@@ -32,14 +32,14 @@ class ONLINESUBSYSTEMUTILS_API AOnlineBeacon : public AActor, public FNetworkNot
 	GENERATED_UCLASS_BODY()
 
 	// Begin AActor Interface
-	virtual void OnActorChannelOpen(class FInBunch& InBunch, class UNetConnection* Connection) OVERRIDE;
+	virtual void OnActorChannelOpen(class FInBunch& InBunch, class UNetConnection* Connection) override;
 	// End AActor Interface
 
 	// Begin FNetworkNotify Interface
-	virtual EAcceptConnection::Type NotifyAcceptingConnection() OVERRIDE;
-	virtual void NotifyAcceptedConnection(class UNetConnection* Connection) OVERRIDE;
-	virtual bool NotifyAcceptingChannel(class UChannel* Channel) OVERRIDE;
-	virtual void NotifyControlMessage(class UNetConnection* Connection, uint8 MessageType, class FInBunch& Bunch) OVERRIDE;
+	virtual EAcceptConnection::Type NotifyAcceptingConnection() override;
+	virtual void NotifyAcceptedConnection(class UNetConnection* Connection) override;
+	virtual bool NotifyAcceptingChannel(class UChannel* Channel) override;
+	virtual void NotifyControlMessage(class UNetConnection* Connection, uint8 MessageType, class FInBunch& Bunch) override;
 	// End FNetworkNotify Interface
 
 	/**
@@ -99,7 +99,7 @@ class ONLINESUBSYSTEMUTILS_API AOnlineBeacon : public AActor, public FNetworkNot
 	 *
 	 * @return net connection used in communication
 	 */
-	virtual UNetConnection* GetNetConnection() OVERRIDE;	
+	virtual UNetConnection* GetNetConnection() override;	
 
 protected:
 	/** Net driver routing network traffic */
diff --git a/Engine/Source/Runtime/Online/OnlineSubsystemUtils/Classes/OnlineBeaconClient.h b/Engine/Source/Runtime/Online/OnlineSubsystemUtils/Classes/OnlineBeaconClient.h
--- a/Engine/Source/Runtime/Online/OnlineSubsystemUtils/Classes/OnlineBeaconClient.h
+++ b/Engine/Source/Runtime/Online/OnlineSubsystemUtils/Classes/OnlineBeaconClient.h
@@ -20,11 +20,11 @@ class ONLINESUBSYSTEMUTILS_API AOnlineBeaconClient : public AOnlineBeacon
 	// End AActor Interface
 
 	// Begin FNetworkNotify Interface
-	virtual void NotifyControlMessage(class UNetConnection* Connection, uint8 MessageType, class FInBunch& Bunch) OVERRIDE;
+	virtual void NotifyControlMessage(class UNetConnection* Connection, uint8 MessageType, class FInBunch& Bunch) override;
 	// End FNetworkNotify Interface
 
 	// Begin OnlineBeacon Interface
-	virtual void DestroyBeacon() OVERRIDE;
+	virtual void DestroyBeacon() override;
 	// End OnlineBeacon Interface
 
 	/**
diff --git a/Engine/Source/Runtime/Online/OnlineSubsystemUtils/Classes/OnlineBeaconHost.h b/Engine/Source/Runtime/Online/OnlineSubsystemUtils/Classes/OnlineBeaconHost.h
--- a/Engine/Source/Runtime/Online/OnlineSubsystemUtils/Classes/OnlineBeaconHost.h
+++ b/Engine/Source/Runtime/Online/OnlineSubsystemUtils/Classes/OnlineBeaconHost.h
@@ -26,15 +26,15 @@ class ONLINESUBSYSTEMUTILS_API AOnlineBeaconHost : public AOnlineBeacon
 	int32 ListenPort;
 
 	// Begin AActor Interface
-	virtual void OnNetCleanup(class UNetConnection* Connection) OVERRIDE;
+	virtual void OnNetCleanup(class UNetConnection* Connection) override;
 	// End AActor Interface
 
 	// Begin OnlineBeacon Interface
-	virtual void HandleNetworkFailure(UWorld* World, class UNetDriver *NetDriver, ENetworkFailure::Type FailureType, const FString& ErrorString) OVERRIDE;
+	virtual void HandleNetworkFailure(UWorld* World, class UNetDriver *NetDriver, ENetworkFailure::Type FailureType, const FString& ErrorString) override;
 	// End OnlineBeacon Interface
 
 	// Begin FNetworkNotify Interface
-	virtual void NotifyControlMessage(class UNetConnection* Connection, uint8 MessageType, class FInBunch& Bunch) OVERRIDE;
+	virtual void NotifyControlMessage(class UNetConnection* Connection, uint8 MessageType, class FInBunch& Bunch) override;
 	// End FNetworkNotify Interface
 
 	/**
diff --git a/Engine/Source/Runtime/Online/OnlineSubsystemUtils/Classes/OnlineSessionClient.h b/Engine/Source/Runtime/Online/OnlineSubsystemUtils/Classes/OnlineSessionClient.h
--- a/Engine/Source/Runtime/Online/OnlineSubsystemUtils/Classes/OnlineSessionClient.h
+++ b/Engine/Source/Runtime/Online/OnlineSubsystemUtils/Classes/OnlineSessionClient.h
@@ -131,17 +131,17 @@ public:
 	/**
 	 * Register all delegates needed to manage online sessions
 	 */
-	virtual void RegisterOnlineDelegates(class UWorld* InWorld) OVERRIDE;
+	virtual void RegisterOnlineDelegates(class UWorld* InWorld) override;
 
 	/**
 	 * Tear down all delegates used to manage online sessions
 	 */
-	virtual void ClearOnlineDelegates(class UWorld* InWorld) OVERRIDE;
+	virtual void ClearOnlineDelegates(class UWorld* InWorld) override;
 
 	/**
 	 * Called to tear down any online sessions and return to main menu
 	 */
-	void HandleDisconnect(UWorld *World, UNetDriver *NetDriver) OVERRIDE;
+	void HandleDisconnect(UWorld *World, UNetDriver *NetDriver) override;
 
 };
 
diff --git a/Engine/Source/Runtime/Online/OnlineSubsystemUtils/Classes/PartyBeaconClient.h b/Engine/Source/Runtime/Online/OnlineSubsystemUtils/Classes/PartyBeaconClient.h
--- a/Engine/Source/Runtime/Online/OnlineSubsystemUtils/Classes/PartyBeaconClient.h
+++ b/Engine/Source/Runtime/Online/OnlineSubsystemUtils/Classes/PartyBeaconClient.h
@@ -25,12 +25,12 @@ class ONLINESUBSYSTEMUTILS_API APartyBeaconClient : public AOnlineBeaconClient
 	GENERATED_UCLASS_BODY()
 
 	// Begin AOnlineBeacon Interface
-	virtual FString GetBeaconType() OVERRIDE { return TEXT("PartyBeacon"); }
+	virtual FString GetBeaconType() override { return TEXT("PartyBeacon"); }
 	// End AOnlineBeacon Interface
 
 	// Begin AOnlineBeaconClient Interface
-	virtual void OnConnected() OVERRIDE;
-	virtual void OnFailure() OVERRIDE;
+	virtual void OnConnected() override;
+	virtual void OnFailure() override;
 	// End AOnlineBeaconClient Interface
 
 	/**
diff --git a/Engine/Source/Runtime/Online/OnlineSubsystemUtils/Classes/PartyBeaconHost.h b/Engine/Source/Runtime/Online/OnlineSubsystemUtils/Classes/PartyBeaconHost.h
--- a/Engine/Source/Runtime/Online/OnlineSubsystemUtils/Classes/PartyBeaconHost.h
+++ b/Engine/Source/Runtime/Online/OnlineSubsystemUtils/Classes/PartyBeaconHost.h
@@ -47,11 +47,11 @@ class ONLINESUBSYSTEMUTILS_API APartyBeaconHost : public AOnlineBeaconHostObject
 	GENERATED_UCLASS_BODY()
 
 	// Begin AActor Interface
-	virtual void Tick(float DeltaTime) OVERRIDE;
+	virtual void Tick(float DeltaTime) override;
 	// End AActor Interface
 
 	// Begin AOnlineBeaconHostObject Interface 
-	virtual AOnlineBeaconClient* SpawnBeaconActor(class UNetConnection* ClientConnection) OVERRIDE;
+	virtual AOnlineBeaconClient* SpawnBeaconActor(class UNetConnection* ClientConnection) override;
 	virtual void ClientConnected(class AOnlineBeaconClient* NewClientActor, class UNetConnection* ClientConnection);
 	virtual void RemoveClientActor(class AOnlineBeaconClient* ClientActor);
 	// End AOnlineBeaconHost Interface 
diff --git a/Engine/Source/Runtime/Online/OnlineSubsystemUtils/Classes/TestBeaconClient.h b/Engine/Source/Runtime/Online/OnlineSubsystemUtils/Classes/TestBeaconClient.h
--- a/Engine/Source/Runtime/Online/OnlineSubsystemUtils/Classes/TestBeaconClient.h
+++ b/Engine/Source/Runtime/Online/OnlineSubsystemUtils/Classes/TestBeaconClient.h
@@ -13,11 +13,11 @@ class ONLINESUBSYSTEMUTILS_API ATestBeaconClient : public AOnlineBeaconClient
 	GENERATED_UCLASS_BODY()
 
 	// Begin AOnlineBeacon Interface
-	virtual FString GetBeaconType() OVERRIDE { return TEXT("TestBeacon"); }
+	virtual FString GetBeaconType() override { return TEXT("TestBeacon"); }
 	// End AOnlineBeacon Interface
 
 	// Begin AOnlineBeaconClient Interface
-	virtual void OnFailure() OVERRIDE;
+	virtual void OnFailure() override;
 	// End AOnlineBeaconClient Interface
 
 	/** Send a ping RPC to the client */
diff --git a/Engine/Source/Runtime/Online/OnlineSubsystemUtils/Classes/TestBeaconHost.h b/Engine/Source/Runtime/Online/OnlineSubsystemUtils/Classes/TestBeaconHost.h
--- a/Engine/Source/Runtime/Online/OnlineSubsystemUtils/Classes/TestBeaconHost.h
+++ b/Engine/Source/Runtime/Online/OnlineSubsystemUtils/Classes/TestBeaconHost.h
@@ -13,7 +13,7 @@ class ONLINESUBSYSTEMUTILS_API ATestBeaconHost : public AOnlineBeaconHostObject
 	GENERATED_UCLASS_BODY()
 
 	// Begin AOnlineBeaconHost Interface 
-	virtual AOnlineBeaconClient* SpawnBeaconActor(class UNetConnection* ClientConnection) OVERRIDE;
+	virtual AOnlineBeaconClient* SpawnBeaconActor(class UNetConnection* ClientConnection) override;
 	virtual void ClientConnected(class AOnlineBeaconClient* NewClientActor, class UNetConnection* ClientConnection);
 	// End AOnlineBeaconHost Interface 
 
diff --git a/Engine/Source/Runtime/Online/OnlineSubsystemUtils/Private/Tests/TestCloudInterface.h b/Engine/Source/Runtime/Online/OnlineSubsystemUtils/Private/Tests/TestCloudInterface.h
--- a/Engine/Source/Runtime/Online/OnlineSubsystemUtils/Private/Tests/TestCloudInterface.h
+++ b/Engine/Source/Runtime/Online/OnlineSubsystemUtils/Private/Tests/TestCloudInterface.h
@@ -191,7 +191,7 @@
 
 	// FTickerObjectBase
 
-	bool Tick( float DeltaTime ) OVERRIDE;
+	bool Tick( float DeltaTime ) override;
 
 	// FTestCloudInterface
 
diff --git a/Engine/Source/Runtime/Online/OnlineSubsystemUtils/Private/Tests/TestLeaderboardInterface.h b/Engine/Source/Runtime/Online/OnlineSubsystemUtils/Private/Tests/TestLeaderboardInterface.h
--- a/Engine/Source/Runtime/Online/OnlineSubsystemUtils/Private/Tests/TestLeaderboardInterface.h
+++ b/Engine/Source/Runtime/Online/OnlineSubsystemUtils/Private/Tests/TestLeaderboardInterface.h
@@ -88,7 +88,7 @@
 
 	// FTickerObjectBase
 
-	bool Tick( float DeltaTime ) OVERRIDE;
+	bool Tick( float DeltaTime ) override;
 
 	// FTestLeaderboardInterface
 
diff --git a/Engine/Source/Runtime/Online/OnlineSubsystemUtils/Private/Tests/TestSessionInterface.h b/Engine/Source/Runtime/Online/OnlineSubsystemUtils/Private/Tests/TestSessionInterface.h
--- a/Engine/Source/Runtime/Online/OnlineSubsystemUtils/Private/Tests/TestSessionInterface.h
+++ b/Engine/Source/Runtime/Online/OnlineSubsystemUtils/Private/Tests/TestSessionInterface.h
@@ -280,10 +280,10 @@
 	}
 
 	// FTickerObjectBase
-	bool Tick(float DeltaTime) OVERRIDE;
+	bool Tick(float DeltaTime) override;
 
 	// FSelfRegisteringExec
-	virtual bool Exec(UWorld* InWorld, const TCHAR* Cmd, FOutputDevice& Ar) OVERRIDE;
+	virtual bool Exec(UWorld* InWorld, const TCHAR* Cmd, FOutputDevice& Ar) override;
 
 	// FTestSessionInterface
 
diff --git a/Engine/Source/Runtime/Online/OnlineSubsystemUtils/Private/Tests/TestVoice.h b/Engine/Source/Runtime/Online/OnlineSubsystemUtils/Private/Tests/TestVoice.h
--- a/Engine/Source/Runtime/Online/OnlineSubsystemUtils/Private/Tests/TestVoice.h
+++ b/Engine/Source/Runtime/Online/OnlineSubsystemUtils/Private/Tests/TestVoice.h
@@ -12,10 +12,10 @@ public:
 	virtual ~FTestVoice();
 
 	// FTickerObjectBase
-	virtual bool Tick(float DeltaTime) OVERRIDE;
+	virtual bool Tick(float DeltaTime) override;
 
 	// FSelfRegisteringExec
-	virtual bool Exec(UWorld* InWorld, const TCHAR* Cmd, FOutputDevice& Ar) OVERRIDE;
+	virtual bool Exec(UWorld* InWorld, const TCHAR* Cmd, FOutputDevice& Ar) override;
 
 	void Test();
 
diff --git a/Engine/Source/Runtime/Online/OnlineSubsystemUtils/Public/VoiceEngineImpl.h b/Engine/Source/Runtime/Online/OnlineSubsystemUtils/Public/VoiceEngineImpl.h
--- a/Engine/Source/Runtime/Online/OnlineSubsystemUtils/Public/VoiceEngineImpl.h
+++ b/Engine/Source/Runtime/Online/OnlineSubsystemUtils/Public/VoiceEngineImpl.h
@@ -63,7 +63,7 @@ class FVoiceEngineImpl : public IVoiceEngine
 		~FVoiceSerializeHelper() {}
 		
 		/** FGCObject interface */
-		virtual void AddReferencedObjects(FReferenceCollector& Collector) OVERRIDE
+		virtual void AddReferencedObjects(FReferenceCollector& Collector) override
 		{
 			// Prevent garbage collection of audio components
 			for (FRemoteTalkerData::TIterator It(VoiceEngine->RemoteTalkerBuffers); It; ++It)
@@ -160,7 +160,7 @@ PACKAGE_SCOPE:
 	{};
 
 	// IVoiceEngine
-	virtual bool Init(int32 MaxLocalTalkers, int32 MaxRemoteTalkers) OVERRIDE;
+	virtual bool Init(int32 MaxLocalTalkers, int32 MaxRemoteTalkers) override;
 
 public:
 
@@ -168,22 +168,22 @@ public:
 	virtual ~FVoiceEngineImpl();
 
 	// IVoiceEngine
-    virtual uint32 StartLocalVoiceProcessing(uint32 LocalUserNum) OVERRIDE;
-    virtual uint32 StopLocalVoiceProcessing(uint32 LocalUserNum) OVERRIDE;
+    virtual uint32 StartLocalVoiceProcessing(uint32 LocalUserNum) override;
+    virtual uint32 StopLocalVoiceProcessing(uint32 LocalUserNum) override;
     
-	virtual uint32 StartRemoteVoiceProcessing(const FUniqueNetId& UniqueId) OVERRIDE
+	virtual uint32 StartRemoteVoiceProcessing(const FUniqueNetId& UniqueId) override
 	{
 		// Not needed
 		return S_OK;
 	}
 
-    virtual uint32 StopRemoteVoiceProcessing(const FUniqueNetId& UniqueId) OVERRIDE
+    virtual uint32 StopRemoteVoiceProcessing(const FUniqueNetId& UniqueId) override
 	{
 		// Not needed
 		return S_OK;
 	}
 
-    virtual uint32 RegisterLocalTalker(uint32 LocalUserNum) OVERRIDE
+    virtual uint32 RegisterLocalTalker(uint32 LocalUserNum) override
 	{
 		if (OwningUserIndex == INVALID_INDEX)
 		{
@@ -194,7 +194,7 @@ public:
 		return E_FAIL;
 	}
 
-    virtual uint32 UnregisterLocalTalker(uint32 LocalUserNum) OVERRIDE
+    virtual uint32 UnregisterLocalTalker(uint32 LocalUserNum) override
 	{
 		if (IsOwningUser(LocalUserNum))
 		{
@@ -205,44 +205,44 @@ public:
 		return E_FAIL;
 	}
 
-    virtual uint32 RegisterRemoteTalker(const FUniqueNetId& UniqueId) OVERRIDE
+    virtual uint32 RegisterRemoteTalker(const FUniqueNetId& UniqueId) override
 	{
 		// Not needed
 		return S_OK;
 	}
 
-    virtual uint32 UnregisterRemoteTalker(const FUniqueNetId& UniqueId) OVERRIDE
+    virtual uint32 UnregisterRemoteTalker(const FUniqueNetId& UniqueId) override
 	{
 		// Not needed
 		return S_OK;
 	}
 
-    virtual bool IsHeadsetPresent(uint32 LocalUserNum) OVERRIDE
+    virtual bool IsHeadsetPresent(uint32 LocalUserNum) override
 	{
 		return IsOwningUser(LocalUserNum) ? true : false;
 	}
 
-    virtual bool IsLocalPlayerTalking(uint32 LocalUserNum) OVERRIDE
+    virtual bool IsLocalPlayerTalking(uint32 LocalUserNum) override
 	{
 		return (GetVoiceDataReadyFlags() & (LocalUserNum << 1)) != 0;
 	}
 
-	virtual bool IsRemotePlayerTalking(const FUniqueNetId& UniqueId) OVERRIDE
+	virtual bool IsRemotePlayerTalking(const FUniqueNetId& UniqueId) override
 	{
 		return RemoteTalkerBuffers.Find((const FUniqueNetIdString&)UniqueId) != NULL;
 	}
 
-	virtual uint32 GetVoiceDataReadyFlags() const OVERRIDE;
-	virtual uint32 SetPlaybackPriority(uint32 LocalUserNum, const FUniqueNetId& RemoteTalkerId, uint32 Priority) OVERRIDE
+	virtual uint32 GetVoiceDataReadyFlags() const override;
+	virtual uint32 SetPlaybackPriority(uint32 LocalUserNum, const FUniqueNetId& RemoteTalkerId, uint32 Priority) override
 	{
 		// Not supported
 		return S_OK;
 	}
 
-	virtual uint32 ReadLocalVoiceData(uint32 LocalUserNum, uint8* Data, uint32* Size) OVERRIDE;
-	virtual uint32 SubmitRemoteVoiceData(const FUniqueNetId& RemoteTalkerId, uint8* Data, uint32* Size) OVERRIDE;
-	virtual void Tick(float DeltaTime) OVERRIDE;
-	FString GetVoiceDebugState() const OVERRIDE;
+	virtual uint32 ReadLocalVoiceData(uint32 LocalUserNum, uint8* Data, uint32* Size) override;
+	virtual uint32 SubmitRemoteVoiceData(const FUniqueNetId& RemoteTalkerId, uint8* Data, uint32* Size) override;
+	virtual void Tick(float DeltaTime) override;
+	FString GetVoiceDebugState() const override;
 
 	/**
 	 * Update the state of all remote talkers, possibly dropping data or the talker entirely
diff --git a/Engine/Source/Runtime/Online/OnlineSubsystemUtils/Public/VoiceInterfaceImpl.h b/Engine/Source/Runtime/Online/OnlineSubsystemUtils/Public/VoiceInterfaceImpl.h
--- a/Engine/Source/Runtime/Online/OnlineSubsystemUtils/Public/VoiceInterfaceImpl.h
+++ b/Engine/Source/Runtime/Online/OnlineSubsystemUtils/Public/VoiceInterfaceImpl.h
@@ -69,8 +69,8 @@ PACKAGE_SCOPE:
 	{};
 
 	// IOnlineVoice
-	virtual bool Init() OVERRIDE;
-	void ProcessMuteChangeNotification() OVERRIDE;
+	virtual bool Init() override;
+	void ProcessMuteChangeNotification() override;
 
 	/**
 	 * Processes any talking delegates that need to be fired off
@@ -106,27 +106,27 @@ public:
 	virtual ~FOnlineVoiceImpl();
 
 	// IOnlineVoice
-	virtual void StartNetworkedVoice(uint8 LocalUserNum) OVERRIDE;
-	virtual void StopNetworkedVoice(uint8 LocalUserNum) OVERRIDE;
-    virtual bool RegisterLocalTalker(uint32 LocalUserNum) OVERRIDE;
-	virtual void RegisterLocalTalkers() OVERRIDE;
-    virtual bool UnregisterLocalTalker(uint32 LocalUserNum) OVERRIDE;
-	virtual void UnregisterLocalTalkers() OVERRIDE;
-    virtual bool RegisterRemoteTalker(const FUniqueNetId& UniqueId) OVERRIDE;
-    virtual bool UnregisterRemoteTalker(const FUniqueNetId& UniqueId) OVERRIDE;
-	virtual void RemoveAllRemoteTalkers() OVERRIDE;
-    virtual bool IsHeadsetPresent(uint32 LocalUserNum) OVERRIDE;
-    virtual bool IsLocalPlayerTalking(uint32 LocalUserNum) OVERRIDE;
-	virtual bool IsRemotePlayerTalking(const FUniqueNetId& UniqueId) OVERRIDE;
-	bool IsMuted(uint32 LocalUserNum, const FUniqueNetId& UniqueId) const OVERRIDE;
-	bool MuteRemoteTalker(uint8 LocalUserNum, const FUniqueNetId& PlayerId, bool bIsSystemWide) OVERRIDE;
-	bool UnmuteRemoteTalker(uint8 LocalUserNum, const FUniqueNetId& PlayerId, bool bIsSystemWide) OVERRIDE;
-	virtual TSharedPtr<class FVoicePacket> SerializeRemotePacket(FArchive& Ar) OVERRIDE;
-	virtual TSharedPtr<class FVoicePacket> GetLocalPacket(uint32 LocalUserNum) OVERRIDE;
-	virtual int32 GetNumLocalTalkers() OVERRIDE { return LocalTalkers.Num(); };
-	virtual void ClearVoicePackets() OVERRIDE;
-	virtual void Tick(float DeltaTime) OVERRIDE;
-	virtual FString GetVoiceDebugState() const OVERRIDE;
+	virtual void StartNetworkedVoice(uint8 LocalUserNum) override;
+	virtual void StopNetworkedVoice(uint8 LocalUserNum) override;
+    virtual bool RegisterLocalTalker(uint32 LocalUserNum) override;
+	virtual void RegisterLocalTalkers() override;
+    virtual bool UnregisterLocalTalker(uint32 LocalUserNum) override;
+	virtual void UnregisterLocalTalkers() override;
+    virtual bool RegisterRemoteTalker(const FUniqueNetId& UniqueId) override;
+    virtual bool UnregisterRemoteTalker(const FUniqueNetId& UniqueId) override;
+	virtual void RemoveAllRemoteTalkers() override;
+    virtual bool IsHeadsetPresent(uint32 LocalUserNum) override;
+    virtual bool IsLocalPlayerTalking(uint32 LocalUserNum) override;
+	virtual bool IsRemotePlayerTalking(const FUniqueNetId& UniqueId) override;
+	bool IsMuted(uint32 LocalUserNum, const FUniqueNetId& UniqueId) const override;
+	bool MuteRemoteTalker(uint8 LocalUserNum, const FUniqueNetId& PlayerId, bool bIsSystemWide) override;
+	bool UnmuteRemoteTalker(uint8 LocalUserNum, const FUniqueNetId& PlayerId, bool bIsSystemWide) override;
+	virtual TSharedPtr<class FVoicePacket> SerializeRemotePacket(FArchive& Ar) override;
+	virtual TSharedPtr<class FVoicePacket> GetLocalPacket(uint32 LocalUserNum) override;
+	virtual int32 GetNumLocalTalkers() override { return LocalTalkers.Num(); };
+	virtual void ClearVoicePackets() override;
+	virtual void Tick(float DeltaTime) override;
+	virtual FString GetVoiceDebugState() const override;
 };
 
 typedef TSharedPtr<FOnlineVoiceImpl, ESPMode::ThreadSafe> FOnlineVoiceImplPtr;
diff --git a/Engine/Source/Runtime/Online/OnlineSubsystemUtils/Public/VoicePacketImpl.h b/Engine/Source/Runtime/Online/OnlineSubsystemUtils/Public/VoicePacketImpl.h
--- a/Engine/Source/Runtime/Online/OnlineSubsystemUtils/Public/VoicePacketImpl.h
+++ b/Engine/Source/Runtime/Online/OnlineSubsystemUtils/Public/VoicePacketImpl.h
@@ -44,22 +44,22 @@ public:
 	FVoicePacketImpl(const FVoicePacketImpl& Other);
 
 	/** Returns the amount of space this packet will consume in a buffer */
-	virtual uint16 GetTotalPacketSize() OVERRIDE;
+	virtual uint16 GetTotalPacketSize() override;
 
 	/** @return the amount of space used by the internal voice buffer */
-	virtual uint16 GetBufferSize() OVERRIDE;
+	virtual uint16 GetBufferSize() override;
 
 	/** @return the sender of this voice packet */
-	virtual TSharedPtr<class FUniqueNetId> GetSender() OVERRIDE;
+	virtual TSharedPtr<class FUniqueNetId> GetSender() override;
 
-	virtual bool IsReliable() OVERRIDE { return false; }
+	virtual bool IsReliable() override { return false; }
 
 	/** 
 	 * Serialize the voice packet data to a buffer 
 	 *
 	 * @param Ar buffer to write into
 	 */
-	virtual void Serialize(class FArchive& Ar) OVERRIDE;
+	virtual void Serialize(class FArchive& Ar) override;
 };
 
 /** Holds the current voice packet data state */
diff --git a/Engine/Source/Runtime/Online/Voice/Private/VoiceCodecOpus.h b/Engine/Source/Runtime/Online/Voice/Private/VoiceCodecOpus.h
--- a/Engine/Source/Runtime/Online/Voice/Private/VoiceCodecOpus.h
+++ b/Engine/Source/Runtime/Online/Voice/Private/VoiceCodecOpus.h
@@ -21,9 +21,9 @@ public:
 	~FVoiceEncoderOpus();
 
 	// IVoiceEncoder
-	virtual bool Init(int32 SampleRate, int32 NumChannels) OVERRIDE;
-	virtual int32 Encode(const uint8* RawPCMData, uint32 RawDataSize, uint8* OutCompressedData, uint32& OutCompressedDataSize) OVERRIDE;
-	virtual void Destroy() OVERRIDE;
+	virtual bool Init(int32 SampleRate, int32 NumChannels) override;
+	virtual int32 Encode(const uint8* RawPCMData, uint32 RawDataSize, uint8* OutCompressedData, uint32& OutCompressedDataSize) override;
+	virtual void Destroy() override;
 
 private:
 
@@ -60,9 +60,9 @@ public:
 	~FVoiceDecoderOpus();
 
 	//IVoiceDecoder
-	virtual bool Init(int32 SampleRate, int32 NumChannels) OVERRIDE;
-	virtual void Decode(const uint8* CompressedData, uint32 CompressedDataSize, uint8* OutRawPCMData, uint32& OutRawDataSize) OVERRIDE;
-	virtual void Destroy() OVERRIDE;
+	virtual bool Init(int32 SampleRate, int32 NumChannels) override;
+	virtual void Decode(const uint8* CompressedData, uint32 CompressedDataSize, uint8* OutRawPCMData, uint32& OutRawDataSize) override;
+	virtual void Destroy() override;
 
 private:
 
diff --git a/Engine/Source/Runtime/Online/Voice/Private/Windows/VoiceCaptureWindows.h b/Engine/Source/Runtime/Online/Voice/Private/Windows/VoiceCaptureWindows.h
--- a/Engine/Source/Runtime/Online/Voice/Private/Windows/VoiceCaptureWindows.h
+++ b/Engine/Source/Runtime/Online/Voice/Private/Windows/VoiceCaptureWindows.h
@@ -92,16 +92,16 @@ public:
 	~FVoiceCaptureWindows();
 
 	// IVoiceCapture
-	virtual bool Init(int32 SampleRate, int32 NumChannels) OVERRIDE;
-	virtual void Shutdown() OVERRIDE;
-	virtual bool Start() OVERRIDE;
-	virtual void Stop() OVERRIDE;
-	virtual bool IsCapturing() OVERRIDE;
-	virtual EVoiceCaptureState::Type GetCaptureState(uint32& OutAvailableVoiceData) const OVERRIDE;
-	virtual EVoiceCaptureState::Type GetVoiceData(uint8* OutVoiceBuffer, uint32 InVoiceBufferSize, uint32& OutAvailableVoiceData) OVERRIDE;
+	virtual bool Init(int32 SampleRate, int32 NumChannels) override;
+	virtual void Shutdown() override;
+	virtual bool Start() override;
+	virtual void Stop() override;
+	virtual bool IsCapturing() override;
+	virtual EVoiceCaptureState::Type GetCaptureState(uint32& OutAvailableVoiceData) const override;
+	virtual EVoiceCaptureState::Type GetVoiceData(uint8* OutVoiceBuffer, uint32 InVoiceBufferSize, uint32& OutAvailableVoiceData) override;
 
 	// FTickerObjectBase
-	virtual bool Tick(float DeltaTime) OVERRIDE;
+	virtual bool Tick(float DeltaTime) override;
 
 private:
 
diff --git a/Engine/Source/Runtime/Online/Voice/Public/VoiceModule.h b/Engine/Source/Runtime/Online/Voice/Public/VoiceModule.h
--- a/Engine/Source/Runtime/Online/Voice/Public/VoiceModule.h
+++ b/Engine/Source/Runtime/Online/Voice/Public/VoiceModule.h
@@ -24,7 +24,7 @@ class FVoiceModule :
 public:
 
 	// FSelfRegisteringExec
-	virtual bool Exec(UWorld* InWorld, const TCHAR* Cmd, FOutputDevice& Ar) OVERRIDE;
+	virtual bool Exec(UWorld* InWorld, const TCHAR* Cmd, FOutputDevice& Ar) override;
 
 	/**
 	 * Singleton-like access to this module's interface.  This is just for convenience!
@@ -84,13 +84,13 @@ private:
 	 * Called when voice module is loaded
 	 * Initialize platform specific parts of voice handling
 	 */
-	virtual void StartupModule() OVERRIDE;
+	virtual void StartupModule() override;
 	
 	/**
 	 * Called when voice module is unloaded
 	 * Shutdown platform specific parts of voice handling
 	 */
-	virtual void ShutdownModule() OVERRIDE;
+	virtual void ShutdownModule() override;
 
 	/** Is voice interface enabled */
 	bool bEnabled;
