[GitHub] 235 : wrong FUniqueNetId subclass used for FOnlineSessionSteam
https://github.com/EpicGames/UnrealEngine/pull/235

diff --git a/Engine/Source/Runtime/Online/OnlineSubsystemSteam/Private/OnlineSessionInterfaceSteam.cpp b/Engine/Source/Runtime/Online/OnlineSubsystemSteam/Private/OnlineSessionInterfaceSteam.cpp
--- a/Engine/Source/Runtime/Online/OnlineSubsystemSteam/Private/OnlineSessionInterfaceSteam.cpp
+++ b/Engine/Source/Runtime/Online/OnlineSubsystemSteam/Private/OnlineSessionInterfaceSteam.cpp
@@ -1226,7 +1226,7 @@ void FOnlineSessionSteam::UnregisterVoice(const FUniqueNetId& PlayerId)
 bool FOnlineSessionSteam::RegisterPlayer(FName SessionName, const FUniqueNetId& PlayerId, bool bWasInvited)
 {
 	TArray< TSharedRef<FUniqueNetId> > Players;
-	Players.Add(MakeShareable(new FUniqueNetIdString(PlayerId)));
+	Players.Add(MakeShareable(new FUniqueNetIdSteam(PlayerId)));
 	return RegisterPlayers(SessionName, Players, bWasInvited);
 }
 
@@ -1303,7 +1303,7 @@ void FOnlineSessionSteam::RegisterLocalPlayers(FNamedOnlineSession* Session)
 bool FOnlineSessionSteam::UnregisterPlayer(FName SessionName, const FUniqueNetId& PlayerId)
 {
 	TArray< TSharedRef<FUniqueNetId> > Players;
-	Players.Add(MakeShareable(new FUniqueNetIdString(PlayerId)));
+	Players.Add(MakeShareable(new FUniqueNetIdSteam(PlayerId)));
 	return UnregisterPlayers(SessionName, Players);
 }
 
