#online, http
- potential fix for http elapsed response time getting reset and preventing timeout from triggering for failed downloads
- also enabled verbose logging for an http request that exceeds total timeout

diff --git a/Engine/Source/Runtime/Online/HTTP/Private/Windows/HttpWinInet.cpp b/Engine/Source/Runtime/Online/HTTP/Private/Windows/HttpWinInet.cpp
--- a/Engine/Source/Runtime/Online/HTTP/Private/Windows/HttpWinInet.cpp
+++ b/Engine/Source/Runtime/Online/HTTP/Private/Windows/HttpWinInet.cpp
@@ -146,6 +146,17 @@ bool FWinInetConnection::ShutdownConnection()
 	return true;
 }
 
+// override for debug logging
+#define DEBUG_LOG_HTTP(bIsDebug, Verbosity, Format, ...) \
+if (bIsDebug) \
+{ \
+	UE_LOG(LogHttp, Log, Format, ##__VA_ARGS__); \
+} \
+else \
+{ \
+	UE_LOG(LogHttp, Verbosity, Format, ##__VA_ARGS__); \
+}
+
 /**
  * Global callback for WinInet API. Will use the dwInternetStatus and dwContext
  * fields to route the results back to the appropriate class instance if necessary.
@@ -202,52 +213,54 @@ void CALLBACK InternetStatusCallbackWinInet(
 		Response = Request->Response.Get();
 	}
 
+	bool bDebugLog = Request != NULL && Request->bDebugVerbose;
+
 	switch (dwInternetStatus)
 	{
 	case INTERNET_STATUS_PREFETCH:
-		UE_LOG(LogHttp, VeryVerbose, TEXT("STATUS_PREFETCH: %p"), dwContext);
+		DEBUG_LOG_HTTP(bDebugLog, VeryVerbose, TEXT("STATUS_PREFETCH: %p"), dwContext);
 		break;
 	case INTERNET_STATUS_USER_INPUT_REQUIRED:
-		UE_LOG(LogHttp, VeryVerbose, TEXT("USER_INPUT_REQUIRED: %p"), dwContext);
+		DEBUG_LOG_HTTP(bDebugLog, VeryVerbose, TEXT("USER_INPUT_REQUIRED: %p"), dwContext);
 		break;
 	case INTERNET_STATUS_DETECTING_PROXY:
-		UE_LOG(LogHttp, VeryVerbose, TEXT("DETECTING_PROXY: %p"), dwContext);
+		DEBUG_LOG_HTTP(bDebugLog, VeryVerbose, TEXT("DETECTING_PROXY: %p"), dwContext);
 		break;
 	case INTERNET_STATUS_CLOSING_CONNECTION:
-		UE_LOG(LogHttp, VeryVerbose, TEXT("CLOSING_CONNECTION: %p"), dwContext);
+		DEBUG_LOG_HTTP(bDebugLog, VeryVerbose, TEXT("CLOSING_CONNECTION: %p"), dwContext);
 		break;
 	case INTERNET_STATUS_CONNECTED_TO_SERVER:
-		UE_LOG(LogHttp, VeryVerbose, TEXT("CONNECTED_TO_SERVER: %p"), dwContext);
+		DEBUG_LOG_HTTP(bDebugLog, VeryVerbose, TEXT("CONNECTED_TO_SERVER: %p"), dwContext);
 		break;
 	case INTERNET_STATUS_CONNECTING_TO_SERVER:
-		UE_LOG(LogHttp, VeryVerbose, TEXT("CONNECTING_TO_SERVER: %p"), dwContext);
+		DEBUG_LOG_HTTP(bDebugLog, VeryVerbose, TEXT("CONNECTING_TO_SERVER: %p"), dwContext);
 		break;
 	case INTERNET_STATUS_CONNECTION_CLOSED:
-		UE_LOG(LogHttp, VeryVerbose, TEXT("CONNECTION_CLOSED: %p"), dwContext);
+		DEBUG_LOG_HTTP(bDebugLog, VeryVerbose, TEXT("CONNECTION_CLOSED: %p"), dwContext);
 		break;
 	case INTERNET_STATUS_HANDLE_CLOSING:
-		UE_LOG(LogHttp, VeryVerbose, TEXT("HANDLE_CLOSING: %p"), dwContext);
+		DEBUG_LOG_HTTP(bDebugLog, VeryVerbose, TEXT("HANDLE_CLOSING: %p"), dwContext);
 		break;
 	case INTERNET_STATUS_HANDLE_CREATED:
-		UE_LOG(LogHttp, VeryVerbose, TEXT("HANDLE_CREATED: %p"), dwContext);
+		DEBUG_LOG_HTTP(bDebugLog, VeryVerbose, TEXT("HANDLE_CREATED: %p"), dwContext);
 		break;
 	case INTERNET_STATUS_INTERMEDIATE_RESPONSE:
-		UE_LOG(LogHttp, VeryVerbose, TEXT("INTERMEDIATE_RESPONSE: %p"), dwContext);
+		DEBUG_LOG_HTTP(bDebugLog, VeryVerbose, TEXT("INTERMEDIATE_RESPONSE: %p"), dwContext);
 		break;
 	case INTERNET_STATUS_NAME_RESOLVED:
-		UE_LOG(LogHttp, VeryVerbose, TEXT("NAME_RESOLVED: %p"), dwContext);
+		DEBUG_LOG_HTTP(bDebugLog, VeryVerbose, TEXT("NAME_RESOLVED: %p"), dwContext);
 		break;
 	case INTERNET_STATUS_RECEIVING_RESPONSE:
-		UE_LOG(LogHttp, VeryVerbose, TEXT("RECEIVING_RESPONSE: %p"), dwContext);
+		DEBUG_LOG_HTTP(bDebugLog, VeryVerbose, TEXT("RECEIVING_RESPONSE: %p"), dwContext);
 		break;
 	case INTERNET_STATUS_RESPONSE_RECEIVED:
-		UE_LOG(LogHttp, VeryVerbose, TEXT("RESPONSE_RECEIVED (%d bytes): %p"), *(uint32*)lpvStatusInformation, dwContext);		
+		DEBUG_LOG_HTTP(bDebugLog, VeryVerbose, TEXT("RESPONSE_RECEIVED (%d bytes): %p"), *(uint32*)lpvStatusInformation, dwContext);		
 		break;
 	case INTERNET_STATUS_REDIRECT:
-		UE_LOG(LogHttp, VeryVerbose, TEXT("STATUS_REDIRECT: %p"), dwContext);
+		DEBUG_LOG_HTTP(bDebugLog, VeryVerbose, TEXT("STATUS_REDIRECT: %p"), dwContext);
 		break;
 	case INTERNET_STATUS_REQUEST_COMPLETE:
-		UE_LOG(LogHttp, VeryVerbose, TEXT("REQUEST_COMPLETE: %p"), dwContext);
+		DEBUG_LOG_HTTP(bDebugLog, VeryVerbose, TEXT("REQUEST_COMPLETE: %p"), dwContext);
 
 		if (lpvStatusInformation != NULL)
 		{
@@ -255,7 +268,7 @@ void CALLBACK InternetStatusCallbackWinInet(
 			// callback failed so can't process response
 			if (AsyncResult->dwResult == 0)
 			{
-				UE_LOG(LogHttp, Log, TEXT("InternetStatusCallbackWinInet request=%p AsyncResult.dwError: %08X. "), 
+				DEBUG_LOG_HTTP(bDebugLog, Log, TEXT("InternetStatusCallbackWinInet request=%p AsyncResult.dwError: %08X. "), 
 					dwContext, AsyncResult->dwError, *InternetTranslateError(AsyncResult->dwError));
 
 				if (Response != NULL)
@@ -270,42 +283,41 @@ void CALLBACK InternetStatusCallbackWinInet(
 		if (Request != NULL && Request->GetStatus() == EHttpRequestStatus::Processing &&
 			Response != NULL && !Response->bIsReady)
 		{
-			Request->ResetRequestTimeout();
 			Response->ProcessResponse();
 		}
 
 		break;
 	case INTERNET_STATUS_REQUEST_SENT:
-		UE_LOG(LogHttp, VeryVerbose, TEXT("REQUEST_SENT: %p"), dwContext);
+		DEBUG_LOG_HTTP(bDebugLog, VeryVerbose, TEXT("REQUEST_SENT: %p"), dwContext);
 		break;
 	case INTERNET_STATUS_RESOLVING_NAME:
-		UE_LOG(LogHttp, VeryVerbose, TEXT("RESOLVING_NAME: %p"), dwContext);
+		DEBUG_LOG_HTTP(bDebugLog, VeryVerbose, TEXT("RESOLVING_NAME: %p"), dwContext);
 		break;
 	case INTERNET_STATUS_SENDING_REQUEST:
-		UE_LOG(LogHttp, VeryVerbose, TEXT("SENDING_REQUEST: %p"), dwContext);
+		DEBUG_LOG_HTTP(bDebugLog, VeryVerbose, TEXT("SENDING_REQUEST: %p"), dwContext);
 		break;
 	case INTERNET_STATUS_STATE_CHANGE:
-		UE_LOG(LogHttp, VeryVerbose, TEXT("STATE_CHANGE: %p"), dwContext);
+		DEBUG_LOG_HTTP(bDebugLog, VeryVerbose, TEXT("STATE_CHANGE: %p"), dwContext);
 		break;
 	case INTERNET_STATUS_COOKIE_SENT:
-		UE_LOG(LogHttp, VeryVerbose, TEXT("COOKIE_SENT: %p"), dwContext);
+		DEBUG_LOG_HTTP(bDebugLog, VeryVerbose, TEXT("COOKIE_SENT: %p"), dwContext);
 		break;
 	case INTERNET_STATUS_COOKIE_RECEIVED:
-		UE_LOG(LogHttp, VeryVerbose, TEXT("COOKIE_RECEIVED: %p"), dwContext);
+		DEBUG_LOG_HTTP(bDebugLog, VeryVerbose, TEXT("COOKIE_RECEIVED: %p"), dwContext);
 		break;
 	case INTERNET_STATUS_PRIVACY_IMPACTED:
-		UE_LOG(LogHttp, VeryVerbose, TEXT("PRIVACY_IMPACTED: %p"), dwContext);
+		DEBUG_LOG_HTTP(bDebugLog, VeryVerbose, TEXT("PRIVACY_IMPACTED: %p"), dwContext);
 		break;
 	case INTERNET_STATUS_P3P_HEADER:
-		UE_LOG(LogHttp, VeryVerbose, TEXT("P3P_HEADER: %p"), dwContext);
+		DEBUG_LOG_HTTP(bDebugLog, VeryVerbose, TEXT("P3P_HEADER: %p"), dwContext);
 		break;
 	case INTERNET_STATUS_P3P_POLICYREF:
-		UE_LOG(LogHttp, VeryVerbose, TEXT("P3P_POLICYREF: %p"), dwContext);
+		DEBUG_LOG_HTTP(bDebugLog, VeryVerbose, TEXT("P3P_POLICYREF: %p"), dwContext);
 		break;
 	case INTERNET_STATUS_COOKIE_HISTORY:
 		{
 			const InternetCookieHistory* CookieHistory = (const InternetCookieHistory*)lpvStatusInformation;
-			UE_LOG(LogHttp, VeryVerbose, TEXT("COOKIE_HISTORY: %p. Accepted: %u. Leashed: %u. Downgraded: %u. Rejected: %u.")
+			DEBUG_LOG_HTTP(bDebugLog, VeryVerbose, TEXT("COOKIE_HISTORY: %p. Accepted: %u. Leashed: %u. Downgraded: %u. Rejected: %u.")
 				, dwContext
 				, CookieHistory->fAccepted
 				, CookieHistory->fLeashed
@@ -314,7 +326,7 @@ void CALLBACK InternetStatusCallbackWinInet(
 		}
 		break;
 	default:
-		UE_LOG(LogHttp, VeryVerbose, TEXT("Unknown Status: %u. %p"), (uint32)dwInternetStatus, dwContext);
+		DEBUG_LOG_HTTP(bDebugLog, VeryVerbose, TEXT("Unknown Status: %u. %p"), (uint32)dwInternetStatus, dwContext);
 		break;
 	}
 }
@@ -328,6 +340,7 @@ FHttpRequestWinInet::FHttpRequestWinInet()
 ,	ElapsedTimeSinceLastServerResponse(0)
 ,	ProgressBytesSent(0)
 ,	StartRequestTime(0)
+,	bDebugVerbose(false)
 {
 
 }
@@ -719,11 +732,23 @@ void FHttpRequestWinInet::Tick(float DeltaSeconds)
 	}
 
 	// Convert to seconds for comparison to the timeout value
-	const float TotalElapsedSeconds = ElapsedTimeSinceLastServerResponse / 1000.f;
-
 	const float HttpTimeout = FHttpModule::Get().GetHttpTimeout();
+	
+	// Log verbose once total elapsed time surpasses HttpTimeout
+	const double TotalElapsed = FPlatformTime::Seconds() - StartRequestTime;
+	if (HttpTimeout > 0 &&
+		TotalElapsed > HttpTimeout)
+	{
+		UE_LOG(LogHttp, Warning, TEXT("Http request taking too long! Elapsed %.3f. Enabling verbose logs %p url=%s"),
+			TotalElapsed, this, *GetURL());
+
+		bDebugVerbose = true;
+	}
+	
+	// Timeout if no server response for HttpTimeout
+	const float SecondsSinceLastResponse = ElapsedTimeSinceLastServerResponse / 1000.f;
 	if (HttpTimeout > 0 && 
-		TotalElapsedSeconds >= HttpTimeout)
+		SecondsSinceLastResponse >= HttpTimeout)
 	{
 		UE_LOG(LogHttp, Warning, TEXT("Timeout processing Http request. %p url=%s"),
 			this, *GetURL());
@@ -742,7 +767,7 @@ void FHttpRequestWinInet::Tick(float DeltaSeconds)
 
 // FHttpResponseWinInet
 
-FHttpResponseWinInet::FHttpResponseWinInet(const FHttpRequestWinInet& InRequest)
+FHttpResponseWinInet::FHttpResponseWinInet(FHttpRequestWinInet& InRequest)
 :	Request(InRequest)
 ,	RequestURL(InRequest.RequestURL)
 ,	AsyncBytesRead(0)
@@ -866,6 +891,7 @@ void FHttpResponseWinInet::ProcessResponse()
 		ResponsePayload.SetNum(BufferSize);
 	}
 
+	bool bDebugLog = Request.bDebugVerbose;
 	bool bFailed = false;
 	int32 LoopCount = 0;
 	do
@@ -878,8 +904,8 @@ void FHttpResponseWinInet::ProcessResponse()
 			NumBytesToRead, 
 			(LPDWORD)&AsyncBytesRead);
 
-		UE_LOG(LogHttp, VeryVerbose, TEXT("InternetReadFile result=%d (%u bytes read) (%u bytes total read). LoopCount=%d %p"), 
-			bReadFile, AsyncBytesRead, TotalBytesRead, LoopCount, &Request);
+		DEBUG_LOG_HTTP(bDebugLog, VeryVerbose, TEXT("InternetReadFile result=%d (%u bytes read) (%u bytes total read) (%u bytes to read). LoopCount=%d %p"), 
+			bReadFile, AsyncBytesRead, TotalBytesRead, NumBytesToRead, LoopCount, &Request);
 
 		if (!bReadFile)
 		{
@@ -891,13 +917,13 @@ void FHttpResponseWinInet::ProcessResponse()
 				// to be called again later when you receive another REQUEST_COMPLETE signal.
 				// You are supposed to keep doing this until InternetReadFile returns 0 with a successful return value.
 
-				UE_LOG(LogHttp, VeryVerbose, TEXT("InternetReadFile is completing asynchronously, so waiting for callback again. %p"), 
+				DEBUG_LOG_HTTP(bDebugLog, VeryVerbose, TEXT("InternetReadFile is completing asynchronously, so waiting for callback again. %p"), 
 					&Request);
 			}
 			else if (ErrorCode != ERROR_SUCCESS)
 			{
-				UE_LOG(LogHttp, Log, TEXT("InternetReadFile failed (%u bytes read). Returning what we've read so far: %s. %p"), 
-					AsyncBytesRead, *InternetTranslateError(ErrorCode), &Request);
+				DEBUG_LOG_HTTP(bDebugLog, Log, TEXT("InternetReadFile failed (%u bytes read). Returning what we've read so far: error=%d %s. %p"), 
+					AsyncBytesRead, ErrorCode, *InternetTranslateError(ErrorCode), &Request);
 			}
 			// Allow processing to continue
 			return;
@@ -911,7 +937,7 @@ void FHttpResponseWinInet::ProcessResponse()
 			{
 				if (ContentLength > 0)
 				{
-					UE_LOG(LogHttp, Log, TEXT("Response payload (%d bytes read so far) is larger than the content-length (%d). Resizing buffer to accommodate. %p"), 
+					DEBUG_LOG_HTTP(bDebugLog, Log, TEXT("Response payload (%d bytes read so far) is larger than the content-length (%d). Resizing buffer to accommodate. %p"), 
 						TotalBytesRead, ContentLength, &Request);
 				}
 				ResponsePayload.AddZeroed(MaxReadBufferSize);
@@ -921,16 +947,17 @@ void FHttpResponseWinInet::ProcessResponse()
 
 		// Update progress bytes
 		ProgressBytesRead.Set( TotalBytesRead );
+		Request.ResetRequestTimeout();
 
 	} while (AsyncBytesRead > 0);
 
 	if (ContentLength != 0 && 
 		TotalBytesRead != ContentLength)
 	{
-		UE_LOG(LogHttp, Warning, TEXT("Response payload was %d bytes, content-length indicated (%d) bytes. %p"), 
+		DEBUG_LOG_HTTP(bDebugLog, Warning, TEXT("Response payload was %d bytes, content-length indicated (%d) bytes. %p"), 
 			TotalBytesRead, ContentLength, &Request);
 	}
-	UE_LOG(LogHttp, Verbose, TEXT("TotalBytesRead = %d. %p"), TotalBytesRead, &Request);
+	DEBUG_LOG_HTTP(bDebugLog, Verbose, TEXT("TotalBytesRead = %d. %p"), TotalBytesRead, &Request);
 
 	// Shrink array to only the valid data
 	ResponsePayload.SetNum(TotalBytesRead);
diff --git a/Engine/Source/Runtime/Online/HTTP/Private/Windows/HttpWinInet.h b/Engine/Source/Runtime/Online/HTTP/Private/Windows/HttpWinInet.h
--- a/Engine/Source/Runtime/Online/HTTP/Private/Windows/HttpWinInet.h
+++ b/Engine/Source/Runtime/Online/HTTP/Private/Windows/HttpWinInet.h
@@ -170,6 +170,8 @@ private:
 	int32 ProgressBytesSent;
 	/** Used to calculate total elapsed time for the request */
 	double StartRequestTime;
+	/** enables verbose logging for any http request that exceeds the total timeout */
+	bool bDebugVerbose;
 };
 
 /**
@@ -202,7 +204,7 @@ public:
 	 *
 	 * @param InRequest - original request that created this response
 	 */
-	FHttpResponseWinInet(const FHttpRequestWinInet& InRequest);
+	FHttpResponseWinInet(FHttpRequestWinInet& InRequest);
 
 	/**
 	 * Destructor
@@ -249,7 +251,7 @@ private:
 private:
 
 	/** Request that owns this response */
-	const FHttpRequestWinInet& Request;
+	FHttpRequestWinInet& Request;
 	/** Original URL used for the request */
 	FURLWinInet RequestURL;
 	/** Last bytes read from async call to InternetReadFile */
