#online
- added paging info for certain online queries

diff --git a/Engine/Plugins/Editor/EpicSurvey/Source/EpicSurvey/Private/SurveyTitleCdnStorage.cpp b/Engine/Plugins/Editor/EpicSurvey/Source/EpicSurvey/Private/SurveyTitleCdnStorage.cpp
--- a/Engine/Plugins/Editor/EpicSurvey/Source/EpicSurvey/Private/SurveyTitleCdnStorage.cpp
+++ b/Engine/Plugins/Editor/EpicSurvey/Source/EpicSurvey/Private/SurveyTitleCdnStorage.cpp
@@ -74,7 +74,7 @@ bool FSurveyTitleCdnStorage::ClearFile(const FString& FileName)
 	return false;
 }
 
-bool FSurveyTitleCdnStorage::EnumerateFiles()
+bool FSurveyTitleCdnStorage::EnumerateFiles(const FPagedQuery& Page)
 {
 	// Make sure an enumeration request  is not currently pending
 	if(!EnumerateFilesRequests.IsEmpty())
@@ -96,11 +96,6 @@ bool FSurveyTitleCdnStorage::EnumerateFiles()
 	
 }	
 
-bool FSurveyTitleCdnStorage::EnumerateFiles(int32 Start, int32 Count)
-{
-	return EnumerateFiles();
-}
-
 void FSurveyTitleCdnStorage::GetFileList(TArray<FCloudFileHeader>& OutFiles) 
 {
 	TArray<FCloudFileHeader>* FilesPtr = &FileHeaders;
diff --git a/Engine/Plugins/Editor/EpicSurvey/Source/EpicSurvey/Private/SurveyTitleCdnStorage.h b/Engine/Plugins/Editor/EpicSurvey/Source/EpicSurvey/Private/SurveyTitleCdnStorage.h
--- a/Engine/Plugins/Editor/EpicSurvey/Source/EpicSurvey/Private/SurveyTitleCdnStorage.h
+++ b/Engine/Plugins/Editor/EpicSurvey/Source/EpicSurvey/Private/SurveyTitleCdnStorage.h
@@ -20,9 +20,7 @@ public:
 
 	virtual bool ClearFile(const FString& DLName) override;
 
-	virtual bool EnumerateFiles() override;
-
-	virtual bool EnumerateFiles(int32 Start, int32 Count) override;
+	virtual bool EnumerateFiles(const FPagedQuery& Page = FPagedQuery()) override;
 
 	virtual void GetFileList(TArray<FCloudFileHeader>& InFileHeaders) override;
 
diff --git a/Engine/Plugins/Editor/EpicSurvey/Source/EpicSurvey/Private/SurveyTitleLocalStorage.cpp b/Engine/Plugins/Editor/EpicSurvey/Source/EpicSurvey/Private/SurveyTitleLocalStorage.cpp
--- a/Engine/Plugins/Editor/EpicSurvey/Source/EpicSurvey/Private/SurveyTitleLocalStorage.cpp
+++ b/Engine/Plugins/Editor/EpicSurvey/Source/EpicSurvey/Private/SurveyTitleLocalStorage.cpp
@@ -54,7 +54,7 @@ bool FSurveyTitleLocalStorage::ClearFile(const FString& DLName)
 	return ClearedFile;
 }
 
-bool FSurveyTitleLocalStorage::EnumerateFiles()
+bool FSurveyTitleLocalStorage::EnumerateFiles(const FPagedQuery& Page)
 {
 	if (!IFileManager::Get().DirectoryExists(*RootDirectory))
 	{
@@ -85,11 +85,6 @@ bool FSurveyTitleLocalStorage::EnumerateFiles()
 	return true;
 }
 
-bool FSurveyTitleLocalStorage::EnumerateFiles(int32 Start, int32 Count)
-{
-	return EnumerateFiles();
-}
-
 void FSurveyTitleLocalStorage::GetFileList(TArray<FCloudFileHeader>& InFileHeaders)
 {
 	InFileHeaders.Append( FileHeaders );
diff --git a/Engine/Plugins/Editor/EpicSurvey/Source/EpicSurvey/Private/SurveyTitleLocalStorage.h b/Engine/Plugins/Editor/EpicSurvey/Source/EpicSurvey/Private/SurveyTitleLocalStorage.h
--- a/Engine/Plugins/Editor/EpicSurvey/Source/EpicSurvey/Private/SurveyTitleLocalStorage.h
+++ b/Engine/Plugins/Editor/EpicSurvey/Source/EpicSurvey/Private/SurveyTitleLocalStorage.h
@@ -16,9 +16,7 @@ public:
 
 	virtual bool ClearFile(const FString& DLName) override;
 
-	virtual bool EnumerateFiles() override;
-
-	virtual bool EnumerateFiles(int32 Start, int32 Count) override;
+	virtual bool EnumerateFiles(const FPagedQuery& Page = FPagedQuery()) override;
 
 	virtual void GetFileList(TArray<FCloudFileHeader>& InFileHeaders) override;
 
diff --git a/Engine/Source/Editor/NewsFeed/Private/Implementation/CdnNewsFeedTitleFile.cpp b/Engine/Source/Editor/NewsFeed/Private/Implementation/CdnNewsFeedTitleFile.cpp
--- a/Engine/Source/Editor/NewsFeed/Private/Implementation/CdnNewsFeedTitleFile.cpp
+++ b/Engine/Source/Editor/NewsFeed/Private/Implementation/CdnNewsFeedTitleFile.cpp
@@ -70,7 +70,7 @@ bool FCdnNewsFeedTitleFile::ClearFile(const FString& FileName)
 	return false;
 }
 
-bool FCdnNewsFeedTitleFile::EnumerateFiles()
+bool FCdnNewsFeedTitleFile::EnumerateFiles(const FPagedQuery& Page)
 {
 	// Make sure an enumeration request  is not currently pending
 	if(!EnumerateFilesRequests.IsEmpty())
@@ -92,11 +92,6 @@ bool FCdnNewsFeedTitleFile::EnumerateFiles()
 	
 }	
 
-bool FCdnNewsFeedTitleFile::EnumerateFiles(int32 Start, int32 Count)
-{
-	return EnumerateFiles();
-}
-
 void FCdnNewsFeedTitleFile::GetFileList(TArray<FCloudFileHeader>& OutFiles) 
 {
 	TArray<FCloudFileHeader>* FilesPtr = &FileHeaders;
diff --git a/Engine/Source/Editor/NewsFeed/Private/Implementation/CdnNewsFeedTitleFile.h b/Engine/Source/Editor/NewsFeed/Private/Implementation/CdnNewsFeedTitleFile.h
--- a/Engine/Source/Editor/NewsFeed/Private/Implementation/CdnNewsFeedTitleFile.h
+++ b/Engine/Source/Editor/NewsFeed/Private/Implementation/CdnNewsFeedTitleFile.h
@@ -45,8 +45,7 @@ public:
 	virtual bool GetFileContents(const FString& DLName, TArray<uint8>& FileContents) override;
 	virtual bool ClearFiles() override;
 	virtual bool ClearFile(const FString& DLName) override;
-	virtual bool EnumerateFiles() override;
-	virtual bool EnumerateFiles(int32 Start, int32 Count) override;
+	virtual bool EnumerateFiles(const FPagedQuery& Page = FPagedQuery()) override;
 	virtual void GetFileList(TArray<FCloudFileHeader>& InFileHeaders) override;
 	virtual bool ReadFile(const FString& DLName) override;
 
diff --git a/Engine/Source/Editor/NewsFeed/Private/Implementation/LocalNewsFeedTitleFile.cpp b/Engine/Source/Editor/NewsFeed/Private/Implementation/LocalNewsFeedTitleFile.cpp
--- a/Engine/Source/Editor/NewsFeed/Private/Implementation/LocalNewsFeedTitleFile.cpp
+++ b/Engine/Source/Editor/NewsFeed/Private/Implementation/LocalNewsFeedTitleFile.cpp
@@ -58,7 +58,7 @@ bool FLocalNewsFeedTitleFile::ClearFile( const FString& DLName )
 }
 
 
-bool FLocalNewsFeedTitleFile::EnumerateFiles()
+bool FLocalNewsFeedTitleFile::EnumerateFiles(const FPagedQuery& Page)
 {
 	const FString WildCard = FPaths::Combine(*RootDirectory, TEXT("*"));
 
@@ -88,12 +88,6 @@ bool FLocalNewsFeedTitleFile::EnumerateFiles()
 	return Success;
 }
 
-bool FLocalNewsFeedTitleFile::EnumerateFiles(int32 Start, int32 Count)
-{
-	return EnumerateFiles();
-}
-
-
 void FLocalNewsFeedTitleFile::GetFileList(TArray<FCloudFileHeader>& InFileHeaders)
 {
 	InFileHeaders.Append( FileHeaders );
diff --git a/Engine/Source/Editor/NewsFeed/Private/Implementation/LocalNewsFeedTitleFile.h b/Engine/Source/Editor/NewsFeed/Private/Implementation/LocalNewsFeedTitleFile.h
--- a/Engine/Source/Editor/NewsFeed/Private/Implementation/LocalNewsFeedTitleFile.h
+++ b/Engine/Source/Editor/NewsFeed/Private/Implementation/LocalNewsFeedTitleFile.h
@@ -20,8 +20,7 @@ public:
 	virtual bool GetFileContents(const FString& DLName, TArray<uint8>& FileContents) override;
 	virtual bool ClearFiles() override;
 	virtual bool ClearFile(const FString& DLName) override;
-	virtual bool EnumerateFiles() override;
-	virtual bool EnumerateFiles(int32 Start, int32 Count) override;
+	virtual bool EnumerateFiles(const FPagedQuery& Page = FPagedQuery()) override;
 	virtual void GetFileList(TArray<FCloudFileHeader>& InFileHeaders) override;
 	virtual bool ReadFile(const FString& DLName) override;
 
diff --git a/Engine/Source/Runtime/Online/OnlineSubsystem/Public/Interfaces/OnlineTitleFileInterface.h b/Engine/Source/Runtime/Online/OnlineSubsystem/Public/Interfaces/OnlineTitleFileInterface.h
--- a/Engine/Source/Runtime/Online/OnlineSubsystem/Public/Interfaces/OnlineTitleFileInterface.h
+++ b/Engine/Source/Runtime/Online/OnlineSubsystem/Public/Interfaces/OnlineTitleFileInterface.h
@@ -59,22 +59,11 @@ public:
 	/**
 	* Requests a list of available files from the network store
 	*
-	* @param Start first file entry to enumerate (default 0)
-	* @param Count number of files to enumerate (default -1 for all)
+	* @param Page paging info to use for query
 	*
 	* @return true if the request has started, false if not
 	*/
-	virtual bool EnumerateFiles() = 0;
-
-	/**
-	 * Requests a list of available files from the network store
-	 *
-	 * @param Start first file entry to enumerate (default 0)
-	 * @param Count number of files to enumerate (default -1 for all)
-	 *
-	 * @return true if the request has started, false if not
-	 */
-	virtual bool EnumerateFiles(int32 Start, int32 Count) = 0;
+	virtual bool EnumerateFiles(const FPagedQuery& Page = FPagedQuery()) = 0;
 
 	/**
 	 * Delegate fired when the list of files has been returned from the network store
diff --git a/Engine/Source/Runtime/Online/OnlineSubsystem/Public/OnlineSubsystemTypes.h b/Engine/Source/Runtime/Online/OnlineSubsystem/Public/OnlineSubsystemTypes.h
--- a/Engine/Source/Runtime/Online/OnlineSubsystem/Public/OnlineSubsystemTypes.h
+++ b/Engine/Source/Runtime/Online/OnlineSubsystem/Public/OnlineSubsystemTypes.h
@@ -922,9 +922,32 @@ public:
 	virtual const FUniqueNetId& GetSessionId() const = 0;
 };
 
+/**
+ * Paging info needed for a request that can return paged results
+ */
+class FPagedQuery
+{
+public:
+	FPagedQuery(int32 InStart = 0, int32 InCount = -1)
+		: Start(InStart)
+		, Count(InCount)
+	{}
+
+	/** @return true if valid range */
+	bool IsValidRange() const
+	{
+		return Start >= 0 && Count >= 0;
+	}
+
+	/** first entry to fetch */
+	int32 Start;
+	/** total entries to fetch. -1 means ALL */
+	int32 Count;
+};
+
 /** Holds metadata about a given downloadable file */
 struct FCloudFileHeader
-{
+{	
 	/** Hash value, if applicable, of the given file contents */
     FString Hash;
 	/** Filename as downloaded */
