In-App purchase support
- Got this up and running again
- Fixed a threading issue with StoreKit communication
- Added new flag to switch on IAP, no longer depends on game center support. Off by default.
- Steamlined the setup for users, now only need product id. Other information is obtained from the server, and localized. No need to keep this locally, at the minute anyway.

#Note
- Documentation is incoming

diff --git a/Engine/Config/IOS/IOSEngine.ini b/Engine/Config/IOS/IOSEngine.ini
--- a/Engine/Config/IOS/IOSEngine.ini
+++ b/Engine/Config/IOS/IOSEngine.ini
@@ -12,4 +12,7 @@ DefaultProviderName=IOSAdvertising
 +EnabledPlugins="IOSDeviceProfileSelector"
 
 [DeviceProfileManager]
-DeviceProfileSelectionModule="IOSDeviceProfileSelector"
\ No newline at end of file
+DeviceProfileSelectionModule="IOSDeviceProfileSelector"
+
+[OnlineSubsystemIOS.Store]
+bSupportInAppPurchasing=False
\ No newline at end of file
diff --git a/Engine/Source/Runtime/Online/IOS/OnlineSubsystemIOS/Private/OnlineStoreInterfaceIOS.cpp b/Engine/Source/Runtime/Online/IOS/OnlineSubsystemIOS/Private/OnlineStoreInterfaceIOS.cpp
--- a/Engine/Source/Runtime/Online/IOS/OnlineSubsystemIOS/Private/OnlineStoreInterfaceIOS.cpp
+++ b/Engine/Source/Runtime/Online/IOS/OnlineSubsystemIOS/Private/OnlineStoreInterfaceIOS.cpp
@@ -105,10 +105,15 @@
 - (void)productsRequest: (SKProductsRequest *)request didReceiveResponse:(SKProductsResponse *)response
 {
 	UE_LOG(LogOnline, Display, TEXT( "FStoreKitHelper::didReceiveResponse" ));
+	// Write the achievements back to our Achievements array
+	[FIOSAsyncTask CreateTaskWithBlock : ^ bool(void)
+	{
+		IOnlineSubsystem* OnlineSub = IOnlineSubsystem::Get(FName(TEXT("IOS")));
+		FOnlineStoreInterfaceIOS* StoreInterface = (FOnlineStoreInterfaceIOS*)OnlineSub->GetStoreInterface().Get();
+		StoreInterface->ProcessProductsResponse(response);
 
-	IOnlineSubsystem* OnlineSub = IOnlineSubsystem::Get(FName(TEXT("IOS")));
-	FOnlineStoreInterfaceIOS* StoreInterface = (FOnlineStoreInterfaceIOS*)OnlineSub->GetStoreInterface().Get();
-	StoreInterface->ProcessProductsResponse( response );
+		return true;
+	}];
 
 	[request autorelease];
 }
@@ -127,27 +132,13 @@ FOnlineStoreInterfaceIOS::FOnlineStoreInterfaceIOS()
 	StoreHelper = [[FStoreKitHelper alloc] init];
 	[[SKPaymentQueue defaultQueue] addTransactionObserver:StoreHelper];
 
-	// Gather the microtransaction data dfrom the ini's
+	// Gather the microtransaction data from the ini's
 	TArray<FString> ProductIDs;
-	TArray<FString> DisplayNames;
-	TArray<FString> DisplayDescriptions;
-	TArray<FString> DisplayPrices;
-	
-	GConfig->GetArray(TEXT("OnlineSubsystemIOS.Store"), TEXT("ProductIDs"),		ProductIDs,				GEngineIni);
-	GConfig->GetArray(TEXT("OnlineSubsystemIOS.Store"), TEXT("DisplayNames"),	DisplayNames,			GEngineIni);
-	GConfig->GetArray(TEXT("OnlineSubsystemIOS.Store"), TEXT("DisplayDescriptions"),	DisplayDescriptions,	GEngineIni);
-	GConfig->GetArray(TEXT("OnlineSubsystemIOS.Store"), TEXT("DisplayPrices"),	DisplayPrices,			GEngineIni);
+	GConfig->GetArray(TEXT("OnlineSubsystemIOS.Store"), TEXT("ProductIDs"), ProductIDs, GEngineIni);
 
-	// @todo: add some sanity checking and logging to make sure the user has added the same number of entries to all of the arrays
 	for(int32 ProductIndex = 0; ProductIndex < ProductIDs.Num(); ProductIndex++)
 	{
-		FMicrotransactionPurchaseInfo Info; 
-
-		Info.Identifier = ProductIDs[ProductIndex];
-		Info.DisplayName = DisplayNames[ProductIndex];
-		Info.DisplayDescription = DisplayDescriptions[ProductIndex];
-		Info.DisplayPrice = DisplayPrices[ProductIndex];
-
+		FOnlineStoreInterfaceIOS::FMicrotransactionPurchaseInfo Info(ProductIDs[ProductIndex]);
 		AvailableProducts.Add( Info );
 	}
 
@@ -327,12 +318,31 @@ void FOnlineStoreInterfaceIOS::ProcessProductsResponse( SKProductsResponse* Resp
 			[numberFormatter setFormatterBehavior:NSNumberFormatterBehavior10_4];
 			[numberFormatter setNumberStyle:NSNumberFormatterCurrencyStyle];
 			[numberFormatter setLocale:Product.priceLocale];
-		
-			UE_LOG(LogOnline, Display, TEXT(": Product: %s, Price: %s"), *FString([Product localizedTitle]), *FString([numberFormatter stringFromNumber:Product.price]));
+
+			for( FOnlineStoreInterfaceIOS::FMicrotransactionPurchaseInfo& ProductInfo : AvailableProducts )
+			{
+				if( ProductInfo.Identifier == *FString([Product productIdentifier]) )
+				{
+					ProductInfo.DisplayName = [Product localizedTitle];
+					ProductInfo.DisplayDescription = [Product localizedDescription];
+					ProductInfo.DisplayPrice = [numberFormatter stringFromNumber : Product.price];
+
+					ProductInfo.bWasUpdatedByServer = true;
+
+					UE_LOG(LogOnline, Display, TEXT("\nProduct Identifier: %s, Name: %s, Description: %s, Price: %s\n"),
+						*ProductInfo.Identifier,
+						*ProductInfo.DisplayName,
+						*ProductInfo.DisplayDescription,
+						*ProductInfo.DisplayPrice);
+
+					break;
+				}
+			}
+
 			[numberFormatter release];
 		}
 		
-		for(NSString *invalidProduct in Response.invalidProductIdentifiers)
+		for( NSString *invalidProduct in Response.invalidProductIdentifiers )
 		{
 			UE_LOG(LogOnline, Display, TEXT("Problem in iTunes connect configuration for product: %s"), *FString(invalidProduct));
 		}
diff --git a/Engine/Source/Runtime/Online/IOS/OnlineSubsystemIOS/Private/OnlineStoreInterfaceIOS.h b/Engine/Source/Runtime/Online/IOS/OnlineSubsystemIOS/Private/OnlineStoreInterfaceIOS.h
--- a/Engine/Source/Runtime/Online/IOS/OnlineSubsystemIOS/Private/OnlineStoreInterfaceIOS.h
+++ b/Engine/Source/Runtime/Online/IOS/OnlineSubsystemIOS/Private/OnlineStoreInterfaceIOS.h
@@ -13,16 +13,6 @@
 #include <StoreKit/SKPaymentQueue.h>
 
 
-struct FMicrotransactionPurchaseInfo
-{
-	FString Identifier;
-	FString DisplayName;
-	FString DisplayDescription;
-	FString DisplayPrice;
-	FString CurrencyType;
-};
-
-
 @interface FStoreKitHelper : NSObject<SKProductsRequestDelegate, SKPaymentTransactionObserver>
 {
 };
@@ -66,6 +56,38 @@ public:
 	 */
 	void ProcessProductsResponse( SKProductsResponse* Response );
 
+
+	/**
+	 * Micro-transaction purchase information
+	 */
+	struct FMicrotransactionPurchaseInfo
+	{
+		FMicrotransactionPurchaseInfo(FString InProductIdentifier)
+		{
+			Identifier = InProductIdentifier;
+			bWasUpdatedByServer = false;
+		}
+
+		// The product identifier that matches the one in iTunesConnect
+		FString Identifier;
+
+		// Display name, gathered from querying the Apple Store
+		FString DisplayName;
+
+		// Display description, gathered from querying the Apple Store
+		FString DisplayDescription;
+
+		// Product price, gathered from querying the Apple Store
+		FString DisplayPrice;
+
+		// The currency type for this MT, gathered from querying the Apple Store
+		FString CurrencyType;
+
+		// flag that lets the user know if the data has been completed from the server
+		bool bWasUpdatedByServer;
+	};
+
+
 private:
 	FStoreKitHelper* StoreHelper;
 	
diff --git a/Engine/Source/Runtime/Online/IOS/OnlineSubsystemIOS/Private/OnlineSubsystemIOS.cpp b/Engine/Source/Runtime/Online/IOS/OnlineSubsystemIOS/Private/OnlineSubsystemIOS.cpp
--- a/Engine/Source/Runtime/Online/IOS/OnlineSubsystemIOS/Private/OnlineSubsystemIOS.cpp
+++ b/Engine/Source/Runtime/Online/IOS/OnlineSubsystemIOS/Private/OnlineSubsystemIOS.cpp
@@ -139,11 +139,15 @@ bool FOnlineSubsystemIOS::Init()
 		IdentityInterface = MakeShareable(new FOnlineIdentityIOS());
 		FriendsInterface = MakeShareable(new FOnlineFriendsIOS(this));
 		LeaderboardsInterface = MakeShareable(new FOnlineLeaderboardsIOS(this));
-		StoreInterface = MakeShareable(new FOnlineStoreInterfaceIOS());
 		AchievementsInterface = MakeShareable(new FOnlineAchievementsIOS(this));
 		ExternalUIInterface = MakeShareable(new FOnlineExternalUIIOS());
 	}
 
+	if( IsInAppPurchasingEnabled() )
+	{
+		StoreInterface = MakeShareable(new FOnlineStoreInterfaceIOS());
+	}
+
 	return bSuccessfullyStartedUp;
 }
 
@@ -180,3 +184,10 @@ bool FOnlineSubsystemIOS::IsEnabled()
 	GConfig->GetBool(TEXT("/Script/IOSRuntimeSettings.IOSRuntimeSettings"), TEXT("bEnableGameCenterSupport"), bEnableGameCenter, GEngineIni);
 	return bEnableGameCenter;
 }
+
+bool FOnlineSubsystemIOS::IsInAppPurchasingEnabled()
+{
+	bool bEnableIAP = false;
+	GConfig->GetBool(TEXT("OnlineSubsystemIOS.Store"), TEXT("bSupportInAppPurchasing"), bEnableIAP, GEngineIni);
+	return bEnableIAP;
+}
\ No newline at end of file
diff --git a/Engine/Source/Runtime/Online/IOS/OnlineSubsystemIOS/Public/OnlineSubsystemIOS.h b/Engine/Source/Runtime/Online/IOS/OnlineSubsystemIOS/Public/OnlineSubsystemIOS.h
--- a/Engine/Source/Runtime/Online/IOS/OnlineSubsystemIOS/Public/OnlineSubsystemIOS.h
+++ b/Engine/Source/Runtime/Online/IOS/OnlineSubsystemIOS/Public/OnlineSubsystemIOS.h
@@ -54,6 +54,12 @@ PACKAGE_SCOPE:
 	 */
 	bool IsEnabled();
 
+	/**
+	* Is IAP available for use
+	* @return true if IAP should be available, false otherwise
+	*/
+	bool IsInAppPurchasingEnabled();
+
 private:
 
 	/** Online async task thread */
