#online
-added some type checks for online json serializer

diff --git a/Engine/Source/Runtime/Online/OnlineSubsystem/Public/OnlineJsonSerializer.h b/Engine/Source/Runtime/Online/OnlineSubsystem/Public/OnlineJsonSerializer.h
--- a/Engine/Source/Runtime/Online/OnlineSubsystem/Public/OnlineJsonSerializer.h
+++ b/Engine/Source/Runtime/Online/OnlineSubsystem/Public/OnlineJsonSerializer.h
@@ -312,7 +312,7 @@ public:
 	 */
 	virtual void Serialize(const TCHAR* Name, int32& Value) override
 	{
-		if (JsonObject->HasField(Name))
+		if (JsonObject->HasTypedField<EJson::Number>(Name))
 		{
 			Value = FMath::TruncToInt(JsonObject->GetNumberField(Name));
 		}
@@ -325,7 +325,7 @@ public:
 	 */
 	virtual void Serialize(const TCHAR* Name, uint32& Value) override
 	{
-		if (JsonObject->HasField(Name))
+		if (JsonObject->HasTypedField<EJson::Number>(Name))
 		{
 			Value = FMath::TruncToInt(JsonObject->GetNumberField(Name));
 		}
@@ -338,7 +338,7 @@ public:
 	 */
 	virtual void Serialize(const TCHAR* Name, bool& Value) override
 	{
-		if (JsonObject->HasField(Name))
+		if (JsonObject->HasTypedField<EJson::Boolean>(Name))
 		{
 			Value = JsonObject->GetBoolField(Name);
 		}
@@ -351,7 +351,7 @@ public:
 	 */
 	virtual void Serialize(const TCHAR* Name, FString& Value) override
 	{
-		if (JsonObject->HasField(Name))
+		if (JsonObject->HasTypedField<EJson::String>(Name))
 		{
 			Value = JsonObject->GetStringField(Name);
 		}
@@ -364,7 +364,7 @@ public:
 	 */
 	virtual void Serialize(const TCHAR* Name, float& Value) override
 	{
-		if (JsonObject->HasField(Name))
+		if (JsonObject->HasTypedField<EJson::Number>(Name))
 		{
 			Value = JsonObject->GetNumberField(Name);
 		}
@@ -377,7 +377,7 @@ public:
 	*/
 	virtual void Serialize(const TCHAR* Name, double& Value) override
 	{
-		if (JsonObject->HasField(Name))
+		if (JsonObject->HasTypedField<EJson::Number>(Name))
 		{
 			Value = JsonObject->GetNumberField(Name);
 		}
@@ -401,7 +401,7 @@ public:
 	 */
 	virtual void SerializeArray(const TCHAR* Name, FJsonSerializableArray& Array) override
 	{
-		if (JsonObject->HasField(Name))
+		if (JsonObject->HasTypedField<EJson::Array>(Name))
 		{
 			TArray< TSharedPtr<FJsonValue> > JsonArray = JsonObject->GetArrayField(Name);
 			// Iterate all of the keys and their values
@@ -419,7 +419,7 @@ public:
 	 */
 	virtual void SerializeMap(const TCHAR* Name, FJsonSerializableKeyValueMap& Map) override
 	{
-		if (JsonObject->HasField(Name))
+		if (JsonObject->HasTypedField<EJson::Object>(Name))
 		{
 			TSharedPtr<FJsonObject> JsonMap = JsonObject->GetObjectField(Name);
 			// Iterate all of the keys and their values
@@ -439,7 +439,7 @@ public:
 	 */
 	virtual void SerializeMap(const TCHAR* Name, FJsonSerializableKeyValueMapInt& Map) override
 	{
-		if (JsonObject->HasField(Name))
+		if (JsonObject->HasTypedField<EJson::Object>(Name))
 		{
 			TSharedPtr<FJsonObject> JsonMap = JsonObject->GetObjectField(Name);
 			// Iterate all of the keys and their values
