#online,json
- allow for parsing of arrays of objects using online json macros

diff --git a/Engine/Source/Runtime/Online/OnlineSubsystem/Public/OnlineJsonSerializer.h b/Engine/Source/Runtime/Online/OnlineSubsystem/Public/OnlineJsonSerializer.h
--- a/Engine/Source/Runtime/Online/OnlineSubsystem/Public/OnlineJsonSerializer.h
+++ b/Engine/Source/Runtime/Online/OnlineSubsystem/Public/OnlineJsonSerializer.h
@@ -7,7 +7,7 @@
 /**
  * Macros used to generate a serialization function for a class derived from FJsonSerializable
  */
- #define BEGIN_ONLINE_JSON_SERIALIZER \
+#define BEGIN_ONLINE_JSON_SERIALIZER \
 	virtual void Serialize(FOnlineJsonSerializerBase& Serializer) override \
 	{ \
 		Serializer.StartObject();
@@ -16,6 +16,13 @@
 		Serializer.EndObject(); \
 	}
 
+#define BEGIN_ONLINE_JSON_SERIALIZER_FLAT \
+	virtual void Serialize(FOnlineJsonSerializerBase& Serializer) override \
+	{
+
+#define END_ONLINE_JSON_SERIALIZER_FLAT \
+	}
+
 #define ONLINE_JSON_SERIALIZE(JsonName, JsonValue) \
 		Serializer.Serialize(TEXT(JsonName), JsonValue)
 
@@ -28,6 +35,28 @@
 #define ONLINE_JSON_SERIALIZE_SERIALIZABLE(JsonName, JsonValue) \
 		JsonValue.Serialize(Serializer)
 
+#define ONLINE_JSON_SERIALIZE_ARRAY_SERIALIZABLE(JsonName, JsonArray, ElementType) \
+		if (Serializer.IsLoading()) \
+		{ \
+			if (Serializer.GetObject()->HasTypedField<EJson::Array>(JsonName)) \
+			{ \
+				for (TArray< TSharedPtr<FJsonValue> >::TConstIterator It(Serializer.GetObject()->GetArrayField(JsonName)); It; ++It) \
+				{ \
+					ElementType* Obj = new(JsonArray) ElementType(); \
+					Obj->FromJson((*It)->AsObject()); \
+				} \
+			} \
+		} \
+		else \
+		{ \
+			Serializer.StartArray(JsonName); \
+			for (auto It = JsonArray.CreateIterator(); It; ++It) \
+			{ \
+				It->Serialize(Serializer); \
+			} \
+			Serializer.EndArray(); \
+		}
+
 /** Array of string data */
 typedef TArray<FString> FJsonSerializableArray;
 
@@ -58,6 +87,7 @@ struct FOnlineJsonSerializerBase
 	virtual void SerializeArray(const TCHAR* Name, FJsonSerializableArray& Value) = 0;
 	virtual void SerializeMap(const TCHAR* Name, FJsonSerializableKeyValueMap& Map) = 0;
 	virtual void SerializeMap(const TCHAR* Name, FJsonSerializableKeyValueMapInt& Map) = 0;
+	virtual TSharedPtr<FJsonObject> GetObject() = 0;
 };
 
 /**
@@ -85,6 +115,8 @@ public:
 	virtual bool IsLoading() const override { return false; }
     /** Is the JSON being written to */
 	virtual bool IsSaving() const override { return true; }
+	/** Access to the root object */
+	virtual TSharedPtr<FJsonObject> GetObject() { return TSharedPtr<FJsonObject>(); }
 
 	/**
 	 * Starts a new object "{"
@@ -273,6 +305,8 @@ public:
 	virtual bool IsLoading() const override { return true; }
     /** Is the JSON being written to */
 	virtual bool IsSaving() const override { return false; }
+	/** Access to the root Json object being read */
+	virtual TSharedPtr<FJsonObject> GetObject() { return JsonObject; }
 
 	/** Ignored */
 	virtual void StartObject() override
