#online,json
- parsing of FDateTime entries from json string fields

diff --git a/Engine/Source/Runtime/Online/OnlineSubsystem/Public/OnlineJsonSerializer.h b/Engine/Source/Runtime/Online/OnlineSubsystem/Public/OnlineJsonSerializer.h
--- a/Engine/Source/Runtime/Online/OnlineSubsystem/Public/OnlineJsonSerializer.h
+++ b/Engine/Source/Runtime/Online/OnlineSubsystem/Public/OnlineJsonSerializer.h
@@ -83,6 +83,7 @@ struct FOnlineJsonSerializerBase
 	virtual void Serialize(const TCHAR* Name, FString& Value) = 0;
 	virtual void Serialize(const TCHAR* Name, float& Value) = 0;
 	virtual void Serialize(const TCHAR* Name, double& Value) = 0;
+	virtual void Serialize(const TCHAR* Name, FDateTime& Value) = 0;
 	virtual void SerializeArray(FJsonSerializableArray& Array) = 0;
 	virtual void SerializeArray(const TCHAR* Name, FJsonSerializableArray& Value) = 0;
 	virtual void SerializeMap(const TCHAR* Name, FJsonSerializableKeyValueMap& Map) = 0;
@@ -215,6 +216,19 @@ public:
 	{
 		JsonWriter->WriteValue(Name, (const double)Value);
 	}
+	/**
+	* Writes the field name and the corresponding value to the JSON data
+	*
+	* @param Name the field name to write out
+	* @param Value the value to write out
+	*/
+	virtual void Serialize(const TCHAR* Name, FDateTime& Value) override
+	{
+		if (Value.GetTicks() > 0)
+		{
+			JsonWriter->WriteValue(Name, Value.ToIso8601());
+		}
+	}
 	/**
 	 * Serializes an array of values
 	 *
@@ -416,6 +430,19 @@ public:
 			Value = JsonObject->GetNumberField(Name);
 		}
 	}
+	/**
+	* Writes the field name and the corresponding value to the JSON data
+	*
+	* @param Name the field name to write out
+	* @param Value the value to write out
+	*/
+	virtual void Serialize(const TCHAR* Name, FDateTime& Value) override
+	{
+		if (JsonObject->HasTypedField<EJson::String>(Name))
+		{
+			FDateTime::ParseIso8601(*JsonObject->GetStringField(Name), Value);
+		}
+	}
 	/**
 	 * Serializes an array of values
 	 *
