Improved code readability and logging of Game Center identity. Should make it clearer in future what the issues are.

diff --git a/Engine/Source/Runtime/Online/IOS/OnlineSubsystemIOS/Private/OnlineIdentityInterfaceIOS.cpp b/Engine/Source/Runtime/Online/IOS/OnlineSubsystemIOS/Private/OnlineIdentityInterfaceIOS.cpp
--- a/Engine/Source/Runtime/Online/IOS/OnlineSubsystemIOS/Private/OnlineIdentityInterfaceIOS.cpp
+++ b/Engine/Source/Runtime/Online/IOS/OnlineSubsystemIOS/Private/OnlineIdentityInterfaceIOS.cpp
@@ -5,8 +5,6 @@
 FOnlineIdentityIOS::FOnlineIdentityIOS() :
 	UniqueNetId( NULL )
 {
-	UE_LOG(LogOnline, Verbose, TEXT("FOnlineIdentityIOS::FOnlineIdentityIOS()"));
-
 	Login( 0, FOnlineAccountCredentials() );
 }
 
@@ -26,78 +24,79 @@ TArray<TSharedPtr<FUserOnlineAccount> > FOnlineIdentityIOS::GetAllUserAccounts()
 
 bool FOnlineIdentityIOS::Login(int32 LocalUserNum, const FOnlineAccountCredentials& AccountCredentials)
 {
-	UE_LOG(LogOnline, Verbose, TEXT("FOnlineIdentityIOS::Login"));
-
 	bool bStartedLogin = false;
-	
+
+	// Was the login handled by Game Center
 	if( GetLocalGameCenterUser() && 
 		GetLocalGameCenterUser().isAuthenticated )
 	{
-		UE_LOG(LogOnline, Verbose, TEXT("FOnlineIdentityIOS::Login - Already logged in and authenticated!"));
-
-		// Login was handled by gamecenter
-		FString PlayerId(GetLocalGameCenterUser().playerID);
+        // Now logged in
+		bStartedLogin = true;
+        
+		const FString PlayerId(GetLocalGameCenterUser().playerID);
 		UniqueNetId = MakeShareable( new FUniqueNetIdString( PlayerId ) );
 		TriggerOnLoginCompleteDelegates(LocalUserNum, true, *UniqueNetId, TEXT(""));
-		bStartedLogin = true;
+        
+        UE_LOG(LogOnline, Log, TEXT("The user %s has logged into Game Center"), *PlayerId);
 	}
 	else if([IOSAppDelegate GetDelegate].OSVersion >= 6.0f)
 	{
+		// Trigger the login event on the main thread.
 		bStartedLogin = true;
 		dispatch_async(dispatch_get_main_queue(), ^
 		{
 			[GetLocalGameCenterUser() setAuthenticateHandler:(^(UIViewController* viewcontroller, NSError *error)
 			{
-				FString ErrorMessage;
 				bool bWasSuccessful = false;
 
-				if (!error && viewcontroller)
+				// The login process has completed.
+				if (viewcontroller == nil)
 				{
-					[[IOSAppDelegate GetDelegate].IOSController
-					presentViewController:viewcontroller animated:YES completion:nil];
-				}
-				else
-				{
-					if (GetLocalGameCenterUser().isAuthenticated)
+					FString ErrorMessage;
+
+					if (GetLocalGameCenterUser().isAuthenticated == YES)
 					{
 						/* Perform additional tasks for the authenticated player here */
-						FString PlayerId(GetLocalGameCenterUser().playerID);
-						UE_LOG(LogOnline, Verbose, TEXT("FOnlineIdentityIOS::Player Authenticated >IOS6.0 - playerid: %s"), *PlayerId);
-						UniqueNetId = MakeShareable( new FUniqueNetIdString( PlayerId ) );
-						
+						const FString PlayerId(GetLocalGameCenterUser().playerID);
+						UniqueNetId = MakeShareable(new FUniqueNetIdString(PlayerId));
+
 						bWasSuccessful = true;
+						UE_LOG(LogOnline, Log, TEXT("The user %s has logged into Game Center"), *PlayerId);
 					}
 					else
 					{
-						UE_LOG(LogOnline, Verbose, TEXT("FOnlineIdentityIOS::Player NOT Authenticated >IOS6.0"));
-						ErrorMessage = TEXT("Player NOT Authenticated");
+						ErrorMessage = TEXT("The user could not be authenticated by Game Center");
+						UE_LOG(LogOnline, Log, TEXT("%s"), *ErrorMessage);
 					}
 
 					if (error)
 					{
-						NSDictionary *userInfo = [error userInfo];
-						NSString *errstr = [[userInfo objectForKey : NSUnderlyingErrorKey] localizedDescription];
-						UE_LOG(LogOnline, Warning, TEXT("FOnlineIdentityIOS::Login failed with error: %d [%s]"), (int32)(error.code), *FString(errstr));
+						NSString *errstr = [error localizedDescription];
+						UE_LOG(LogOnline, Warning, TEXT("Game Center login has failed. %s]"), *FString(errstr));
 					}
-				}
 
-				// Report back to the game thread whether this succeeded.
-				[FIOSAsyncTask CreateTaskWithBlock : ^ bool(void)
-				{
-					TSharedPtr<FUniqueNetIdString> UniqueIdForUser = UniqueNetId.IsValid() ? UniqueNetId : MakeShareable(new FUniqueNetIdString());
-					TriggerOnLoginCompleteDelegates(LocalUserNum, bWasSuccessful, *UniqueIdForUser, *ErrorMessage);
+					// Report back to the game thread whether this succeeded.
+					[FIOSAsyncTask CreateTaskWithBlock : ^ bool(void)
+					{
+						TSharedPtr<FUniqueNetIdString> UniqueIdForUser = UniqueNetId.IsValid() ? UniqueNetId : MakeShareable(new FUniqueNetIdString());
+						TriggerOnLoginCompleteDelegates(LocalUserNum, bWasSuccessful, *UniqueIdForUser, *ErrorMessage);
 
-					return true;
-				}];
+						return true;
+					}];
+				}
+				else
+				{
+					// Game Center has provided a view controller for us to login, we present it.
+					[[IOSAppDelegate GetDelegate].IOSController 
+						presentViewController:viewcontroller animated:YES completion:nil];
+				}
 			})];
 		});
 	}
 	else
 	{
-		UE_LOG(LogOnline, Verbose, TEXT("GameCenter login was unsuccessful"));
-
 		// User is not currently logged into game center
-		TriggerOnLoginCompleteDelegates(LocalUserNum, false, FUniqueNetIdString(), TEXT("Local player is invalid"));
+		TriggerOnLoginCompleteDelegates(LocalUserNum, false, FUniqueNetIdString(), TEXT("IOS version is not compatible with the game center implementation"));
 	}
 	
 	return bStartedLogin;
@@ -105,15 +104,12 @@ bool FOnlineIdentityIOS::Login(int32 LocalUserNum, const FOnlineAccountCredentia
 
 bool FOnlineIdentityIOS::Logout(int32 LocalUserNum)
 {
-	UE_LOG(LogOnline, Verbose, TEXT("FOnlineIdentityIOS::Logout"));
 	TriggerOnLogoutCompleteDelegates(LocalUserNum, true);
 	return false;
 }
 
 bool FOnlineIdentityIOS::AutoLogin(int32 LocalUserNum)
 {
-	UE_LOG(LogOnline, Verbose, TEXT("FOnlineIdentityIOS::AutoLogin"));
-	
 	return Login( LocalUserNum, FOnlineAccountCredentials() );
 }
 
@@ -123,24 +119,14 @@ ELoginStatus::Type FOnlineIdentityIOS::GetLoginStatus(int32 LocalUserNum) const
 
 	if(LocalUserNum < MAX_LOCAL_PLAYERS && GetLocalGameCenterUser() != NULL && GetLocalGameCenterUser().isAuthenticated == YES)
 	{
-		UE_LOG(LogOnline, Verbose, TEXT("FOnlineIdentityIOS::GetLoginStatus - Logged in"));
 		LoginStatus = ELoginStatus::LoggedIn;
 	}
-	else
-	{
-		UE_LOG(LogOnline, Verbose, TEXT("FOnlineIdentityIOS::GetLoginStatus - Not logged in"));
-	}
 
 	return LoginStatus;
 }
 
 TSharedPtr<FUniqueNetId> FOnlineIdentityIOS::GetUniquePlayerId(int32 LocalUserNum) const
 {
-	if( UniqueNetId.IsValid() )
-	{
-		UE_LOG(LogOnline, Verbose, TEXT("FOnlineIdentityIOS::GetUniquePlayerId %s"), *UniqueNetId->ToString());
-	}
-
 	return UniqueNetId;
 }
 
@@ -155,6 +141,7 @@ TSharedPtr<FUniqueNetId> FOnlineIdentityIOS::CreateUniquePlayerId(uint8* Bytes,
 			return MakeShareable(new FUniqueNetIdString(StrId));
 		}
 	}
+    
 	return NULL;
 }
 
@@ -165,28 +152,21 @@ TSharedPtr<FUniqueNetId> FOnlineIdentityIOS::CreateUniquePlayerId(const FString&
 
 FString FOnlineIdentityIOS::GetPlayerNickname(int32 LocalUserNum) const
 {
-	UE_LOG(LogOnline, Verbose, TEXT("FOnlineIdentityIOS::GetPlayerNickname"));
-
-
 	if (LocalUserNum < MAX_LOCAL_PLAYERS && GetLocalGameCenterUser() != NULL)
 	{
 		NSString* PersonaName = [GetLocalGameCenterUser() alias];
 		
         if (PersonaName != nil)
         {
-            NSString* UserNameString = [NSString stringWithFormat:@"FOnlineIdentityIOS::GetPlayerNickname - %@", PersonaName];
-            UE_LOG(LogOnline, Verbose, TEXT("%s"), *FString(UserNameString));
             return FString(PersonaName);
         }
 	}
 
-	UE_LOG(LogOnline, Verbose, TEXT("FOnlineIdentityIOS::GetPlayerNickname was empty"));
 	return FString();
 }
 
 FString FOnlineIdentityIOS::GetAuthToken(int32 LocalUserNum) const
 {
-	UE_LOG(LogOnline, Verbose, TEXT("FOnlineIdentityIOS::GetAuthToken not implemented"));
 	FString ResultToken;
 	return ResultToken;
 }
\ No newline at end of file
