#online,mcp,entitlements
- refactor of entitlements interface
- only used for querying list of available entitlements now

diff --git a/Engine/Source/Runtime/Online/OnlineSubsystem/Public/Interfaces/OnlineEntitlementsInterface.h b/Engine/Source/Runtime/Online/OnlineSubsystem/Public/Interfaces/OnlineEntitlementsInterface.h
--- a/Engine/Source/Runtime/Online/OnlineSubsystem/Public/Interfaces/OnlineEntitlementsInterface.h
+++ b/Engine/Source/Runtime/Online/OnlineSubsystem/Public/Interfaces/OnlineEntitlementsInterface.h
@@ -4,250 +4,80 @@
 #include "OnlineDelegateMacros.h"
 #include "OnlineSubsystemTypes.h"
 
+/** 
+ * unique identifier for entitlements
+ */
+typedef FString FUniqueEntitlementId;
+
 /**
  * Details of an entitlement
  */
-struct FEntitlement
+struct FOnlineEntitlement
 {
 	/** Unique Entitlement Id associated with this entitlement */
-	FString EntitlementId;
-	/** Id unique to this player's instance of this entitlement */
-	FString EntitlementInstanceId;
-	/** Quantity of this entitlement */
-	int32 Quantity;
-
+	FUniqueEntitlementId EntitlementId;
 
 	/**
 	 * Equality operator
 	 */
-	bool operator==(const FEntitlement& Other) const
-	{
-		return Other.EntitlementId == EntitlementId &&
-			Other.EntitlementInstanceId == EntitlementInstanceId &&
-			Other.Quantity == Quantity;
-	}
-
-	/**
-	 * Assignment operator.  Copies the other into this object.
-	 */
-	FEntitlement& operator=(const FEntitlement& Other)
-	{
-		if (this != &Other)
-		{
-			this->EntitlementId = Other.EntitlementId;
-			this->EntitlementInstanceId = Other.EntitlementInstanceId;
-			this->Quantity = Other.Quantity;
-		}
-		return *this;
-	}
-
-	FEntitlement(const FString& InEntitlementId, const FString& InEntitlemendInstanceId, int32 InQuantity)
-		: EntitlementId(InEntitlementId)
-		, EntitlementInstanceId(InEntitlemendInstanceId)
-		, Quantity(InQuantity)
-	{
-	}
-
-	FEntitlement()
+	bool operator==(const FOnlineEntitlement& Other) const
 	{
+		return Other.EntitlementId == EntitlementId;
 	}
 };
 
 
-/**
- * Delegate declaration for when an entitlements group is created
- *
- * @param bWasSuccessful true if server was contacted and a valid result received
- * @param Error string representing the error condition
- * @param UserId of the user who was granted entitlements in this callback
- */
-DECLARE_MULTICAST_DELEGATE_ThreeParams(FOnCreateEntitlementsGroupComplete, bool, const FString&, const FUniqueNetId&);
-typedef FOnCreateEntitlementsGroupComplete::FDelegate FOnCreateEntitlementsGroupCompleteDelegate;
-
-/**
- * Delegate declaration for when an entitlements group is deleted
- *
- * @param bWasSuccessful true if server was contacted and a valid result received
- * @param Error string representing the error condition
- * @param UserId of the user who was granted entitlements in this callback
- */
-DECLARE_MULTICAST_DELEGATE_ThreeParams(FOnDeleteEntitlementsGroupComplete, bool, const FString&, const FUniqueNetId&);
-typedef FOnDeleteEntitlementsGroupComplete::FDelegate FOnDeleteEntitlementsGroupCompleteDelegate;
-
 /**
  * Delegate declaration for when entitlements are enumerated
  *
  * @param bWasSuccessful true if server was contacted and a valid result received
- * @param Error string representing the error condition
  * @param UserId of the user who was granted entitlements in this callback
- */
-DECLARE_MULTICAST_DELEGATE_ThreeParams(FOnEnumerateEntitlementsComplete, bool, const FString&, const FUniqueNetId&);
-typedef FOnEnumerateEntitlementsComplete::FDelegate FOnEnumerateEntitlementsCompleteDelegate;
-
-/**
- * Delegate declaration for when an array of entitlements are granted
- *
- * @param bWasSuccessful true if server was contacted and a valid result received
  * @param Error string representing the error condition
- * @param UserId of the user who was granted entitlements in this callback
- * @param OutNewEntitlements an array of entitlements that were changed as a result of this call
- * @param OutModifiedEntitlements an array of entitlements that were changed as a result of this call
  */
-DECLARE_MULTICAST_DELEGATE_FiveParams(FOnGrantEntitlementsComplete, bool, const FString&, const FUniqueNetId&, TArray<FEntitlement>&, TArray<FEntitlement>&);
-typedef FOnGrantEntitlementsComplete::FDelegate FOnGrantEntitlementsCompleteDelegate;
+DECLARE_MULTICAST_DELEGATE_ThreeParams(FOnQueryEntitlementsComplete, bool, const FUniqueNetId&, const FString&);
+typedef FOnQueryEntitlementsComplete::FDelegate FOnQueryEntitlementsCompleteDelegate;
 
 /**
- * Delegate declaration for when an entitlement is consumed
- *
- * @param bWasSuccessful true if server was contacted and a valid result received
- * @param Error string representing the error condition
- * @param UserId of the user who was granted entitlements in this callback
- * @param OutNewEntitlements an array of entitlements that were changed as a result of this call
- * @param OutModifiedEntitlements an array of entitlements that were changed as a result of this call
- */
-DECLARE_MULTICAST_DELEGATE_FiveParams(FOnConsumeEntitlementComplete, bool, const FString&, const FUniqueNetId&, TArray<FEntitlement>&, TArray<FEntitlement>&);
-typedef FOnConsumeEntitlementComplete::FDelegate FOnConsumeEntitlementCompleteDelegate;
-
-
-/**
- * Interface for granting and retrieving user entitlements
+ * Interface for retrieving user entitlements
  */
 class IOnlineEntitlements
 {
-
-protected:
-	IOnlineEntitlements() {};
-
 public:
-	virtual ~IOnlineEntitlements() {};
-
-	/**
-	 * Checks if a user has a single entitlement in the local cache
-	 *
-	 * @param UserId the ID of the user to get this entitlement for
-	 * @param EntitlementId the ID of the entitlement to retrieve
-	 * @return true if the requested entitlement was found in the local cache for the user
-	 */
-	virtual bool HasEntitlement(const FUniqueNetId& UserId, const FString& EntitlementId) = 0;
-
 	/**
 	 * Checks for and retrieves a single cached entitlement for a user
 	 *
 	 * @param UserId the ID of the user to get this entitlement for
 	 * @param EntitlementId the ID of the entitlement to retrieve
-	 * @param OutEntitlement out parameter for copying the entitlement details into if found
-	 * @return true if the requested entitlement was found in the local cache for the user
+	 *
+	 * @return entitlement entry if found, null otherwise
 	 */
-	virtual bool GetEntitlement(const FUniqueNetId& UserId, const FString& EntitlementId, FEntitlement& OutEntitlement) = 0;
+	virtual TSharedPtr<FOnlineEntitlement> GetEntitlement(const FUniqueNetId& UserId, const FUniqueEntitlementId& EntitlementId) = 0;
 	
 	/**
 	 * Gets the cached entitlement set for the requested user
 	 *
 	 * @param UserId the ID of the user to get entitlements for
 	 * @param OutUserEntitlements out parameter to copy the user's entitlements into
-	 * @return true if the cached entitlements were available for the requested user
-	 */
-	virtual bool GetEntitlements(const FUniqueNetId& UserId, TMap<FString, FEntitlement>& OutUserEntitlements) = 0;
-
-	/**
-	 * Contacts server and creates the entitlements group where the user's entitlements will be stored
-	 *
-	 * @param UserId the ID of the user to act on
-	 * @return true if the operation started successfully
-	 */
-	virtual bool CreateEntitlementsGroup(const FUniqueNetId& UserId) = 0;
-
-	/**
-	 * Delegate instanced called when entitlements group creation has completed
-	 *
-	 * @param bWasSuccessful true if server was contacted and a valid result received
-	 * @param Error string representing the error condition
-	 * @param UserId of the user who was granted entitlements in this callback
-	 */
-	DEFINE_ONLINE_DELEGATE_THREE_PARAM(OnCreateEntitlementsGroupComplete, bool, const FString&, const FUniqueNetId&);
-
-	/**
-	 * Contacts server and deletes the entitlements group and all entitlements within it
-	 *
-	 * @param UserId the ID of the user to act on
-	 * @return true if the operation started successfully
-	 */
-	virtual bool DeleteEntitlementsGroup(const FUniqueNetId& UserId) = 0;
-	
-	/**
-	 * Delegate instanced called when entitlements group deletion has completed
-	 *
-	 * @param bWasSuccessful true if server was contacted and a valid result received
-	 * @param Error string representing the error condition
-	 * @param UserId of the user who was granted entitlements in this callback
 	 */
-	DEFINE_ONLINE_DELEGATE_THREE_PARAM(OnDeleteEntitlementsGroupComplete, bool, const FString&, const FUniqueNetId&);
+	virtual void GetAllEntitlements(const FUniqueNetId& UserId, TArray<TSharedRef<FOnlineEntitlement>>& OutUserEntitlements) = 0;
 
 	/**
 	 * Contacts server and retrieves the list of the user's entitlements, caching them locally
 	 *
 	 * @param UserId the ID of the user to act on
+	 *
 	 * @return true if the operation started successfully
 	 */
-	virtual bool EnumerateEntitlements(const FUniqueNetId& UserId) = 0;
+	virtual bool QueryEntitlements(const FUniqueNetId& UserId) = 0;
 	
 	/**
 	 * Delegate instanced called when enumerating entitlements has completed
 	 *
 	 * @param bWasSuccessful true if server was contacted and a valid result received
-	 * @param Error string representing the error condition
 	 * @param UserId of the user who was granted entitlements in this callback
-	 */
-	DEFINE_ONLINE_DELEGATE_THREE_PARAM(OnEnumerateEntitlementsComplete, bool, const FString&, const FUniqueNetId&);
-
-	/**
-	 * Contacts server and gives an entitlement to the user
-	 *
-	 * @param UserId the ID of the user to give this entitlement to
-	 * @param EntitlementId the entitlement to grant to this user
-	 * @return true if the operation started successfully
-	 */
-	virtual bool GrantEntitlement(const FUniqueNetId& UserId, const FString& EntitlementId) = 0;
-
-	/**
-	 * Contacts server and gives multiple entitlements to the user
-	 *
-	 * @param UserId the ID of the user to give this entitlement to
-	 * @param EntitlementIds the array of entitlementIds to grant to this user
-	 * @return true if the operation started successfully
-	 */
-	virtual bool GrantEntitlements(const FUniqueNetId& UserId, const TArray<FString>& EntitlementIds) = 0;
-
-	/**
-	 * Delegate instanced called when granting entitlements has completed
-	 *
-	 * @param bWasSuccessful true if server was contacted and a valid result received
 	 * @param Error string representing the error condition
-	 * @param UserId of the user who was granted entitlements in this callback
-	 * @param OutNewEntitlements an array of entitlements that were changed as a result of this call
-	 * @param OutModifiedEntitlements an array of entitlements that were changed as a result of this call
-	 */
-	DEFINE_ONLINE_DELEGATE_FIVE_PARAM(OnGrantEntitlementsComplete, bool, const FString&, const FUniqueNetId&, TArray<FEntitlement>&, TArray<FEntitlement>&);
-
-	/**
-	 * Contacts server and consumes a particular entitlement from the user
-	 *
-	 * @param UserId the ID of the user to act on
-	 * @param EntitlementId the entitlement to remove from this user
-	 * @return true if the operation started successfully
-	 */
-	virtual bool ConsumeEntitlement(const FUniqueNetId& UserId, const FString& EntitlementId) = 0;
-
-	/**
-	 * Delegate instanced called when deleting an entitlement has completed
-	 *
-	 * @param bWasSuccessful true if server was contacted and a valid result received
-	 * @param Error string representing the error condition
-	 * @param UserId of the user who was granted entitlements in this callback
-	 * @param OutNewEntitlements an array of entitlements that were changed as a result of this call
-	 * @param OutModifiedEntitlements an array of entitlements that were changed as a result of this call
 	 */
-	DEFINE_ONLINE_DELEGATE_FIVE_PARAM(OnConsumeEntitlementComplete, bool, const FString&, const FUniqueNetId&, TArray<FEntitlement>&, TArray<FEntitlement>&);
+	DEFINE_ONLINE_DELEGATE_THREE_PARAM(OnQueryEntitlementsComplete, bool, const FUniqueNetId&, const FString&);
 };
 
 typedef TSharedPtr<IOnlineEntitlements, ESPMode::ThreadSafe> IOnlineEntitlementsPtr;
diff --git a/Engine/Source/Runtime/Online/OnlineSubsystemUtils/Private/Tests/TestEntitlementsInterface.cpp b/Engine/Source/Runtime/Online/OnlineSubsystemUtils/Private/Tests/TestEntitlementsInterface.cpp
--- a/Engine/Source/Runtime/Online/OnlineSubsystemUtils/Private/Tests/TestEntitlementsInterface.cpp
+++ b/Engine/Source/Runtime/Online/OnlineSubsystemUtils/Private/Tests/TestEntitlementsInterface.cpp
@@ -18,19 +18,14 @@ void FTestEntitlementsInterface::Test(UWorld* InWorld)
 		EntitlementsOSS.IsValid())
 	{
 		// Add delegates for the various async calls
-		EntitlementsOSS->AddOnCreateEntitlementsGroupCompleteDelegate(OnCreateEntitlementsGroupCompleteDelegate);
-		EntitlementsOSS->AddOnGrantEntitlementsCompleteDelegate(OnGrantEntitlementsCompleteDelegate);
-		EntitlementsOSS->AddOnConsumeEntitlementCompleteDelegate(OnConsumeEntitlementCompleteDelegate);
-		EntitlementsOSS->AddOnEnumerateEntitlementsCompleteDelegate(OnEnumerateEntitlementsCompleteDelegate);
-		EntitlementsOSS->AddOnDeleteEntitlementsGroupCompleteDelegate(OnDeleteEntitlementsGroupCompleteDelegate);
+		EntitlementsOSS->AddOnQueryEntitlementsCompleteDelegate(OnQueryEntitlementsCompleteDelegate);
 
 		// kick off next test
 		StartNextTest();
 	}
 	else
 	{
-		UE_LOG(LogOnline, Warning,
-			TEXT("Failed to get account creation service Mcp "));
+		UE_LOG(LogOnline, Warning, TEXT("Failed to get entitlement service Mcp "));
 
 		// done with the test
 		FinishTest();
@@ -42,40 +37,9 @@ void FTestEntitlementsInterface::Test(UWorld* InWorld)
  */
 void FTestEntitlementsInterface::StartNextTest()
 {
-	// Always register user first, because this user will be used for all subsequent calls
-	if (bCreateEntitlementsGroup)
+	if (bQueryEntitlements)
 	{
-		EntitlementsOSS->CreateEntitlementsGroup(*UserId);
-	}
-	else if (bGrantEntitlement)
-	{
-		EntitlementsOSS->GrantEntitlement(*UserId, TEXT("entitlement1"));
-	}
-	else if (bGrantEntitlements)
-	{
-		TArray<FString> EntitlementIds;
-		EntitlementIds.Add(TEXT("entitlement1"));
-		EntitlementIds.Add(TEXT("entitlement2"));
-		EntitlementIds.Add(TEXT("entitlement2"));
-		EntitlementIds.Add(TEXT("entitlement3"));
-		// Add 4 of entitlement4 (should add quantity passed in here)
-		EntitlementIds.Add(TEXT("entitlement4"));
-		EntitlementIds.Add(TEXT("entitlement4"));
-		EntitlementIds.Add(TEXT("entitlement4"));
-		EntitlementIds.Add(TEXT("entitlement4"));
-		EntitlementsOSS->GrantEntitlements(*UserId, EntitlementIds);
-	}
-	else if (bConsumeEntitlement)
-	{
-		EntitlementsOSS->ConsumeEntitlement(*UserId, TEXT("entitlement4"));
-	}
-	else if (bEnumerateEntitlements)
-	{
-		EntitlementsOSS->EnumerateEntitlements(*UserId);
-	}
-	else if (bDeleteEntitlementsGroup)
-	{
-		EntitlementsOSS->DeleteEntitlementsGroup(*UserId);
+		EntitlementsOSS->QueryEntitlements(*UserId);
 	}
 	else
 	{
@@ -88,225 +52,22 @@ void FTestEntitlementsInterface::FinishTest()
 	if (EntitlementsOSS.IsValid())
 	{
 		// Clear delegates for the entitlements async calls
-		EntitlementsOSS->ClearOnCreateEntitlementsGroupCompleteDelegate(OnCreateEntitlementsGroupCompleteDelegate);
-		EntitlementsOSS->ClearOnGrantEntitlementsCompleteDelegate(OnGrantEntitlementsCompleteDelegate);
-		EntitlementsOSS->ClearOnConsumeEntitlementCompleteDelegate(OnConsumeEntitlementCompleteDelegate);
-		EntitlementsOSS->ClearOnEnumerateEntitlementsCompleteDelegate(OnEnumerateEntitlementsCompleteDelegate);
-		EntitlementsOSS->ClearOnDeleteEntitlementsGroupCompleteDelegate(OnDeleteEntitlementsGroupCompleteDelegate);
-
-		if (bDeleteEntitlementsGroup == true && UserId.IsValid())
-		{
-			// If we haven't already tried, then kick off a request to clear the profile, but don't wait around to see if it worked
-			// Not strictly needed, but nice to cleanup when possible
-			EntitlementsOSS->DeleteEntitlementsGroup(*UserId);
-		}
+		EntitlementsOSS->ClearOnQueryEntitlementsCompleteDelegate(OnQueryEntitlementsCompleteDelegate);
 	}
 
 	delete this;
 }
 
-void FTestEntitlementsInterface::OnCreateEntitlementsGroupComplete(bool bWasSuccessful, const FString& Error, const FUniqueNetId& )
-{
-	UE_LOG(LogOnline, Log, TEXT("created entitlements group. Result=%s Error=[%s]"), 
-		bWasSuccessful ? TEXT("true") : TEXT("false"), *Error);
-
-	// done with this part of the test
-	bCreateEntitlementsGroup = false;
-
-	// kick off next test
-	if (bWasSuccessful)
-	{
-		StartNextTest();
-	}
-	else
-	{
-		FinishTest();
-	}
-}
-
-void FTestEntitlementsInterface::OnGrantEntitlementsComplete(bool bWasSuccessful, const FString& Error, const FUniqueNetId& InUserId, TArray<FEntitlement>& NewEntitlements, TArray<FEntitlement>& ModifiedEntitlements)
-{
-	UE_LOG(LogOnline, Log, TEXT("granted entitlements. Result=%s Error=[%s]"), 
-		bWasSuccessful ? TEXT("true") : TEXT("false"), *Error);
-
-	// Same callback for single and plural grants, single is run first so turn it off first
-	if (bGrantEntitlement == true)
-	{
-		// @todo check the returned new/modified in addition to getting the full cache to check
-		// Check that the cached entitlements are correct after granting 1
-		TMap<FString, FEntitlement> Entitlements;
-		EntitlementsOSS->GetEntitlements(InUserId, Entitlements);
-		if (Entitlements.Num() != 1)
-		{
-			UE_LOG(LogOnline, Error, TEXT("Incorrect number of entitlement1 found after granting 1 entitlement: %d"), Entitlements.Num());
-			bWasSuccessful = false;
-		}
-		else
-		{
-			FEntitlement* TestEntitlement;
-			TestEntitlement = Entitlements.Find("entitlement1");
-			if(TestEntitlement == NULL)
-			{
-				UE_LOG(LogOnline, Error, TEXT("Unable to find entitlement1 after granting 1 entitlement"));
-				bWasSuccessful = false;
-			}
-			else
-			{
-				if(TestEntitlement->Quantity != 1)
-				{
-					UE_LOG(LogOnline, Error, TEXT("Incorrect quantity [%d] for entitlement1 after granting 1 entitlement"), TestEntitlement->Quantity);
-					bWasSuccessful = false;
-				}
-			}
-		}
-
-		bGrantEntitlement = false;
-	}
-	else
-	{
-		// Check on the plural entitlement granting results in local cache
-		TMap<FString, FEntitlement> Entitlements;
-		EntitlementsOSS->GetEntitlements(InUserId, Entitlements);
-		if (Entitlements.Num() != 4)
-		{
-			UE_LOG(LogOnline, Error, TEXT("Incorrect number of entitlements found after granting multiple entitlement: %d"), Entitlements.Num());
-			bWasSuccessful = false;
-		}
-		else
-		{
-			FEntitlement* TestEntitlement;
-			TestEntitlement = Entitlements.Find("entitlement4");
-			if(TestEntitlement == NULL)
-			{
-				UE_LOG(LogOnline, Error, TEXT("Unable to find entitlement4 after granting multiple entitlement"));
-				bWasSuccessful = false;
-			}
-			else
-			{
-				if(TestEntitlement->Quantity != 4)
-				{
-					UE_LOG(LogOnline, Error, TEXT("Incorrect quantity [%d] for entitlement4 after granting multiple entitlement"), TestEntitlement->Quantity);
-					bWasSuccessful = false;
-				}
-			}
-		}
-
-		bGrantEntitlements = false;
-	}
-
-	// kick off next test
-	if (bWasSuccessful)
-	{
-		StartNextTest();
-	}
-	else
-	{
-		FinishTest();
-	}
-}
-
-void FTestEntitlementsInterface::OnConsumeEntitlementComplete(bool bWasSuccessful, const FString& Error, const FUniqueNetId& InUserId, TArray<FEntitlement>& NewEntitlements, TArray<FEntitlement>& ModifiedEntitlements)
-{
-	UE_LOG(LogOnline, Log, TEXT("consumed entitlement. Result=%s Error=[%s]"), 
-		bWasSuccessful ? TEXT("true") : TEXT("false"), *Error);
-
-	// @todo check the returned new/modified in addition to getting the full cache to check
-	// Check on the entitlement consumption results in local cache
-	TMap<FString, FEntitlement> Entitlements;
-	EntitlementsOSS->GetEntitlements(InUserId, Entitlements);
-	if (Entitlements.Num() != 4)
-	{
-		UE_LOG(LogOnline, Error, TEXT("Incorrect number of entitlements found after granting multiple entitlement: %d"), Entitlements.Num());
-		bWasSuccessful = false;
-	}
-	else
-	{
-		FEntitlement* TestEntitlement;
-		TestEntitlement = Entitlements.Find("entitlement4");
-		if(TestEntitlement == NULL)
-		{
-			UE_LOG(LogOnline, Error, TEXT("Unable to find entitlement4 after granting multiple entitlement"));
-			bWasSuccessful = false;
-		}
-		else
-		{
-			if(TestEntitlement->Quantity != 3)
-			{
-				UE_LOG(LogOnline, Error, TEXT("Incorrect quantity [%d] for entitlement4 after granting multiple entitlement"), TestEntitlement->Quantity);
-				bWasSuccessful = false;
-			}
-		}
-	}
-
-	// done with this part of the test
-	bConsumeEntitlement = false;
-
-	// kick off next test
-	if (bWasSuccessful)
-	{
-		StartNextTest();
-	}
-	else
-	{
-		FinishTest();
-	}
-}
-
-void FTestEntitlementsInterface::OnEnumerateEntitlementsComplete(bool bWasSuccessful, const FString& Error, const FUniqueNetId& InUserId)
+void FTestEntitlementsInterface::OnQueryEntitlementsComplete(bool bWasSuccessful, const FUniqueNetId& UserId, const FString& Error)
 {
-	UE_LOG(LogOnline, Log, TEXT("enumerated entitlements. Result=%s Error=[%s]"), 
-		bWasSuccessful ? TEXT("true") : TEXT("false"), *Error);
+	UE_LOG(LogOnline, Log, TEXT("enumerated entitlements. UserId=%s Result=%s Error=[%s]"), 
+		*UserId.ToDebugString(), bWasSuccessful ? TEXT("true") : TEXT("false"), *Error);
 
 	// Now check em out
-	TMap<FString, FEntitlement> Entitlements;
-	EntitlementsOSS->GetEntitlements(InUserId, Entitlements);
-	if (Entitlements.Num() != 4)
-	{
-		UE_LOG(LogOnline, Error, TEXT("Incorrect number of entitlements found: %d"), Entitlements.Num());
-		return;
-	}
-	FEntitlement* TestEntitlement;
-	TestEntitlement = Entitlements.Find("entitlement4");
-	if(TestEntitlement == NULL)
-	{
-		UE_LOG(LogOnline, Error, TEXT("Unable to find entitlement4 after enumerating entitlements"));
-		return;
-	}
-	if(TestEntitlement->Quantity != 3)
-	{
-		// Should be 3 because we add 4 then consume 1
-		UE_LOG(LogOnline, Error, TEXT("Incorrect quantity [%d] for entitlement4 after enumerating entitlements"), TestEntitlement->Quantity);
-	}
-
-	// done with this part of the test
-	bEnumerateEntitlements = false;
-
-	// kick off next test
-	if (bWasSuccessful)
-	{
-		StartNextTest();
-	}
-	else
-	{
-		FinishTest();
-	}
-}
-
-void FTestEntitlementsInterface::OnDeleteEntitlementsGroupComplete(bool bWasSuccessful, const FString& Error, const FUniqueNetId& )
-{
-	UE_LOG(LogOnline, Log, TEXT("deleting entitlements group. Result=%s Error=[%s]"), 
-		bWasSuccessful ? TEXT("true") : TEXT("false"), *Error);
-
-	// done with this part of the test
-	bDeleteEntitlementsGroup = false;
+	TArray<TSharedRef<FOnlineEntitlement>> Entitlements;
+	EntitlementsOSS->GetAllEntitlements(UserId, Entitlements);
 
 	// kick off next test
-	if (bWasSuccessful)
-	{
-		StartNextTest();
-	}
-	else
-	{
-		FinishTest();
-	}
+	bQueryEntitlements = false;
+	StartNextTest();
 }
\ No newline at end of file
diff --git a/Engine/Source/Runtime/Online/OnlineSubsystemUtils/Private/Tests/TestEntitlementsInterface.h b/Engine/Source/Runtime/Online/OnlineSubsystemUtils/Private/Tests/TestEntitlementsInterface.h
--- a/Engine/Source/Runtime/Online/OnlineSubsystemUtils/Private/Tests/TestEntitlementsInterface.h
+++ b/Engine/Source/Runtime/Online/OnlineSubsystemUtils/Private/Tests/TestEntitlementsInterface.h
@@ -18,20 +18,10 @@
 	IOnlineEntitlementsPtr EntitlementsOSS;
 	
 	/** Delegates for callbacks of each test */
-	FOnCreateEntitlementsGroupCompleteDelegate OnCreateEntitlementsGroupCompleteDelegate;
-	FOnGrantEntitlementsCompleteDelegate OnGrantEntitlementsCompleteDelegate;
-	FOnConsumeEntitlementCompleteDelegate OnConsumeEntitlementCompleteDelegate;
-	FOnEnumerateEntitlementsCompleteDelegate OnEnumerateEntitlementsCompleteDelegate;
-	FOnDeleteEntitlementsGroupCompleteDelegate OnDeleteEntitlementsGroupCompleteDelegate;
+	FOnQueryEntitlementsCompleteDelegate OnQueryEntitlementsCompleteDelegate;
 
 	/** toggles for whether to run each test */
-	bool bRegisterUser;
-	bool bCreateEntitlementsGroup;
-	bool bEnumerateEntitlements;
-	bool bGrantEntitlement;
-	bool bConsumeEntitlement;
-	bool bGrantEntitlements;
-	bool bDeleteEntitlementsGroup;
+	bool bQueryEntitlements;
 
 	/** Account registered by the test that should be used for all tests */
 	TSharedPtr<FUniqueNetId> UserId;
@@ -55,76 +45,23 @@
 	 */
 	void FinishTest();
 
-	/**
-	 * Called when the async task for creating an entitlements group has completed
-	 *
-	 * @param bWasSuccessful true if server was contacted and a valid result received
-	 * @param Error string representing the error condition if any
-	 * @param UserId of the user who was granted entitlements in this callback
-	 */
-	void OnCreateEntitlementsGroupComplete(bool bWasSuccessful, const FString& Error, const FUniqueNetId& UserId);
-
-	/**
-	 * Called when the async task for deleting an entitlements group has completed
-	 *
-	 * @param bWasSuccessful true if server was contacted and a valid result received
-	 * @param Error string representing the error condition if any
-	 * @param UserId of the user who was granted entitlements in this callback
-	 */
-	void OnDeleteEntitlementsGroupComplete(bool bWasSuccessful, const FString& Error, const FUniqueNetId& UserId);
-
 	/**
 	 * Called when the async task for enumerating entitlements has completed
 	 *
 	 * @param bWasSuccessful true if server was contacted and a valid result received
-	 * @param Error string representing the error condition if any
-	 * @param UserId of the user who was granted entitlements in this callback
-	 */
-	void OnEnumerateEntitlementsComplete(bool bWasSuccessful, const FString& Error, const FUniqueNetId& UserId);
-
-	/**
-	 * Called when the async task for granting entitlements has completed
-	 *
-	 * @param bWasSuccessful true if server was contacted and a valid result received
-	 * @param Error string representing the error condition if any
 	 * @param UserId of the user who was granted entitlements in this callback
-	 * @param NewEntitlements new entitlements as a result of this call
-	 * @param ModifiedEntitlements an array of entitlements that were changed as a result of this call
-	 */
-	void OnGrantEntitlementsComplete(bool bWasSuccessful, const FString& Error, const FUniqueNetId& UserId, TArray<FEntitlement>& NewEntitlements, TArray<FEntitlement>& ModifiedEntitlements);
-
-	/**
-	 * Called when the async task for consuming an entitlement has completed
-	 *
-	 * @param bWasSuccessful true if server was contacted and a valid result received
 	 * @param Error string representing the error condition if any
-	 * @param UserId of the user who was granted entitlements in this callback
-	 * @param NewEntitlements new entitlements as a result of this call
-	 * @param ModifiedEntitlements an array of entitlements that were changed as a result of this call
 	 */
-	void OnConsumeEntitlementComplete(bool bWasSuccessful, const FString& Error, const FUniqueNetId& UserId, TArray<FEntitlement>& NewEntitlements, TArray<FEntitlement>& ModifiedEntitlements);
+	void OnQueryEntitlementsComplete(bool bWasSuccessful, const FUniqueNetId& UserId, const FString& Error);
 
  public:
-	/**
-	 * Constructor
-	 *
-	 */
+	
 	FTestEntitlementsInterface(const FString& InSubsystemName)
 		: SubsystemName(InSubsystemName)
 		, IdentityOSS(NULL)
 		, EntitlementsOSS(NULL)
-		, OnCreateEntitlementsGroupCompleteDelegate(FOnCreateEntitlementsGroupCompleteDelegate::CreateRaw(this, &FTestEntitlementsInterface::OnCreateEntitlementsGroupComplete))
-		, OnGrantEntitlementsCompleteDelegate(FOnGrantEntitlementsCompleteDelegate::CreateRaw(this, &FTestEntitlementsInterface::OnGrantEntitlementsComplete))
-		, OnConsumeEntitlementCompleteDelegate(FOnConsumeEntitlementCompleteDelegate::CreateRaw(this, &FTestEntitlementsInterface::OnConsumeEntitlementComplete))
-		, OnEnumerateEntitlementsCompleteDelegate(FOnEnumerateEntitlementsCompleteDelegate::CreateRaw(this, &FTestEntitlementsInterface::OnEnumerateEntitlementsComplete))
-		, OnDeleteEntitlementsGroupCompleteDelegate(FOnDeleteEntitlementsGroupCompleteDelegate::CreateRaw(this, &FTestEntitlementsInterface::OnDeleteEntitlementsGroupComplete))
-		, bRegisterUser(true)
-		, bCreateEntitlementsGroup(true)
-		, bEnumerateEntitlements(true)
-		, bGrantEntitlement(true)
-		, bConsumeEntitlement(true)
-		, bGrantEntitlements(true)
-		, bDeleteEntitlementsGroup(true)
+		, OnQueryEntitlementsCompleteDelegate(FOnQueryEntitlementsCompleteDelegate::CreateRaw(this, &FTestEntitlementsInterface::OnQueryEntitlementsComplete))
+		, bQueryEntitlements(true)
 		, LocalUserIdx(0)
 	{
 	}
