#online,mcp,entitlements
- refactor of existing interface to simply query for list of entitlements
- implemented for Mcp OSS

diff --git a/Engine/Source/Runtime/Online/OnlineSubsystem/Public/Interfaces/OnlineEntitlementsInterface.h b/Engine/Source/Runtime/Online/OnlineSubsystem/Public/Interfaces/OnlineEntitlementsInterface.h
--- a/Engine/Source/Runtime/Online/OnlineSubsystem/Public/Interfaces/OnlineEntitlementsInterface.h
+++ b/Engine/Source/Runtime/Online/OnlineSubsystem/Public/Interfaces/OnlineEntitlementsInterface.h
@@ -15,14 +15,16 @@ typedef FString FUniqueEntitlementId;
 struct FOnlineEntitlement
 {
 	/** Unique Entitlement Id associated with this entitlement */
-	FUniqueEntitlementId EntitlementId;
+	FUniqueEntitlementId Id;
+	/** Display name for the entitlement */
+	FString Name;
 
 	/**
 	 * Equality operator
 	 */
 	bool operator==(const FOnlineEntitlement& Other) const
 	{
-		return Other.EntitlementId == EntitlementId;
+		return Other.Id == Id;
 	}
 };
 
@@ -68,7 +70,7 @@ public:
 	 *
 	 * @return true if the operation started successfully
 	 */
-	virtual bool QueryEntitlements(const FUniqueNetId& UserId) = 0;
+	virtual bool QueryEntitlements(const FUniqueNetId& UserId, const FPagedQuery& Page = FPagedQuery()) = 0;
 	
 	/**
 	 * Delegate instanced called when enumerating entitlements has completed
diff --git a/Engine/Source/Runtime/Online/OnlineSubsystemUtils/Private/Tests/TestEntitlementsInterface.cpp b/Engine/Source/Runtime/Online/OnlineSubsystemUtils/Private/Tests/TestEntitlementsInterface.cpp
--- a/Engine/Source/Runtime/Online/OnlineSubsystemUtils/Private/Tests/TestEntitlementsInterface.cpp
+++ b/Engine/Source/Runtime/Online/OnlineSubsystemUtils/Private/Tests/TestEntitlementsInterface.cpp
@@ -8,24 +8,33 @@ void FTestEntitlementsInterface::Test(UWorld* InWorld)
 {
 	IOnlineSubsystem* OSS = Online::GetSubsystem(InWorld, SubsystemName.Len() ? FName(*SubsystemName, FNAME_Find) : NAME_None);
 
-	EntitlementsOSS = OSS->GetEntitlementsInterface();
 	IdentityOSS = OSS->GetIdentityInterface();
 	if (IdentityOSS.IsValid())
 	{
 		UserId = IdentityOSS->GetUniquePlayerId(LocalUserIdx);
 	}
-	if (UserId.IsValid() && UserId->IsValid() &&
-		EntitlementsOSS.IsValid())
+	if (UserId.IsValid() && UserId->IsValid())
 	{
-		// Add delegates for the various async calls
-		EntitlementsOSS->AddOnQueryEntitlementsCompleteDelegate(OnQueryEntitlementsCompleteDelegate);
+		EntitlementsOSS = OSS->GetEntitlementsInterface();
+		if (EntitlementsOSS.IsValid())
+		{
+			// Add delegates for the various async calls
+			EntitlementsOSS->AddOnQueryEntitlementsCompleteDelegate(OnQueryEntitlementsCompleteDelegate);
 
-		// kick off next test
-		StartNextTest();
+			// kick off next test
+			StartNextTest();
+		}
+		else
+		{
+			UE_LOG(LogOnline, Warning, TEXT("Entitlement test failed. Failed to get entitlement service API"));
+
+			// done with the test
+			FinishTest();
+		}
 	}
 	else
 	{
-		UE_LOG(LogOnline, Warning, TEXT("Failed to get entitlement service Mcp "));
+		UE_LOG(LogOnline, Warning, TEXT("Entitlement test failed. No logged in user"));
 
 		// done with the test
 		FinishTest();
@@ -67,6 +76,13 @@ void FTestEntitlementsInterface::OnQueryEntitlementsComplete(bool bWasSuccessful
 	TArray<TSharedRef<FOnlineEntitlement>> Entitlements;
 	EntitlementsOSS->GetAllEntitlements(UserId, Entitlements);
 
+	for (auto It = Entitlements.CreateConstIterator(); It; ++It)
+	{
+		const TSharedRef<FOnlineEntitlement> Entry = *It;
+		UE_LOG(LogOnline, Log, TEXT("	entitlement id=%s name=%s"),
+			*Entry->Id, *Entry->Name);
+	}
+
 	// kick off next test
 	bQueryEntitlements = false;
 	StartNextTest();
