#online,mcp,entitlements
- keep track of rest of entitlement values as key/val attributes

diff --git a/Engine/Source/Runtime/Online/OnlineSubsystem/Public/Interfaces/OnlineEntitlementsInterface.h b/Engine/Source/Runtime/Online/OnlineSubsystem/Public/Interfaces/OnlineEntitlementsInterface.h
--- a/Engine/Source/Runtime/Online/OnlineSubsystem/Public/Interfaces/OnlineEntitlementsInterface.h
+++ b/Engine/Source/Runtime/Online/OnlineSubsystem/Public/Interfaces/OnlineEntitlementsInterface.h
@@ -19,6 +19,14 @@ struct FOnlineEntitlement
 	/** Display name for the entitlement */
 	FString Name;
 
+	/**
+	 * @return Any additional data associated with the entitlement
+ 	 */
+	virtual bool GetAttribute(const FString& AttrName, FString& OutAttrValue) const
+	{
+		return false;
+	}
+
 	/**
 	 * Equality operator
 	 */
