#online,mcp,entitlements
- added reference to item id for an entitlement

diff --git a/Engine/Source/Runtime/Online/OnlineSubsystem/Public/Interfaces/OnlineEntitlementsInterface.h b/Engine/Source/Runtime/Online/OnlineSubsystem/Public/Interfaces/OnlineEntitlementsInterface.h
--- a/Engine/Source/Runtime/Online/OnlineSubsystem/Public/Interfaces/OnlineEntitlementsInterface.h
+++ b/Engine/Source/Runtime/Online/OnlineSubsystem/Public/Interfaces/OnlineEntitlementsInterface.h
@@ -18,6 +18,8 @@ struct FOnlineEntitlement
 	FUniqueEntitlementId Id;
 	/** Display name for the entitlement */
 	FString Name;
+	/** Id for the item that this entitlement is associated with */
+	FString ItemId;
 
 	/**
 	 * @return Any additional data associated with the entitlement
