#online,mcp,entitlements
- find an entitlement entry by associated item id

diff --git a/Engine/Source/Runtime/Online/OnlineSubsystem/Public/Interfaces/OnlineEntitlementsInterface.h b/Engine/Source/Runtime/Online/OnlineSubsystem/Public/Interfaces/OnlineEntitlementsInterface.h
--- a/Engine/Source/Runtime/Online/OnlineSubsystem/Public/Interfaces/OnlineEntitlementsInterface.h
+++ b/Engine/Source/Runtime/Online/OnlineSubsystem/Public/Interfaces/OnlineEntitlementsInterface.h
@@ -64,6 +64,16 @@ public:
 	 * @return entitlement entry if found, null otherwise
 	 */
 	virtual TSharedPtr<FOnlineEntitlement> GetEntitlement(const FUniqueNetId& UserId, const FUniqueEntitlementId& EntitlementId) = 0;
+
+	/**
+	 * Checks for and retrieves a single cached entitlement for a user
+	 *
+	 * @param UserId the ID of the user to get this entitlement for
+	 * @param ItemId the ID of the item to retrieve an entitlement for
+	 *
+	 * @return entitlement entry if found, null otherwise
+	 */
+	virtual TSharedPtr<FOnlineEntitlement> GetItemEntitlement(const FUniqueNetId& UserId, const FString& ItemId) = 0;
 	
 	/**
 	 * Gets the cached entitlement set for the requested user
