#online,json
- added TMap serialization of nested objects

diff --git a/Engine/Source/Runtime/Online/OnlineSubsystem/Public/OnlineJsonSerializer.h b/Engine/Source/Runtime/Online/OnlineSubsystem/Public/OnlineJsonSerializer.h
--- a/Engine/Source/Runtime/Online/OnlineSubsystem/Public/OnlineJsonSerializer.h
+++ b/Engine/Source/Runtime/Online/OnlineSubsystem/Public/OnlineJsonSerializer.h
@@ -40,7 +40,7 @@
 		{ \
 			if (Serializer.GetObject()->HasTypedField<EJson::Array>(JsonName)) \
 			{ \
-				for (TArray< TSharedPtr<FJsonValue> >::TConstIterator It(Serializer.GetObject()->GetArrayField(JsonName)); It; ++It) \
+				for (auto It = Serializer.GetObject()->GetArrayField(JsonName).CreateConstIterator(); It; ++It) \
 				{ \
 					ElementType* Obj = new(JsonArray) ElementType(); \
 					Obj->FromJson((*It)->AsObject()); \
@@ -57,6 +57,32 @@
 			Serializer.EndArray(); \
 		}
 
+#define ONLINE_JSON_SERIALIZE_MAP_SERIALIZABLE(JsonName, JsonMap, ElementType) \
+		if (Serializer.IsLoading()) \
+		{ \
+			if (Serializer.GetObject()->HasTypedField<EJson::Object>(JsonName)) \
+			{ \
+				TSharedPtr<FJsonObject> JsonObj = Serializer.GetObject()->GetObjectField(JsonName); \
+				for (auto MapIt = JsonObj->Values.CreateConstIterator(); MapIt; ++MapIt) \
+				{ \
+					ElementType NewEntry; \
+					NewEntry.FromJson(MapIt.Value()->AsObject()); \
+					JsonMap.Add(MapIt.Key(), NewEntry); \
+				} \
+			} \
+		} \
+		else \
+		{ \
+			Serializer.StartObject(JsonName); \
+			for (auto It = JsonMap.CreateIterator(); It; ++It) \
+			{ \
+				Serializer.StartObject(It.Key()); \
+				It.Value().Serialize(Serializer); \
+				Serializer.EndObject(); \
+			} \
+			Serializer.EndObject(); \
+		}
+
 /** Array of string data */
 typedef TArray<FString> FJsonSerializableArray;
 
