#online,http
- warning logs of response headers on bad server response

diff --git a/Engine/Source/Runtime/Online/HTTP/Private/Windows/HttpWinInet.cpp b/Engine/Source/Runtime/Online/HTTP/Private/Windows/HttpWinInet.cpp
--- a/Engine/Source/Runtime/Online/HTTP/Private/Windows/HttpWinInet.cpp
+++ b/Engine/Source/Runtime/Online/HTTP/Private/Windows/HttpWinInet.cpp
@@ -623,8 +623,11 @@ void FHttpRequestWinInet::FinishedRequest()
 		UE_LOG(LogHttp, Log, TEXT("Finished request %p. response=%d %s url=%s elapsed=%.3f DownloadSize=%d"), 
 			this, Response->GetResponseCode(), *GetVerb(), *GetURL(), ElapsedTime, Response->GetContentLength());
 
+		const bool bDebugServerResponse = Response->GetResponseCode() >= 500 && Response->GetResponseCode() <= 505;
+
 		// log info about cloud front to identify failed downloads
-		if (UE_LOG_ACTIVE(LogHttp, Verbose))
+		if (UE_LOG_ACTIVE(LogHttp, Verbose) ||
+			bDebugServerResponse)
 		{
 			TArray<FString> AllHeaders = Response->GetAllHeaders();
 			for (TArray<FString>::TConstIterator It(AllHeaders); It; ++It)
@@ -632,7 +635,14 @@ void FHttpRequestWinInet::FinishedRequest()
 				const FString& HeaderStr = *It;
 				if (!HeaderStr.Contains(TEXT("Authorization")))
 				{
-					UE_LOG(LogHttp, Verbose, TEXT("%p Response Header %s"), this, *HeaderStr);
+					if (bDebugServerResponse)
+					{
+						UE_LOG(LogHttp, Warning, TEXT("%p Response Header %s"), this, *HeaderStr);
+					}
+					else
+					{
+						UE_LOG(LogHttp, Verbose, TEXT("%p Response Header %s"), this, *HeaderStr);
+					}
 				}
 			}
 		}
