#UE4 - small code cleanup

diff --git a/Engine/Source/Runtime/Online/OnlineSubsystemUtils/Private/OnlineSessionClient.cpp b/Engine/Source/Runtime/Online/OnlineSubsystemUtils/Private/OnlineSessionClient.cpp
--- a/Engine/Source/Runtime/Online/OnlineSubsystemUtils/Private/OnlineSessionClient.cpp
+++ b/Engine/Source/Runtime/Online/OnlineSubsystemUtils/Private/OnlineSessionClient.cpp
@@ -80,12 +80,13 @@ void UOnlineSessionClient::ClearOnlineDelegates(UWorld* InWorld)
 	IOnlineSubsystem* OnlineSub = Online::GetSubsystem(InWorld);
 	if (OnlineSub)
 	{
-		IOnlineSessionPtr OnlineSessionInt = OnlineSub->GetSessionInterface();
-
-		int32 ControllerId = GetControllerId();
-		if (ControllerId != INVALID_CONTROLLERID)
+		if (SessionInt.IsValid())
 		{
-			OnlineSessionInt->ClearOnSessionInviteAcceptedDelegate(ControllerId, OnSessionInviteAcceptedDelegate);
+			int32 ControllerId = GetControllerId();
+			if (ControllerId != INVALID_CONTROLLERID)
+			{
+				SessionInt->ClearOnSessionInviteAcceptedDelegate(ControllerId, OnSessionInviteAcceptedDelegate);
+			}
 		}
 	}
 
