#UE4 - changes to voip delegate processing
- prevent some possible crashes
- fix LIVE implementation (NULL can't call ToSharedRef)

diff --git a/Engine/Source/Runtime/Online/OnlineSubsystemSteam/Private/VoiceInterfaceSteam.cpp b/Engine/Source/Runtime/Online/OnlineSubsystemSteam/Private/VoiceInterfaceSteam.cpp
--- a/Engine/Source/Runtime/Online/OnlineSubsystemSteam/Private/VoiceInterfaceSteam.cpp
+++ b/Engine/Source/Runtime/Online/OnlineSubsystemSteam/Private/VoiceInterfaceSteam.cpp
@@ -211,7 +211,14 @@ bool FOnlineVoiceSteam::UnregisterLocalTalker(uint32 LocalUserNum)
 			if (OnPlayerTalkingStateChangedDelegates.IsBound() && (Talker.bIsTalking || Talker.bWasTalking))
 			{
 				TSharedPtr<FUniqueNetId> UniqueId = IdentityInt->GetUniquePlayerId(LocalUserNum);
-				OnPlayerTalkingStateChangedDelegates.Broadcast(UniqueId.ToSharedRef(), Talker.bIsTalking);
+				if (UniqueId.IsValid())
+				{
+					OnPlayerTalkingStateChangedDelegates.Broadcast(UniqueId.ToSharedRef(), false);
+				}
+				else
+				{
+					UE_LOG(LogVoice, Warning, TEXT("Invalid UserId for local player %d in UnregisterLocalTalker"), LocalUserNum);
+				}
 			}
 
 			// Remove them from engine too
@@ -335,7 +342,7 @@ void FOnlineVoiceSteam::RemoveAllRemoteTalkers()
 		{
 			const FRemoteTalker& Talker = RemoteTalkers[Index];
 
-			if (OnPlayerTalkingStateChangedDelegates.IsBound() && Talker.bIsTalking)
+			if (OnPlayerTalkingStateChangedDelegates.IsBound() && (Talker.bIsTalking || Talker.bWasTalking))
 			{
 				OnPlayerTalkingStateChangedDelegates.Broadcast(Talker.TalkerId.ToSharedRef(), false);
 			}
diff --git a/Engine/Source/Runtime/Online/OnlineSubsystemUtils/Private/VoiceInterfaceImpl.cpp b/Engine/Source/Runtime/Online/OnlineSubsystemUtils/Private/VoiceInterfaceImpl.cpp
--- a/Engine/Source/Runtime/Online/OnlineSubsystemUtils/Private/VoiceInterfaceImpl.cpp
+++ b/Engine/Source/Runtime/Online/OnlineSubsystemUtils/Private/VoiceInterfaceImpl.cpp
@@ -210,7 +210,14 @@ bool FOnlineVoiceImpl::UnregisterLocalTalker(uint32 LocalUserNum)
 			if (OnPlayerTalkingStateChangedDelegates.IsBound() && (Talker.bIsTalking || Talker.bWasTalking))
 			{
 				TSharedPtr<FUniqueNetId> UniqueId = IdentityInt->GetUniquePlayerId(LocalUserNum);
-				OnPlayerTalkingStateChangedDelegates.Broadcast(UniqueId.ToSharedRef(), Talker.bIsTalking);
+				if (UniqueId.IsValid())
+				{
+					OnPlayerTalkingStateChangedDelegates.Broadcast(UniqueId.ToSharedRef(), false);
+				}
+				else
+				{
+					UE_LOG(LogVoice, Warning, TEXT("Invalid UserId for local player %d in UnregisterLocalTalker"), LocalUserNum);
+				}
 			}
 
 			// Remove them from engine too
@@ -334,7 +341,7 @@ void FOnlineVoiceImpl::RemoveAllRemoteTalkers()
 		{
 			const FRemoteTalker& Talker = RemoteTalkers[Index];
 
-			if (OnPlayerTalkingStateChangedDelegates.IsBound() && Talker.bIsTalking)
+			if (OnPlayerTalkingStateChangedDelegates.IsBound() && (Talker.bIsTalking || Talker.bWasTalking))
 			{
 				OnPlayerTalkingStateChangedDelegates.Broadcast(Talker.TalkerId.ToSharedRef(), false);
 			}
