Exposed more data fields on Offer/Item. Fleshed out the skeleton for the Marketplace Details page.

diff --git a/Engine/Source/Runtime/Online/OnlineSubsystem/Public/OnlineJsonSerializer.h b/Engine/Source/Runtime/Online/OnlineSubsystem/Public/OnlineJsonSerializer.h
--- a/Engine/Source/Runtime/Online/OnlineSubsystem/Public/OnlineJsonSerializer.h
+++ b/Engine/Source/Runtime/Online/OnlineSubsystem/Public/OnlineJsonSerializer.h
@@ -107,6 +107,7 @@ struct FOnlineJsonSerializerBase
 	virtual void Serialize(const TCHAR* Name, uint32& Value) = 0;
 	virtual void Serialize(const TCHAR* Name, bool& Value) = 0;
 	virtual void Serialize(const TCHAR* Name, FString& Value) = 0;
+	virtual void Serialize(const TCHAR* Name, FText& Value) = 0;
 	virtual void Serialize(const TCHAR* Name, float& Value) = 0;
 	virtual void Serialize(const TCHAR* Name, double& Value) = 0;
 	virtual void Serialize(const TCHAR* Name, FDateTime& Value) = 0;
@@ -222,6 +223,16 @@ public:
 	{
 		JsonWriter->WriteValue(Name, Value);
 	}
+	/**
+	 * Writes the field name and the corresponding value to the JSON data
+	 *
+	 * @param Name the field name to write out
+	 * @param Value the value to write out
+	 */
+	virtual void Serialize(const TCHAR* Name, FText& Value) override
+	{
+		JsonWriter->WriteValue(Name, Value.ToString());
+	}
 	/**
 	 * Writes the field name and the corresponding value to the JSON data
 	 *
@@ -430,6 +441,19 @@ public:
 			Value = JsonObject->GetStringField(Name);
 		}
 	}
+	/**
+	 * If the underlying json object has the field, it is read into the value
+	 *
+	 * @param Name the name of the field to read
+	 * @param Value the out value to read the data into
+	 */
+	virtual void Serialize(const TCHAR* Name, FText& Value) override
+	{
+		if (JsonObject->HasTypedField<EJson::String>(Name))
+		{
+			Value = FText::FromString(JsonObject->GetStringField(Name));
+		}
+	}
 	/**
 	 * If the underlying json object has the field, it is read into the value
 	 *
