Improved In-App purchase support
- Adds Blueprint functionality to query and transact available products.
- General improvements to the IAP functionality on IOS
- Improved error handling and logging.

#Note - Further improvements may be incoming as we start using this.

diff --git a/Engine/Source/Editor/OnlineBlueprintSupport/Classes/K2Node_InAppPurchase.h b/Engine/Source/Editor/OnlineBlueprintSupport/Classes/K2Node_InAppPurchase.h
new file mode 100644
--- /dev/null
+++ b/Engine/Source/Editor/OnlineBlueprintSupport/Classes/K2Node_InAppPurchase.h
@@ -0,0 +1,11 @@
+// Copyright 1998-2014 Epic Games, Inc. All Rights Reserved.
+
+#pragma once
+
+#include "K2Node_InAppPurchase.generated.h"
+
+UCLASS()
+class ONLINEBLUEPRINTSUPPORT_API UK2Node_InAppPurchase : public UK2Node_BaseAsyncTask
+{
+	GENERATED_UCLASS_BODY()
+};
diff --git a/Engine/Source/Editor/OnlineBlueprintSupport/Classes/K2Node_InAppPurchaseQuery.h b/Engine/Source/Editor/OnlineBlueprintSupport/Classes/K2Node_InAppPurchaseQuery.h
new file mode 100644
--- /dev/null
+++ b/Engine/Source/Editor/OnlineBlueprintSupport/Classes/K2Node_InAppPurchaseQuery.h
@@ -0,0 +1,11 @@
+// Copyright 1998-2014 Epic Games, Inc. All Rights Reserved.
+
+#pragma once
+
+#include "K2Node_InAppPurchaseQuery.generated.h"
+
+UCLASS()
+class ONLINEBLUEPRINTSUPPORT_API UK2Node_InAppPurchaseQuery : public UK2Node_BaseAsyncTask
+{
+	GENERATED_UCLASS_BODY()
+};
diff --git a/Engine/Source/Editor/OnlineBlueprintSupport/Private/K2Node_InAppPurchase.cpp b/Engine/Source/Editor/OnlineBlueprintSupport/Private/K2Node_InAppPurchase.cpp
new file mode 100644
--- /dev/null
+++ b/Engine/Source/Editor/OnlineBlueprintSupport/Private/K2Node_InAppPurchase.cpp
@@ -0,0 +1,18 @@
+// Copyright 1998-2014 Epic Games, Inc. All Rights Reserved.
+
+#include "OnlineBlueprintSupportPrivatePCH.h"
+#include "InAppPurchaseCallbackProxy.h"
+#include "K2Node_InAppPurchase.h"
+
+#define LOCTEXT_NAMESPACE "K2Node"
+
+UK2Node_InAppPurchase::UK2Node_InAppPurchase(const FPostConstructInitializeProperties& PCIP)
+	: Super(PCIP)
+{
+	ProxyFactoryFunctionName = GET_FUNCTION_NAME_CHECKED(UInAppPurchaseCallbackProxy, CreateProxyObjectForInAppPurchase);
+	ProxyFactoryClass = UInAppPurchaseCallbackProxy::StaticClass();
+
+	ProxyClass = UInAppPurchaseCallbackProxy::StaticClass();
+}
+
+#undef LOCTEXT_NAMESPACE
diff --git a/Engine/Source/Editor/OnlineBlueprintSupport/Private/K2Node_InAppPurchaseQuery.cpp b/Engine/Source/Editor/OnlineBlueprintSupport/Private/K2Node_InAppPurchaseQuery.cpp
new file mode 100644
--- /dev/null
+++ b/Engine/Source/Editor/OnlineBlueprintSupport/Private/K2Node_InAppPurchaseQuery.cpp
@@ -0,0 +1,18 @@
+// Copyright 1998-2014 Epic Games, Inc. All Rights Reserved.
+
+#include "OnlineBlueprintSupportPrivatePCH.h"
+#include "InAppPurchaseQueryCallbackProxy.h"
+#include "K2Node_InAppPurchaseQuery.h"
+
+#define LOCTEXT_NAMESPACE "K2Node"
+
+UK2Node_InAppPurchaseQuery::UK2Node_InAppPurchaseQuery(const FPostConstructInitializeProperties& PCIP)
+	: Super(PCIP)
+{
+	ProxyFactoryFunctionName = GET_FUNCTION_NAME_CHECKED(UInAppPurchaseQueryCallbackProxy, CreateProxyObjectForInAppPurchaseQuery);
+	ProxyFactoryClass = UInAppPurchaseQueryCallbackProxy::StaticClass();
+
+	ProxyClass = UInAppPurchaseQueryCallbackProxy::StaticClass();
+}
+
+#undef LOCTEXT_NAMESPACE
diff --git a/Engine/Source/Runtime/Online/IOS/OnlineSubsystemIOS/Private/OnlineStoreInterfaceIOS.cpp b/Engine/Source/Runtime/Online/IOS/OnlineSubsystemIOS/Private/OnlineStoreInterfaceIOS.cpp
--- a/Engine/Source/Runtime/Online/IOS/OnlineSubsystemIOS/Private/OnlineStoreInterfaceIOS.cpp
+++ b/Engine/Source/Runtime/Online/IOS/OnlineSubsystemIOS/Private/OnlineStoreInterfaceIOS.cpp
@@ -23,7 +23,7 @@
 
 - (void)paymentQueue:(SKPaymentQueue *)queue updatedTransactions:(NSArray *)transactions
 {
-	UE_LOG(LogOnline, Display, TEXT( "FStoreKitHelper::updatedTransactions" ));
+	UE_LOG(LogOnline, Log, TEXT( "FStoreKitHelper::updatedTransactions" ));
     for (SKPaymentTransaction *transaction in transactions)
     {
         switch ([transaction transactionState])
@@ -48,11 +48,26 @@
 
 - (void) completeTransaction: (SKPaymentTransaction *)transaction
 {
-	UE_LOG(LogOnline, Display, TEXT( "FStoreKitHelper::completeTransaction" ));
+	UE_LOG(LogOnline, Log, TEXT( "FStoreKitHelper::completeTransaction" ));
 
+    
+    [FIOSAsyncTask CreateTaskWithBlock : ^ bool(void)
+    {
+        IOnlineSubsystem* OnlineSub = IOnlineSubsystem::Get(FName(TEXT("IOS")));
+        FOnlineStoreInterfaceIOS* StoreInterface = (FOnlineStoreInterfaceIOS*)OnlineSub->GetStoreInterface().Get();
+        if (StoreInterface->CachedPurchaseStateObject.IsValid())
+        {
+            StoreInterface->CachedPurchaseStateObject->ReadState = EOnlineAsyncTaskState::Done;
+        }
+     
+        StoreInterface->TriggerOnInAppPurchaseCompleteDelegates(EInAppPurchaseState::Success);
+     
+        return true;
+    }];
+    
     // Your application should implement these two methods.
     //[self recordTransaction:transaction];
-    //[self provideContent:transaction.payment.productIdentifier];
+	//[self provideContent:transaction.payment.productIdentifier];
  
     // Remove the transaction from the payment queue.
     [[SKPaymentQueue defaultQueue] finishTransaction: transaction];
@@ -61,7 +76,7 @@
 
 - (void) restoreTransaction: (SKPaymentTransaction *)transaction
 {
-	UE_LOG(LogOnline, Display, TEXT( "FStoreKitHelper::restoreTransaction" ));
+	UE_LOG(LogOnline, Log, TEXT( "FStoreKitHelper::restoreTransaction" ));
 
     //[self recordTransaction: transaction];
     //[self provideContent: transaction.originalTransaction.payment.productIdentifier];
@@ -71,18 +86,44 @@
 
 - (void) failedTransaction: (SKPaymentTransaction *)transaction
 {
-	UE_LOG(LogOnline, Display, TEXT( "FStoreKitHelper::failedTransaction" ));
+	UE_LOG(LogOnline, Log, TEXT( "FStoreKitHelper::failedTransaction - %s" ), *FString([transaction.error localizedDescription]));
+
+	EInAppPurchaseState::Type CompletionState = EInAppPurchaseState::Unknown;
+	switch (transaction.error.code)
+	{
+	case SKErrorPaymentCancelled:
+		CompletionState = EInAppPurchaseState::Cancelled;
+		break;
+        case SKErrorClientInvalid:
+        case SKErrorStoreProductNotAvailable:
+        case SKErrorPaymentInvalid:
+		CompletionState = EInAppPurchaseState::Invalid;
+		break;
+	case SKErrorPaymentNotAllowed:
+		CompletionState = EInAppPurchaseState::NotAllowed;
+		break;
+	}
+	
+    [FIOSAsyncTask CreateTaskWithBlock : ^ bool(void)
+    {
+        IOnlineSubsystem* OnlineSub = IOnlineSubsystem::Get(FName(TEXT("IOS")));
+        FOnlineStoreInterfaceIOS* StoreInterface = (FOnlineStoreInterfaceIOS*)OnlineSub->GetStoreInterface().Get();
+        if (StoreInterface->CachedPurchaseStateObject.IsValid())
+        {
+            StoreInterface->CachedPurchaseStateObject->ReadState = EOnlineAsyncTaskState::Done;
+        }
+        StoreInterface->TriggerOnInAppPurchaseCompleteDelegates(CompletionState);
+     
+        return true;
+     }];
 
-    if (transaction.error.code != SKErrorPaymentCancelled) {
-        // Optionally, display an error here.
-    }
     [[SKPaymentQueue defaultQueue] finishTransaction: transaction];
 }
 
 
 - (void) requestProductData: (NSMutableSet*) productIDs
 {
-	UE_LOG(LogOnline, Display, TEXT( "FStoreKitHelper::requestProductData" ));
+	UE_LOG(LogOnline, Verbose, TEXT( "FStoreKitHelper::requestProductData" ));
 
 	Request = [[SKProductsRequest alloc] initWithProductIdentifiers: productIDs];
 	Request.delegate = self;
@@ -93,7 +134,7 @@
 
 - (void) makePurchase: (NSMutableSet*) productIDs
 {
-	UE_LOG(LogOnline, Display, TEXT( "FStoreKitHelper::makePurchase" ));
+	UE_LOG(LogOnline, Verbose, TEXT( "FStoreKitHelper::makePurchase" ));
 
 	Request = [[SKProductsRequest alloc] initWithProductIdentifiers: productIDs];
 	Request.delegate = self;
@@ -104,8 +145,8 @@
 
 - (void)productsRequest: (SKProductsRequest *)request didReceiveResponse:(SKProductsResponse *)response
 {
-	UE_LOG(LogOnline, Display, TEXT( "FStoreKitHelper::didReceiveResponse" ));
-	// Write the achievements back to our Achievements array
+	UE_LOG(LogOnline, Verbose, TEXT( "FStoreKitHelper::didReceiveResponse" ));
+	// Direct the response back to the store interface
 	[FIOSAsyncTask CreateTaskWithBlock : ^ bool(void)
 	{
 		IOnlineSubsystem* OnlineSub = IOnlineSubsystem::Get(FName(TEXT("IOS")));
@@ -128,19 +169,10 @@
 
 FOnlineStoreInterfaceIOS::FOnlineStoreInterfaceIOS() 
 {
-	UE_LOG(LogOnline, Display, TEXT( "FOnlineStoreInterfaceIOS::FOnlineStoreInterfaceIOS" ));
+	UE_LOG(LogOnline, Verbose, TEXT( "FOnlineStoreInterfaceIOS::FOnlineStoreInterfaceIOS" ));
 	StoreHelper = [[FStoreKitHelper alloc] init];
-	[[SKPaymentQueue defaultQueue] addTransactionObserver:StoreHelper];
-
-	// Gather the microtransaction data from the ini's
-	TArray<FString> ProductIDs;
-	GConfig->GetArray(TEXT("OnlineSubsystemIOS.Store"), TEXT("ProductIDs"), ProductIDs, GEngineIni);
-
-	for(int32 ProductIndex = 0; ProductIndex < ProductIDs.Num(); ProductIndex++)
-	{
-		FOnlineStoreInterfaceIOS::FMicrotransactionPurchaseInfo Info(ProductIDs[ProductIndex]);
-		AvailableProducts.Add( Info );
-	}
+    
+    [[SKPaymentQueue defaultQueue] addTransactionObserver:StoreHelper];
 
 	bIsPurchasing = false;
 	bIsProductRequestInFlight = false;
@@ -149,26 +181,25 @@ FOnlineStoreInterfaceIOS::FOnlineStoreInterfaceIOS()
 
 FOnlineStoreInterfaceIOS::~FOnlineStoreInterfaceIOS() 
 {
-	UE_LOG(LogOnline, Display, TEXT( "FOnlineStoreInterfaceIOS::~FOnlineStoreInterfaceIOS" ));
+	UE_LOG(LogOnline, Verbose, TEXT( "FOnlineStoreInterfaceIOS::~FOnlineStoreInterfaceIOS" ));
 	[StoreHelper release];
 }
 
 
-bool FOnlineStoreInterfaceIOS::QueryForAvailablePurchases()
+bool FOnlineStoreInterfaceIOS::QueryForAvailablePurchases(const TArray<FString>& ProductIDs, FOnlineProductInformationReadRef& InReadObject)
 {
-	UE_LOG(LogOnline, Display, TEXT( "FOnlineStoreInterfaceIOS::QueryForAvailablePurchases" ));
+	UE_LOG(LogOnline, Verbose, TEXT( "FOnlineStoreInterfaceIOS::QueryForAvailablePurchases" ));
 	bool bSentAQueryRequest = false;
 
-	TArray<FString> ProductIDs;
-	GConfig->GetArray(TEXT("OnlineSubsystemIOS.Store"), TEXT("ProductIDs"), ProductIDs, GEngineIni);
+	CachedReadObject = InReadObject;
 	
 	if( bIsPurchasing || bIsProductRequestInFlight )
 	{
-		UE_LOG(LogOnline, Display, TEXT( "FOnlineStoreInterfaceIOS::BeginPurchase - cannot make a purchase whilst one is in transaction." ));
+		UE_LOG(LogOnline, Verbose, TEXT( "FOnlineStoreInterfaceIOS::BeginPurchase - cannot make a purchase whilst one is in transaction." ));
 	}
 	else if (ProductIDs.Num() == 0)
 	{
-		UE_LOG(LogOnline, Display, TEXT( "There are no product IDs configured for Microtransactions in the engine.ini" ));
+		UE_LOG(LogOnline, Verbose, TEXT( "There are no product IDs configured for Microtransactions in the engine.ini" ));
 	}
 	else
 	{
@@ -180,8 +211,11 @@ bool FOnlineStoreInterfaceIOS::QueryForAvailablePurchases()
 			// convert to NSString for the set objects
 			[ProductSet addObject:ID];
 		}
-
-		[StoreHelper requestProductData:ProductSet];
+        
+        dispatch_async(dispatch_get_main_queue(), ^
+        {
+            [StoreHelper requestProductData:ProductSet];
+        });
 
 		bIsProductRequestInFlight = true;
 
@@ -194,78 +228,58 @@ bool FOnlineStoreInterfaceIOS::QueryForAvailablePurchases()
 
 bool FOnlineStoreInterfaceIOS::IsAllowedToMakePurchases()
 {
-	UE_LOG(LogOnline, Display, TEXT( "FOnlineStoreInterfaceIOS::IsAllowedToMakePurchases" ));
+	UE_LOG(LogOnline, Verbose, TEXT( "FOnlineStoreInterfaceIOS::IsAllowedToMakePurchases" ));
 	bool bCanMakePurchases = [SKPaymentQueue canMakePayments];
 	return bCanMakePurchases;
 }
 
 
-bool FOnlineStoreInterfaceIOS::BeginPurchase(int Index)
+bool FOnlineStoreInterfaceIOS::BeginPurchase(const FString& ProductId, FOnlineInAppPurchaseTransactionRef& InPurchaseStateObject)
 {
-	UE_LOG(LogOnline, Display, TEXT( "FOnlineStoreInterfaceIOS::BeginPurchase" ));
+	UE_LOG(LogOnline, Verbose, TEXT( "FOnlineStoreInterfaceIOS::BeginPurchase" ));
 	bool bCreatedNewTransaction = false;
-	bool bIsAValidIndex = Index >= 0 && Index < AvailableProducts.Num();
 	
 	if( bIsPurchasing || bIsProductRequestInFlight )
 	{
-		UE_LOG(LogOnline, Display, TEXT( "FOnlineStoreInterfaceIOS::BeginPurchase - cannot make a purchase whilst one is in transaction." ));
-	}
-	else if( !bIsAValidIndex )
-	{
-		UE_LOG(LogOnline, Display, TEXT( "FOnlineStoreInterfaceIOS - Index: %i is out of range." ), Index);
+		UE_LOG(LogOnline, Verbose, TEXT( "FOnlineStoreInterfaceIOS::BeginPurchase - cannot make a purchase whilst one is in transaction." ));
+		TriggerOnInAppPurchaseCompleteDelegates(EInAppPurchaseState::Failed);
 	}
 	else if( IsAllowedToMakePurchases() )
 	{
-		UE_LOG(LogOnline, Display, TEXT( "FOnlineStoreInterfaceIOS - Making a transaction of product %i" ), Index);
-
-		NSMutableSet* ProductSet = [NSMutableSet setWithCapacity:1];
-		NSString* ID = [NSString stringWithFString:AvailableProducts[Index].Identifier];
-		// convert to NSString for the set objects
-		[ProductSet addObject:ID];
-				
-		UE_LOG(LogOnline, Display, TEXT( "FOnlineStoreInterfaceIOS - Making a transaction of product %i(%s) - count=%i" ), Index, *AvailableProducts[Index].Identifier, [ProductSet count]);
+		UE_LOG(LogOnline, Verbose, TEXT("FOnlineStoreInterfaceIOS - Making a transaction of product %s"), *ProductId);
 
-		[StoreHelper makePurchase:ProductSet];
+		NSString* ID = [NSString stringWithFString : ProductId];
+        
+        dispatch_async(dispatch_get_main_queue(), ^
+        {
+            NSMutableSet* ProductSet = [NSMutableSet setWithCapacity:1];
+            [ProductSet addObject:ID];
+            
+            // Purchase the product through the StoreKit framework
+            [StoreHelper makePurchase:ProductSet];
+        });
 		
 		// Flag that we are purchasing so we can manage subsequent callbacks and transaction requests
 		bIsPurchasing = true;
 		
-		// We have created a new transaction on ths pass.
+		// We have created a new transaction on this pass.
 		bCreatedNewTransaction = true;
+
+		// Cache the readobject so we can add product information to it after the purchase
+		CachedPurchaseStateObject = InPurchaseStateObject;
+		CachedPurchaseStateObject->ReadState = EOnlineAsyncTaskState::InProgress;
 	}
 	else
 	{
-		UE_LOG(LogOnline, Display, TEXT( "This device is not able to make purchases." ) );
+		UE_LOG(LogOnline, Verbose, TEXT("This device is not able to make purchases."));
+		InPurchaseStateObject->ReadState = EOnlineAsyncTaskState::Failed;
+		TriggerOnInAppPurchaseCompleteDelegates(EInAppPurchaseState::NotAllowed);
 	}
 
 	return bCreatedNewTransaction;
 }
 
 
-bool FOnlineStoreInterfaceIOS::Exec(UWorld* InWorld, const TCHAR* Cmd, FOutputDevice& Ar)
-{
-	UE_LOG(LogOnline, Display, TEXT( "FOnlineStoreInterfaceIOS::Exec" ));
-	// Ignore any execs that don't start with MT
-	if (FParse::Command(&Cmd, TEXT("MT")))
-	{
-		if (FParse::Command(&Cmd, TEXT("q")))
-		{
-			QueryForAvailablePurchases();
-		}
-		else
-		{
-			int32 Index;
-			if (FParse::Value(Cmd, TEXT("PURCHASE="), Index ) )
-			{
-				BeginPurchase( Index );
-			}
-		}
-		return true;
-	}
-	return false;
-}
-
-
 void FOnlineStoreInterfaceIOS::ProcessProductsResponse( SKProductsResponse* Response )
 {
 	if( bIsPurchasing )
@@ -280,28 +294,37 @@ void FOnlineStoreInterfaceIOS::ProcessProductsResponse( SKProductsResponse* Resp
 			[numberFormatter setNumberStyle:NSNumberFormatterCurrencyStyle];
 			[numberFormatter setLocale:Product.priceLocale];
 
-			UE_LOG(LogOnline, Display, TEXT("Making a purchase: Product: %s, Price: %s"), *FString([Product localizedTitle]), *FString([numberFormatter stringFromNumber:Product.price]));
+			UE_LOG(LogOnline, Log, TEXT("Making a purchase: Product: %s, Price: %s"), *FString([Product localizedTitle]), *FString([numberFormatter stringFromNumber:Product.price]));
 
-			// now that we have recently refreshed the info, we can purchase it
-			SKPayment* Payment = [SKPayment paymentWithProduct:Product];
-			[[SKPaymentQueue defaultQueue] addPayment:Payment];
+			FInAppPurchaseProductInfo PurchaseProductInfo;
+			PurchaseProductInfo.Identifier = [Product productIdentifier];
+			PurchaseProductInfo.DisplayName = [Product localizedTitle];
+			PurchaseProductInfo.DisplayDescription = [Product localizedDescription];
+			PurchaseProductInfo.DisplayPrice = [numberFormatter stringFromNumber : Product.price];
+
+			[numberFormatter release];
+
+			CachedPurchaseStateObject->ProvidedProductInformation = PurchaseProductInfo;
+            
+            dispatch_async(dispatch_get_main_queue(), ^
+            {
+                // now that we have recently refreshed the info, we can purchase it
+                SKPayment* Payment = [SKPayment paymentWithProduct:Product];
+                [[SKPaymentQueue defaultQueue] addPayment:Payment];
+            });
 		}
 		else if( [Response.products count] > 1 )
 		{
-			UE_LOG(LogOnline, Display, TEXT("Wrong number of products, [%d], in the response when trying to make a single purchase"), [Response.products count]);
+			UE_LOG(LogOnline, Warning, TEXT("Wrong number of products, [%d], in the response when trying to make a single purchase"), [Response.products count]);
 		}
 		else
 		{
 			for(NSString *invalidProduct in Response.invalidProductIdentifiers)
 			{
-				UE_LOG(LogOnline, Display, TEXT("Problem in iTunes connect configuration for product: %s"), *FString(invalidProduct));
+				UE_LOG(LogOnline, Error, TEXT("Problem in iTunes connect configuration for product: %s"), *FString(invalidProduct));
 			}
-
-			UE_LOG(LogOnline, Display, TEXT("Wrong number of products, [%d], in the response when trying to make a single purchase"), [Response.products count]);
 		}
 
-		TriggerOnPurchaseCompleteDelegates( bWasSuccessful );
-
 		bIsPurchasing = false;
 	}
 	else if( bIsProductRequestInFlight )
@@ -309,42 +332,44 @@ void FOnlineStoreInterfaceIOS::ProcessProductsResponse( SKProductsResponse* Resp
 		bool bWasSuccessful = [Response.products count] > 0;
 		if( [Response.products count] == 0 && [Response.invalidProductIdentifiers count] == 0 )
 		{
-			UE_LOG(LogOnline, Display, TEXT("Wrong number of products [%d] in the response when trying to make a single purchase"), [Response.products count]);
+			UE_LOG(LogOnline, Warning, TEXT("Wrong number of products [%d] in the response when trying to make a single purchase"), [Response.products count]);
 		}
-	
-		for( SKProduct* Product in Response.products )
-		{
-			NSNumberFormatter *numberFormatter = [[NSNumberFormatter alloc] init];
-			[numberFormatter setFormatterBehavior:NSNumberFormatterBehavior10_4];
-			[numberFormatter setNumberStyle:NSNumberFormatterCurrencyStyle];
-			[numberFormatter setLocale:Product.priceLocale];
 
-			for( FOnlineStoreInterfaceIOS::FMicrotransactionPurchaseInfo& ProductInfo : AvailableProducts )
+		if( CachedReadObject.IsValid() )
+		{
+			for (SKProduct* Product in Response.products)
 			{
-				if( ProductInfo.Identifier == *FString([Product productIdentifier]) )
-				{
-					ProductInfo.DisplayName = [Product localizedTitle];
-					ProductInfo.DisplayDescription = [Product localizedDescription];
-					ProductInfo.DisplayPrice = [numberFormatter stringFromNumber : Product.price];
-
-					ProductInfo.bWasUpdatedByServer = true;
-
-					UE_LOG(LogOnline, Display, TEXT("\nProduct Identifier: %s, Name: %s, Description: %s, Price: %s\n"),
-						*ProductInfo.Identifier,
-						*ProductInfo.DisplayName,
-						*ProductInfo.DisplayDescription,
-						*ProductInfo.DisplayPrice);
-
-					break;
-				}
+				NSNumberFormatter *numberFormatter = [[NSNumberFormatter alloc] init];
+				[numberFormatter setFormatterBehavior : NSNumberFormatterBehavior10_4];
+				[numberFormatter setNumberStyle : NSNumberFormatterCurrencyStyle];
+				[numberFormatter setLocale : Product.priceLocale];
+
+				FInAppPurchaseProductInfo NewProductInfo;
+				NewProductInfo.Identifier = [Product productIdentifier];
+				NewProductInfo.DisplayName = [Product localizedTitle];
+				NewProductInfo.DisplayDescription = [Product localizedDescription];
+				NewProductInfo.DisplayPrice = [numberFormatter stringFromNumber : Product.price];
+
+				UE_LOG(LogOnline, Log, TEXT("\nProduct Identifier: %s, Name: %s, Description: %s, Price: %s\n"),
+					*NewProductInfo.Identifier,
+					*NewProductInfo.DisplayName,
+					*NewProductInfo.DisplayDescription,
+					*NewProductInfo.DisplayPrice);
+
+				[numberFormatter release];
+
+				CachedReadObject->ProvidedProductInformation.Add( NewProductInfo );
 			}
-
-			[numberFormatter release];
 		}
+		else
+		{
+			UE_LOG(LogOnline, Log, TEXT("Read Object is invalid."));
+		}
+
 		
 		for( NSString *invalidProduct in Response.invalidProductIdentifiers )
 		{
-			UE_LOG(LogOnline, Display, TEXT("Problem in iTunes connect configuration for product: %s"), *FString(invalidProduct));
+			UE_LOG(LogOnline, Warning, TEXT("Problem in iTunes connect configuration for product: %s"), *FString(invalidProduct));
 		}
 
 		TriggerOnQueryForAvailablePurchasesCompleteDelegates( bWasSuccessful );
diff --git a/Engine/Source/Runtime/Online/IOS/OnlineSubsystemIOS/Private/OnlineStoreInterfaceIOS.h b/Engine/Source/Runtime/Online/IOS/OnlineSubsystemIOS/Private/OnlineStoreInterfaceIOS.h
--- a/Engine/Source/Runtime/Online/IOS/OnlineSubsystemIOS/Private/OnlineStoreInterfaceIOS.h
+++ b/Engine/Source/Runtime/Online/IOS/OnlineSubsystemIOS/Private/OnlineStoreInterfaceIOS.h
@@ -4,6 +4,7 @@
 
 #include "OnlineStoreInterface.h"
 
+// SK includes
 #include <StoreKit/SKRequest.h>
 #include <StoreKit/SKError.h>
 #include <StoreKit/SKProduct.h>
@@ -13,15 +14,21 @@
 #include <StoreKit/SKPaymentQueue.h>
 
 
+/** Helper class, which allows us to manage IAP product information requests, AND transactions */
 @interface FStoreKitHelper : NSObject<SKProductsRequestDelegate, SKPaymentTransactionObserver>
 {
 };
-
+/** Store kit request object, holds information about the products we are purchasing, or querying. */
 @property (nonatomic, strong) SKRequest *Request;
+/** collection of available products attaced through a store kit request */
 @property (nonatomic, strong) NSArray *AvailableProducts;
 
-- (void) makePurchase:(NSMutableSet*) productIDs;
+/** Helper fn to start a store kit purchase request */
+-(void)makePurchase:(NSMutableSet*)productIDs;
+/** Helper fn to start a store kit purchase information query request */
 - (void) requestProductData:(NSMutableSet*) productIDs;
+
+/** Helper fn to direct a product request response back to our store interface */
 - (void)productsRequest:(SKProductsRequest *)request didReceiveResponse:(SKProductsResponse *)response;
 
 @end
@@ -30,72 +37,44 @@
 /**
  *	FOnlineStoreInterfaceIOS - Implementation of the online store for IOS
  */
-class FOnlineStoreInterfaceIOS : public IOnlineStore, public FSelfRegisteringExec
+class FOnlineStoreInterfaceIOS : public IOnlineStore
 {
-private:
-	bool bIsPurchasing, bIsProductRequestInFlight;
-
 public:
+	/** C-tor */
 	FOnlineStoreInterfaceIOS();
+	/** Destructor */
 	virtual ~FOnlineStoreInterfaceIOS();
 
 	// Begin IOnlineStore interface
-	virtual bool QueryForAvailablePurchases() override;
-
+	virtual bool QueryForAvailablePurchases(const TArray<FString>& ProductIDs, FOnlineProductInformationReadRef& InReadObject) override;
+	virtual bool BeginPurchase(const FString& ProductId, FOnlineInAppPurchaseTransactionRef& InReadObject) override;
 	virtual bool IsAllowedToMakePurchases() override;
-
-	virtual bool BeginPurchase(int Index) override;
 	// End IOnlineStore interface
 
-	// Begin FExec Interface
-	virtual bool Exec( UWorld* InWorld, const TCHAR* Cmd, FOutputDevice& Ar=*GLog ) override;
-	// End FExec Interface
-
 	/**
 	 * Process a response from the StoreKit
 	 */
 	void ProcessProductsResponse( SKProductsResponse* Response );
-
-
-	/**
-	 * Micro-transaction purchase information
-	 */
-	struct FMicrotransactionPurchaseInfo
-	{
-		FMicrotransactionPurchaseInfo(FString InProductIdentifier)
-		{
-			Identifier = InProductIdentifier;
-			bWasUpdatedByServer = false;
-		}
-
-		// The product identifier that matches the one in iTunesConnect
-		FString Identifier;
-
-		// Display name, gathered from querying the Apple Store
-		FString DisplayName;
-
-		// Display description, gathered from querying the Apple Store
-		FString DisplayDescription;
-
-		// Product price, gathered from querying the Apple Store
-		FString DisplayPrice;
-
-		// The currency type for this MT, gathered from querying the Apple Store
-		FString CurrencyType;
-
-		// flag that lets the user know if the data has been completed from the server
-		bool bWasUpdatedByServer;
-	};
+    
+	/** Cached in-app purchase query object, used to provide the user with product information attained from the server */
+	FOnlineProductInformationReadPtr CachedReadObject;
+    
+	/** Cached in-app purchase transaction object, used to provide details to the user, of the product that has just been purchased. */
+	FOnlineInAppPurchaseTransactionPtr CachedPurchaseStateObject;
 
 
 private:
+	/** Flags which determine the state of our transaction buffer, only one action at a time */
+	bool bIsPurchasing, bIsProductRequestInFlight;
+
+	/** Access to the IOS Store kit interface */
 	FStoreKitHelper* StoreHelper;
-	
-	TArray< FMicrotransactionPurchaseInfo > AvailableProducts;
 
+	/** Delegate fired when a query for purchases has completed, whether successful or unsuccessful */
 	FOnQueryForAvailablePurchasesComplete OnQueryForAvailablePurchasesCompleteDelegate;
 
-	FOnPurchaseComplete OnPurchaseCompleteDelegate;
+	/** Delegate fired when a purchase transaction has completed, whether successful or unsuccessful */
+	FOnInAppPurchaseComplete OnPurchaseCompleteDelegate;
 };
 
 typedef TSharedPtr<FOnlineStoreInterfaceIOS, ESPMode::ThreadSafe> FOnlineStoreInterfaceIOSPtr;
diff --git a/Engine/Source/Runtime/Online/OnlineSubsystem/Public/Interfaces/OnlineStoreInterface.h b/Engine/Source/Runtime/Online/OnlineSubsystem/Public/Interfaces/OnlineStoreInterface.h
--- a/Engine/Source/Runtime/Online/OnlineSubsystem/Public/Interfaces/OnlineStoreInterface.h
+++ b/Engine/Source/Runtime/Online/OnlineSubsystem/Public/Interfaces/OnlineStoreInterface.h
@@ -1,7 +1,25 @@
 // Copyright 1998-2014 Epic Games, Inc. All Rights Reserved.
 
 #pragma once
- #include "OnlineDelegateMacros.h"
+#include "OnlineDelegateMacros.h"
+#include "OnlineStoreInterface.generated.h"
+
+/**
+ * Possible result states of an in-app purchase transaction
+ */
+UENUM(BlueprintType)
+namespace EInAppPurchaseState
+{
+	enum Type
+	{
+		Success = 0 UMETA(DisplayName = "Success"),
+		Failed UMETA(DisplayName = "Failed"),
+		Cancelled UMETA(DisplayName = "Cancelled"),
+		Invalid UMETA(DisplayName = "Invalid"),
+		NotAllowed UMETA(DisplayName = "NotAllowed"),
+		Unknown  UMETA(DisplayName = "Unknown")
+	};
+}
 
 /**
  * Delegate fired when a session create request has completed
@@ -18,9 +36,76 @@ typedef FOnQueryForAvailablePurchasesComplete::FDelegate FOnQueryForAvailablePur
  * @param SessionName the name of the session the that has transitioned to started
  * @param bWasSuccessful true if the async action completed without error, false if there was an error
  */
-DECLARE_MULTICAST_DELEGATE_OneParam(FOnPurchaseComplete, bool);
-typedef FOnPurchaseComplete::FDelegate FOnPurchaseCompleteDelegate;
+DECLARE_MULTICAST_DELEGATE_OneParam(FOnInAppPurchaseComplete, EInAppPurchaseState::Type);
+typedef FOnInAppPurchaseComplete::FDelegate FOnInAppPurchaseCompleteDelegate;
+
+
+/**
+* Micro-transaction purchase information
+*/
+USTRUCT(BlueprintType)
+struct FInAppPurchaseProductInfo
+{
+	GENERATED_USTRUCT_BODY()
+
+	// The unique product identifier
+	UPROPERTY(BlueprintReadOnly, Category = ProductInfo)
+	FString Identifier;
+
+	// The localized display name
+	UPROPERTY(BlueprintReadOnly, Category = ProductInfo)
+	FString DisplayName;
+
+	// The localized display description name
+	UPROPERTY(BlueprintReadOnly, Category = ProductInfo)
+	FString DisplayDescription;
+
+	// The localized display price name
+	UPROPERTY(BlueprintReadOnly, Category = ProductInfo)
+	FString DisplayPrice;
+};
+
+
+/**
+*	Interface for reading data from an In App Purchase service
+*/
+class ONLINESUBSYSTEM_API FOnlineProductInformationRead
+{
+public:
+	/** Indicates an error reading data occurred while processing */
+	EOnlineAsyncTaskState::Type ReadState;
+
+	FOnlineProductInformationRead() :
+		ReadState(EOnlineAsyncTaskState::NotStarted)
+	{
+	}
+
+	TArray<FInAppPurchaseProductInfo> ProvidedProductInformation;
+};
+
+typedef TSharedRef<FOnlineProductInformationRead, ESPMode::ThreadSafe> FOnlineProductInformationReadRef;
+typedef TSharedPtr<FOnlineProductInformationRead, ESPMode::ThreadSafe> FOnlineProductInformationReadPtr;
+
+
+/**
+ *	Interface for reading data from an In App Purchase service
+ */
+class ONLINESUBSYSTEM_API FOnlineInAppPurchaseTransaction
+{
+public:
+	/** Indicates an error reading data occurred while processing */
+	EOnlineAsyncTaskState::Type ReadState;
+
+	FOnlineInAppPurchaseTransaction() :
+		ReadState(EOnlineAsyncTaskState::NotStarted)
+	{
+	}
 
+	FInAppPurchaseProductInfo ProvidedProductInformation;
+};
+
+typedef TSharedRef<FOnlineInAppPurchaseTransaction, ESPMode::ThreadSafe> FOnlineInAppPurchaseTransactionRef;
+typedef TSharedPtr<FOnlineInAppPurchaseTransaction, ESPMode::ThreadSafe> FOnlineInAppPurchaseTransactionPtr;
 
 /**
  *	IOnlineStore - Interface class for microtransactions
@@ -34,7 +119,7 @@ public:
 	 *
 	 * @return - Whether a request was sent to check for purchases
 	 */
-	virtual bool QueryForAvailablePurchases() = 0;
+	virtual bool QueryForAvailablePurchases(const TArray<FString>& ProductIDs, FOnlineProductInformationReadRef& InReadObject) = 0;
 
 	/**
 	 * Delegate which is executed when QueryForAvailablePurchases completes
@@ -47,7 +132,7 @@ public:
 	 * @return - Whether the device can make purchases
 	 */
 	virtual bool IsAllowedToMakePurchases() = 0;
-	
+
 	/**
 	 * Begin a purchase transaction for the product which relates to the given index
 	 *
@@ -55,10 +140,10 @@ public:
 	 *
 	 * @return - whether a purchase request was sent
 	 */
-	virtual bool BeginPurchase(int Index) = 0;
+	virtual bool BeginPurchase(const FString& ProductId, FOnlineInAppPurchaseTransactionRef& InReadObject) = 0;
 	
 	/**
 	 * Delegate which is executed when a Purchase completes
 	 */
-	DEFINE_ONLINE_DELEGATE_ONE_PARAM(OnPurchaseComplete, bool);
+	DEFINE_ONLINE_DELEGATE_ONE_PARAM(OnInAppPurchaseComplete, EInAppPurchaseState::Type);
 };
\ No newline at end of file
diff --git a/Engine/Source/Runtime/Online/OnlineSubsystemUtils/Classes/InAppPurchaseCallbackProxy.h b/Engine/Source/Runtime/Online/OnlineSubsystemUtils/Classes/InAppPurchaseCallbackProxy.h
new file mode 100644
--- /dev/null
+++ b/Engine/Source/Runtime/Online/OnlineSubsystemUtils/Classes/InAppPurchaseCallbackProxy.h
@@ -0,0 +1,63 @@
+// Copyright 1998-2014 Epic Games, Inc. All Rights Reserved.
+#pragma once
+
+#include "OnlineStoreInterface.h"
+#include "InAppPurchaseCallbackProxy.generated.h"
+
+DECLARE_DYNAMIC_MULTICAST_DELEGATE_TwoParams(FInAppPurchaseResult, EInAppPurchaseState::Type, CompletionStatus, const FInAppPurchaseProductInfo&, InAppPurchaseInformation);
+
+UCLASS(MinimalAPI)
+class UInAppPurchaseCallbackProxy : public UObject
+{
+	GENERATED_UCLASS_BODY()
+
+	// Called when there is a successful In-App Purchase transaction
+	UPROPERTY(BlueprintAssignable)
+	FInAppPurchaseResult OnSuccess;
+
+	// Called when there is an unsuccessful In-App Purchase transaction
+	UPROPERTY(BlueprintAssignable)
+	FInAppPurchaseResult OnFailure;
+
+	// Kicks off a transaction for the provided product identifier
+	UFUNCTION(BlueprintCallable, meta = (FriendlyName="Make an In-App Purchase"), Category="Online|InAppPurchase")
+	static UInAppPurchaseCallbackProxy* CreateProxyObjectForInAppPurchase(class APlayerController* PlayerController, const FString& ProductIdentifier);
+
+public:
+
+	// Begin UObject interface
+	virtual void BeginDestroy() override;
+	// End UObject interface
+
+private:
+
+	/** Called by the InAppPurchase system when the transaction has finished */
+	void OnInAppPurchaseComplete_Delayed();
+	void OnInAppPurchaseComplete(EInAppPurchaseState::Type CompletionState);
+
+	/** Unregisters our delegate from the In-App Purchase system */
+	void RemoveDelegate();
+
+	/** Triggers the In-App Purchase Transaction for the specifed user; the Purchase Request object must already be set up */
+	void Trigger(class APlayerController* PlayerController, const FString& ProductIdentifier);
+
+private:
+
+	/** Delegate called when a InAppPurchase has been successfully read */
+	FOnInAppPurchaseCompleteDelegate InAppPurchaseCompleteDelegate;
+
+	/** The InAppPurchase read request */
+	FOnlineInAppPurchaseTransactionPtr PurchaseRequest;
+
+	/** Did we fail immediately? */
+	bool bFailedToEvenSubmit;
+
+	/** Pointer to the world, needed to delay the results slightly */
+	TWeakObjectPtr<UWorld> WorldPtr;
+
+	/** Did the purchase succeed? */
+	EInAppPurchaseState::Type SavedPurchaseState;
+
+	/** Cached product details of the purchased product */
+	FInAppPurchaseProductInfo SavedProductInformation;
+};
diff --git a/Engine/Source/Runtime/Online/OnlineSubsystemUtils/Classes/InAppPurchaseQueryCallbackProxy.h b/Engine/Source/Runtime/Online/OnlineSubsystemUtils/Classes/InAppPurchaseQueryCallbackProxy.h
new file mode 100644
--- /dev/null
+++ b/Engine/Source/Runtime/Online/OnlineSubsystemUtils/Classes/InAppPurchaseQueryCallbackProxy.h
@@ -0,0 +1,58 @@
+// Copyright 1998-2014 Epic Games, Inc. All Rights Reserved.
+#pragma once
+
+#include "OnlineStoreInterface.h"
+#include "InAppPurchaseQueryCallbackProxy.generated.h"
+
+DECLARE_DYNAMIC_MULTICAST_DELEGATE_OneParam(FInAppPurchaseQueryResult, const TArray<FInAppPurchaseProductInfo>&, InAppPurchaseInformation);
+
+UCLASS(MinimalAPI)
+class UInAppPurchaseQueryCallbackProxy : public UObject
+{
+	GENERATED_UCLASS_BODY()
+
+	// Called when there is a successful InAppPurchase query
+	UPROPERTY(BlueprintAssignable)
+	FInAppPurchaseQueryResult OnSuccess;
+
+	// Called when there is an unsuccessful InAppPurchase query
+	UPROPERTY(BlueprintAssignable)
+	FInAppPurchaseQueryResult OnFailure;
+
+	// Queries a InAppPurchase for an integer value
+	UFUNCTION(BlueprintCallable, meta = (FriendlyName="Read In App Purchase Information"), Category="Online|InAppPurchase")
+	static UInAppPurchaseQueryCallbackProxy* CreateProxyObjectForInAppPurchaseQuery(class APlayerController* PlayerController, const TArray<FString>& ProductIdentifiers);
+
+public:
+	// Begin UObject interface
+	virtual void BeginDestroy() override;
+	// End UObject interface
+
+private:
+	/** Called by the InAppPurchase system when the read is finished */
+	void OnInAppPurchaseRead_Delayed();
+	void OnInAppPurchaseRead(bool bWasSuccessful);
+
+	/** Unregisters our delegate from the InAppPurchase system */
+	void RemoveDelegate();
+
+	/** Triggers the query for a specifed user; the ReadObject must already be set up */
+	void TriggerQuery(class APlayerController* PlayerController, const TArray<FString>& ProductIdentifiers);
+
+private:
+	/** Delegate called when a InAppPurchase has been successfully read */
+	FOnQueryForAvailablePurchasesCompleteDelegate InAppPurchaseReadCompleteDelegate;
+
+	/** The InAppPurchase read request */
+	FOnlineProductInformationReadPtr ReadObject;
+
+	/** Did we fail immediately? */
+	bool bFailedToEvenSubmit;
+
+	// Pointer to the world, needed to delay the results slightly
+	TWeakObjectPtr<UWorld> WorldPtr;
+
+	// Did the read succeed?
+	bool bSavedWasSuccessful;
+	TArray<FInAppPurchaseProductInfo> SavedProductInformation;
+};
diff --git a/Engine/Source/Runtime/Online/OnlineSubsystemUtils/Private/InAppPurchaseCallbackProxy.cpp b/Engine/Source/Runtime/Online/OnlineSubsystemUtils/Private/InAppPurchaseCallbackProxy.cpp
new file mode 100644
--- /dev/null
+++ b/Engine/Source/Runtime/Online/OnlineSubsystemUtils/Private/InAppPurchaseCallbackProxy.cpp
@@ -0,0 +1,124 @@
+// Copyright 1998-2014 Epic Games, Inc. All Rights Reserved.
+
+#include "OnlineSubsystemUtilsPrivatePCH.h"
+
+//////////////////////////////////////////////////////////////////////////
+// UInAppPurchaseCallbackProxy
+
+UInAppPurchaseCallbackProxy::UInAppPurchaseCallbackProxy(const class FPostConstructInitializeProperties& PCIP)
+	: Super(PCIP)
+{
+	PurchaseRequest = nullptr;
+	WorldPtr = nullptr;
+}
+
+
+void UInAppPurchaseCallbackProxy::Trigger(APlayerController* PlayerController, const FString& ProductIdentifier)
+{
+	bFailedToEvenSubmit = true;
+
+	WorldPtr = (PlayerController != nullptr) ? PlayerController->GetWorld() : nullptr;
+	if (APlayerState* PlayerState = (PlayerController != NULL) ? PlayerController->PlayerState : nullptr)
+	{
+		if (IOnlineSubsystem* const OnlineSub = IOnlineSubsystem::Get())
+		{
+			IOnlineStorePtr StoreInterface = OnlineSub->GetStoreInterface();
+			if (StoreInterface.IsValid())
+			{
+				bFailedToEvenSubmit = false;
+
+				// Register the completion callback
+				InAppPurchaseCompleteDelegate = FOnInAppPurchaseCompleteDelegate::CreateUObject(this, &UInAppPurchaseCallbackProxy::OnInAppPurchaseComplete);
+				StoreInterface->AddOnInAppPurchaseCompleteDelegate(InAppPurchaseCompleteDelegate);
+
+				// Set-up, and trigger the transaction through the store interface
+				PurchaseRequest = MakeShareable(new FOnlineInAppPurchaseTransaction());
+				FOnlineInAppPurchaseTransactionRef PurchaseRequestRef = PurchaseRequest.ToSharedRef();
+				StoreInterface->BeginPurchase(ProductIdentifier, PurchaseRequestRef);
+			}
+			else
+			{
+				FFrame::KismetExecutionMessage(TEXT("UInAppPurchaseCallbackProxy::Trigger - In-App Purchases are not supported by Online Subsystem"), ELogVerbosity::Warning);
+			}
+		}
+		else
+		{
+			FFrame::KismetExecutionMessage(TEXT("UInAppPurchaseCallbackProxy::Trigger - Invalid or uninitialized OnlineSubsystem"), ELogVerbosity::Warning);
+		}
+	}
+	else
+	{
+		FFrame::KismetExecutionMessage(TEXT("UInAppPurchaseCallbackProxy::Trigger - Invalid player state"), ELogVerbosity::Warning);
+	}
+
+	if (bFailedToEvenSubmit && (PlayerController != NULL))
+	{
+		OnInAppPurchaseComplete(EInAppPurchaseState::Failed);
+	}
+}
+
+
+void UInAppPurchaseCallbackProxy::OnInAppPurchaseComplete(EInAppPurchaseState::Type CompletionState)
+{
+	RemoveDelegate();
+
+	SavedPurchaseState = CompletionState;
+
+	if (SavedPurchaseState == EInAppPurchaseState::Success && PurchaseRequest.IsValid())
+	{
+		SavedProductInformation = PurchaseRequest->ProvidedProductInformation;
+	}
+
+	if (UWorld* World = WorldPtr.Get())
+	{
+		World->GetTimerManager().SetTimer(this, &UInAppPurchaseCallbackProxy::OnInAppPurchaseComplete_Delayed, 0.001f, false);
+	}
+
+	PurchaseRequest = NULL;
+}
+
+
+void UInAppPurchaseCallbackProxy::OnInAppPurchaseComplete_Delayed()
+{
+	if (SavedPurchaseState == EInAppPurchaseState::Success)
+	{
+		OnSuccess.Broadcast(SavedPurchaseState, SavedProductInformation);
+	}
+	else
+	{
+		OnFailure.Broadcast(SavedPurchaseState, SavedProductInformation);
+	}
+}
+
+
+void UInAppPurchaseCallbackProxy::RemoveDelegate()
+{
+	if (!bFailedToEvenSubmit)
+	{
+		if (IOnlineSubsystem* OnlineSub = IOnlineSubsystem::Get())
+		{
+			IOnlineStorePtr InAppPurchases = OnlineSub->GetStoreInterface();
+			if (InAppPurchases.IsValid())
+			{
+				InAppPurchases->ClearOnInAppPurchaseCompleteDelegate(InAppPurchaseCompleteDelegate);
+			}
+		}
+	}
+}
+
+
+void UInAppPurchaseCallbackProxy::BeginDestroy()
+{
+	PurchaseRequest = nullptr;
+	RemoveDelegate();
+
+	Super::BeginDestroy();
+}
+
+
+UInAppPurchaseCallbackProxy* UInAppPurchaseCallbackProxy::CreateProxyObjectForInAppPurchase(class APlayerController* PlayerController, const FString& ProductIdentifier)
+{
+	UInAppPurchaseCallbackProxy* Proxy = NewObject<UInAppPurchaseCallbackProxy>();
+	Proxy->Trigger(PlayerController, ProductIdentifier);
+	return Proxy;
+}
diff --git a/Engine/Source/Runtime/Online/OnlineSubsystemUtils/Private/InAppPurchaseQueryCallbackProxy.cpp b/Engine/Source/Runtime/Online/OnlineSubsystemUtils/Private/InAppPurchaseQueryCallbackProxy.cpp
new file mode 100644
--- /dev/null
+++ b/Engine/Source/Runtime/Online/OnlineSubsystemUtils/Private/InAppPurchaseQueryCallbackProxy.cpp
@@ -0,0 +1,119 @@
+// Copyright 1998-2014 Epic Games, Inc. All Rights Reserved.
+
+#include "OnlineSubsystemUtilsPrivatePCH.h"
+
+//////////////////////////////////////////////////////////////////////////
+// UInAppPurchaseQueryCallbackProxy
+
+UInAppPurchaseQueryCallbackProxy::UInAppPurchaseQueryCallbackProxy(const class FPostConstructInitializeProperties& PCIP)
+	: Super(PCIP)
+{
+}
+
+void UInAppPurchaseQueryCallbackProxy::TriggerQuery(APlayerController* PlayerController, const TArray<FString>& ProductIdentifiers)
+{
+	bFailedToEvenSubmit = true;
+
+	WorldPtr = (PlayerController != nullptr) ? PlayerController->GetWorld() : nullptr;
+	if (APlayerState* PlayerState = (PlayerController != NULL) ? PlayerController->PlayerState : nullptr)
+	{
+		if (IOnlineSubsystem* const OnlineSub = IOnlineSubsystem::Get())
+		{
+			IOnlineStorePtr StoreInterface = OnlineSub->GetStoreInterface();
+			if (StoreInterface.IsValid())
+			{
+				bFailedToEvenSubmit = false;
+
+				// Register the completion callback
+				InAppPurchaseReadCompleteDelegate = FOnQueryForAvailablePurchasesCompleteDelegate::CreateUObject(this, &UInAppPurchaseQueryCallbackProxy::OnInAppPurchaseRead);
+				StoreInterface->AddOnQueryForAvailablePurchasesCompleteDelegate(InAppPurchaseReadCompleteDelegate);
+
+
+				ReadObject = MakeShareable(new FOnlineProductInformationRead());
+				FOnlineProductInformationReadRef ReadObjectRef = ReadObject.ToSharedRef();
+				StoreInterface->QueryForAvailablePurchases(ProductIdentifiers, ReadObjectRef);
+			}
+			else
+			{
+				FFrame::KismetExecutionMessage(TEXT("UInAppPurchaseQueryCallbackProxy::TriggerQuery - In App Purchases are not supported by Online Subsystem"), ELogVerbosity::Warning);
+			}
+		}
+		else
+		{
+			FFrame::KismetExecutionMessage(TEXT("UInAppPurchaseQueryCallbackProxy::TriggerQuery - Invalid or uninitialized OnlineSubsystem"), ELogVerbosity::Warning);
+		}
+	}
+	else
+	{
+		FFrame::KismetExecutionMessage(TEXT("UInAppPurchaseQueryCallbackProxy::TriggerQuery - Invalid player state"), ELogVerbosity::Warning);
+	}
+
+	if (bFailedToEvenSubmit && (PlayerController != NULL))
+	{
+		OnInAppPurchaseRead(false);
+	}
+}
+
+void UInAppPurchaseQueryCallbackProxy::OnInAppPurchaseRead(bool bWasSuccessful)
+{
+	RemoveDelegate();
+
+	if (bWasSuccessful && ReadObject.IsValid())
+	{
+		SavedProductInformation = ReadObject->ProvidedProductInformation;
+		bSavedWasSuccessful = true;
+	}
+	else
+	{
+		bSavedWasSuccessful = false;
+	}
+
+	if (UWorld* World = WorldPtr.Get())
+	{
+		World->GetTimerManager().SetTimer(this, &UInAppPurchaseQueryCallbackProxy::OnInAppPurchaseRead_Delayed, 0.001f, false);
+	}
+
+	ReadObject = NULL;
+}
+
+void UInAppPurchaseQueryCallbackProxy::OnInAppPurchaseRead_Delayed()
+{
+	if (bSavedWasSuccessful)
+	{
+		OnSuccess.Broadcast(SavedProductInformation);
+	}
+	else
+	{
+		OnFailure.Broadcast(SavedProductInformation);
+	}
+}
+
+void UInAppPurchaseQueryCallbackProxy::RemoveDelegate()
+{
+	if (!bFailedToEvenSubmit)
+	{
+		if (IOnlineSubsystem* OnlineSub = IOnlineSubsystem::Get())
+		{
+			IOnlineStorePtr InAppPurchases = OnlineSub->GetStoreInterface();
+			if (InAppPurchases.IsValid())
+			{
+				InAppPurchases->ClearOnQueryForAvailablePurchasesCompleteDelegate(InAppPurchaseReadCompleteDelegate);
+			}
+		}
+	}
+}
+
+void UInAppPurchaseQueryCallbackProxy::BeginDestroy()
+{
+	ReadObject = NULL;
+	RemoveDelegate();
+
+	Super::BeginDestroy();
+}
+
+UInAppPurchaseQueryCallbackProxy* UInAppPurchaseQueryCallbackProxy::CreateProxyObjectForInAppPurchaseQuery(class APlayerController* PlayerController, const TArray<FString>& ProductIdentifiers)
+{
+	UInAppPurchaseQueryCallbackProxy* Proxy = NewObject<UInAppPurchaseQueryCallbackProxy>();
+	Proxy->TriggerQuery(PlayerController, ProductIdentifiers);
+	return Proxy;
+}
