#online,mcp
- added country/locale info to catalog & orders

diff --git a/Engine/Source/Runtime/Online/OnlineSubsystem/Public/OnlineSubsystemTypes.h b/Engine/Source/Runtime/Online/OnlineSubsystem/Public/OnlineSubsystemTypes.h
--- a/Engine/Source/Runtime/Online/OnlineSubsystem/Public/OnlineSubsystemTypes.h
+++ b/Engine/Source/Runtime/Online/OnlineSubsystem/Public/OnlineSubsystemTypes.h
@@ -945,6 +945,24 @@ public:
 	int32 Count;
 };
 
+/**
+ * Locale and country code 
+ */
+class FRegionInfo
+{
+public:
+
+	FRegionInfo(const FString& InCountry = FString(), const FString& InLocale = FString())
+		: Country(InCountry)
+		, Locale(InLocale)
+	{}
+
+	/** country code. eg. US */
+	FString Country;
+	/** local code. eg. en_US */
+	FString Locale;
+};
+
 /** Holds metadata about a given downloadable file */
 struct FCloudFileHeader
 {	
