Update the Google Play achievement cache in native code after a reset command

diff --git a/Engine/Build/Android/Java/src/com/epicgames/ue4/GameActivity.java b/Engine/Build/Android/Java/src/com/epicgames/ue4/GameActivity.java
--- a/Engine/Build/Android/Java/src/com/epicgames/ue4/GameActivity.java
+++ b/Engine/Build/Android/Java/src/com/epicgames/ue4/GameActivity.java
@@ -860,6 +860,9 @@ public class GameActivity extends NativeActivity implements GoogleApiClient.Conn
 			{
 				urlConnection.disconnect();
 			}
+
+			// Kick off a update to the native side achievements
+			AndroidThunkJava_QueryAchievements();
         }
         catch(Exception e)
         {
diff --git a/Engine/Source/Runtime/Online/Android/OnlineSubsystemGooglePlay/Private/OnlineAchievementsInterfaceGooglePlay.cpp b/Engine/Source/Runtime/Online/Android/OnlineSubsystemGooglePlay/Private/OnlineAchievementsInterfaceGooglePlay.cpp
--- a/Engine/Source/Runtime/Online/Android/OnlineSubsystemGooglePlay/Private/OnlineAchievementsInterfaceGooglePlay.cpp
+++ b/Engine/Source/Runtime/Online/Android/OnlineSubsystemGooglePlay/Private/OnlineAchievementsInterfaceGooglePlay.cpp
@@ -64,6 +64,7 @@ extern "C" void Java_com_epicgames_ue4_GameActivity_nativeUpdateAchievements(JNI
 		}
 
 		PendingQuery.Delegate.ExecuteIfBound( PendingQuery.PlayerID, true );
+		PendingQuery.Delegate.Unbind();
 		PendingQuery.IsQueryPending = false;
 	});
 }
@@ -88,6 +89,7 @@ extern "C" void Java_com_epicgames_ue4_GameActivity_nativeFailedUpdateAchievemen
 		auto& PendingQuery = FOnlineAchievementsGooglePlay::PendingAchievementQuery;
 
 		PendingQuery.Delegate.ExecuteIfBound( PendingQuery.PlayerID, false );
+		PendingQuery.Delegate.Unbind();
 		PendingQuery.IsQueryPending = false;
 	});
 }
