#online,mcp,payment
- parse billing account name as FText

diff --git a/Engine/Source/Runtime/Online/OnlineSubsystem/Public/OnlineSubsystemTypes.h b/Engine/Source/Runtime/Online/OnlineSubsystem/Public/OnlineSubsystemTypes.h
--- a/Engine/Source/Runtime/Online/OnlineSubsystem/Public/OnlineSubsystemTypes.h
+++ b/Engine/Source/Runtime/Online/OnlineSubsystem/Public/OnlineSubsystemTypes.h
@@ -945,9 +945,7 @@ public:
 	int32 Count;
 };
 
-/**
- * Locale and country code 
- */
+/** Locale and country code */
 class FRegionInfo
 {
 public:
@@ -957,9 +955,9 @@ public:
 		, Locale(InLocale)
 	{}
 
-	/** country code. eg. US */
+	/** country code used for configuring things like currency/pricing specific to a country. eg. US */
 	FString Country;
-	/** local code. eg. en_US */
+	/** local code used to select the localization language. eg. en_US */
 	FString Locale;
 };
 
