Improvements for user and controller handling in ShooterGame

Added a ControllerPairingChanged delegate which Xbox games will probably want to respond to
Xbox now fires the OnControllerConnectionChange delegate
Xbox Live OSS now triggers OnLoginStatusChanged instead of OnLoginChanged
ShooterGameInstance now subscribes to OnLoginStatusChanged instead of OnLoginChanged
Added overload of FindLocalPlayerFromUniqueNetId taking a const FUniqueNetId& since some OSS delegates return them instead of TSharedPtrs.

diff --git a/Engine/Source/Runtime/Engine/Classes/Engine/GameInstance.h b/Engine/Source/Runtime/Engine/Classes/Engine/GameInstance.h
--- a/Engine/Source/Runtime/Engine/Classes/Engine/GameInstance.h
+++ b/Engine/Source/Runtime/Engine/Classes/Engine/GameInstance.h
@@ -118,6 +118,7 @@ public:
 	APlayerController*		GetFirstLocalPlayerController() const;
 	ULocalPlayer*			FindLocalPlayerFromControllerId(int32 ControllerId) const;
 	ULocalPlayer*			FindLocalPlayerFromUniqueNetId(TSharedPtr< FUniqueNetId > UniqueNetId) const;
+	ULocalPlayer*			FindLocalPlayerFromUniqueNetId(const FUniqueNetId& UniqueNetId) const;
 	ULocalPlayer*			GetFirstGamePlayer() const;
 
 	void					CleanupGameViewport();
diff --git a/Engine/Source/Runtime/Engine/Private/GameInstance.cpp b/Engine/Source/Runtime/Engine/Private/GameInstance.cpp
--- a/Engine/Source/Runtime/Engine/Private/GameInstance.cpp
+++ b/Engine/Source/Runtime/Engine/Private/GameInstance.cpp
@@ -515,13 +515,8 @@ ULocalPlayer* UGameInstance::FindLocalPlayerFromControllerId(int32 ControllerId)
 	return nullptr;
 }
 
-ULocalPlayer* UGameInstance::FindLocalPlayerFromUniqueNetId(TSharedPtr<FUniqueNetId> UniqueNetId) const
+ULocalPlayer* UGameInstance::FindLocalPlayerFromUniqueNetId(const FUniqueNetId& UniqueNetId) const
 {
-	if (!UniqueNetId.IsValid())
-	{
-		return NULL;
-	}
-
 	for (ULocalPlayer* Player : LocalPlayers)
 	{
 		if (Player == NULL)
@@ -536,7 +531,7 @@ ULocalPlayer* UGameInstance::FindLocalPlayerFromUniqueNetId(TSharedPtr<FUniqueNe
 			continue;
 		}
 
-		if (*OtherUniqueNetId == *UniqueNetId)
+		if (*OtherUniqueNetId == UniqueNetId)
 		{
 			// Match
 			return Player;
@@ -547,6 +542,16 @@ ULocalPlayer* UGameInstance::FindLocalPlayerFromUniqueNetId(TSharedPtr<FUniqueNe
 	return nullptr;
 }
 
+ULocalPlayer* UGameInstance::FindLocalPlayerFromUniqueNetId(TSharedPtr<FUniqueNetId> UniqueNetId) const
+{
+	if (!UniqueNetId.IsValid())
+	{
+		return nullptr;
+	}
+
+	return FindLocalPlayerFromUniqueNetId(*UniqueNetId);
+}
+
 ULocalPlayer* UGameInstance::GetFirstGamePlayer() const
 {
 	return (LocalPlayers.Num() > 0) ? LocalPlayers[0] : nullptr;
diff --git a/Engine/Source/Runtime/Online/OnlineSubsystem/Public/Interfaces/OnlineIdentityInterface.h b/Engine/Source/Runtime/Online/OnlineSubsystem/Public/Interfaces/OnlineIdentityInterface.h
--- a/Engine/Source/Runtime/Online/OnlineSubsystem/Public/Interfaces/OnlineIdentityInterface.h
+++ b/Engine/Source/Runtime/Online/OnlineSubsystem/Public/Interfaces/OnlineIdentityInterface.h
@@ -64,6 +64,16 @@ typedef FOnLoginChanged::FDelegate FOnLoginChangedDelegate;
 DECLARE_MULTICAST_DELEGATE_ThreeParams(FOnLoginStatusChanged, int32, ELoginStatus::Type, const FUniqueNetId&);
 typedef FOnLoginStatusChanged::FDelegate FOnLoginStatusChangedDelegate;
 
+/**
+ * Delegate called when a controller-user pairing changes
+ *
+ * @param LocalUserNum the controller number of the controller whose pairing changed
+ * @param PreviousUser the user that used to be paired with this controller
+ * @param NewUser the user that is currently paired with this controller
+ */
+DECLARE_MULTICAST_DELEGATE_ThreeParams(FOnControllerPairingChanged, int, const FUniqueNetId&, const FUniqueNetId&);
+typedef FOnControllerPairingChanged::FDelegate FOnControllerPairingChangedDelegate;
+
 /**
  * Called when user account login has completed after calling Login() or AutoLogin()
  *
@@ -113,6 +123,15 @@ public:
 	 */
 	DEFINE_ONLINE_PLAYER_DELEGATE_TWO_PARAM(MAX_LOCAL_PLAYERS, OnLoginStatusChanged, ELoginStatus::Type, const FUniqueNetId&);
 
+	/**
+	 * Delegate called when a controller-user pairing changes
+	 *
+	 * @param LocalUserNum the controller number of the associated user
+	 * @param PreviousUser the new login status for the user
+	 * @param NewUser the new id to associate with the user
+	 */
+	DEFINE_ONLINE_DELEGATE_THREE_PARAM(OnControllerPairingChanged, int, const FUniqueNetId&, const FUniqueNetId&);
+
 	/**
 	 * Login/Authenticate with user credentials.
 	 * Will call OnLoginComplete() delegate when async task completes
