TPP #341357 MobileMVP: Add support for OSS connect action in blueprint and response handling
Second pass
Removed new interface and change implementation to use the Indentify interface
Created Android implementation of the same.

diff --git a/Engine/Build/Android/Java/src/com/epicgames/ue4/GameActivity.java b/Engine/Build/Android/Java/src/com/epicgames/ue4/GameActivity.java
--- a/Engine/Build/Android/Java/src/com/epicgames/ue4/GameActivity.java
+++ b/Engine/Build/Android/Java/src/com/epicgames/ue4/GameActivity.java
@@ -532,7 +532,7 @@ public class GameActivity extends NativeActivity implements GoogleApiClient.Conn
 		// Set the flag that we successfully connected. Checked in onStart to re-establish the connection.
 		bHaveConnectedToGooglePlay = true;
 
-		nativeCompletedConnection(RESULT_OK);
+		nativeCompletedConnection(0, RESULT_OK);
 
 		// Load achievements. Since games are expected to pass in achievement progress as a percentage,
 		// we need to know what the maximum steps are in order to convert the percentage to an integer
@@ -618,7 +618,7 @@ public class GameActivity extends NativeActivity implements GoogleApiClient.Conn
 			{
 				// translate result code? 
 				// 0 if we cancel out the attempt to error recover...
-				nativeCompletedConnection(resultCode);
+				nativeCompletedConnection(0, resultCode);
 			}
 		}
 	}
@@ -741,7 +741,7 @@ public class GameActivity extends NativeActivity implements GoogleApiClient.Conn
 		// check if google play services is available on this device, or is available with an update
 		if ((status != ConnectionResult.SUCCESS) && (status != ConnectionResult.SERVICE_VERSION_UPDATE_REQUIRED))
 		{
-			nativeCompletedConnection(status);
+			nativeCompletedConnection(0, status);
 			return;
 		}
 
@@ -1051,7 +1051,7 @@ public class GameActivity extends NativeActivity implements GoogleApiClient.Conn
 
 	public native void nativeResumeMainInit();
 	
-	public native void nativeCompletedConnection(int errorCode);
+	public native void nativeCompletedConnection(int playerID, int errorCode);
 
 	static
 	{
diff --git a/Engine/Source/Runtime/Online/Android/OnlineSubsystemGooglePlay/Private/OnlineConnectionInterfaceGooglePlay.cpp b/Engine/Source/Runtime/Online/Android/OnlineSubsystemGooglePlay/Private/OnlineConnectionInterfaceGooglePlay.cpp
deleted file mode 100644
diff --git a/Engine/Source/Runtime/Online/Android/OnlineSubsystemGooglePlay/Private/OnlineConnectionInterfaceGooglePlay.h b/Engine/Source/Runtime/Online/Android/OnlineSubsystemGooglePlay/Private/OnlineConnectionInterfaceGooglePlay.h
deleted file mode 100644
diff --git a/Engine/Source/Runtime/Online/Android/OnlineSubsystemGooglePlay/Private/OnlineIdentityInterfaceGooglePlay.h b/Engine/Source/Runtime/Online/Android/OnlineSubsystemGooglePlay/Private/OnlineIdentityInterfaceGooglePlay.h
--- a/Engine/Source/Runtime/Online/Android/OnlineSubsystemGooglePlay/Private/OnlineIdentityInterfaceGooglePlay.h
+++ b/Engine/Source/Runtime/Online/Android/OnlineSubsystemGooglePlay/Private/OnlineIdentityInterfaceGooglePlay.h
@@ -5,6 +5,8 @@
 #include "OnlineIdentityInterface.h"
 #include "OnlineSubsystemGooglePlayPackage.h"
 
+extern "C" void Java_com_epicgames_ue4_GameActivity_nativeCompletedConnection(JNIEnv* LocalJNIEnv, jobject LocalThiz, jint userID, jint errorCode);
+
 class FOnlineIdentityGooglePlay :
 	public IOnlineIdentity
 {
@@ -19,6 +21,27 @@ private:
 	bool bLoggingInUser;
 	bool bRegisteringUser;
 
+	/** UID for this identity */
+	TSharedPtr< FUniqueNetIdString > UniqueNetId;
+
+	struct FPendingConnection
+	{
+		FOnlineIdentityGooglePlay * ConnectionInterface;
+		bool IsConnectionPending;
+	};
+
+	/** Instance of the connection context data */
+	static FPendingConnection PendingConnectRequest;
+
+	/**
+	* Native function called from Java when we get the connection result callback.
+	*
+	* @param LocalJNIEnv the current JNI environment
+	* @param LocalThiz instance of the Java GameActivity class
+	* @param errorCode The reported error code from the connection attempt
+	*/
+	friend void Java_com_epicgames_ue4_GameActivity_nativeCompletedConnection(JNIEnv* LocalJNIEnv, jobject LocalThiz, jint userID, jint errorCode);
+
 PACKAGE_SCOPE:
 
 	/**
@@ -52,6 +75,8 @@ public:
 	void OnLogin(const bool InLoggedIn,  const FString& InPlayerId, const FString& InPlayerAlias);
 	void OnLogout(const bool InLoggedIn);
 
+	void OnLoginCompleted(const int playerID, const int errorCode);
+
 };
 
 
diff --git a/Engine/Source/Runtime/Online/Android/OnlineSubsystemGooglePlay/Private/OnlineIdentityInterfaceGoolgePlay.cpp b/Engine/Source/Runtime/Online/Android/OnlineSubsystemGooglePlay/Private/OnlineIdentityInterfaceGoolgePlay.cpp
--- a/Engine/Source/Runtime/Online/Android/OnlineSubsystemGooglePlay/Private/OnlineIdentityInterfaceGoolgePlay.cpp
+++ b/Engine/Source/Runtime/Online/Android/OnlineSubsystemGooglePlay/Private/OnlineIdentityInterfaceGoolgePlay.cpp
@@ -2,6 +2,32 @@
 
 #include "OnlineSubsystemGooglePlayPrivatePCH.h"
 
+#include "Android/AndroidJNI.h"
+FOnlineIdentityGooglePlay::FPendingConnection FOnlineIdentityGooglePlay::PendingConnectRequest;
+
+// Java interface to deal with callbacks
+extern "C" void Java_com_epicgames_ue4_GameActivity_nativeCompletedConnection(JNIEnv* LocalJNIEnv, jobject LocalThiz, jint userID, jint errorCode)
+{
+	auto ConnectionInterface = FOnlineIdentityGooglePlay::PendingConnectRequest.ConnectionInterface;
+
+	if (!ConnectionInterface ||
+		!ConnectionInterface->MainSubsystem ||
+		!ConnectionInterface->MainSubsystem->GetAsyncTaskManager())
+	{
+		// We should call the delegate with a false parameter here, but if we don't have
+		// the async task manager we're not going to call it on the game thread.
+		return;
+	}
+
+	ConnectionInterface->MainSubsystem->GetAsyncTaskManager()->AddGenericToOutQueue([errorCode, userID]()
+	{
+		auto& PendingConnection = FOnlineIdentityGooglePlay::PendingConnectRequest;		
+
+		PendingConnection.ConnectionInterface->OnLoginCompleted(userID, errorCode);
+	});
+
+}
+
 FOnlineIdentityGooglePlay::FOnlineIdentityGooglePlay(FOnlineSubsystemGooglePlay* InSubsystem)
 	: bPrevLoggedIn(false)
 	, bLoggedIn(false)
@@ -12,6 +38,7 @@ FOnlineIdentityGooglePlay::FOnlineIdentityGooglePlay(FOnlineSubsystemGooglePlay*
 
 {
 	UE_LOG(LogOnline, Display, TEXT("FOnlineIdentityAndroid::FOnlineIdentityAndroid()"));
+	PendingConnectRequest.ConnectionInterface = this;
 }
 
 TSharedPtr<FUserOnlineAccount> FOnlineIdentityGooglePlay::GetUserAccount(const FUniqueNetId& UserId) const
@@ -30,8 +57,33 @@ TArray<TSharedPtr<FUserOnlineAccount> > FOnlineIdentityGooglePlay::GetAllUserAcc
 
 bool FOnlineIdentityGooglePlay::Login(int32 LocalUserNum, const FOnlineAccountCredentials& AccountCredentials) 
 {
-	TriggerOnLoginCompleteDelegates(LocalUserNum, false, FUniqueNetIdString(TEXT("")), FString("Login not supported"));
-	return false;
+	bool bStartedLogin = false;
+	if (bLoggedIn)
+	{
+		// already logged in so just report all is ok!
+		// Now logged in
+		bStartedLogin = true;
+
+		static const int32 MAX_TEXT_LINE_LEN = 32;
+		TCHAR Line[MAX_TEXT_LINE_LEN + 1] = { 0 };
+		int32 Len = FCString::Snprintf(Line, MAX_TEXT_LINE_LEN, TEXT("%d"), LocalUserNum);
+
+		const FString PlayerId(Line);
+		UniqueNetId = MakeShareable(new FUniqueNetIdString(PlayerId));
+		TriggerOnLoginCompleteDelegates(LocalUserNum, true, *UniqueNetId, TEXT(""));
+	}
+	else if (!PendingConnectRequest.IsConnectionPending)
+	{
+		// Kick the login sequence...
+		bStartedLogin = true;
+		PendingConnectRequest.IsConnectionPending = true;
+	}
+	else
+	{
+		TriggerOnLoginCompleteDelegates(LocalUserNum, false, FUniqueNetIdString(TEXT("")), FString("Already trying to login"));
+	}
+	
+	return bStartedLogin;
 }
 
 
@@ -44,7 +96,7 @@ bool FOnlineIdentityGooglePlay::Logout(int32 LocalUserNum)
 
 bool FOnlineIdentityGooglePlay::AutoLogin(int32 LocalUserNum)
 {
-	return false;
+	return Login(LocalUserNum, FOnlineAccountCredentials());
 }
 
 
@@ -111,3 +163,16 @@ FString FOnlineIdentityGooglePlay::GetAuthToken(int32 LocalUserNum) const
 void FOnlineIdentityGooglePlay::Tick(float DeltaTime)
 {
 }
+
+void FOnlineIdentityGooglePlay::OnLoginCompleted(const int playerID, const int errorCode)
+{
+	static const int32 MAX_TEXT_LINE_LEN = 32;
+	TCHAR Line[MAX_TEXT_LINE_LEN + 1] = { 0 };
+	int32 Len = FCString::Snprintf(Line, MAX_TEXT_LINE_LEN, TEXT("%d"), playerID);
+
+	UniqueNetId = MakeShareable(new FUniqueNetIdString(Line));
+	bLoggedIn = errorCode == 0;
+	TriggerOnLoginCompleteDelegates(playerID, errorCode == 0, *UniqueNetId, TEXT(""));
+
+	PendingConnectRequest.IsConnectionPending = false;
+}
diff --git a/Engine/Source/Runtime/Online/Android/OnlineSubsystemGooglePlay/Private/OnlineSubsystemGooglePlay.cpp b/Engine/Source/Runtime/Online/Android/OnlineSubsystemGooglePlay/Private/OnlineSubsystemGooglePlay.cpp
--- a/Engine/Source/Runtime/Online/Android/OnlineSubsystemGooglePlay/Private/OnlineSubsystemGooglePlay.cpp
+++ b/Engine/Source/Runtime/Online/Android/OnlineSubsystemGooglePlay/Private/OnlineSubsystemGooglePlay.cpp
@@ -102,7 +102,6 @@ bool FOnlineSubsystemGooglePlay::Init()
 	LeaderboardsInterface = MakeShareable(new FOnlineLeaderboardsGooglePlay(this));
 	AchievementsInterface = MakeShareable(new FOnlineAchievementsGooglePlay(this));
 	ExternalUIInterface = MakeShareable(new FOnlineExternalUIGooglePlay());
-	ConnectionInterface = MakeShareable(new FOnlineConnectionGooglePlay(this));
 
 	return true;
 }
@@ -134,7 +133,6 @@ bool FOnlineSubsystemGooglePlay::Shutdown()
 	DESTRUCT_INTERFACE(AchievementsInterface);
 	DESTRUCT_INTERFACE(LeaderboardsInterface);
 	DESTRUCT_INTERFACE(IdentityInterface);
-	DESTRUCT_INTERFACE(ConnectionInterface);
 #undef DESTRUCT_INTERFACE
 
 	OnlineAsyncTaskThreadRunnable.Reset();
diff --git a/Engine/Source/Runtime/Online/Android/OnlineSubsystemGooglePlay/Public/OnlineSubsystemGooglePlay.h b/Engine/Source/Runtime/Online/Android/OnlineSubsystemGooglePlay/Public/OnlineSubsystemGooglePlay.h
--- a/Engine/Source/Runtime/Online/Android/OnlineSubsystemGooglePlay/Public/OnlineSubsystemGooglePlay.h
+++ b/Engine/Source/Runtime/Online/Android/OnlineSubsystemGooglePlay/Public/OnlineSubsystemGooglePlay.h
@@ -7,7 +7,6 @@
 #include "OnlineAchievementsInterfaceGooglePlay.h"
 #include "OnlineLeaderboardInterfaceGooglePlay.h"
 #include "OnlineExternalUIInterfaceGooglePlay.h"
-#include "OnlineConnectionInterfaceGooglePlay.h"
 #include "UniquePtr.h"
 
 /**
@@ -41,7 +40,6 @@ public:
 	virtual IOnlineUserPtr GetUserInterface() const override { return NULL; }
 	virtual IOnlineAchievementsPtr GetAchievementsInterface() const override;
 	virtual IOnlinePresencePtr GetPresenceInterface() const override { return NULL; }
-	virtual IOnlineConnectionPtr GetConnectionInterface() const override { return ConnectionInterface;  }
 
 	virtual class UObject* GetNamedInterface(FName InterfaceName) override { return NULL; }
 	virtual void SetNamedInterface(FName InterfaceName, class UObject* NewInterface) override {}
@@ -85,9 +83,6 @@ private:
 
 	/** Interface to the external UI services */
 	FOnlineExternalUIGooglePlayPtr ExternalUIInterface;
-
-	/** Interface to the connection system */
-	FOnlineConnectionGooglePlayPtr ConnectionInterface;
 };
 
 typedef TSharedPtr<FOnlineSubsystemGooglePlay, ESPMode::ThreadSafe> FOnlineSubsystemGooglePlayPtr;
diff --git a/Engine/Source/Runtime/Online/IOS/OnlineSubsystemIOS/Private/OnlineSubsystemIOS.cpp b/Engine/Source/Runtime/Online/IOS/OnlineSubsystemIOS/Private/OnlineSubsystemIOS.cpp
--- a/Engine/Source/Runtime/Online/IOS/OnlineSubsystemIOS/Private/OnlineSubsystemIOS.cpp
+++ b/Engine/Source/Runtime/Online/IOS/OnlineSubsystemIOS/Private/OnlineSubsystemIOS.cpp
@@ -117,11 +117,6 @@ IOnlinePresencePtr FOnlineSubsystemIOS::GetPresenceInterface() const
 	return NULL;
 }
 
-IOnlineConnectionPtr FOnlineSubsystemIOS::GetConnectionInterface() const
-{
-	return NULL;
-}
-
 bool FOnlineSubsystemIOS::Init() 
 {
 	bool bSuccessfullyStartedUp = true;
diff --git a/Engine/Source/Runtime/Online/IOS/OnlineSubsystemIOS/Public/OnlineSubsystemIOS.h b/Engine/Source/Runtime/Online/IOS/OnlineSubsystemIOS/Public/OnlineSubsystemIOS.h
--- a/Engine/Source/Runtime/Online/IOS/OnlineSubsystemIOS/Public/OnlineSubsystemIOS.h
+++ b/Engine/Source/Runtime/Online/IOS/OnlineSubsystemIOS/Public/OnlineSubsystemIOS.h
@@ -39,7 +39,6 @@ public:
 	virtual IOnlineUserPtr GetUserInterface() const override;
 	virtual IOnlineMessagePtr GetMessageInterface() const override;
 	virtual IOnlinePresencePtr GetPresenceInterface() const override;
-	virtual IOnlineConnectionPtr GetConnectionInterface() const override;
 	virtual bool Init() override;
 	virtual bool Shutdown() override;
 	virtual FString GetAppId() const override;
diff --git a/Engine/Source/Runtime/Online/OnlineSubsystem/Public/Interfaces/OnlineConnectionInterface.h b/Engine/Source/Runtime/Online/OnlineSubsystem/Public/Interfaces/OnlineConnectionInterface.h
deleted file mode 100644
diff --git a/Engine/Source/Runtime/Online/OnlineSubsystem/Public/OnlineSubsystem.h b/Engine/Source/Runtime/Online/OnlineSubsystem/Public/OnlineSubsystem.h
--- a/Engine/Source/Runtime/Online/OnlineSubsystem/Public/OnlineSubsystem.h
+++ b/Engine/Source/Runtime/Online/OnlineSubsystem/Public/OnlineSubsystem.h
@@ -50,7 +50,6 @@ typedef TSharedPtr<class IOnlineUser, ESPMode::ThreadSafe> IOnlineUserPtr;
 typedef TSharedPtr<class IOnlineMessage, ESPMode::ThreadSafe> IOnlineMessagePtr;
 typedef TSharedPtr<class IOnlinePresence, ESPMode::ThreadSafe> IOnlinePresencePtr;
 typedef TSharedPtr<class IOnlineParty, ESPMode::ThreadSafe> IOnlinePartyPtr;
-typedef TSharedPtr<class IOnlineConnection, ESPMode::ThreadSafe> IOnlineConnectionPtr;
 
 /**
  * Called when the connection state as reported by the online platform changes
@@ -247,12 +246,6 @@ public:
 	 */
 	virtual class UObject* GetNamedInterface(FName InterfaceName) = 0;
 
-	/**
-	* Get the interface for managing the connection to an online service
-	* @return Interface pointer for the appropriate online service
-	*/
-	virtual IOnlineConnectionPtr GetConnectionInterface() const = 0;
-
 	/**
 	 * Set a custom UObject to be preserved by the online subsystem
 	 *
diff --git a/Engine/Source/Runtime/Online/OnlineSubsystemAmazon/Public/OnlineSubsystemAmazon.h b/Engine/Source/Runtime/Online/OnlineSubsystemAmazon/Public/OnlineSubsystemAmazon.h
--- a/Engine/Source/Runtime/Online/OnlineSubsystemAmazon/Public/OnlineSubsystemAmazon.h
+++ b/Engine/Source/Runtime/Online/OnlineSubsystemAmazon/Public/OnlineSubsystemAmazon.h
@@ -117,11 +117,6 @@ public:
 		return NULL;
 	}
 
-	IOnlineConnectionPtr GetConnectionInterface() const override
-	{
-		return NULL;
-	}
-
 	virtual bool Init() override;
 	virtual bool Shutdown() override;
 	virtual FString GetAppId() const override;
diff --git a/Engine/Source/Runtime/Online/OnlineSubsystemFacebook/Private/IOS/OnlineSubsystemFacebook.cpp b/Engine/Source/Runtime/Online/OnlineSubsystemFacebook/Private/IOS/OnlineSubsystemFacebook.cpp
--- a/Engine/Source/Runtime/Online/OnlineSubsystemFacebook/Private/IOS/OnlineSubsystemFacebook.cpp
+++ b/Engine/Source/Runtime/Online/OnlineSubsystemFacebook/Private/IOS/OnlineSubsystemFacebook.cpp
@@ -116,11 +116,6 @@ IOnlinePartyPtr FOnlineSubsystemFacebook::GetPartyInterface() const
 	return NULL;
 }
 
-IOnlineConnectionPtr FOnlineSubsystemFacebook::GetConnectionInterface() const
-{
-	return NULL;
-}
-
 bool FOnlineSubsystemFacebook::Init() 
 {
 	bool bSuccessfullyStartedUp = true;
diff --git a/Engine/Source/Runtime/Online/OnlineSubsystemFacebook/Private/Windows/OnlineSubsystemFacebook.cpp b/Engine/Source/Runtime/Online/OnlineSubsystemFacebook/Private/Windows/OnlineSubsystemFacebook.cpp
--- a/Engine/Source/Runtime/Online/OnlineSubsystemFacebook/Private/Windows/OnlineSubsystemFacebook.cpp
+++ b/Engine/Source/Runtime/Online/OnlineSubsystemFacebook/Private/Windows/OnlineSubsystemFacebook.cpp
@@ -114,11 +114,6 @@ IOnlinePartyPtr FOnlineSubsystemFacebook::GetPartyInterface() const
 	return NULL;
 }
 
-IOnlineConnectionPtr FOnlineSubsystemFacebook::GetConnectionInterface() const
-{
-	return NULL;
-}
-
 bool FOnlineSubsystemFacebook::Tick(float DeltaTime)
 {
 	if (FacebookIdentity.IsValid())
diff --git a/Engine/Source/Runtime/Online/OnlineSubsystemFacebook/Public/OnlineSubsystemFacebook.h b/Engine/Source/Runtime/Online/OnlineSubsystemFacebook/Public/OnlineSubsystemFacebook.h
--- a/Engine/Source/Runtime/Online/OnlineSubsystemFacebook/Public/OnlineSubsystemFacebook.h
+++ b/Engine/Source/Runtime/Online/OnlineSubsystemFacebook/Public/OnlineSubsystemFacebook.h
@@ -41,8 +41,6 @@ public:
 	virtual IOnlineUserPtr GetUserInterface() const override;
 	virtual IOnlineMessagePtr GetMessageInterface() const override;
 	virtual IOnlinePresencePtr GetPresenceInterface() const override;
-	virtual IOnlineConnectionPtr GetConnectionInterface() const override;
-
 	virtual bool Init() override;
 	virtual bool Shutdown() override;
 	virtual FString GetAppId() const override;
diff --git a/Engine/Source/Runtime/Online/OnlineSubsystemNull/Private/OnlineSubsystemNull.cpp b/Engine/Source/Runtime/Online/OnlineSubsystemNull/Private/OnlineSubsystemNull.cpp
--- a/Engine/Source/Runtime/Online/OnlineSubsystemNull/Private/OnlineSubsystemNull.cpp
+++ b/Engine/Source/Runtime/Online/OnlineSubsystemNull/Private/OnlineSubsystemNull.cpp
@@ -108,11 +108,6 @@ IOnlinePresencePtr FOnlineSubsystemNull::GetPresenceInterface() const
 	return NULL;
 }
 
-IOnlineConnectionPtr FOnlineSubsystemNull::GetConnectionInterface() const
-{
-	return NULL;
-}
-
 bool FOnlineSubsystemNull::Tick(float DeltaTime)
 {
 	if (OnlineAsyncTaskThreadRunnable)
diff --git a/Engine/Source/Runtime/Online/OnlineSubsystemNull/Public/OnlineSubsystemNull.h b/Engine/Source/Runtime/Online/OnlineSubsystemNull/Public/OnlineSubsystemNull.h
--- a/Engine/Source/Runtime/Online/OnlineSubsystemNull/Public/OnlineSubsystemNull.h
+++ b/Engine/Source/Runtime/Online/OnlineSubsystemNull/Public/OnlineSubsystemNull.h
@@ -52,7 +52,6 @@ public:
 	virtual IOnlineUserPtr GetUserInterface() const override;
 	virtual IOnlineMessagePtr GetMessageInterface() const override;
 	virtual IOnlinePresencePtr GetPresenceInterface() const override;
-	virtual IOnlineConnectionPtr GetConnectionInterface() const override;
 
 	virtual bool Init() override;
 	virtual bool Shutdown() override;
diff --git a/Engine/Source/Runtime/Online/OnlineSubsystemSteam/Private/OnlineSubsystemSteam.cpp b/Engine/Source/Runtime/Online/OnlineSubsystemSteam/Private/OnlineSubsystemSteam.cpp
--- a/Engine/Source/Runtime/Online/OnlineSubsystemSteam/Private/OnlineSubsystemSteam.cpp
+++ b/Engine/Source/Runtime/Online/OnlineSubsystemSteam/Private/OnlineSubsystemSteam.cpp
@@ -249,12 +249,6 @@ IOnlinePresencePtr FOnlineSubsystemSteam::GetPresenceInterface() const
 	return NULL;
 }
 
-IOnlineConnectionPtr FOnlineSubsystemSteam::GetConnectionInterface() const
-{
-	return NULL;
-}
-
-
 void FOnlineSubsystemSteam::QueueAsyncTask(FOnlineAsyncTask* AsyncTask)
 {
 	check(OnlineAsyncTaskThreadRunnable);
diff --git a/Engine/Source/Runtime/Online/OnlineSubsystemSteam/Public/OnlineSubsystemSteam.h b/Engine/Source/Runtime/Online/OnlineSubsystemSteam/Public/OnlineSubsystemSteam.h
--- a/Engine/Source/Runtime/Online/OnlineSubsystemSteam/Public/OnlineSubsystemSteam.h
+++ b/Engine/Source/Runtime/Online/OnlineSubsystemSteam/Public/OnlineSubsystemSteam.h
@@ -230,8 +230,6 @@ public:
 	virtual IOnlineUserPtr GetUserInterface() const override;
 	virtual IOnlineMessagePtr GetMessageInterface() const override;
 	virtual IOnlinePresencePtr GetPresenceInterface() const override;
-	virtual IOnlineConnectionPtr GetConnectionInterface() const override;
-
 	virtual bool IsLocalPlayer(const FUniqueNetId& UniqueId) const override;
 	virtual bool Init() override;
 	virtual bool Shutdown() override;
diff --git a/Engine/Source/Runtime/Online/OnlineSubsystemUtils/Classes/ConnectionCallbackProxy.h b/Engine/Source/Runtime/Online/OnlineSubsystemUtils/Classes/ConnectionCallbackProxy.h
--- a/Engine/Source/Runtime/Online/OnlineSubsystemUtils/Classes/ConnectionCallbackProxy.h
+++ b/Engine/Source/Runtime/Online/OnlineSubsystemUtils/Classes/ConnectionCallbackProxy.h
@@ -2,7 +2,7 @@
 
 #pragma once
 
-#include "Runtime/Online/OnlineSubsystem/Public/Interfaces/OnlineConnectionInterface.h"
+#include "Runtime/Online/OnlineSubsystem/Public/Interfaces/OnlineIdentityInterface.h"
 #include "ConnectionCallbackProxy.generated.h"
 
 DECLARE_DYNAMIC_MULTICAST_DELEGATE_OneParam(FOnlineConnectionResult, int32, ErrorCode);
@@ -31,9 +31,12 @@ class UConnectionCallbackProxy : public UOnlineBlueprintCallProxyBase
 
 private:
 	// Internal callback when the achievement query completes, calls out to the public success/failure callbacks
-	void OnConnectCompleted(int errorCode, bool bSuccess);
+	void OnLoginCompleted(int32 LocalUserNum, bool bWasSuccessful, const FUniqueNetId& UserId, const FString& Error);
 
 private:
 	// The player controller triggering things
 	TWeakObjectPtr<APlayerController> PlayerControllerWeakPtr;
+
+	// Connection callback delegate
+	FOnLoginCompleteDelegate OnLoginCompleteDelegate;
 };
\ No newline at end of file
diff --git a/Engine/Source/Runtime/Online/OnlineSubsystemUtils/Private/ConnectionCallbackProxy.cpp b/Engine/Source/Runtime/Online/OnlineSubsystemUtils/Private/ConnectionCallbackProxy.cpp
--- a/Engine/Source/Runtime/Online/OnlineSubsystemUtils/Private/ConnectionCallbackProxy.cpp
+++ b/Engine/Source/Runtime/Online/OnlineSubsystemUtils/Private/ConnectionCallbackProxy.cpp
@@ -8,6 +8,9 @@
 UConnectionCallbackProxy::UConnectionCallbackProxy(const class FPostConstructInitializeProperties& PCIP)
 	: Super(PCIP)
 {
+
+	OnLoginCompleteDelegate.BindUObject(this, &ThisClass::OnLoginCompleted);
+
 }
 
 UConnectionCallbackProxy* UConnectionCallbackProxy::ConnectToService(class APlayerController* PlayerController)
@@ -24,13 +27,21 @@ void UConnectionCallbackProxy::Activate()
 
 	if (Helper.IsValid())
 	{
-		IOnlineConnectionPtr Connection = Helper.OnlineSub->GetConnectionInterface();
-		if (Connection.IsValid())
+		IOnlineIdentityPtr OnlineIdentity = Helper.OnlineSub->GetIdentityInterface();
+		if (OnlineIdentity.IsValid())
 		{
-			FOnConnectionCompleteDelegate ConnectionFinishedDelegate = FOnConnectionCompleteDelegate::CreateUObject(this, &ThisClass::OnConnectCompleted);
-
-			Connection->Connect(ConnectionFinishedDelegate); // this will get called when everything is resolved
+			ULocalPlayer * localPlayer = CastChecked<ULocalPlayer>(PlayerControllerWeakPtr.Get()->Player);
 
+			if (!OnlineIdentity->OnLoginCompleteDelegates[localPlayer->ControllerId].IsBoundToObject(this))
+			{
+				OnlineIdentity->AddOnLoginCompleteDelegate(localPlayer->ControllerId, OnLoginCompleteDelegate);
+				OnlineIdentity->Login(localPlayer->ControllerId, FOnlineAccountCredentials()); /// Probably need to supply real creds here somehow... doesn't apply for all imple however.
+			}
+			else
+			{
+				// already trying to login!
+			}
+			
 			// OnQueryCompleted will get called, nothing more to do now
 			return;
 		}
@@ -44,9 +55,21 @@ void UConnectionCallbackProxy::Activate()
 	OnFailure.Broadcast(0);
 }
 
-void UConnectionCallbackProxy::OnConnectCompleted(int errorCode, bool bSuccess)
+void UConnectionCallbackProxy::OnLoginCompleted(int32 LocalUserNum, bool bWasSuccessful, const FUniqueNetId& UserId, const FString& Error)
 {
-	if (bSuccess)
+	FOnlineSubsystemBPCallHelper Helper(TEXT("ConnectToService"));
+	Helper.QueryIDFromPlayerController(PlayerControllerWeakPtr.Get());
+
+	if (Helper.IsValid())
+	{
+		IOnlineIdentityPtr OnlineIdentity = Helper.OnlineSub->GetIdentityInterface();
+		if (OnlineIdentity.IsValid())
+		{
+			OnlineIdentity->ClearOnLoginCompleteDelegate(LocalUserNum, OnLoginCompleteDelegate);
+		}
+	}
+
+	if (bWasSuccessful)
 	{
 		OnSuccess.Broadcast(0);
 	}
