Moving Build Patch Services from NoRedist to Runtime
Moving BuildPatchTool out of NoRedist

diff --git a/Engine/Source/Programs/BuildPatchTool/BuildPatchTool.Build.cs b/Engine/Source/Programs/BuildPatchTool/BuildPatchTool.Build.cs
new file mode 100644
--- /dev/null
+++ b/Engine/Source/Programs/BuildPatchTool/BuildPatchTool.Build.cs
@@ -0,0 +1,23 @@
+// Copyright 1998-2014 Epic Games, Inc. All Rights Reserved.
+
+using UnrealBuildTool;
+
+public class BuildPatchTool : ModuleRules
+{
+	public BuildPatchTool( TargetInfo Target )
+	{
+		PublicIncludePaths.Add("Runtime/Launch/Public");
+
+		// For LaunchEngineLoop.cpp include
+		PrivateIncludePaths.Add("Runtime/Launch/Private");
+
+		PrivateDependencyModuleNames.AddRange(
+			new string[] 
+			{
+				"Core",
+				"BuildPatchServices",
+				"Projects",
+			}
+		);
+	}
+}
\ No newline at end of file
diff --git a/Engine/Source/Programs/BuildPatchTool/BuildPatchTool.Target.cs b/Engine/Source/Programs/BuildPatchTool/BuildPatchTool.Target.cs
new file mode 100644
--- /dev/null
+++ b/Engine/Source/Programs/BuildPatchTool/BuildPatchTool.Target.cs
@@ -0,0 +1,59 @@
+// Copyright 1998-2014 Epic Games, Inc. All Rights Reserved.
+
+using UnrealBuildTool;
+using System.Collections.Generic;
+
+public class BuildPatchToolTarget : TargetRules
+{
+	public BuildPatchToolTarget( TargetInfo Target )
+	{
+		Type = TargetType.Program;
+	}
+
+	//
+	// TargetRules interface.
+	//
+
+	public override void SetupBinaries(
+		TargetInfo Target,
+		ref List<UEBuildBinaryConfiguration> OutBuildBinaryConfigurations,
+		ref List<string> OutExtraModuleNames)
+	{
+		OutBuildBinaryConfigurations.Add(
+			new UEBuildBinaryConfiguration(
+				InType: UEBuildBinaryType.Executable,
+				InModuleNames: new List<string>() {
+					"BuildPatchTool"
+				}
+			)
+		);
+	}
+
+	public override bool ShouldCompileMonolithic(UnrealTargetPlatform InPlatform, UnrealTargetConfiguration InConfiguration)
+	{
+		return true;
+	}
+
+	public override void SetupGlobalEnvironment(
+		TargetInfo Target,
+		ref LinkEnvironmentConfiguration OutLinkEnvironmentConfiguration,
+		ref CPPEnvironmentConfiguration OutCPPEnvironmentConfiguration)
+	{
+		UEBuildConfiguration.bBuildEditor = false;
+		UEBuildConfiguration.bCompileAgainstEngine = false;
+		UEBuildConfiguration.bCompileAgainstCoreUObject = false;
+		UEBuildConfiguration.bCompileLeanAndMeanUE = true;
+
+		OutCPPEnvironmentConfiguration.Definitions.Add("WITH_BUILDPATCHGENERATION=1");
+
+		OutLinkEnvironmentConfiguration.bIsBuildingConsoleApplication = true;
+		OutLinkEnvironmentConfiguration.bBuildAdditionalConsoleApplication = false;
+		OutLinkEnvironmentConfiguration.bHasExports = false;
+	}
+    public override bool GUBP_AlwaysBuildWithTools(UnrealTargetPlatform InHostPlatform, bool bBuildingRocket, out bool bInternalToolOnly, out bool SeparateNode)
+    {
+        bInternalToolOnly = true;
+        SeparateNode = true;
+        return true;
+    }
+}
diff --git a/Engine/Source/Programs/BuildPatchTool/Private/BuildPatchToolMain.cpp b/Engine/Source/Programs/BuildPatchTool/Private/BuildPatchToolMain.cpp
new file mode 100644
--- /dev/null
+++ b/Engine/Source/Programs/BuildPatchTool/Private/BuildPatchToolMain.cpp
@@ -0,0 +1,388 @@
+// Copyright 1998-2014 Epic Games, Inc. All Rights Reserved.
+
+/*=============================================================================
+	BuildPatchToolMain.cpp: Implements the BuildPatchTool application's main loop.
+
+	The tool can be used in two modes of operation:
+	(1) to generate patch data (manifest, chunks, files) for build images given the existing cloud data; or
+	(2) to "compactify" a cloud directory by removing all orphaned chunks, not referenced by any manifest file.
+
+	In order to trigger compactify functionality, the -compactify commandline argument should be specified.
+
+
+	GENERATE PATCH DATA MODE
+
+	The tool supports generating chunk based patches or simple file based patches. Chunk based patch data will be generated by default
+	but you can switch to file based patch data by adding the -nochunks commandline argument.
+
+	File based patch data is only recommended for small build images that probably wouldn't benefit from chunking. I.e. Updates are commonly only a few small changed files.
+
+	Required arguments:
+	-BuildRoot=""		Specifies in quotes the directory containing the build image to be read.
+	-CloudDir=""		Specifies in quotes the cloud directory where existing data will be recognised from, and new data added to.
+	-AppID=123456		Specifies without quotes, the ID number for the app
+	-AppName=""			Specifies in quotes, the name of the app
+	-BuildVersion=""	Specifies in quotes, the version string for the build image
+	-AppLaunch=""		Specifies in quotes, the path to the app executable, must be relative to, and inside of BuildRoot.
+	-AppArgs=""			Specifies in quotes, the commandline to send to the app on launch.
+
+	Optional arguments:
+	-stdout				Adds stdout logging to the app.
+	-nochunks			Creates file based patch data instead of chunk based patch data.
+	-FileIgnoreList=""	Specifies in quotes, the path to a text file containing BuildRoot relative files, separated by \r\n line endings, to not be included in the build.
+	-PrereqName=""		Specifies in quotes, the display name for the prerequisites installer
+	-PrereqPath=""		Specifies in quotes, the prerequisites installer to launch on successful product install
+	-PrereqArgs=""		Specifies in quotes, the commandline to send to prerequisites installer on launch
+	
+	-custom="field=value"	 	Adds a custom string field to the build manifest
+	-customint="field=number"	Adds a custom int64 field to the build manifest 
+	-customfloat="field=number"	Adds a custom double field to the build manifest
+
+	COMPACTIFY MODE
+
+	Required arguments:
+	-CloudDir=""		Specifies in quotes the cloud directory where manifest files and chunks to be compactified can be found.
+	-compactify			Must be specified to launch the tool in compactify mode
+
+	Optional arguments:
+	-stdout				Adds stdout logging to the app.
+	-preview			Log all the actions it will take to update internal structures, but don't actually execute them.
+
+
+	Example command lines:
+	-BuildRoot="D:\Builds\ExampleGame_[07-02_03.00]" -FileIgnoreList="D:\Builds\ExampleGame_[07-02_03.00]\Manifest_DebugFiles.txt" -CloudDir="D:\BuildPatchCloud" -AppID=123456 -AppName="Example" -BuildVersion="ExampleGame_[07-02_03.00]" -AppLaunch=".\ExampleGame\Binaries\Win32\ExampleGame.exe" -AppArgs="-pak -nosteam"
+	-BuildRoot="C:\Program Files (x86)\Community Portal" -CloudDir="E:\BuildPatchCloud" -AppID=0 -AppName="Community Portal" -BuildVersion="++depot+UE4-CL-1791234" -AppLaunch=".\Engine\Binaries\Win32\CommunityPortal.exe" -AppArgs="" -nochunks
+	-CloudDir="E:\BuildPatchCloude" -compactify
+=============================================================================*/
+
+#include "RequiredProgramMainCPPInclude.h"
+
+#include "BuildPatchServices.h"
+
+// Ensure compiled with patch generation code
+#if WITH_BUILDPATCHGENERATION
+#else
+#error BuildPatchTool must define WITH_BUILDPATCHGENERATION
+#endif
+
+IMPLEMENT_APPLICATION( BuildPatchTool, "BuildPatchTool" );
+
+struct FCommandLineMatcher
+{
+	FString Command;
+
+	FCommandLineMatcher(){}
+
+	FORCEINLINE bool Matches( const FString& ToMatch ) const
+	{
+		if( ToMatch.StartsWith( Command, ESearchCase::CaseSensitive ) )
+		{
+			return true;
+		}
+
+		return false;
+	}
+};
+
+class FBuildPatchOutputDevice : public FOutputDevice
+{
+public:
+	virtual void Serialize( const TCHAR* V, ELogVerbosity::Type Verbosity, const class FName& Category ) override
+	{
+#if PLATFORM_USE_LS_SPEC_FOR_WIDECHAR
+		printf( "\n%ls", *FOutputDevice::FormatLogLine( Verbosity, Category, V, GPrintLogTimes ) );
+#else
+		wprintf( TEXT( "\n%s" ), *FOutputDevice::FormatLogLine( Verbosity, Category, V, GPrintLogTimes ) );
+#endif
+		fflush( stdout );
+	}
+};
+
+
+int32 BuildPatchToolMain( const TCHAR* CommandLine )
+{
+	// Initialize the command line
+	FCommandLine::Set(CommandLine);
+
+	// Add log devices
+	if (FParse::Param(FCommandLine::Get(), TEXT("stdout")))
+	{
+		GLog->AddOutputDevice(new FBuildPatchOutputDevice());
+	}
+	if (FPlatformMisc::IsDebuggerPresent())
+	{
+		GLog->AddOutputDevice(new FOutputDeviceDebug());
+	}
+
+	GLog->Logf(TEXT("BuildPatchToolMain ran with: %s"), CommandLine);
+
+	FPlatformProcess::SetCurrentWorkingDirectoryToBaseDir();
+	bool bSuccess = false;
+
+	FString RootDirectory;
+	FString CloudDirectory;
+	uint32  AppID=0;
+	FString AppName;
+	FString BuildVersion;
+	FString LaunchExe;
+	FString LaunchCommand;
+	FString IgnoreListFile;
+	FString PrereqName;
+	FString PrereqPath;
+	FString PrereqArgs;
+	TMap<FString, FVariant> CustomFields;
+
+	bool bCompactify = false;
+	bool bPreviewCompactify = false;
+
+	// Collect all the info from the CommandLine
+	TArray< FString > Tokens, Switches;
+	FCommandLine::Parse(FCommandLine::Get(), Tokens, Switches);
+	if (Switches.Num() > 0)
+	{
+		int32 BuildRootIdx;
+		int32 CloudDirIdx;
+		int32 AppIDIdx;
+		int32 AppNameIdx;
+		int32 BuildVersionIdx;
+		int32 AppLaunchIdx;
+		int32 AppArgsIdx;
+		int32 FileIgnoreListIdx;
+		int32 PrereqNameIdx;
+		int32 PrereqPathIdx;
+		int32 PrereqArgsIdx;
+
+
+		FCommandLineMatcher Matcher;
+		bSuccess = true;
+
+		Matcher.Command = TEXT("compactify");
+		bCompactify = Switches.FindMatch(Matcher) != INDEX_NONE;
+
+		Matcher.Command = TEXT("preview");
+		bPreviewCompactify = bCompactify && Switches.FindMatch(Matcher) != INDEX_NONE;
+		
+		Matcher.Command = TEXT( "BuildRoot" );
+		BuildRootIdx = Switches.FindMatch( Matcher );
+
+		Matcher.Command = TEXT( "CloudDir" );
+		CloudDirIdx = Switches.FindMatch( Matcher );
+
+		Matcher.Command = TEXT( "AppID" );
+		AppIDIdx = Switches.FindMatch( Matcher );
+
+		Matcher.Command = TEXT( "AppName" );
+		AppNameIdx = Switches.FindMatch( Matcher );
+
+		Matcher.Command = TEXT( "BuildVersion" );
+		BuildVersionIdx = Switches.FindMatch( Matcher );
+
+		Matcher.Command = TEXT( "AppLaunch" );
+		AppLaunchIdx = Switches.FindMatch( Matcher );
+
+		Matcher.Command = TEXT( "AppArgs" );
+		AppArgsIdx = Switches.FindMatch( Matcher );
+
+		Matcher.Command = TEXT( "FileIgnoreList" );
+		FileIgnoreListIdx = Switches.FindMatch( Matcher );
+
+		Matcher.Command = TEXT( "PrereqName" );
+		PrereqNameIdx = Switches.FindMatch( Matcher );
+
+		Matcher.Command = TEXT( "PrereqPath" );
+		PrereqPathIdx = Switches.FindMatch( Matcher );
+
+		Matcher.Command = TEXT( "PrereqArgs" );
+		PrereqArgsIdx = Switches.FindMatch( Matcher );
+
+		// Check required param indexes
+		bSuccess = bSuccess && CloudDirIdx != INDEX_NONE;
+		if (!bCompactify)
+		{
+			bSuccess = bSuccess && BuildRootIdx != INDEX_NONE;
+			bSuccess = bSuccess && AppIDIdx != INDEX_NONE;
+			bSuccess = bSuccess && AppNameIdx != INDEX_NONE;
+			bSuccess = bSuccess && BuildVersionIdx != INDEX_NONE;
+			bSuccess = bSuccess && AppLaunchIdx != INDEX_NONE;
+			bSuccess = bSuccess && AppArgsIdx != INDEX_NONE;
+		}
+
+		// Get required param values
+		bSuccess = bSuccess && FParse::Value( *Switches[CloudDirIdx], TEXT( "CloudDir=" ), CloudDirectory );
+		if (!bCompactify)
+		{
+			bSuccess = bSuccess && FParse::Value(*Switches[BuildRootIdx], TEXT("BuildRoot="), RootDirectory);
+			bSuccess = bSuccess && FParse::Value(*Switches[AppIDIdx], TEXT("AppID="), AppID);
+			bSuccess = bSuccess && FParse::Value(*Switches[AppNameIdx], TEXT("AppName="), AppName);
+			bSuccess = bSuccess && FParse::Value(*Switches[BuildVersionIdx], TEXT("BuildVersion="), BuildVersion);
+			bSuccess = bSuccess && FParse::Value(*Switches[AppLaunchIdx], TEXT("AppLaunch="), LaunchExe);
+			bSuccess = bSuccess && FParse::Value(*Switches[AppArgsIdx], TEXT("AppArgs="), LaunchCommand);
+		}
+
+		// Get optional param values
+		if( FileIgnoreListIdx != INDEX_NONE )
+		{
+			FParse::Value( *Switches[ FileIgnoreListIdx ], TEXT( "FileIgnoreList=" ), IgnoreListFile );
+		}
+
+		if( PrereqNameIdx != INDEX_NONE )
+		{
+			FParse::Value( *Switches[ PrereqNameIdx ], TEXT( "FileIgnoreList=" ), PrereqName );
+		}
+
+		if( PrereqPathIdx != INDEX_NONE )
+		{
+			FParse::Value( *Switches[ PrereqPathIdx ], TEXT( "FileIgnoreList=" ), PrereqPath );
+		}
+
+		if( PrereqArgsIdx != INDEX_NONE )
+		{
+			FParse::Value( *Switches[ PrereqArgsIdx ], TEXT( "FileIgnoreList=" ), PrereqArgs );
+		}
+
+		FString CustomValue;
+		FString Left;
+		FString Right;
+		for (const auto& Switch : Switches)
+		{
+			if (FParse::Value(*Switch, TEXT("custom="), CustomValue))
+			{
+				if (CustomValue.Split(TEXT("="), &Left, &Right))
+				{
+					Left.Trim();
+					Left.TrimTrailing();
+					Right.Trim();
+					Right.TrimTrailing();
+					CustomFields.Add(Left, FVariant(Right));
+				}
+			}
+			else if (FParse::Value(*Switch, TEXT("customfloat="), CustomValue))
+			{
+				if (CustomValue.Split(TEXT("="), &Left, &Right))
+				{
+					Left.Trim();
+					Left.TrimTrailing();
+					Right.Trim();
+					Right.TrimTrailing();
+					if (!Right.IsNumeric())
+					{
+						GLog->Log(ELogVerbosity::Error, TEXT("An error occurred processing token -customint. Non Numeric character found right of ="));
+						bSuccess = false;
+					}
+					CustomFields.Add(Left, FVariant(TCString<TCHAR>::Atod(*Right)));
+				}
+			}
+			else if (FParse::Value(*Switch, TEXT("customint="), CustomValue))
+			{
+				if (CustomValue.Split(TEXT("="), &Left, &Right))
+				{
+					Left.Trim();
+					Left.TrimTrailing();
+					Right.Trim();
+					Right.TrimTrailing();
+					if (!Right.IsNumeric())
+					{
+						GLog->Log(ELogVerbosity::Error, TEXT("An error occurred processing token -customfloat. Non Numeric character found right of ="));
+						bSuccess = false;
+					}
+					CustomFields.Add(Left, FVariant(TCString<TCHAR>::Atoi64(*Right)));
+				}
+			}
+		}
+
+		FPaths::NormalizeDirectoryName( RootDirectory );
+		FPaths::NormalizeDirectoryName( CloudDirectory );
+	}
+
+	// Check for argument error
+	if( !bSuccess )
+	{
+		GLog->Log(ELogVerbosity::Error, TEXT("An error occurred processing arguments"));
+		return 1;
+	}
+
+
+	// Initialize the file manager
+	IFileManager::Get().ProcessCommandLineOptions();
+	
+	// Populate cultures for localization (used by BuildPatch progress info)
+	BeginInitTextLocalization();
+	FInternationalization I18N = FInternationalization::Get();
+
+	// Load the BuildPatchServices Module
+	TSharedPtr<IBuildPatchServicesModule> BuildPatchServicesModule = StaticCastSharedPtr<IBuildPatchServicesModule>( FModuleManager::Get().LoadModule( TEXT( "BuildPatchServices" ) ) );
+
+	// Setup the module
+	BuildPatchServicesModule->SetCloudDirectory( CloudDirectory + TEXT( "/" ) );
+
+	if (bCompactify)
+	{
+		// Run the compactify routine
+		bSuccess = BuildPatchServicesModule->CompactifyCloudDirectory(bPreviewCompactify);
+	}
+	else
+	{
+		FBuildPatchSettings Settings;
+		Settings.RootDirectory = RootDirectory + TEXT("/");
+		Settings.InAppID = AppID;
+		Settings.AppName = AppName;
+		Settings.BuildVersion = BuildVersion;
+		Settings.LaunchExe = LaunchExe;
+		Settings.LaunchCommand = LaunchCommand;
+		Settings.IgnoreListFile = IgnoreListFile;
+		Settings.PrereqName = PrereqName;
+		Settings.PrereqPath = PrereqPath;
+		Settings.PrereqArgs = PrereqArgs;
+		Settings.CustomFields = CustomFields;
+
+		// Run the build generation
+		if (FParse::Param(FCommandLine::Get(), TEXT("nochunks")))
+		{
+			bSuccess = BuildPatchServicesModule->GenerateFilesManifestFromDirectory(Settings);
+		}
+		else
+		{
+			bSuccess = BuildPatchServicesModule->GenerateChunksManifestFromDirectory(Settings);
+		}
+	}
+
+	// Release the module ptr
+	BuildPatchServicesModule.Reset();
+
+	// Check for processing error
+	if (!bSuccess)
+	{
+		GLog->Log(ELogVerbosity::Error, TEXT("An fatal error occurred executing BuildPatchTool.exe"));
+		return 2;
+	}
+
+	GLog->Log(TEXT("BuildPatchToolMain completed successfuly"));
+	return 0;
+}
+
+INT32_MAIN_INT32_ARGC_TCHAR_ARGV()
+{
+	FString CommandLine;
+	for( int32 Option = 1; Option < ArgC; Option++ )
+	{
+		CommandLine += TEXT(" ");
+		FString Argument( ArgV[Option] );
+		if( Argument.Contains( TEXT(" ") ) )
+		{
+			if (Argument.Contains(TEXT("=")))
+			{
+				FString ArgName;
+				FString ArgValue;
+				Argument.Split( TEXT("="), &ArgName, &ArgValue );
+				Argument = FString::Printf( TEXT("%s=\"%s\""), *ArgName, *ArgValue );
+			}
+			else
+			{
+				Argument = FString::Printf(TEXT("\"%s\""), *Argument);
+			}
+		}
+		CommandLine += Argument;
+	}
+
+	return BuildPatchToolMain( *CommandLine );
+}
diff --git a/Engine/Source/Runtime/Online/BuildPatchServices/BuildPatchServices.Build.cs b/Engine/Source/Runtime/Online/BuildPatchServices/BuildPatchServices.Build.cs
new file mode 100644
--- /dev/null
+++ b/Engine/Source/Runtime/Online/BuildPatchServices/BuildPatchServices.Build.cs
@@ -0,0 +1,37 @@
+// Copyright 1998-2014 Epic Games, Inc. All Rights Reserved.
+
+using UnrealBuildTool;
+
+public class BuildPatchServices : ModuleRules
+{
+	public BuildPatchServices(TargetInfo Target)
+	{
+		PublicDependencyModuleNames.AddRange(
+			new string[]
+			{
+				"Core",
+			}
+		);
+
+		PrivateIncludePaths.AddRange(
+			new string[]
+			{
+				"Developer/NoRedist/BuildPatchServices/Private",
+			}
+		);
+
+		PrivateDependencyModuleNames.AddRange(
+		new string[] {
+				"HTTP",
+				"Analytics",
+				"AnalyticsET",
+			}
+		);
+
+		PrivateIncludePathModuleNames.AddRange(
+			new string[] {
+				"HTTP"
+			}
+		);
+	}
+}
diff --git a/Engine/Source/Runtime/Online/BuildPatchServices/Private/BuildPatchAnalytics.cpp b/Engine/Source/Runtime/Online/BuildPatchServices/Private/BuildPatchAnalytics.cpp
new file mode 100644
--- /dev/null
+++ b/Engine/Source/Runtime/Online/BuildPatchServices/Private/BuildPatchAnalytics.cpp
@@ -0,0 +1,90 @@
+// Copyright 1998-2014 Epic Games, Inc. All Rights Reserved.
+
+/*=============================================================================
+	BuildPatchAnalytics.cpp: Implements static helper class for sending analytics
+	events for the module.
+=============================================================================*/
+
+#define ERROR_EVENT_SEND_LIMIT 20
+
+#include "BuildPatchServicesPrivatePCH.h"
+
+TSharedPtr< IAnalyticsProvider > FBuildPatchAnalytics::Analytics;
+
+FThreadSafeCounter FBuildPatchAnalytics::DownloadErrors;
+
+FThreadSafeCounter FBuildPatchAnalytics::CacheErrors;
+
+FThreadSafeCounter FBuildPatchAnalytics::ConstructionErrors;
+
+void FBuildPatchAnalytics::SetAnalyticsProvider( TSharedPtr< IAnalyticsProvider > AnalyticsProvider )
+{
+	Analytics = AnalyticsProvider;
+}
+
+void FBuildPatchAnalytics::ResetCounters()
+{
+	DownloadErrors.Reset();
+	CacheErrors.Reset();
+	ConstructionErrors.Reset();
+}
+
+void FBuildPatchAnalytics::RecordChunkDownloadError( const FString& ChunkUrl, const int32& ResponseCode, const FString& ErrorString )
+{
+	if ( DownloadErrors.Increment() <= ERROR_EVENT_SEND_LIMIT )
+	{
+		TArray<FAnalyticsEventAttribute> Attributes;
+		Attributes.Add( FAnalyticsEventAttribute( TEXT( "ChunkURL" ), ChunkUrl ) );
+		Attributes.Add( FAnalyticsEventAttribute( TEXT( "ResponseCode" ), ResponseCode ) );
+		Attributes.Add( FAnalyticsEventAttribute( TEXT( "ErrorString" ), ErrorString ) );
+		FBuildPatchHTTP::QueueAnalyticsEvent( TEXT( "BUILDPATCHERROR.CHUNK.DOWNLOAD" ) , Attributes );
+	}
+}
+
+void FBuildPatchAnalytics::RecordChunkCacheError( const FGuid& ChunkGuid, const FString& Filename, const int32 LastError, const FString& SystemName, const FString& ErrorString )
+{
+	if ( CacheErrors.Increment() <= ERROR_EVENT_SEND_LIMIT )
+	{
+		TArray<FAnalyticsEventAttribute> Attributes;
+		Attributes.Add( FAnalyticsEventAttribute( TEXT( "ChunkGuid" ), ChunkGuid.ToString() ) );
+		Attributes.Add( FAnalyticsEventAttribute( TEXT( "Filename" ), Filename ) );
+		Attributes.Add( FAnalyticsEventAttribute( TEXT( "LastError" ), LastError ) );
+		Attributes.Add( FAnalyticsEventAttribute( TEXT( "SystemName" ), SystemName ) );
+		Attributes.Add( FAnalyticsEventAttribute( TEXT( "ErrorString" ), ErrorString ) );
+		FBuildPatchHTTP::QueueAnalyticsEvent( TEXT( "BUILDPATCHERROR.CHUNK.CACHE" ) , Attributes );
+	}
+}
+
+void FBuildPatchAnalytics::RecordConstructionError( const FString& Filename, const int32 LastError, const FString& ErrorString )
+{
+	if ( ConstructionErrors.Increment() <= ERROR_EVENT_SEND_LIMIT )
+	{
+		TArray<FAnalyticsEventAttribute> Attributes;
+		Attributes.Add( FAnalyticsEventAttribute( TEXT( "Filename" ), Filename ) );
+		Attributes.Add( FAnalyticsEventAttribute( TEXT( "LastError" ), LastError ) );
+		Attributes.Add( FAnalyticsEventAttribute( TEXT( "ErrorString" ), ErrorString ) );
+		FBuildPatchHTTP::QueueAnalyticsEvent( TEXT( "BUILDPATCHERROR.CONSTRUCTION" ) , Attributes );
+	}
+}
+
+void FBuildPatchAnalytics::RecordPrereqInstallnError( const FString& Filename, const FString& CommandLine, const int32 ErrorCode, const FString& ErrorString )
+{
+	TArray<FAnalyticsEventAttribute> Attributes;
+	Attributes.Add( FAnalyticsEventAttribute( TEXT( "Filename" ), Filename ) );
+	Attributes.Add( FAnalyticsEventAttribute( TEXT( "CommandLine" ), CommandLine ) );
+	Attributes.Add( FAnalyticsEventAttribute( TEXT( "ErrorCode" ), ErrorCode ) );
+	Attributes.Add( FAnalyticsEventAttribute( TEXT( "ErrorString" ), ErrorString ) );
+	FBuildPatchHTTP::QueueAnalyticsEvent( TEXT( "BUILDPATCHERROR.PREREQUISITES" ) , Attributes );
+
+}
+
+void FBuildPatchAnalytics::RecordEvent( const FString& EventName, const TArray<FAnalyticsEventAttribute>& Attributes )
+{
+	// Check for being only main thread
+	check( IsInGameThread() );
+	if ( Analytics.IsValid() )
+	{
+		IAnalyticsProvider& AnalyticsProvider = *Analytics.Get();
+		AnalyticsProvider.RecordEvent( EventName, Attributes );
+	}
+}
diff --git a/Engine/Source/Runtime/Online/BuildPatchServices/Private/BuildPatchAnalytics.h b/Engine/Source/Runtime/Online/BuildPatchServices/Private/BuildPatchAnalytics.h
new file mode 100644
--- /dev/null
+++ b/Engine/Source/Runtime/Online/BuildPatchServices/Private/BuildPatchAnalytics.h
@@ -0,0 +1,82 @@
+// Copyright 1998-2014 Epic Games, Inc. All Rights Reserved.
+
+/*=============================================================================
+	BuildPatchAnalytics.h: Declares static helper class for sending analytics events
+	for the module.
+	An IAnalyticsProvider should be provided by the code making use of this module.
+=============================================================================*/
+
+#pragma once
+
+/**
+ * A struct to hold static fatal error flagging
+ */
+struct FBuildPatchAnalytics
+{
+private:
+	// The analytics provider interface
+	static TSharedPtr< IAnalyticsProvider > Analytics;
+
+	// Avoid spamming download errors
+	static FThreadSafeCounter DownloadErrors;
+
+	// Avoid spamming cache errors
+	static FThreadSafeCounter CacheErrors;
+
+	// Avoid spamming construction errors
+	static FThreadSafeCounter ConstructionErrors;
+
+public:
+	/**
+	 * Set the analytics provider to be used when recording events
+	 * @param AnalyticsProvider		Shared ptr to an analytics interface to use. If NULL analytics will be disabled.
+	 */
+	static void SetAnalyticsProvider( TSharedPtr< IAnalyticsProvider > AnalyticsProvider );
+
+	/**
+	 * Records an analytics event, must be called from the main thread only
+	 * @param EventName		The name of the event being sent
+	 * @param Attributes	The attributes for the event
+	 */
+	static void RecordEvent( const FString& EventName, const TArray<FAnalyticsEventAttribute>& Attributes );
+
+	/**
+	 * Resets the event counters that monitor for spamming and limit the number of events that are sent
+	 */
+	static void ResetCounters();
+
+	/**
+	 * Record a chunk download error event
+	 * @param ChunkUrl		The url for the chunk being downloaded
+	 * @param ResponseCode	The http response code, should pass INDEX_NONE if no response or not http.
+	 * @param ErrorString	The point at which the download has failed
+	 */
+	static void RecordChunkDownloadError( const FString& ChunkUrl, const int32& ResponseCode, const FString& ErrorString );
+
+	/**
+	 * Record a chunk cache error event
+	 * @param ChunkGuid		The chunk guid that the error occurred for
+	 * @param Filename		The filename for the related event
+	 * @param LastError		The OS error code if appropriate
+	 * @param SystemName	The cache system that the error comes from
+	 * @param ErrorString	The type of error that has occurred
+	 */
+	static void RecordChunkCacheError( const FGuid& ChunkGuid, const FString& Filename, const int32 LastError, const FString& SystemName, const FString& ErrorString );
+
+	/**
+	 * Record a file construction error event
+	 * @param Filename		The filename for the related event
+	 * @param LastError		The OS error code if appropriate
+	 * @param ErrorString	The type of error that has occurred
+	 */
+	static void RecordConstructionError( const FString& Filename, const int32 LastError, const FString& ErrorString );
+
+	/**
+	 * Record a file construction error event
+	 * @param Filename		The installer file
+	 * @param CommandLine	The command line passed to the installer
+	 * @param ErrorCode		The OS error code if appropriate
+	 * @param ErrorString	The type of error that has occurred
+	 */
+	static void RecordPrereqInstallnError( const FString& Filename, const FString& CommandLine, const int32 ErrorCode, const FString& ErrorString );
+};
diff --git a/Engine/Source/Runtime/Online/BuildPatchServices/Private/BuildPatchChunk.cpp b/Engine/Source/Runtime/Online/BuildPatchServices/Private/BuildPatchChunk.cpp
new file mode 100644
--- /dev/null
+++ b/Engine/Source/Runtime/Online/BuildPatchServices/Private/BuildPatchChunk.cpp
@@ -0,0 +1,318 @@
+// Copyright 1998-2014 Epic Games, Inc. All Rights Reserved.
+
+/*=============================================================================
+	BuildPatchChunk.cpp: Implements classes involved with chunks for the build system.
+=============================================================================*/
+
+#include "BuildPatchServicesPrivatePCH.h"
+
+/* FChunkFile implementation
+*****************************************************************************/
+FChunkFile::FChunkFile( const uint32& InReferenceCount, const bool& bInIsFromDisk )
+	: bIsFromDisk( bInIsFromDisk )
+	, ReferenceCount(InReferenceCount)
+{
+	LastAccessTime = FPlatformTime::Seconds();
+}
+
+void FChunkFile::GetDataLock( uint8** OutChunkData, FChunkHeader** OutChunkHeader )
+{
+	ThreadLock.Lock();
+	if( OutChunkData )
+	{
+		(*OutChunkData) = ChunkData;
+	}
+	if( OutChunkHeader )
+	{
+		(*OutChunkHeader) = &ChunkHeader;
+	}
+}
+
+void FChunkFile::ReleaseDataLock()
+{
+	LastAccessTime = FPlatformTime::Seconds();
+	ThreadLock.Unlock();
+}
+
+const uint32 FChunkFile::GetRefCount() const
+{
+	return ReferenceCount;
+}
+
+uint32 FChunkFile::Dereference()
+{
+	FScopeLock ScopeLock( &ThreadLock );
+	check( ReferenceCount > 0 );
+	return --ReferenceCount;
+}
+
+const double FChunkFile::GetLastAccessTime() const
+{
+	return LastAccessTime;
+}
+
+#if WITH_BUILDPATCHGENERATION
+
+/* FQueuedChunkWriter implementation
+*****************************************************************************/
+FChunkWriter::FQueuedChunkWriter::FQueuedChunkWriter()
+{
+	bMoreChunks = true;
+}
+
+FChunkWriter::FQueuedChunkWriter::~FQueuedChunkWriter()
+{
+	for( auto QueuedIt = ChunkFileQueue.CreateConstIterator(); QueuedIt; ++QueuedIt)
+	{
+		FChunkFile* ChunkFile = *QueuedIt;
+		delete ChunkFile;
+	}
+	ChunkFileQueue.Empty();
+}
+
+bool FChunkWriter::FQueuedChunkWriter::Init()
+{
+	IFileManager::Get().MakeDirectory( *ChunkDirectory, true );
+	return IFileManager::Get().DirectoryExists( *ChunkDirectory );
+}
+
+uint32 FChunkWriter::FQueuedChunkWriter::Run()
+{
+	// Loop until there's no more chunks
+	while ( ShouldBeRunning() )
+	{
+		FChunkFile* ChunkFile = GetNextChunk();
+		if( ChunkFile != NULL )
+		{
+			const FGuid& ChunkGuid = ChunkFile->ChunkHeader.Guid;
+			const uint64& ChunkHash = ChunkFile->ChunkHeader.RollingHash;
+#if SAVE_OLD_CHUNKDATA_FILENAMES
+			const FString OldChunkFilename = FBuildPatchUtils::GetChunkOldFilename( ChunkDirectory, ChunkGuid );
+#endif
+			const FString NewChunkFilename = FBuildPatchUtils::GetChunkNewFilename( EBuildPatchAppManifestVersion::GetLatestVersion(), ChunkDirectory, ChunkGuid, ChunkHash );
+
+			// To be a bit safer, make a few attempts at writing chunks
+			int32 RetryCount = 5;
+			bool bChunkSaveSuccess = false;
+			while ( RetryCount > 0 )
+			{
+				// Write out chunks
+				bChunkSaveSuccess = WriteChunkData( NewChunkFilename, ChunkFile, ChunkGuid );
+#if SAVE_OLD_CHUNKDATA_FILENAMES
+				bChunkSaveSuccess = bChunkSaveSuccess && WriteChunkData( OldChunkFilename, ChunkFile, ChunkGuid );
+#endif
+				// Check success
+				if( bChunkSaveSuccess )
+				{
+					RetryCount = 0;
+				}
+				else
+				{
+					// Retry after a second if failed
+					--RetryCount;
+					FPlatformProcess::Sleep( 1.0f );
+				}
+			}
+
+			// If we really could not save out chunk data successfully, this build will never work, so panic flush logs and then cause a hard error.
+			if( !bChunkSaveSuccess )
+			{
+				GLog->PanicFlushThreadedLogs();
+				check( bChunkSaveSuccess );
+			}
+
+			// Delete the data memory
+			delete ChunkFile;
+		}
+		FPlatformProcess::Sleep( 0.0f );
+	}
+	return 0;
+}
+
+const bool FChunkWriter::FQueuedChunkWriter::WriteChunkData(const FString& ChunkFilename, FChunkFile* ChunkFile, const FGuid& ChunkGuid)
+{
+	// Chunks are saved with GUID, so if a file already exists it will never be different.
+	// Skip with return true if already exists
+	if( FPaths::FileExists( ChunkFilename ) )
+	{
+		const int64 ChunkFilesSize = IFileManager::Get().FileSize(*ChunkFilename);
+		ChunkFileSizesCS.Lock();
+		ChunkFileSizes.Add(ChunkGuid, ChunkFilesSize);
+		ChunkFileSizesCS.Unlock();
+		return true;
+	}
+	FArchive* FileOut = IFileManager::Get().CreateFileWriter( *ChunkFilename );
+	bool bSuccess = FileOut != NULL;
+	if( bSuccess )
+	{
+		// Setup to handle compression
+		bool bDataIsCompressed = true;
+		uint8* ChunkDataSource = ChunkFile->ChunkData;
+		int32 ChunkDataSourceSize = FBuildPatchData::ChunkDataSize;
+		TArray< uint8 > TempCompressedData;
+		TempCompressedData.Empty( ChunkDataSourceSize );
+		TempCompressedData.AddUninitialized( ChunkDataSourceSize );
+		int32 CompressedSize = ChunkDataSourceSize;
+
+		// Compressed can increase in size, but the function will return as failure in that case
+		// we can allow that to happen since we would not keep larger compressed data anyway.
+		bDataIsCompressed = FCompression::CompressMemory(
+			static_cast< ECompressionFlags >( COMPRESS_ZLIB | COMPRESS_BiasMemory ),
+			TempCompressedData.GetData(),
+			CompressedSize,
+			ChunkFile->ChunkData,
+			FBuildPatchData::ChunkDataSize );
+
+		// If compression succeeded, set data vars
+		if( bDataIsCompressed )
+		{
+			ChunkDataSource = TempCompressedData.GetData();
+			ChunkDataSourceSize = CompressedSize;
+		}
+
+		// Setup Header
+		FChunkHeader& Header = ChunkFile->ChunkHeader;
+		*FileOut << Header;
+		Header.HeaderSize = FileOut->Tell();
+		Header.StoredAs = bDataIsCompressed ? FChunkHeader::STORED_COMPRESSED : FChunkHeader::STORED_RAW;
+		Header.DataSize = ChunkDataSourceSize;
+		Header.HashType = FChunkHeader::HASH_ROLLING;
+
+		// Write out files
+		FileOut->Seek( 0 );
+		*FileOut << Header;
+		FileOut->Serialize( ChunkDataSource, ChunkDataSourceSize );
+		const int64 ChunkFilesSize = FileOut->TotalSize();
+		FileOut->Close();
+
+		ChunkFileSizesCS.Lock();
+		ChunkFileSizes.Add(ChunkGuid, ChunkFilesSize);
+		ChunkFileSizesCS.Unlock();
+
+		bSuccess = !FileOut->GetError();
+
+		delete FileOut;
+	}
+	// Log errors
+	if( !bSuccess )
+	{
+		GLog->Logf( TEXT( "BuildPatchServices: Error: Could not save out generated chunk file %s" ), *ChunkFilename );
+	}
+	return bSuccess;
+}
+
+const bool FChunkWriter::FQueuedChunkWriter::ShouldBeRunning()
+{
+	bool rtn;
+	MoreChunksCS.Lock();
+	rtn = bMoreChunks;
+	MoreChunksCS.Unlock();
+	rtn = rtn || HasQueuedChunk();
+	return rtn;
+}
+
+const bool FChunkWriter::FQueuedChunkWriter::HasQueuedChunk()
+{
+	bool rtn;
+	ChunkFileQueueCS.Lock();
+	rtn = ChunkFileQueue.Num() > 0;
+	ChunkFileQueueCS.Unlock();
+	return rtn;
+}
+
+const bool FChunkWriter::FQueuedChunkWriter::CanQueueChunk()
+{
+	bool rtn;
+	ChunkFileQueueCS.Lock();
+	rtn = ChunkFileQueue.Num() < FBuildPatchData::ChunkQueueSize;
+	ChunkFileQueueCS.Unlock();
+	return rtn;
+}
+
+FChunkFile* FChunkWriter::FQueuedChunkWriter::GetNextChunk()
+{
+	FChunkFile* rtn = NULL;
+	ChunkFileQueueCS.Lock();
+	if( ChunkFileQueue.Num() > 0 )
+	{
+		rtn = ChunkFileQueue.Last();
+		ChunkFileQueue.RemoveAt( ChunkFileQueue.Num() - 1 );
+	}
+	ChunkFileQueueCS.Unlock();
+	return rtn;
+}
+
+void FChunkWriter::FQueuedChunkWriter::QueueChunk( const uint8* ChunkData, const FGuid& ChunkGuid, const uint64& ChunkHash )
+{
+	// Create the FChunkFile and copy data
+	FChunkFile* NewChunk = new FChunkFile();
+	NewChunk->ChunkHeader.Guid = ChunkGuid;
+	NewChunk->ChunkHeader.RollingHash = ChunkHash;
+	FMemory::Memcpy( NewChunk->ChunkData, ChunkData, FBuildPatchData::ChunkDataSize );
+	// Wait until we can fit this chunk in the queue
+	while ( !CanQueueChunk() )
+	{
+		FPlatformProcess::Sleep( 0.0f );
+	}
+	// Queue the chunk
+	ChunkFileQueueCS.Lock();
+	ChunkFileQueue.Add( NewChunk );
+	ChunkFileQueueCS.Unlock();
+}
+
+void FChunkWriter::FQueuedChunkWriter::SetNoMoreChunks()
+{
+	MoreChunksCS.Lock();
+	bMoreChunks = false;
+	MoreChunksCS.Unlock();
+}
+
+void FChunkWriter::FQueuedChunkWriter::GetChunkFilesizes(TMap<FGuid, int64>& OutChunkFileSizes)
+{
+	FScopeLock ScopeLock(&ChunkFileSizesCS);
+	OutChunkFileSizes.Empty(ChunkFileSizes.Num());
+	OutChunkFileSizes.Append(ChunkFileSizes);
+}
+
+/* FChunkWriter implementation
+*****************************************************************************/
+FChunkWriter::FChunkWriter( const FString& ChunkDirectory )
+{
+	QueuedChunkWriter.ChunkDirectory = ChunkDirectory;
+	WriterThread = FRunnableThread::Create(&QueuedChunkWriter, TEXT("QueuedChunkWriterThread"));
+}
+
+FChunkWriter::~FChunkWriter()
+{
+	if( WriterThread != NULL )
+	{
+		delete WriterThread;
+		WriterThread = NULL;
+	}
+}
+
+void FChunkWriter::QueueChunk( const uint8* ChunkData, const FGuid& ChunkGuid, const uint64& ChunkHash )
+{
+	QueuedChunkWriter.QueueChunk( ChunkData, ChunkGuid, ChunkHash);
+}
+
+void FChunkWriter::NoMoreChunks()
+{
+	QueuedChunkWriter.SetNoMoreChunks();
+}
+
+void FChunkWriter::WaitForThread()
+{
+	if( WriterThread != NULL )
+	{
+		WriterThread->WaitForCompletion();
+	}
+}
+
+void FChunkWriter::GetChunkFilesizes(TMap<FGuid, int64>& OutChunkFileSizes)
+{
+	QueuedChunkWriter.GetChunkFilesizes(OutChunkFileSizes);
+}
+
+#endif // WITH_BUILDPATCHGENERATION
+
diff --git a/Engine/Source/Runtime/Online/BuildPatchServices/Private/BuildPatchChunk.h b/Engine/Source/Runtime/Online/BuildPatchServices/Private/BuildPatchChunk.h
new file mode 100644
--- /dev/null
+++ b/Engine/Source/Runtime/Online/BuildPatchServices/Private/BuildPatchChunk.h
@@ -0,0 +1,404 @@
+// Copyright 1998-2014 Epic Games, Inc. All Rights Reserved.
+
+/*=============================================================================
+	BuildPatchChunk.h: Declares classes involved with chunks for the build system.
+=============================================================================*/
+
+#pragma once
+
+/**
+ * Constant values and typedefs
+ */
+namespace FBuildPatchData
+{
+	enum Value
+	{
+		// Sizes
+		ChunkDataSize		= 1024*1024,	// We are using 1MB chunks for patching
+		ChunkQueueSize		= 5,			// The number of chunks to allow to be queued up when saving out
+	};
+
+	enum Type
+	{
+		// Data Types
+		ChunkData			= 0,
+		FileData			= 1,
+	};
+}
+
+// The chunk header magic codeword, for quick checking that the opened file is a chunk file.
+#define CHUNK_HEADER_MAGIC		0xB1FE3AA2
+
+// The chunk header version number.
+#define CHUNK_HEADER_VERSION	2
+
+/**
+ * Declares a struct to store the info for a chunk header
+ */
+struct FChunkHeader
+{
+	enum 
+	{
+		// Storage flags, can be raw or a combination of others.
+		STORED_RAW			=	0x0, // Zero means raw data
+		STORED_COMPRESSED	=	0x1, // Flag for compressed
+		STORED_ENCRYPTED	=	0x2, // Flag for encrypted
+
+		// Hash types
+		HASH_ROLLING		=	0x1, // Uses FRollingHash class
+		HASH_SHA1			=	0x2, // Uses FSHA class
+	};
+	// Magic value used as a check for correct file
+	uint32 Magic;
+	// Version of this chunk file/header combo
+	uint32 Version;
+	// The size of this header
+	uint32 HeaderSize;
+	// The size of this data
+	uint32 DataSize;
+	// The GUID for this data
+	FGuid Guid;
+	// What type of hash we are using
+	uint8 HashType;
+	// The ChunkHash for this chunk
+	uint64 RollingHash;
+	// The FileHash for this chunk
+	FSHAHash SHAHash;
+	// How the chunk data is stored
+	uint8 StoredAs;
+
+	/**
+	 * Default Constructor sets the magic and version ready for writing out
+	 */
+	FChunkHeader()
+		: Magic( CHUNK_HEADER_MAGIC )
+		, Version( CHUNK_HEADER_VERSION )
+		, HashType( HASH_ROLLING )
+	{}
+
+	/**
+	 * Checks if the Header Magic was set to the correct value
+	 * @return	true if the header magic value is correct
+	 */
+	const bool IsValidMagic() const { return Magic == CHUNK_HEADER_MAGIC; }
+
+	/**
+	 * Serialization operator.
+	 * @param	Ar			Archive to serialize to
+	 * @param	Header		Header to serialize
+	 * @return	Passed in archive
+	 */
+	friend FArchive& operator << ( FArchive& Ar, FChunkHeader& Header )
+	{
+		// The constant sizes for each version of a header struct. Must be updated
+		// If new member variables are added the version MUST be bumped and handled properly here,
+		// and these values must never change.
+		static const uint32 Version1Size = 41;
+		static const uint32 Version2Size = 62;
+		// Calculate how much space left in the archive for reading data ( will be 0 when writing )
+		const int64 ArchiveSizeLeft = Ar.TotalSize() - Ar.Tell();
+		// Make sure the archive has enough data to read from, or we are saving instead.
+		bool bSuccess = Ar.IsSaving() || ( ArchiveSizeLeft >= Version1Size );
+		if( bSuccess )
+		{
+			Ar	<< Header.Magic
+				<< Header.Version
+				<< Header.HeaderSize
+				<< Header.DataSize
+				<< Header.Guid
+				<< Header.RollingHash
+				<< Header.StoredAs;
+
+			// From version 2, we have a hash type choice. Previous versions default as only rolling
+			if( Header.Version >= 2 )
+			{
+				bSuccess = Ar.IsSaving() || ( ArchiveSizeLeft >= Version2Size );
+				if( bSuccess )
+				{
+					Ar << Header.SHAHash;
+					Ar << Header.HashType;
+				}
+			}
+		}
+
+		// If we had a size error, zero out the header values
+		if( !bSuccess )
+		{
+			Header.Magic = 0;
+			Header.Version = 0;
+			Header.HeaderSize = 0;
+			Header.DataSize = 0;
+			Header.Guid.Invalidate();
+			Header.RollingHash = 0;
+			Header.StoredAs = 0;
+		}
+
+		return Ar;
+	}
+};
+
+/**
+ * Declares a struct to store the info for a chunk part
+ */
+struct FChunkPart
+{
+	// The GUID of the chunk containing this part
+	FGuid  Guid;
+	// The offset of the first byte into the chunk
+	uint32 Offset;
+	// The size of this part
+	uint32 Size;
+
+	/**
+	 * Default Constructor
+	 */
+	FChunkPart()
+		: Guid()
+		, Offset( 0 )
+		, Size( 0 )
+	{}
+};
+
+/**
+ * Declares a struct to store the info about a piece of a chunk that is inside a file
+ */
+struct FFileChunkPart
+{
+	// The file containing this piece
+	FString Filename;
+	// The offset into the file of this piece
+	uint64 FileOffset;
+	// The FChunkPart that can be salvaged from this file
+	FChunkPart ChunkPart;
+
+	/**
+	 * Default Constructor
+	 */
+	FFileChunkPart()
+		: Filename()
+		, FileOffset( 0 )
+		, ChunkPart()
+	{}
+};
+
+/**
+ * Declares a struct to hold a chunk header with data
+ */
+class FChunkFile
+{
+	friend class FChunkWriter;
+private:
+	// The header for the chunk
+	FChunkHeader ChunkHeader;
+
+	// The data for the chunk
+	uint8 ChunkData[ FBuildPatchData::ChunkDataSize ];
+
+	// The last time the chunk was accessed - might not need
+	volatile double LastAccessTime;
+
+	// Whether the file came from disk.
+	// If it came from memory, it will need to be saved to disk if it is being forced out
+	bool bIsFromDisk;
+
+	// The number of times tis chunk is still referenced by the installing build
+	volatile uint32 ReferenceCount;
+
+	// The thread safety data lock
+	FCriticalSection ThreadLock;
+
+public:
+	// Must be allocated with ref count to start with
+	FChunkFile( const uint32& InReferenceCount, const bool& bInIsFromDisk );
+
+	/**
+	 * Gets the thread lock on the data, must call ReleaseDataLock when finished with data
+	 * @param	OutChunkData	Receives the pointer to chunk data
+	 * @param	OutChunkHeader	Receives the pointer to header
+	 */
+	void GetDataLock( uint8** OutChunkData, FChunkHeader** OutChunkHeader );
+
+	/**
+	 * Releases access to the data to allow other threads to use
+	 */
+	void ReleaseDataLock();
+
+	/**
+	 * Gets the current reference count for this chunk
+	 * @return	The ref count.
+	 */
+	const uint32 GetRefCount() const;
+
+	/**
+	 * Decrements the reference count for this chunk
+	 * @return	The new ref count.
+	 */
+	uint32 Dereference();
+
+	/**
+	 * Gets the last access time for this chunk
+	 * @return	The last access time that was set from FPlatformTime::Seconds().
+	 */
+	const double GetLastAccessTime() const;
+
+private:
+	// Hide default constructors as they must not be used
+	FChunkFile(){}
+	FChunkFile(const FChunkFile&){}
+};
+
+#if WITH_BUILDPATCHGENERATION
+
+/**
+ * Declares a delegate for the chunk writer that gets called whenever a chunk has completed saving out
+ * @param Param1	The GUID of the completed chunk
+ * @param Param2	The full path to the chunk file
+ */
+DECLARE_MULTICAST_DELEGATE_TwoParams( FOnChunkComplete, const FGuid&, const FString& );
+
+/**
+ * Declares threaded chunk writer class for queuing up chunk file saving
+ */
+class FChunkWriter
+{
+private:
+	// A private runnable class that writes out data
+	class FQueuedChunkWriter : public FRunnable
+	{
+	public:
+		// The directory that we want to save chunks to
+		FString ChunkDirectory;
+
+		// Default Constructor
+		FQueuedChunkWriter();
+
+		// Default Destructor
+		~FQueuedChunkWriter();
+
+		// Start FRunnable interface
+		virtual bool Init() override;
+		virtual uint32 Run() override;
+		// End FRunnable interface
+
+		/**
+		 * Add a complete chunk to the queue.
+		 * BLOCKS RETURN until space in the queue
+		 * @param	ChunkData		The pointer to the chunk data. Should be of FBuildPatchData::ChunkDataSize length
+		 * @param	ChunkGuid		The GUID of this chunk
+		 * @param	ChunkHash		The hash value of this chunk
+		 */
+		void QueueChunk( const uint8* ChunkData, const FGuid& ChunkGuid, const uint64& ChunkHash );
+
+		/**
+		 * Call to flag as no more chunks to come (the thread will finish up and then exit).
+		 */
+		void SetNoMoreChunks();
+
+		/**
+		 * Retrieves the list of chunk files sizes
+		 * @param	OutChunkFileSizes		The Map is emptied and then filled with the list
+		 */
+		void GetChunkFilesizes(TMap<FGuid, int64>& OutChunkFileSizes);
+
+	private:
+		// The queue of chunks to save
+		TArray< FChunkFile* > ChunkFileQueue;
+
+		// A critical section for accessing ChunkFileQueue.
+		FCriticalSection ChunkFileQueueCS;
+
+		// Store whether more chunks are coming
+		bool bMoreChunks;
+
+		// A critical section for accessing ChunkFileQueue.
+		FCriticalSection MoreChunksCS;
+
+		// Store chunk file sizes.
+		TMap<FGuid, int64> ChunkFileSizes;
+
+		// A critical section for accessing ChunkFileSizes.
+		FCriticalSection ChunkFileSizesCS;
+
+		/**
+		 * Called from within run to save out a chunk file.
+		 * @param	ChunkFilename	The chunk filename
+		 * @param	ChunkFile		File data to be written out. Containing header will be edited.
+		 * @param	ChunkGuid		The GUID of this chunk.
+		 * @return	true if no file error
+		 */
+		const bool WriteChunkData(const FString& ChunkFilename, FChunkFile* ChunkFile, const FGuid& ChunkGuid);
+
+		/**
+		 * Thread safe. Checks to see if there are any chunks in the queue, or if there are more chunks expected.
+		 * @return	true if we should keep looping
+		 */
+		const bool ShouldBeRunning();
+
+		/**
+		 * Thread safe. Checks to see if there are any chunks in the queue
+		 * @return	true if ChunkFileQueue is not empty
+		 */
+		const bool HasQueuedChunk();
+
+		/**
+		 * Thread safe. Checks to see if there is space to queue a new chunk
+		 * @return	true if a chunk can be queued
+		 */
+		const bool CanQueueChunk();
+
+		/**
+		 * Thread safe. Gets the next chunk from the chunk queue
+		 * @return	The next chunk file
+		 */
+		FChunkFile* GetNextChunk();
+
+	} QueuedChunkWriter;
+
+public:
+	/**
+	 * Constructor
+	 */
+	FChunkWriter( const FString& ChunkDirectory );
+
+	/**
+	 * Default destructor
+	 */
+	~FChunkWriter();
+
+	/**
+	 * Add a complete chunk to the queue. 
+	 * BLOCKS RETURN until space in the queue
+	 * @param	ChunkData		The pointer to the chunk data. Should be of FBuildPatchData::ChunkDataSize length
+	 * @param	ChunkGuid		The GUID of this chunk
+	 * @param	ChunkHash		The hash value of this chunk
+	 */
+	void QueueChunk( const uint8* ChunkData, const FGuid& ChunkGuid, const uint64& ChunkHash );
+
+	/**
+	 * Call when there are no more chunks. The thread will be stopped and data cleaned up.
+	 */
+	void NoMoreChunks();
+
+	/**
+	 * Will only return when the writer thread has finished.
+	 */
+	void WaitForThread();
+
+	/**
+	 * Retrieves the list of chunk files sizes from the chunk writer
+	 * @param	OutChunkFileSizes		The Map is emptied and then filled with the list
+	 */
+	void GetChunkFilesizes(TMap<FGuid, int64>& OutChunkFileSizes);
+
+private:
+	/**
+	 * Hide the default constructor so that the params must be given
+	 */
+	FChunkWriter(){};
+
+	// Holds a pointer to our thread for the runnable
+	FRunnableThread* WriterThread;
+};
+
+#endif // WITH_BUILDPATCHGENERATION
+
diff --git a/Engine/Source/Runtime/Online/BuildPatchServices/Private/BuildPatchChunkCache.cpp b/Engine/Source/Runtime/Online/BuildPatchServices/Private/BuildPatchChunkCache.cpp
new file mode 100644
--- /dev/null
+++ b/Engine/Source/Runtime/Online/BuildPatchServices/Private/BuildPatchChunkCache.cpp
@@ -0,0 +1,904 @@
+// Copyright 1998-2014 Epic Games, Inc. All Rights Reserved.
+
+/*=============================================================================
+	BuildPatchChunkCache.cpp: Implements classes involved with chunks for the build system.
+=============================================================================*/
+
+#include "BuildPatchServicesPrivatePCH.h"
+
+// How many chunks we should be able to keep in RAM while required
+#define CHUNK_CACHE_SIZE		256
+
+// How many chunks we forward acquire each time we need to load one from disk
+// This value is how many upcoming chunks to scan through, and load from disk whichever ones
+// are on disk. It essentially stands for the minimum number of chunk parts that will be written
+// to the installation in between any chunks that source from HDD
+#define FORWARD_READ_COUNT		40
+
+/* FThreadSafeChunkCache implementation
+*****************************************************************************/
+bool FBuildPatchChunkCache::FThreadSafeChunkCache::Contains( const FGuid& ChunkGuid )
+{
+	FScopeLock ScopeLock( &ThreadLock );
+	return ChunkStore.Contains( ChunkGuid );
+}
+
+int32 FBuildPatchChunkCache::FThreadSafeChunkCache::Num()
+{
+	FScopeLock ScopeLock( &ThreadLock );
+	return ChunkStore.Num() + ReservedChunks.Num();
+}
+
+void FBuildPatchChunkCache::FThreadSafeChunkCache::Add( const FGuid& ChunkGuid, FChunkFile* ChunkFile )
+{
+	FScopeLock ScopeLock( &ThreadLock );
+	// Must not be NULL
+	check( ChunkFile != NULL );
+	// Must not exist already
+	check( Contains( ChunkGuid ) == false );
+	// We only accept chunks with reservations
+	check( HasReservation( ChunkGuid ) );
+	// Remove reservation and add, returning new entry in order to receive data
+	RemoveReservation( ChunkGuid );
+	// Add the entry
+	ChunkStore.Add( ChunkGuid, ChunkFile );
+}
+
+FChunkFile* FBuildPatchChunkCache::FThreadSafeChunkCache::Get( const FGuid& ChunkGuid )
+{
+	FScopeLock ScopeLock( &ThreadLock );
+	return ChunkStore[ ChunkGuid ];
+}
+
+void FBuildPatchChunkCache::FThreadSafeChunkCache::Remove( const FGuid& ChunkGuid )
+{
+	FScopeLock ScopeLock( &ThreadLock );
+	RemoveReservation( ChunkGuid );
+	if( ChunkStore.Contains( ChunkGuid ) )
+	{
+		// Get and release the data lock in case the data is currently being read
+		ChunkStore[ ChunkGuid ]->GetDataLock( NULL, NULL );
+		ChunkStore[ ChunkGuid ]->ReleaseDataLock();
+		delete ChunkStore[ ChunkGuid ];
+		ChunkStore.Remove( ChunkGuid );
+	}
+}
+
+void FBuildPatchChunkCache::FThreadSafeChunkCache::PurgeUnreferenced()
+{
+	FScopeLock ScopeLock( &ThreadLock );
+	// Get a copy of our key list
+	TArray< FGuid > CachedChunkGuids;
+	// Copy the keys so that we can iterate without invalidating
+	ChunkStore.GetKeys( CachedChunkGuids );
+	// Check each entry for 0 ref count
+	for ( auto CachedChunkGuidsIt = CachedChunkGuids.CreateConstIterator(); CachedChunkGuidsIt; ++CachedChunkGuidsIt )
+	{
+		const FGuid& ChunkGuid = *CachedChunkGuidsIt;
+		FChunkFile& ChunkFile = *ChunkStore[ ChunkGuid ];
+		if( ChunkFile.GetRefCount() == 0 )
+		{
+			Remove( ChunkGuid );
+		}
+	}
+}
+
+void FBuildPatchChunkCache::FThreadSafeChunkCache::Empty()
+{
+	FScopeLock ScopeLock( &ThreadLock );
+	for ( auto ChunkStoreIt = ChunkStore.CreateConstIterator(); ChunkStoreIt; ++ChunkStoreIt )
+	{
+		delete ChunkStoreIt->Value;
+	}
+	ChunkStore.Empty();
+}
+
+bool FBuildPatchChunkCache::FThreadSafeChunkCache::TryAddReservation( const FGuid& ChunkGuid )
+{
+	// Important to keep lock throughout this call
+	FScopeLock ScopeLock( &ThreadLock );
+
+	// Application error if we already have this chunk
+	check( ChunkStore.Contains( ChunkGuid ) == false );
+
+	// Pass if already reserved
+	if( ReservedChunks.Contains( ChunkGuid ) )
+	{
+		return true;
+	}
+
+	// Fail if we cannot fit yet
+	if( Num() >= CHUNK_CACHE_SIZE )
+	{
+		return false;
+	}
+
+	// Make the reservation
+	ReservedChunks.AddUnique( ChunkGuid );
+	return true;
+}
+
+void FBuildPatchChunkCache::FThreadSafeChunkCache::RemoveReservation( const FGuid& ChunkGuid )
+{
+	FScopeLock ScopeLock( &ThreadLock );
+	ReservedChunks.Remove( ChunkGuid );
+}
+
+bool FBuildPatchChunkCache::FThreadSafeChunkCache::HasReservation( const FGuid& ChunkGuid )
+{
+	FScopeLock ScopeLock( &ThreadLock );
+	return ReservedChunks.Contains( ChunkGuid );
+}
+
+/* FBuildPatchChunkCache implementation
+*****************************************************************************/
+FBuildPatchChunkCache::FBuildPatchChunkCache( const FBuildPatchAppManifestRef& InInstallManifet, const FBuildPatchAppManifestPtr& InCurrentManifest, const FString& InChunkCacheStage, const FString& InCurrentInstallDir, FBuildPatchProgress* InBuildProgress, TArray<FString>& FilesToConstruct, FBuildPatchInstallationInfo& InstallationInfoRef )
+	: ChunkCacheStage(InChunkCacheStage)
+	, CurrentInstallDir( InCurrentInstallDir )
+	, InstallManifet( InInstallManifet )
+	, CurrentManifest( InCurrentManifest )
+	, BuildProgress( InBuildProgress )
+	, bEnableChunkRecycle( InCurrentManifest.IsValid() )
+	, bDownloadsStarted( 0 )
+	, SkippedChunkDownloadSize( 0 )
+	, InstallationInfo( InstallationInfoRef )
+{
+	// Setup chunk information
+	InstallManifet->GetChunksRequiredForFiles( FilesToConstruct, FullChunkRefList, false );
+	for( int32 FullChunkRefListIdx = FullChunkRefList.Num() - 1; FullChunkRefListIdx >= 0; --FullChunkRefListIdx )
+	{
+		const FGuid& ChunkGuid = FullChunkRefList[ FullChunkRefListIdx ];
+		ChunkOrigins.Add( ChunkGuid, EChunkOrigin::Download );
+		if( ChunkUseOrderStack.Num() == 0 || ChunkUseOrderStack.Last() != ChunkGuid )
+		{
+			ChunkUseOrderStack.Push( ChunkGuid );
+		}
+	}
+	// Get unique list of downloads
+	InstallManifet->GetChunksRequiredForFiles( FilesToConstruct, ChunksToDownload, true );
+	// Check for chunks that can be constructed locally
+	TArray< FGuid > ChunksToConstruct;
+	InstallationInfo.EnumerateProducibleChunks(FullChunkRefList, ChunksToConstruct);
+	for( auto ChunksToConstructIt = ChunksToConstruct.CreateConstIterator(); ChunksToConstructIt; ++ChunksToConstructIt )
+	{
+		const FGuid& ChunkGuid = *ChunksToConstructIt;
+		ChunkOrigins.Add( ChunkGuid, EChunkOrigin::Recycle );
+		ChunksToDownload.Remove( ChunkGuid );
+	}
+
+	// Init chunk stats
+	NumFilesToConstruct = FilesToConstruct.Num();
+	NumChunksToDownload = ChunksToDownload.Num();
+	NumChunksToRecycle = ChunksToConstruct.Num();
+	NumRequiredChunks = NumChunksToDownload + NumChunksToRecycle;
+	TotalChunkDownloadSize = InstallManifet->GetDataSize(ChunksToDownload);
+	NumChunksRecycled.Reset();
+	NumChunksCacheBooted.Reset();
+	NumDriveCacheChunkLoads.Reset();
+	NumRecycleFailures.Reset();
+	NumDriveCacheLoadFailures.Reset();
+}
+
+/* FBuildPatchChunkCache implementation
+*****************************************************************************/
+const uint32& FBuildPatchChunkCache::GetStatNumFilesToConstruct() const
+{
+	return NumFilesToConstruct;
+}
+
+const uint32& FBuildPatchChunkCache::GetStatNumChunksToDownload() const
+{
+	return NumChunksToDownload;
+}
+
+const uint32& FBuildPatchChunkCache::GetStatNumChunksToRecycle() const
+{
+	return NumChunksToRecycle;
+}
+
+const uint32& FBuildPatchChunkCache::GetStatNumRequiredChunks() const
+{
+	return NumRequiredChunks;
+}
+
+const int64& FBuildPatchChunkCache::GetStatTotalChunkDownloadSize() const
+{
+	return TotalChunkDownloadSize;
+}
+
+const int32 FBuildPatchChunkCache::GetCounterChunksRecycled() const
+{
+	return NumChunksRecycled.GetValue();
+}
+
+const int32 FBuildPatchChunkCache::GetCounterChunksCacheBooted() const
+{
+	return NumChunksCacheBooted.GetValue();
+}
+
+const int32 FBuildPatchChunkCache::GetCounterDriveCacheChunkLoads() const
+{
+	return NumDriveCacheChunkLoads.GetValue();
+}
+
+const int32 FBuildPatchChunkCache::GetCounterRecycleFailures() const
+{
+	return NumRecycleFailures.GetValue();
+}
+
+const int32 FBuildPatchChunkCache::GetCounterDriveCacheLoadFailures() const
+{
+	return NumDriveCacheLoadFailures.GetValue();
+}
+
+FChunkFile* FBuildPatchChunkCache::GetChunkFile( const FGuid& ChunkGuid )
+{
+	// Data setup check
+	check( ChunkOrigins.Contains( ChunkGuid ) );
+
+	// If we don't already have this chunk yet, then we are ready to grab the next batch.
+	// We do it here because we want any loading from disk to be performed on the same thread as writing to disk
+	FChunkFile* ChunkFilePtr = NULL;
+	if( ChunkCache.Contains( ChunkGuid ) == false )
+	{
+		// We generate a list of the next chunks to acquire
+		TArray< FGuid > NextChunkBatch;
+		NextChunkBatch.AddUnique( ChunkGuid );
+		ChunkInfoLock.Lock();
+		int32 ForwardReadIdx = 1;
+		const int32 ChunkUseOrderStackNum = ChunkUseOrderStack.Num();
+		while ( ChunkUseOrderStackNum >= ForwardReadIdx && ( NextChunkBatch.Num() < FORWARD_READ_COUNT ) && !FBuildPatchInstallError::HasFatalError() )
+		{
+			const FGuid NextChunkGuid = ChunkUseOrderStack[ ChunkUseOrderStackNum - ForwardReadIdx ];
+			++ForwardReadIdx;
+			if( ChunkCache.Contains( NextChunkGuid ) == false )
+			{
+				NextChunkBatch.AddUnique( NextChunkGuid );
+				ReserveChunkInventorySlotForce( NextChunkGuid );
+			}
+		}
+		ChunkInfoLock.Unlock();
+
+		// Load each chunk that comes from disk
+		for( auto NextChunkBatchIt = NextChunkBatch.CreateConstIterator(); NextChunkBatchIt && !FBuildPatchInstallError::HasFatalError(); ++NextChunkBatchIt )
+		{
+			FScopeLock ScopeLock( &ChunkOriginLock );
+			const FGuid& NextChunkGuid = *NextChunkBatchIt;
+			const EChunkOrigin::Type& ChunkOrigin = ChunkOrigins[ NextChunkGuid ];
+			BuildProgress->WaitWhilePaused();
+			// Check if there were any errors while paused, like cancelling
+			if( FBuildPatchInstallError::HasFatalError() )
+			{
+				return NULL;
+			}
+			switch ( ChunkOrigin )
+			{
+				case EChunkOrigin::Harddisk:
+					{
+						// Force cache space for this chunk
+						ReserveChunkInventorySlotForce( NextChunkGuid );
+						// Read from disk
+						bool bReadFromDisk = ReadChunkFromDriveCache( NextChunkGuid );
+						if( bReadFromDisk == false && !FBuildPatchInstallError::HasFatalError() )
+						{
+							// Queue for download instead
+							ChunkOrigins[ NextChunkGuid ] = EChunkOrigin::Download;
+							FBuildPatchDownloader::Get().AddChunkToDownload( NextChunkGuid, true );
+							// Count fail
+							NumDriveCacheLoadFailures.Increment();
+						}
+					}
+					break;
+				case EChunkOrigin::Recycle:
+					{
+						// Force cache space for this chunk
+						ReserveChunkInventorySlotForce( NextChunkGuid );
+						// Recycle from build
+						bool bRecycled = RecycleChunkFromBuild( NextChunkGuid );
+						if( bRecycled == false && !FBuildPatchInstallError::HasFatalError() )
+						{
+							// Queue for download instead
+							ChunkOrigins[ NextChunkGuid ] = EChunkOrigin::Download;
+							FBuildPatchDownloader::Get().AddChunkToDownload( NextChunkGuid, true );
+							// Count fail
+							NumRecycleFailures.Increment();
+						}
+					}
+					break;
+				case EChunkOrigin::Download:
+					// We skip over downloads that will be running on another thread
+					continue;
+				default:
+					check( false );
+					break;
+			}
+		}
+
+		// Make sure the currently required chunk is available
+		WaitForChunk( ChunkGuid );
+	}
+
+	// Check if there were any errors
+	if( FBuildPatchInstallError::HasFatalError() )
+	{
+		return NULL;
+	}
+
+	// Get the chunk file, should assert if was removed
+	ChunkFilePtr = ChunkCache.Get( ChunkGuid );
+
+	// Check and update chunk information
+	{
+		FScopeLock ScopeLock( &ChunkInfoLock );
+		// The chunk we just handed out must be the first in FullChunkRefList list, otherwise we are constructing out of order.
+		check( FullChunkRefList.Num() > 0 && FullChunkRefList[ 0 ] == ChunkGuid );
+		// Must also be the chunk on top of ChunkUseOrderStack
+		check( ChunkUseOrderStack.Num() > 0 && ChunkUseOrderStack.Last() == ChunkGuid );
+		// Remove this chunk from the full reference list
+		FullChunkRefList.RemoveAt( 0 ); 
+		// If this was the last sequential reference to this chunk, also remove from order stack
+		if( FullChunkRefList.Num() > 0 && FullChunkRefList[ 0 ] != ChunkUseOrderStack.Last() )
+		{
+			ChunkUseOrderStack.Pop();
+			// The next chunk in both lists must be the same now
+			check( FullChunkRefList[ 0 ] == ChunkUseOrderStack.Last() );
+		}
+	}
+
+	// Return the pointer
+	return ChunkFilePtr;
+}
+
+bool FBuildPatchChunkCache::ReadChunkFromDriveCache( const FGuid& ChunkGuid )
+{
+	bool bSuccess = true;
+
+	// Get the chunk filename
+	const FString Filename = FBuildPatchUtils::GetChunkOldFilename( ChunkCacheStage, ChunkGuid );
+
+	// Read the chunk
+	FArchive* FileReader = IFileManager::Get().CreateFileReader( *Filename );
+	bSuccess = FileReader != NULL;
+	if( bSuccess )
+	{
+		// Get file size
+		const int64 FileSize = FileReader->TotalSize();
+
+		// Create the ChunkFile data structure
+		FChunkFile* NewChunkFile = new FChunkFile( GetRemainingReferenceCount( ChunkGuid ), true );
+
+		// Lock data
+		FChunkHeader* ChunkHeader;
+		uint8* ChunkData;
+		NewChunkFile->GetDataLock( &ChunkData, &ChunkHeader );
+
+		// Read the header
+		*FileReader << *ChunkHeader;
+
+		// Check header magic
+		bSuccess = ChunkHeader->IsValidMagic();
+		if ( bSuccess )
+		{
+			// Check a rolling hash has the right data size
+			bSuccess = ( ChunkHeader->HashType == FChunkHeader::HASH_ROLLING && ChunkHeader->DataSize == FBuildPatchData::ChunkDataSize );
+			if( bSuccess )
+			{
+				// Check Header and data size
+				bSuccess = ( ChunkHeader->HeaderSize + ChunkHeader->DataSize ) == FileSize;
+				if( bSuccess )
+				{
+					// Read the data
+					FileReader->Serialize( ChunkData, FBuildPatchData::ChunkDataSize );
+					// Verify the data hash
+					bSuccess = ChunkHeader->RollingHash == FRollingHash< FBuildPatchData::ChunkDataSize >::GetHashForDataSet( ChunkData );
+					if( !bSuccess )
+					{
+						FBuildPatchAnalytics::RecordChunkCacheError( ChunkGuid, Filename, INDEX_NONE, TEXT( "DriveCache" ), TEXT( "Hash Check Failed" ) );
+						GLog->Logf( TEXT( "FBuildPatchChunkCache: ERROR: ReadChunkFromDriveCache chunk failed hash check %s" ), *ChunkGuid.ToString() );
+					}
+					else
+					{
+						// Count loads
+						NumDriveCacheChunkLoads.Increment();
+						GLog->Logf( TEXT( "FBuildPatchChunkCache: ReadChunkFromDriveCache loaded chunk %s" ), *ChunkGuid.ToString() );
+					}
+				}
+				else
+				{
+					FBuildPatchAnalytics::RecordChunkCacheError( ChunkGuid, Filename, INDEX_NONE, TEXT( "DriveCache" ), TEXT( "Incorrect File Size" ) );
+					GLog->Logf( TEXT( "FBuildPatchChunkCache: ERROR: ReadChunkFromDriveCache header info does not match file size %s" ), *ChunkGuid.ToString() );
+				}
+			}
+			else
+			{
+				FBuildPatchAnalytics::RecordChunkCacheError( ChunkGuid, Filename, INDEX_NONE, TEXT( "DriveCache" ), TEXT( "Datasize/Hashtype Mismatch" ) );
+				GLog->Logf( TEXT( "FBuildPatchChunkCache: ERROR: ReadChunkFromDriveCache mismatch datasize/hashtype combination %s" ), *ChunkGuid.ToString() );
+			}
+		}
+		else
+		{
+			FBuildPatchAnalytics::RecordChunkCacheError( ChunkGuid, Filename, INDEX_NONE, TEXT( "DriveCache" ), TEXT( "Corrupt Header" ) );
+			GLog->Logf( TEXT( "FBuildPatchChunkCache: ERROR: ReadChunkFromDriveCache corrupt header %s" ), *ChunkGuid.ToString() );
+		}
+
+		// Release data
+		NewChunkFile->ReleaseDataLock();
+
+		// Add the newly filled data to the cache if successful
+		if( bSuccess )
+		{
+			ChunkCache.Add( ChunkGuid, NewChunkFile );
+		}
+		// If there was a problem, remove from cache and reservation
+		else
+		{
+			ChunkCache.Remove( ChunkGuid );
+		}
+
+		// Close the file
+		FileReader->Close();
+		delete FileReader;
+
+	}
+	else
+	{
+		FBuildPatchAnalytics::RecordChunkCacheError( ChunkGuid, Filename, FPlatformMisc::GetLastError(), TEXT( "DriveCache" ), TEXT( "Open File Fail" ) );
+		GLog->Logf( TEXT( "BuildPatchServices: ERROR: GetChunkData could not open chunk file %s" ), *ChunkGuid.ToString() );
+	}
+
+	return bSuccess;
+}
+
+bool FBuildPatchChunkCache::RecycleChunkFromBuild( const FGuid& ChunkGuid )
+{
+	// Must never double acquire
+	check( ChunkCache.Contains( ChunkGuid ) == false );
+
+	// Debug leaving any files open
+	bool bSuccess = true;
+
+	// Get the app manifest that this chunk can be sourced from
+	FBuildPatchAppManifestPtr ChunkSourceAppManifest = InstallationInfo.GetManifestContainingChunk(ChunkGuid);
+	if (!ChunkSourceAppManifest.IsValid())
+	{
+		return false;
+	}
+
+	// Get the install directory for this manifest
+	const FString ChunkSourceInstallDir = InstallationInfo.GetManifestInstallDir(ChunkSourceAppManifest);
+	if(ChunkSourceInstallDir.Len() <= 0)
+	{
+		return false;
+	}
+
+	// We need to generate an inventory of all chunk parts in this build that refer to the chunk that we require
+	TMap< FGuid, TArray< FFileChunkPart > > ChunkPartInventory;
+	TArray< FGuid > Array;
+	Array.Add( ChunkGuid );
+	ChunkSourceAppManifest->EnumerateChunkPartInventory(Array, ChunkPartInventory);
+
+	// Attempt construction of the chunk from the parts
+	FArchive* BuildFileIn = NULL;
+	FString  BuildFileOpened;
+	int64 BuildFileInSize = 0;
+
+	// We must have a hash for this chunk or else we cant verify it
+	uint64 ChunkHash = 0;
+	TArray< FFileChunkPart >* FileChunkPartsPtr = ChunkPartInventory.Find( ChunkGuid );
+	bSuccess = (FileChunkPartsPtr != NULL && ChunkSourceAppManifest->GetChunkHash(ChunkGuid, ChunkHash));
+	if( bSuccess )
+	{
+		const TArray< FFileChunkPart >& FileChunkParts = *FileChunkPartsPtr;
+		TArray< uint8 > TempArray;
+		TempArray.AddUninitialized( FBuildPatchData::ChunkDataSize );
+		uint8* TempChunkConstruction = TempArray.GetData();
+		FMemory::Memzero( TempChunkConstruction, FBuildPatchData::ChunkDataSize );
+		bSuccess = FileChunkParts.Num() > 0;
+		for( auto FileChunkPartIt = FileChunkParts.CreateConstIterator(); FileChunkPartIt && bSuccess && !FBuildPatchInstallError::HasFatalError(); ++FileChunkPartIt )
+		{
+			const FFileChunkPart& FileChunkPart = *FileChunkPartIt;
+			FString FullFilename = ChunkSourceInstallDir / FileChunkPart.Filename;
+			// Close current build file ?
+			if( BuildFileIn != NULL && BuildFileOpened != FullFilename )
+			{
+				BuildFileIn->Close();
+				delete BuildFileIn;
+				BuildFileIn = NULL;
+				BuildFileOpened = TEXT( "" );
+				BuildFileInSize = 0;
+			}
+			// Open build file ?
+			if( BuildFileIn == NULL )
+			{
+				BuildFileIn = IFileManager::Get().CreateFileReader( *FullFilename );
+				bSuccess = BuildFileIn != NULL;
+				if( !bSuccess )
+				{
+					BuildFileOpened = TEXT( "" );
+					FBuildPatchAnalytics::RecordChunkCacheError( ChunkGuid, FileChunkPart.Filename, FPlatformMisc::GetLastError(), TEXT( "ChunkRecycle" ), TEXT( "Source File Missing" ) );
+					GWarn->Logf( TEXT( "BuildPatchChunkConstruction: Warning: Failed to load source file for chunk. %s" ), *FullFilename );
+				}
+				else
+				{
+					BuildFileOpened = FullFilename;
+					BuildFileInSize = BuildFileIn->TotalSize();
+				}
+			}
+			// Grab the section of the chunk
+			if( BuildFileIn != NULL )
+			{
+				// Make sure we don't attempt to read off the end of the file
+				const int64 LastRequiredByte = FileChunkPart.FileOffset + FileChunkPart.ChunkPart.Size;
+				if( BuildFileInSize >= LastRequiredByte )
+				{
+					BuildFileIn->Seek( FileChunkPart.FileOffset );
+					BuildFileIn->Serialize( TempChunkConstruction + FileChunkPart.ChunkPart.Offset, FileChunkPart.ChunkPart.Size );
+				}
+				else
+				{
+					bSuccess = false;
+					FBuildPatchAnalytics::RecordChunkCacheError( ChunkGuid, FileChunkPart.Filename, INDEX_NONE, TEXT( "ChunkRecycle" ), TEXT( "Source File Too Small" ) );
+					GWarn->Logf( TEXT( "BuildPatchChunkConstruction: Warning: Source file too small for chunk position. %s" ), *FullFilename );
+				}
+			}
+		}
+
+		// Check no other fatal errors were registered in the meantime
+		bSuccess = bSuccess && !FBuildPatchInstallError::HasFatalError();
+
+		// Check chunk hash
+		if( bSuccess )
+		{
+			bSuccess = FRollingHash< FBuildPatchData::ChunkDataSize >::GetHashForDataSet( TempChunkConstruction ) == ChunkHash;
+			if( !bSuccess )
+			{
+				FBuildPatchAnalytics::RecordChunkCacheError( ChunkGuid, TEXT( "" ), INDEX_NONE, TEXT( "ChunkRecycle" ), TEXT( "Chunk Hash Fail" ) );
+				GWarn->Logf( TEXT( "BuildPatchChunkConstruction: Warning: Hash check failed for recycled chunk %s" ), *ChunkGuid.ToString() );
+			}
+		}
+
+		// Save the chunk to cache if all went well
+		if( bSuccess )
+		{
+			// It was added asynchronously!!
+			check( ChunkCache.Contains( ChunkGuid ) == false );
+
+			// Create the ChunkFile data structure
+			FChunkFile* NewChunkFile = new FChunkFile( GetRemainingReferenceCount( ChunkGuid ), true );
+
+			// Lock data
+			FChunkHeader* ChunkHeader;
+			uint8* ChunkData;
+			NewChunkFile->GetDataLock( &ChunkData, &ChunkHeader );
+
+			// Copy the data
+			FMemoryReader MemReader( TempArray );
+			MemReader.Serialize( ChunkData, FBuildPatchData::ChunkDataSize );
+
+			// Setup the header
+			ChunkHeader->Guid = ChunkGuid;
+			ChunkHeader->StoredAs = FChunkHeader::STORED_RAW;
+			ChunkHeader->DataSize = FBuildPatchData::ChunkDataSize; // This would change if compressing/encrypting
+			ChunkHeader->HashType = FChunkHeader::HASH_ROLLING;
+			ChunkHeader->RollingHash = ChunkHash;
+
+			// Release data
+			NewChunkFile->ReleaseDataLock();
+
+			// Count chunk
+			NumChunksRecycled.Increment();
+
+			// Add it to our cache.
+			ChunkCache.Add( ChunkGuid, NewChunkFile );
+		}
+
+		// Close any open file
+		if( BuildFileIn != NULL )
+		{
+			BuildFileIn->Close();
+			delete BuildFileIn;
+			BuildFileIn = NULL;
+		}
+	}
+	
+	return bSuccess;
+}
+
+void FBuildPatchChunkCache::AddDataToCache( const FGuid& ChunkGuid, const TArray< uint8 >& ChunkDataFile )
+{
+	// Must never double acquire
+	check( ChunkCache.Contains( ChunkGuid ) == false );
+
+	// Create the ChunkFile data structure
+	FChunkFile* NewChunkFile = new FChunkFile( GetRemainingReferenceCount( ChunkGuid ), false );
+
+	// Lock data
+	FChunkHeader* ChunkHeader;
+	uint8* ChunkData;
+	NewChunkFile->GetDataLock( &ChunkData, &ChunkHeader );
+
+	// Copy the data
+	FMemoryReader MemReader( ChunkDataFile );
+	MemReader << *ChunkHeader;
+	MemReader.Serialize( ChunkData, FBuildPatchData::ChunkDataSize );
+
+	// Release data
+	NewChunkFile->ReleaseDataLock();
+
+	// Add it to our cache.
+	ChunkCache.Add( ChunkGuid, NewChunkFile );
+}
+
+bool FBuildPatchChunkCache::ReserveChunkInventorySlot( const FGuid& ChunkGuid )
+{
+	// We shouldn't be trying to reserve for something already loaded
+	check( ChunkCache.Contains( ChunkGuid ) == false );
+
+	// If already reserved, return immediate
+	if( ChunkCache.HasReservation( ChunkGuid ) )
+	{
+		return true;
+	}
+
+	// We must be a required chunk if not already reserved
+	ChunkInfoLock.Lock();
+	check( ChunkUseOrderStack.Contains( ChunkGuid ) );
+	ChunkInfoLock.Unlock();
+
+	// We only allow chunks if they will be used soon enough
+	if( GetChunkOrderIndex( ChunkGuid ) >= CHUNK_CACHE_SIZE )
+	{
+		return false;
+	}
+
+	// Run a purge first
+	ChunkCache.PurgeUnreferenced();
+
+	// Try reserve
+	return ChunkCache.TryAddReservation( ChunkGuid );
+}
+
+void FBuildPatchChunkCache::SkipFile( const FString& Filename )
+{
+	check( bDownloadsStarted.GetValue() == 0 );
+	// We're going to skip over this file so we need to remove all the references to the chunks it uses from our
+	// data lists to avoid downloading them and keep track of the next chunks we need to get hold of
+	GLog->Logf( TEXT("FBuildPatchChunkCache::SkipFile %s"), *Filename );
+	// Get the file manifest
+	TSharedPtr< FBuildPatchFileManifest > FileManifest = InstallManifet->GetFileManifest( Filename );
+	// Go through each chunk part, and remove it from out data lists
+	FScopeLock ScopeLock( &ChunkInfoLock );
+	for( auto ChunkPartIt = FileManifest->FileChunkParts.CreateConstIterator(); ChunkPartIt; ++ChunkPartIt )
+	{
+		const FChunkPart& ChunkPart = *ChunkPartIt;
+		SkipChunkPart( ChunkPart );
+	}
+}
+
+void FBuildPatchChunkCache::SkipChunkPart( const FChunkPart& ChunkPart )
+{
+	check( bDownloadsStarted.GetValue() == 0 );
+	FScopeLock ScopeLock( &ChunkInfoLock );
+	const FGuid& ChunkGuid = ChunkPart.Guid;
+	// The chunk we are skipping must be the first in FullChunkRefList list, otherwise we are calling out of order.
+	check( FullChunkRefList.Num() > 0 && FullChunkRefList[ 0 ] == ChunkGuid );
+	// Must also be the chunk on top of ChunkUseOrderStack
+	check( ChunkUseOrderStack.Num() > 0 && ChunkUseOrderStack.Last() == ChunkGuid );
+	// Remove this chunk from the full reference list
+	FullChunkRefList.RemoveAt( 0 ); 
+	// If this was the last sequential reference to this chunk, also remove from order stack
+	if( FullChunkRefList.Num() > 0 && FullChunkRefList[ 0 ] != ChunkUseOrderStack.Last() )
+	{
+		ChunkUseOrderStack.Pop();
+		// The next chunk in both lists must be the same now
+		check( FullChunkRefList[ 0 ] == ChunkUseOrderStack.Last() );
+	}
+	// If this chunk is no longer referenced, remove from download list
+	if( ChunkUseOrderStack.Contains( ChunkGuid ) == false )
+	{
+		ChunksToDownload.Remove( ChunkGuid );
+		SkippedChunkDownloadSize += InstallManifet->GetDataSize( ChunkGuid );
+	}
+	// Keep UI in check
+	const float OriginalCount = TotalChunkDownloadSize;
+	const float SkipCount = SkippedChunkDownloadSize;
+	BuildProgress->SetStateProgress( EBuildPatchProgress::Downloading, SkipCount / OriginalCount );
+}
+
+void FBuildPatchChunkCache::BeginDownloads()
+{
+	check( bDownloadsStarted.GetValue() == 0 );
+	// Queue downloads
+	ChunkInfoLock.Lock();
+	FBuildPatchDownloader::Get().AddChunksToDownload( ChunksToDownload, FBuildPatchData::ChunkData );
+	ChunkInfoLock.Unlock();
+	bDownloadsStarted.Set( 1 );
+}
+
+const bool FBuildPatchChunkCache::HaveDownloadsStarted() const
+{
+	return bDownloadsStarted.GetValue() > 0;
+}
+
+int32 FBuildPatchChunkCache::GetChunkOrderIndex( const FGuid& ChunkGuid )
+{
+	FScopeLock ScopeLock( &ChunkInfoLock );
+	int32 LastIndex = 0;
+	if( ChunkUseOrderStack.FindLast( ChunkGuid, LastIndex ) )
+	{
+		return ChunkUseOrderStack.Num() - LastIndex;
+	}
+	// Apparently FindLast does not return 0 if the only instance is the first one!
+	else if( ChunkUseOrderStack.Num() > 0 && ChunkUseOrderStack[ 0 ] == ChunkGuid )
+	{
+		return ChunkUseOrderStack.Num();
+	}
+	// If we get here there's been a fault
+	check( false );
+	return CHUNK_CACHE_SIZE;
+}
+
+int32 FBuildPatchChunkCache::GetRemainingReferenceCount( const FGuid& ChunkGuid )
+{
+	FScopeLock ScopeLock( &ChunkInfoLock );
+	int32 RefCount = 0;
+	for( auto FullChunkRefListIt = FullChunkRefList.CreateConstIterator(); FullChunkRefListIt; ++FullChunkRefListIt )
+	{
+		if( *FullChunkRefListIt == ChunkGuid )
+		{
+			++RefCount;
+		}
+	}
+	return RefCount;
+}
+
+void FBuildPatchChunkCache::WaitForChunk( const FGuid& ChunkGuid )
+{
+	// Wait for the chunk to be available
+	double StartedWaiting = FPlatformTime::Seconds();
+	while( ChunkCache.Contains( ChunkGuid ) == false && !FBuildPatchInstallError::HasFatalError() )
+	{
+		FPlatformProcess::Sleep( 0.1f );
+		const double CurrentTime = FPlatformTime::Seconds();
+		if( BuildProgress->GetPauseState() )
+		{
+			BuildProgress->WaitWhilePaused();
+			StartedWaiting = CurrentTime;
+		}
+		if( ( CurrentTime - StartedWaiting ) > 5.0 )
+		{
+			GWarn->Logf( TEXT( "FBuildPatchChunkCache: Still waiting for chunk %s" ), *ChunkGuid.ToString() );
+			StartedWaiting = CurrentTime;
+		}
+	}
+}
+
+void FBuildPatchChunkCache::ReserveChunkInventorySlotForce( const FGuid& ChunkGuid )
+{
+	// If already reserved, return immediate
+	if( ChunkCache.HasReservation( ChunkGuid ) || ChunkCache.Contains( ChunkGuid ) )
+	{
+		return;
+	}
+
+	// Begin by checking if any slots can be freed
+	ChunkCache.PurgeUnreferenced();
+
+	// Try to add the reservation
+	bool bReservationAccepted = ChunkCache.TryAddReservation( ChunkGuid );
+
+	// If we couldn't reserve, we need to boot out a chunk for this required one
+	if( bReservationAccepted == false )
+	{
+		// We create a unique ref array from the use order so that chunks not needed
+		// for longer times end up nearer the bottom of the array
+		TArray< FGuid > ChunkPriorityList;
+		ChunkInfoLock.Lock();
+		for( int32 ChunkUseOrderStackIdx = ChunkUseOrderStack.Num() - 1; ChunkUseOrderStackIdx >= 0 ; --ChunkUseOrderStackIdx )
+		{
+			ChunkPriorityList.AddUnique( ChunkUseOrderStack[ ChunkUseOrderStackIdx ] );
+		}
+		ChunkInfoLock.Unlock();
+
+		// Starting at the bottom of the list, we look for a chunk that is contained in the cache
+		for( int32 ChunkPriorityListIdx = ChunkPriorityList.Num() - 1; ChunkPriorityListIdx >= 0 && !bReservationAccepted; --ChunkPriorityListIdx )
+		{
+			const FGuid& LowPriChunk = ChunkPriorityList[ ChunkPriorityListIdx ];
+			BuildProgress->WaitWhilePaused();
+			// Check if there were any errors while paused, like canceling
+			if( FBuildPatchInstallError::HasFatalError() )
+			{
+				return;
+			}
+			if( ChunkCache.Contains( LowPriChunk ) )
+			{
+				GWarn->Logf( TEXT( "FBuildPatchChunkCache: Booting chunk %s" ), *LowPriChunk.ToString() );
+
+				// Save chunk to disk so we don't have to download again
+				bool bSuccess = true;
+				const FString NewChunkFilename = FBuildPatchUtils::GetChunkOldFilename( ChunkCacheStage, LowPriChunk );
+				FChunkFile* LowPriChunkFile = ChunkCache.Get( LowPriChunk );
+				FChunkHeader* LowPriChunkHeader;
+				uint8* LowPriChunkData;
+				LowPriChunkFile->GetDataLock( &LowPriChunkData, &LowPriChunkHeader );
+				FArchive* FileOut = IFileManager::Get().CreateFileWriter( *NewChunkFilename );
+				bSuccess = FileOut != NULL;
+				const int32 LastError = FPlatformMisc::GetLastError();
+				if( bSuccess )
+				{
+					// Setup Header
+					*FileOut << *LowPriChunkHeader;
+					LowPriChunkHeader->HeaderSize = FileOut->Tell();
+					LowPriChunkHeader->StoredAs = FChunkHeader::STORED_RAW;
+					LowPriChunkHeader->DataSize = FBuildPatchData::ChunkDataSize; // This would change if compressing/encrypting
+					LowPriChunkHeader->HashType = FChunkHeader::HASH_ROLLING;
+
+					// Write out file
+					FileOut->Seek( 0 );
+					*FileOut << *LowPriChunkHeader;
+					FileOut->Serialize( LowPriChunkData, FBuildPatchData::ChunkDataSize );
+					FileOut->Close();
+
+					delete FileOut;
+				}
+				LowPriChunkFile->ReleaseDataLock();
+
+				// Setup new chunk origin
+				if( bSuccess )
+				{
+					ChunkOrigins[ LowPriChunk ] = EChunkOrigin::Harddisk;
+				}
+				else
+				{
+					// Queue download if save failed
+					ChunkOrigins[ LowPriChunk ] = EChunkOrigin::Download;
+					FBuildPatchDownloader::Get().AddChunkToDownload( LowPriChunk );
+					FBuildPatchAnalytics::RecordChunkCacheError( ChunkGuid, NewChunkFilename, LastError, TEXT( "ChunkBooting" ), TEXT( "Chunk Save Failed" ) );
+				}
+
+				// Boot this chunk
+				ChunkCache.Remove( LowPriChunk );
+
+				// Try get the reservation again!
+				bReservationAccepted = ChunkCache.TryAddReservation( ChunkGuid );
+
+				// Count the boot
+				NumChunksCacheBooted.Increment();
+			}
+		}
+
+		// We must have been able to make room
+		check( bReservationAccepted );
+	}
+}
+
+/* FBuildPatchChunkCache system singleton setup
+*****************************************************************************/
+TSharedPtr< FBuildPatchChunkCache, ESPMode::ThreadSafe > FBuildPatchChunkCache::SingletonInstance = NULL;
+
+void FBuildPatchChunkCache::Init(const FBuildPatchAppManifestRef& InInstallManifet, const FBuildPatchAppManifestPtr& InCurrentManifest, const FString& InChunkCacheStage, const FString& InCurrentInstallDir, FBuildPatchProgress* InBuildProgress, TArray<FString>& FilesToConstruct, FBuildPatchInstallationInfo& InstallationInfoRef)
+{
+	// We won't allow misuse of these functions
+	check(!SingletonInstance.IsValid());
+	SingletonInstance = MakeShareable(new FBuildPatchChunkCache(InInstallManifet, InCurrentManifest, InChunkCacheStage, InCurrentInstallDir, InBuildProgress, FilesToConstruct, InstallationInfoRef));
+}
+
+bool FBuildPatchChunkCache::IsAvailable()
+{
+	return SingletonInstance.IsValid();
+}
+
+FBuildPatchChunkCache& FBuildPatchChunkCache::Get()
+{
+	// We won't allow misuse of these functions
+	check( SingletonInstance.IsValid() );
+	return *SingletonInstance.Get();
+}
+
+void FBuildPatchChunkCache::Shutdown()
+{
+	// We won't allow misuse of these functions
+	check( SingletonInstance.IsValid() );
+	SingletonInstance->ChunkCache.Empty();
+	SingletonInstance.Reset();
+}
diff --git a/Engine/Source/Runtime/Online/BuildPatchServices/Private/BuildPatchChunkCache.h b/Engine/Source/Runtime/Online/BuildPatchServices/Private/BuildPatchChunkCache.h
new file mode 100644
--- /dev/null
+++ b/Engine/Source/Runtime/Online/BuildPatchServices/Private/BuildPatchChunkCache.h
@@ -0,0 +1,374 @@
+// Copyright 1998-2014 Epic Games, Inc. All Rights Reserved.
+
+/*=============================================================================
+	BuildPatchChunkCache.h: Declares chunk cache which controls the flow of
+	chunk data in and out of RAM, hopefully minimizing HDD usage
+=============================================================================*/
+
+#pragma once
+
+// Forward decelerations
+class FBuildPatchAppManifest;
+class FBuildPatchInstallationInfo;
+
+/**
+ * Declares a static access singleton class that controls the flow of chunk 
+ * acquisition and usage. It handles holding chunks in memory from downloads,
+ * and also reading chunks from an existing installation where necessary.
+ * There is a limited number of chunks that are allowed to be stored in RAM.
+ *
+ * This class is strictly for chunk data builds only, and not packaged file data.
+ */
+class FBuildPatchChunkCache
+{
+private:
+	// Enum describing the origin for a chunk
+	struct EChunkOrigin
+	{
+		enum Type
+		{
+			// The chunk will need downloading
+			Download,
+			// The chunk is in an existing installation
+			Recycle,
+			// The chunk has been saved to disk (usually forced out of a full cache)
+			Harddisk
+		};
+	};
+
+	// Class for storing chunk data in memory
+	// Protects access to the lookup only, chunk data is not protected internally and has it's own locking system
+	class FThreadSafeChunkCache
+	{
+	private:
+		// Critical Section for lookup access
+		FCriticalSection ThreadLock;
+		// The chunk file lookup
+		TMap< FGuid, FChunkFile* > ChunkStore;
+		// A list of chunks that we are reserving memory for
+		TArray< FGuid > ReservedChunks;
+	public:
+		/**
+		 * Find out if a chunk is currently loaded into the cache
+		 * @param	ChunkGuid			The chunk GUID
+		 * @return	true if the chunk is in the cache
+		 */
+		bool Contains( const FGuid& ChunkGuid );
+
+		/**
+		 * Get the number of chunks slots used
+		 * @return	How many chunk slots have been used up (or reserved)
+		 */
+		int32 Num();
+
+		/**
+		 * Add a new chunk to the cache
+		 * @param	ChunkGuid	The chunk GUID
+		 * @param	ChunkFile	The allocated and fully setup chunk file data structure
+		 */
+		void Add( const FGuid& ChunkGuid, FChunkFile* ChunkFile );
+
+		/**
+		 * Get chunk file from the cache
+		 * @param	ChunkGuid			The chunk GUID
+		 * @return		pointer to the chunk file
+		 */
+		FChunkFile* Get( const FGuid& ChunkGuid );
+
+		/**
+		 * Remove a chunk from the cache
+		 * @param	ChunkGuid			The chunk GUID
+		 */
+		void Remove( const FGuid& ChunkGuid );
+
+		/**
+		 * Completely empty the cache
+		 */
+		void Empty();
+
+		/**
+		 * Remove all chunks that have a zero ref count
+		 */
+		void PurgeUnreferenced();
+
+		/**
+		 * Attempt to reserve space for a chunk
+		 * @param	ChunkGuid			The chunk GUID
+		 * @return		true if space was reserved, false if we cannot take this chunk.
+		 */
+		bool TryAddReservation( const FGuid& ChunkGuid );
+
+		/**
+		 * Remove the reservation for a chunk
+		 * @param	ChunkGuid			The chunk GUID
+		 */
+		void RemoveReservation( const FGuid& ChunkGuid );
+
+		/**
+		 * See if a chunk has reserved space
+		 * @param	ChunkGuid			The chunk GUID
+		 * @return		true if the chunk has been reserved
+		 */
+		bool HasReservation( const FGuid& ChunkGuid );
+	};
+
+private:
+
+	// Chunk disk cache for full cache resort
+	FString ChunkCacheStage;
+
+	// Directory of current installation
+	FString CurrentInstallDir;
+
+	// The manifest for the build being installed
+	FBuildPatchAppManifestRef InstallManifet;
+
+	// The manifest for the currently installed build, if applicable
+	FBuildPatchAppManifestPtr CurrentManifest;
+
+	// Pointer to the progress tracking class which can be used to get pause state
+	FBuildPatchProgress* BuildProgress;
+
+	// Whether we are recycling chunks from current build
+	bool bEnableChunkRecycle;
+
+	// Whether downloads have started
+	FThreadSafeCounter bDownloadsStarted;
+
+	// The chunk cache storage class
+	FThreadSafeChunkCache ChunkCache;
+
+	// A critical section to protect the chunk origin lookup
+	FCriticalSection ChunkOriginLock;
+
+	// Lookup for where a chunk will be acquired from
+	TMap< FGuid, EChunkOrigin::Type > ChunkOrigins;
+
+	// A critical section to protect the chunk information lists
+	FCriticalSection ChunkInfoLock;
+
+	// A list of chunk downloads that will be queued
+	TArray< FGuid > ChunksToDownload;
+
+	// The full list of chunk references
+	TArray< FGuid > FullChunkRefList;
+
+	// A stack of chunks to store their usage order (Last is needed sooner)
+	TArray< FGuid > ChunkUseOrderStack;
+
+	// Build Stat: The number of files being constructed
+	uint32 NumFilesToConstruct;
+
+	// Build Stat: The number of chunks that need downloading
+	uint32 NumChunksToDownload;
+
+	// Build Stat: The number of chunks that need recycling
+	uint32 NumChunksToRecycle;
+
+	// Build Stat: The total number of chunks required to produce the build
+	uint32 NumRequiredChunks;
+
+	// Build Stat: The initial download size for chunks we need to download
+	int64 TotalChunkDownloadSize;
+
+	// Build Stat: The skipped download size from resuming
+	int64 SkippedChunkDownloadSize;
+
+	// Build Counter: The number of chunks that were successfully recycled
+	FThreadSafeCounter NumChunksRecycled;
+
+	// Build Counter: The number of chunks that had to be booted from cache to disk
+	FThreadSafeCounter NumChunksCacheBooted;
+
+	// Build Counter: The number of chunks that were successfully loaded from disk cache
+	FThreadSafeCounter NumDriveCacheChunkLoads;
+
+	// Build Counter: The number of chunks that failed to recycle
+	FThreadSafeCounter NumRecycleFailures;
+
+	// Build Counter: The number of chunks that failed to load from disk cache
+	FThreadSafeCounter NumDriveCacheLoadFailures;
+
+	// Reference to the Module's installation info which has the information about installed builds that can pull chunks from
+	FBuildPatchInstallationInfo& InstallationInfo;
+
+public:
+
+	/**
+	 * Constructor takes required arguments
+	 * @param	InInstallManifet		The manifest for build being installed
+	 * @param	InCurrentManifest		The manifest for build currently installed
+	 * @param	InChunkCacheStage		A stage where we can save chunks to if necessary
+	 * @param	InCurrentInstallDir		The directory of the current install
+	 * @param	InBuildProgress			Pointer to the progress tracking class, also used for pause.
+	 * @param	FilesToConstruct		List of files that are going to be constructed
+	 * @param	InstallationInfoRef		Reference to the module's installation info that keeps record of locally installed apps for use as chunk sources
+	 */
+	FBuildPatchChunkCache(const FBuildPatchAppManifestRef& InInstallManifet, const FBuildPatchAppManifestPtr& InCurrentManifest, const FString& InChunkCacheStage, const FString& InCurrentInstallDir, FBuildPatchProgress* InBuildProgress, TArray< FString >& FilesToConstruct, FBuildPatchInstallationInfo& InstallationInfoRef);
+
+	/**
+	 * Gets hold of a chunk file class to provide access to chunk data. It will be blocking until the chunk is available in memory
+	 * via whatever origin the chunk comes from.
+	 *
+	 * NOTE: The caller must dereference this chunk when finished with, inside the data lock
+	 *
+	 * NOTE: It's important to only call this function from the single thread that is writing install data.
+	 *		 Calling from other places and/or threads will have large impact on drive access performance as this function
+	 *		 also internally handles shuttling of chunk files to and from disk when necessary and this needs to take turns
+	 *		 with writing installation data.
+	 *
+	 * @param ChunkGuid		The chunk guid
+	 * @return	the chunk file data with header
+	 */
+	FChunkFile* GetChunkFile( const FGuid& ChunkGuid );
+
+	/**
+	 * Add downloaded data to the cache. The data will be copied.
+	 * @param ChunkGuid			The chunk guid for this data
+	 * @param ChunkDataFile		The file data
+	 */
+	void AddDataToCache( const FGuid& ChunkGuid, const TArray< uint8 >& ChunkDataFile );
+
+	/**
+	 * Reserve space for a chunk. Must be called before attempting download of a chunk
+	 * @param ChunkGuid			The chunk guid for this data
+	 * @return true if there was space available and this chunk is allowed into the cache
+	 */
+	bool ReserveChunkInventorySlot( const FGuid& ChunkGuid );
+
+	/**
+	 * Notify the cache that a file is already fully constructed and so we can skip over chunks for it
+	 * @param Filename		The Filename being skipped
+	 */
+	void SkipFile( const FString& Filename );
+
+	/**
+	 * Notify the cache that a chunk part for a partially complete file is to be skipped
+	 * @param ChunkPart		The chunk part
+	 */
+	void SkipChunkPart( const FChunkPart& ChunkPart );
+
+	/**
+	 * Should be after initialization and after all skippable parts have been registered to kick of the chunk downloader
+	 * with the remaining list of downloadable chunks
+	 */
+	void BeginDownloads();
+
+	/**
+	 * Get whether we have been told to queue downloads yet
+	 * @return True if downloads are queued
+	 */
+	const bool HaveDownloadsStarted() const;
+
+	/**
+	 * Get the statistic for the number of files needing construction
+	 * @return Number of files to be constructed
+	 */
+	const uint32& GetStatNumFilesToConstruct() const;
+
+	/**
+	 * Get the statistic for the number of chunks that need downloading
+	 * @return Number of chunks that need downloading
+	 */
+	const uint32& GetStatNumChunksToDownload() const;
+
+	/**
+	 * Get the statistic for the number of chunks that need recycling
+	 * @return Number of chunks that need recycling
+	 */
+	const uint32& GetStatNumChunksToRecycle() const;
+
+	/**
+	 * Get the statistic for the number of chunks required for the build
+	 * @return Number of chunks required for build
+	 */
+	const uint32& GetStatNumRequiredChunks() const;
+
+	/**
+	 * Get the total number of bytes for the chunks needing download
+	 * @return Total bytes of initial download list
+	 */
+	const int64& GetStatTotalChunkDownloadSize() const;
+
+	/**
+	 * Get the counter for the number of chunks that were recycled
+	 * @return Number of chunks that were recycled
+	 */
+	const int32 GetCounterChunksRecycled() const;
+
+	/**
+	 * Get the counter for the number of chunks that were booted from cache to disk
+	 * @return Number of chunks that were booted from the cache
+	 */
+	const int32 GetCounterChunksCacheBooted() const;
+
+	/**
+	 * Get the counter for the number of chunks that were loaded from disk cache
+	 * @return Number of chunks loaded from disk cache
+	 */
+	const int32 GetCounterDriveCacheChunkLoads() const;
+
+	/**
+	 * Get the counter for the number of chunks that failed to recycle
+	 * @return Number of chunks that failed to recycle
+	 */
+	const int32 GetCounterRecycleFailures() const;
+
+	/**
+	 * Get the counter for the number of chunks that failed to load from disk cache
+	 * @return Number of chunks that failed to load from disk cache
+	 */
+	const int32 GetCounterDriveCacheLoadFailures() const;
+
+private:
+
+	/**
+	 * Get the index for this chunk in the order of requirement time
+	 * @param	ChunkGuid	The chunk GUID
+	 * @return Chunk order index
+	 */
+	int32 GetChunkOrderIndex( const FGuid& ChunkGuid );
+
+	/**
+	 * Calculate how many more times a chunk is required
+	 * @param	ChunkGuid	The chunk GUID
+	 * @return Number of remaining references to the chunk
+	 */
+	int32 GetRemainingReferenceCount( const FGuid& ChunkGuid );
+
+	/**
+	 * Block the calling thread until the given chunk is available in the cache
+	 * @param	ChunkGuid	The chunk GUID
+	 */
+	void WaitForChunk( const FGuid& ChunkGuid );
+
+	/**
+	 * Loads a chunk from the drive cache
+	 * @param	ChunkGuid	The chunk GUID
+	 * @return true if there were no errors
+	 */
+	bool ReadChunkFromDriveCache( const FGuid& ChunkGuid );
+
+	/**
+	 * Recycle the chunk from the currently installed build
+	 * @param	ChunkGuid	The chunk GUID
+	 * @return true if there were no errors
+	 */
+	bool RecycleChunkFromBuild( const FGuid& ChunkGuid );
+
+	/**
+	 * Force there to be space in the cache by booting out the one that is not needed for the longest time
+	 * @param	ChunkGuid	The chunk GUID
+	 */
+	void ReserveChunkInventorySlotForce( const FGuid& ChunkGuid );
+
+/* Here we have static access for the singleton
+*****************************************************************************/
+public:
+	static void Init(const FBuildPatchAppManifestRef& InInstallManifet, const FBuildPatchAppManifestPtr& InCurrentManifest, const FString& InChunkCacheStage, const FString& InCurrentInstallDir, FBuildPatchProgress* InBuildProgress, TArray<FString>& FilesToConstruct, FBuildPatchInstallationInfo& ChunkRecyclerRef);
+	static FBuildPatchChunkCache& Get();
+	static void Shutdown();
+	static bool IsAvailable();
+private:
+	static TSharedPtr< FBuildPatchChunkCache, ESPMode::ThreadSafe > SingletonInstance;
+};
diff --git a/Engine/Source/Runtime/Online/BuildPatchServices/Private/BuildPatchCompactifier.cpp b/Engine/Source/Runtime/Online/BuildPatchServices/Private/BuildPatchCompactifier.cpp
new file mode 100644
--- /dev/null
+++ b/Engine/Source/Runtime/Online/BuildPatchServices/Private/BuildPatchCompactifier.cpp
@@ -0,0 +1,213 @@
+// Copyright 1998-2014 Epic Games, Inc. All Rights Reserved.
+
+/*=============================================================================
+BuildPatchCompactifier.cpp: Implements the classes that clean up chunk and file
+data that are no longer referenced by the manifests in a given cloud directory.
+=============================================================================*/
+
+#include "BuildPatchServicesPrivatePCH.h"
+
+#if WITH_BUILDPATCHGENERATION
+
+/* Constructors
+*****************************************************************************/
+
+FBuildDataCompactifier::FBuildDataCompactifier()
+{
+
+}
+
+FBuildDataCompactifier::FBuildDataCompactifier(const FString& CloudDir)
+	: CloudDir(CloudDir)
+{
+
+}
+
+/* Public static methods
+*****************************************************************************/
+
+bool FBuildDataCompactifier::CompactifyCloudDirectory(const FString& CloudDir, const bool bPreview)
+{
+	FBuildDataCompactifier Compactifier(CloudDir);
+	return Compactifier.Compactify(bPreview);
+}
+
+bool FBuildDataCompactifier::CompactifyCloudDirectory(const bool bPreview)
+{
+	return CompactifyCloudDirectory(FBuildPatchServicesModule::GetCloudDirectory(), bPreview);
+}
+
+/* Private methods
+*****************************************************************************/
+
+bool FBuildDataCompactifier::Compactify(bool bPreview) const
+{
+	GLog->Logf(TEXT("Running Compactify on %s%s"), *CloudDir, bPreview ? TEXT(". Preview mode. NO action will be taken.") : TEXT(""));
+
+	// We'll get ALL files first, so we can use the count to preallocate space within the data filenames array to save excessive reallocs
+	TArray<FString> AllFiles;
+	const bool bFindFiles = true;
+	const bool bFindDirectories = false;
+	IFileManager::Get().FindFilesRecursive(AllFiles, *CloudDir, TEXT("*.*"), bFindFiles, bFindDirectories);
+
+	TSet<FGuid> ReferencedGuids; // The master list of *ALL* referenced chunk / file data Guids
+
+	TArray<FString> ManifestFilenames;
+	TArray<FGuid> DataGuids; // The Guids associated with the data files from a single manifest
+	int32 NumDataFiles = 0;
+	
+	// Preallocate enough storage in DataGuids, to stop repeatedly expanding the allocation
+	DataGuids.Reserve(AllFiles.Num());
+
+	EnumerateManifests(ManifestFilenames);
+
+	// If we don't have any manifest files, we'll treat that as an error condition
+	if (ManifestFilenames.Num() == 0)
+	{
+		GLog->Log(ELogVerbosity::Warning, TEXT("Could not find any manifest files. Aborting operation."));
+		return true;  // We're still going to return a success code, as this isn't a hard error
+	}
+
+	// Process each manifest, and build up a list of all referenced files
+	for (const auto& ManifestFilename : ManifestFilenames)
+	{
+		const FString ManifestPath = CloudDir / ManifestFilename;
+		GLog->Logf(TEXT("Extracting chunk filenames from %s"), *ManifestFilename);
+
+		// Load the manifest data from the manifest
+		FBuildPatchAppManifest Manifest;
+		if (Manifest.LoadFromFile(ManifestPath))
+		{
+			// Work out all data Guids referenced in the manifest, and add them to our list of files to keep
+			DataGuids.Empty();
+			Manifest.GetDataList(DataGuids);
+
+			GLog->Logf(TEXT("Extracted %d chunks from %s. Unioning with %d existing chunks"), DataGuids.Num(), *ManifestFilename, NumDataFiles);
+			NumDataFiles += DataGuids.Num();
+
+			// We're going to keep all the Guids so we know which files to keep later
+			ReferencedGuids.Append(DataGuids);
+		}
+		else
+		{
+			// We failed to read from the manifest file.  This is an error which should halt progress and return a non-zero exit code
+			FString ErrorMessage = FString::Printf(TEXT("Could not parse manifest file %s"), *ManifestFilename);
+			GLog->Log(ELogVerbosity::Error, *ErrorMessage);
+			return false;
+		}
+	}
+
+	GLog->Logf(TEXT("Compactify walking %s to remove all unreferenced chunks and compute statistics."), *CloudDir);
+
+	uint32 FilesProcessed = 0;
+	uint32 NonPatchFilesProcessed = 0;
+	uint32 FilesDeleted = 0;
+	uint64 BytesProcessed = 0;
+	uint64 NonPatchBytesProcessed = 0;
+	uint64 BytesDeleted = 0;
+	uint64 CurrentFileSize;
+	FGuid FileGuid;
+
+	for (const auto& File : AllFiles)
+	{
+		CurrentFileSize = IFileManager::Get().FileSize(*File);
+		if (CurrentFileSize >= 0)
+		{
+			++FilesProcessed;
+			BytesProcessed += CurrentFileSize;
+
+			if (!GetPatchDataGuid(File, FileGuid))
+			{
+				if (!ManifestFilenames.Contains(FPaths::GetCleanFilename(File)))
+				{
+					++NonPatchFilesProcessed;
+					NonPatchBytesProcessed += CurrentFileSize;
+				}
+				continue;
+			}
+
+			if (!ReferencedGuids.Contains(FileGuid))
+			{
+				DeleteFile(File, bPreview);
+				++FilesDeleted;
+				BytesDeleted += CurrentFileSize;
+			}
+		}
+		else
+		{
+			GLog->Logf(TEXT("Warning. Could not determine size of %s. Perhaps it has been removed by another process."), *File);
+		}
+	}
+
+	GLog->Logf(TEXT("Compactify of %s complete!"), *CloudDir);
+	GLog->Logf(TEXT("Compactify found %u files totalling %s."), FilesProcessed, *HumanReadableSize(BytesProcessed));
+	if (NonPatchFilesProcessed > 0)
+	{
+		GLog->Logf(TEXT("Of these, %d (totalling %s) were not chunk/manifest files."), NonPatchFilesProcessed, *HumanReadableSize(NonPatchBytesProcessed));
+	}
+	GLog->Logf(TEXT("Compactify deleted %u chunk/manifest files totalling %s."), FilesDeleted, *HumanReadableSize(BytesDeleted));
+	return true;
+}
+
+void FBuildDataCompactifier::DeleteFile(const FString& FilePath, const bool bPreview) const
+{
+	GLog->Logf(TEXT("Deleting file %s"), *FilePath);
+	if (!bPreview)
+	{
+		IFileManager::Get().Delete(*FilePath);
+	}
+}
+
+void FBuildDataCompactifier::EnumerateManifests(TArray<FString>& OutManifests) const
+{
+	// Get a list of all manifest filenames by using IFileManager
+	FString FilePattern = CloudDir / TEXT("*.manifest");
+	IFileManager::Get().FindFiles(OutManifests, *FilePattern, true, false);
+}
+
+bool FBuildDataCompactifier::GetPatchDataGuid(const FString& FilePath, FGuid& OutGuid) const
+{
+	FString Extension = FPaths::GetExtension(FilePath);
+	if (Extension != TEXT("chunk") && Extension != TEXT("file"))
+	{
+		// Our filename doesn't have one of the allowed file extensions, so it's not patch data
+		return false;
+	}
+
+	FString BaseFileName = FPaths::GetBaseFilename(FilePath);
+	if (BaseFileName.Contains(TEXT("_")))
+	{
+		FString Right;
+		BaseFileName.Split(TEXT("_"), NULL, &Right);
+		if (Right.Len() != 32)
+		{
+			return false; // Valid data files which contain an _ in their filename must have 32 characters to the right of the underscore
+		}
+	}
+	else
+	{
+		if (BaseFileName.Len() != 32)
+		{
+			return false; // Valid data files whose names do not contain underscores must be 32 characters long
+		}
+	}
+
+	return FBuildPatchUtils::GetGUIDFromFilename(FilePath, OutGuid);
+}
+
+FString FBuildDataCompactifier::HumanReadableSize(uint64 NumBytes, uint8 DecimalPlaces, bool bUseBase10) const
+{
+	static const TCHAR* const Suffixes[2][7] = { { TEXT("Bytes"), TEXT("kB"), TEXT("MB"), TEXT("GB"), TEXT("TB"), TEXT("PB"), TEXT("EB") },
+	                                             { TEXT("Bytes"), TEXT("KiB"), TEXT("MiB"), TEXT("GiB"), TEXT("TiB"), TEXT("PiB"), TEXT("EiB") }};
+
+	float DataSize = (float)NumBytes;
+
+	const int32 Base = bUseBase10 ? 1000 : 1024;
+	int32 Index = FMath::FloorToInt(FMath::LogX(Base, NumBytes));
+	Index = FMath::Clamp<int32>(Index, 0, ARRAY_COUNT(Suffixes[0]) - 1);
+
+	DecimalPlaces = FMath::Clamp<int32>(DecimalPlaces, 0, Index*3);
+
+	return FString::Printf(TEXT("%.*f %s"), DecimalPlaces, DataSize / (FMath::Pow(Base, Index)), Suffixes[bUseBase10 ? 0 : 1][Index]);
+}
+#endif
