Moving Build Patch Services from NoRedist to Runtime

diff --git a/Engine/Source/Runtime/Online/BuildPatchServices/Private/BuildPatchCompactifier.h b/Engine/Source/Runtime/Online/BuildPatchServices/Private/BuildPatchCompactifier.h
new file mode 100644
--- /dev/null
+++ b/Engine/Source/Runtime/Online/BuildPatchServices/Private/BuildPatchCompactifier.h
@@ -0,0 +1,89 @@
+// Copyright 1998-2014 Epic Games, Inc. All Rights Reserved.
+
+/*=================================================================================
+BuildPatchCompactifier.h: Declares classes involved with compactifying build data.
+===================================================================================*/
+
+#pragma once
+
+#if WITH_BUILDPATCHGENERATION
+
+/**
+ * Used to run compactify routines on cloud directories.
+ * Compactification removes all data files not currently associated
+ * with an active manifest file.
+ */
+class FBuildDataCompactifier
+{
+public:
+	/**
+	 * Processes the Cloud Directory to identify and delete any orphaned chunks or files.
+	 * The CloudDir should already have been set using FBuildPatchServicesModule::SetCloudDirectory
+	 * NOTE: THIS function is blocking and will not return until finished. Don't run on main thread.
+	 * @param bPreview      If true, then no actual work will be done, but all operations which would be carried out will be logged.
+	 * @return              true if no file errors occurred
+	 */
+	static bool CompactifyCloudDirectory(const bool bPreview);
+
+	/**
+	 * Processes the specified Cloud Directory to identify and delete any orphaned chunks or files.
+	 * NOTE: THIS function is blocking and will not return until finished. Don't run on main thread.
+	 * @param CloudDir      The path to the Cloud Dir to compactify
+	 * @param bPreview      If true, then no actual work will be done, but all operations which would be carried out will be logged.
+	 * @return              true if no file errors occurred
+	 */
+	static bool CompactifyCloudDirectory(const FString& CloudDir, const bool bPreview);
+
+private:
+	FString CloudDir;
+
+	// Hide the default constructor
+	FBuildDataCompactifier();
+
+	/**
+	 * Constructor.
+	 * @param CloudDir      The directory that this compactifier should operate on.
+	 */
+	FBuildDataCompactifier(const FString& CloudDir);
+
+	/**
+	 * Processes the Cloud Directory to identify and delete any orphaned chunks or files.
+	 * NOTE: THIS function is blocking and will not return until finished. Don't run on main thread.
+	 * @param bPreview      If true, then no actual work will be done, but all operations which would be carried out will be logged.
+	 * @return              true if no file errors occurred
+	 */
+	bool Compactify(const bool bPreview) const;
+
+	/**
+	 * Deletes the specified file, and logs the deletion
+	 * @param FilePath     The path to the file to delete
+	 * @param bPreview     If true, then the file is not actually deleted.
+	 */
+	void DeleteFile(const FString& FilePath, const bool bPreview) const;
+
+	/**
+	 * Obtains a list of manifest filenames, along with the manifest data, in the specified Cloud Directory
+	 * @param Manifests    OUT  An array of filenames corresponding to manifest files within the Cloud Directory
+	 */
+	void EnumerateManifests(TArray<FString>& OutManifests) const;
+
+	/**
+	 * Determines whether or not the filename passed in contains patch data, indicated by a file extension of
+	 * .chunk or .file, and obtains the file's data Guid.
+	 * @param FilePath    IN  The filename to get the Guid from
+	 * @param OutGuid     OUT The Guid of the filename (only set if result  is true)
+	 * @return                True if the file is a patch and we could determine its Guid, false otherwise
+	 */
+	bool GetPatchDataGuid(const FString& FilePath, FGuid& OutGuid) const;
+
+	/**
+	 * Converts a data size (in bytes) into a more human readable form using kB, MB, GB etc.
+	 * @param NumBytes       The value (in bytes) to convert
+	 * @param DecimalPlaces  The number of decimal places of the size to show
+	 * @param bUseBase10     If true, 1kB = 1000 bytes and so on, otherwise 1KiB = 1024 bytes ...
+	 * @return               A string containing the human readable size
+	 */
+	FString HumanReadableSize(uint64 NumBytes, uint8 DecimalPlaces = 2, bool bUseBase10 = false) const;
+};
+
+#endif //WITH_BUILDPATCHGENERATION
diff --git a/Engine/Source/Runtime/Online/BuildPatchServices/Private/BuildPatchDownloader.cpp b/Engine/Source/Runtime/Online/BuildPatchServices/Private/BuildPatchDownloader.cpp
new file mode 100644
--- /dev/null
+++ b/Engine/Source/Runtime/Online/BuildPatchServices/Private/BuildPatchDownloader.cpp
@@ -0,0 +1,578 @@
+// Copyright 1998-2014 Epic Games, Inc. All Rights Reserved.
+
+/*=============================================================================
+	FBuildPatchDownloader.cpp: Implements the BuildPatchChunkDownloader
+	that runs a thread to download chunks in it's queue.
+=============================================================================*/
+
+#include "BuildPatchServicesPrivatePCH.h"
+
+#define EXTRA_DOWNLOAD_LOGGING 0
+
+int32 NUM_DOWNLOAD_THREADS = 8;
+
+int32 CHUNK_DOWNLOAD_RETRIES = 20;
+
+float CHUNK_RETRY_TIME = 1.0;
+
+/* FBuildPatchDownloader implementation
+*****************************************************************************/
+FBuildPatchDownloader::FBuildPatchDownloader( const FString& InSaveDirectory, const FBuildPatchAppManifestRef& InInstallManifest, FBuildPatchProgress* InBuildProgress )
+	: SaveDirectory( InSaveDirectory )
+	, InstallManifest( InInstallManifest )
+	, Thread( NULL )
+	, bIsRunning( false )
+	, bIsInited( false )
+	, bIsIdle( false )
+	, bWaitingForJobs( true )
+	, DataToDownloadTotalBytes( 0 )
+	, BuildProgress( InBuildProgress )
+{
+	// config overrides to control downloads
+	GConfig->GetInt(TEXT("Portal.BuildPatch"), TEXT("ChunkDownloads"), NUM_DOWNLOAD_THREADS, GEngineIni);
+	NUM_DOWNLOAD_THREADS = FMath::Clamp<int32>(NUM_DOWNLOAD_THREADS, 1, 100);
+	GConfig->GetInt(TEXT("Portal.BuildPatch"), TEXT("ChunkRetries"), CHUNK_DOWNLOAD_RETRIES, GEngineIni);
+	CHUNK_DOWNLOAD_RETRIES = FMath::Clamp<int32>(CHUNK_DOWNLOAD_RETRIES, 1, 100);
+	GConfig->GetFloat(TEXT("Portal.BuildPatch"), TEXT("RetryTime"), CHUNK_RETRY_TIME, GEngineIni);
+	CHUNK_RETRY_TIME = FMath::Clamp<float>(CHUNK_RETRY_TIME, 0.5f, 30.0f);
+
+	// Start thread!
+	const TCHAR* ThreadName = TEXT( "ChunkDownloaderThread" );
+	Thread = FRunnableThread::Create(this, ThreadName);
+}
+
+FBuildPatchDownloader::~FBuildPatchDownloader()
+{
+	if( Thread != NULL )
+	{
+		Thread->WaitForCompletion();
+		delete Thread;
+		Thread = NULL;
+	}
+}
+
+bool FBuildPatchDownloader::Init()
+{
+	bool bReady = true;
+	return bReady;
+}
+
+uint32 FBuildPatchDownloader::Run()
+{
+	FString ErrorString;
+	SetRunning( true );
+	SetInited( true );
+
+	// Cache is chunk data
+	const bool bIsChunkData = !InstallManifest->IsFileDataManifest();
+	const FBuildPatchData::Type DataType = bIsChunkData ? FBuildPatchData::ChunkData : FBuildPatchData::FileData;
+
+	// While there is the possibility of more chunks to download
+	while ( ShouldBeRunning() && !FBuildPatchInstallError::HasFatalError() )
+	{
+		// Give other threads some time
+		FPlatformProcess::Sleep( 0.0f );
+		// Check inflight downloads for completion and process them
+		InFlightDownloadsLock.Lock();
+		TArray< FGuid > InFlightKeys;
+		InFlightDownloads.GetKeys( InFlightKeys );
+		for( auto InFlightKeysIt = InFlightKeys.CreateConstIterator(); InFlightKeysIt && !FBuildPatchInstallError::HasFatalError(); ++InFlightKeysIt )
+		{
+			const FGuid& InFlightKey = *InFlightKeysIt;
+			FDownloadJob& InFlightJob = InFlightDownloads[ InFlightKey ];
+			// If the download has finished running and not restarting
+			if( InFlightJob.StateFlag.GetValue() > EDownloadState::DownloadRunning )
+			{
+				InFlightDownloadsLock.Unlock();
+				bool bSuccess = InFlightJob.StateFlag.GetValue() == EDownloadState::DownloadSuccess;
+				if( bSuccess )
+				{
+					// Uncompress if needed
+					bSuccess = !bIsChunkData || FBuildPatchUtils::UncompressChunkFile( InFlightJob.DataArray );
+					if( !bSuccess )
+					{
+						FBuildPatchAnalytics::RecordChunkDownloadError( InFlightJob.DownloadUrl, FPlatformMisc::GetLastError(), TEXT( "Uncompress Fail" ) );
+						GWarn->Logf( TEXT( "BuildPatchServices: ERROR: Failed to uncompress chunk data %s" ), *InFlightJob.DownloadUrl );
+					}
+					// Verify the data in memory
+					else if( FBuildPatchUtils::VerifyChunkFile( InFlightJob.DataArray ) )
+					{
+						// When saving file data locally, we can just use the old filename
+						const FString FileDataPath = FBuildPatchUtils::GetFileOldFilename( SaveDirectory, InFlightJob.Guid );
+						// Save file data to staging area
+						if( bIsChunkData || FFileHelper::SaveArrayToFile( InFlightJob.DataArray, *FileDataPath ) )
+						{
+							// Record download
+							{
+								FScopeLock Lock( &DownloadRecordsLock );
+								DownloadRecords.Add( InFlightJob.DownloadRecord );
+							}
+
+							// Register with the chunk cache
+							if( bIsChunkData )
+							{
+								FBuildPatchChunkCache::Get().AddDataToCache( InFlightJob.Guid, InFlightJob.DataArray );
+							}
+							else
+							{
+								FBuildPatchFileConstructor::AddFileDataToInventory( InFlightJob.Guid, FileDataPath );
+							}
+
+							bSuccess = true;
+						}
+						else
+						{
+							bSuccess = false;
+							FBuildPatchAnalytics::RecordChunkDownloadError( InFlightJob.DownloadUrl, FPlatformMisc::GetLastError(), TEXT( "SaveToDisk Fail" ) );
+							GWarn->Logf( TEXT( "BuildPatchServices: ERROR: Failed to save file data %s" ), *FileDataPath );
+						}
+					}
+					else
+					{
+						bSuccess = false;
+						FBuildPatchAnalytics::RecordChunkDownloadError( InFlightJob.DownloadUrl, InFlightJob.ResponseCode, TEXT( "Verify Fail" ) );
+						GWarn->Logf( TEXT( "BuildPatchServices: ERROR: Verify failed on chunk %s" ), *InFlightJob.DownloadUrl );
+					}
+				}
+				else
+				{
+					bSuccess = false;
+					FBuildPatchAnalytics::RecordChunkDownloadError( InFlightJob.DownloadUrl, InFlightJob.ResponseCode, TEXT( "Download Fail" ) );
+					GWarn->Logf( TEXT( "BuildPatchServices: ERROR: Failed to download chunk %s" ), *InFlightJob.DownloadUrl );
+				}
+				
+				// Handle failed
+				if( !bSuccess )
+				{
+					if( InFlightJob.RetryCount.Increment() <= CHUNK_DOWNLOAD_RETRIES )
+					{
+						DataToRetry.Add( InFlightJob.Guid );
+						InFlightJob.StateFlag.Set( EDownloadState::DownloadInit );
+						GWarn->Logf( TEXT( "BuildPatchServices: ERROR: Failed to downloaded data %s. Retrying" ), *InFlightJob.Guid.ToString() );
+					}
+					else
+					{
+						FString ErrorString = TEXT( "Failed to download chunk. No more retries allowed for " );
+						ErrorString += InFlightJob.Guid.ToString();
+						FBuildPatchInstallError::SetFatalError( EBuildPatchInstallError::DownloadError, ErrorString );
+					}
+				}
+
+				// Remove from list if successful, we still need to lock as it ties in with looping
+				InFlightDownloadsLock.Lock();
+				if( bSuccess )
+				{
+					InFlightDownloads.Remove( InFlightKey );
+				}
+			}
+		}
+		InFlightDownloadsLock.Unlock();
+
+		// Pause
+		BuildProgress->WaitWhilePaused();
+
+		// Check if there were any errors while paused, like canceling
+		if( FBuildPatchInstallError::HasFatalError() )
+		{
+			continue;
+		}
+
+		// Try get a job, or if we don't have a job yet, wait a little.
+		FGuid NextGuid;
+		if( !GetNextDownload( NextGuid ) )
+		{
+			SetIdle( true );
+			FPlatformProcess::Sleep( 0.1f );
+			continue;
+		}
+		SetIdle( false );
+
+		// Make the filename and download url
+		FString DownloadUrl;
+		// Support older manifests
+		if( InstallManifest->GetManifestVersion() < EBuildPatchAppManifestVersion::DataFileRenames )
+		{
+			DownloadUrl = FBuildPatchUtils::GetDataTypeOldFilename( DataType, FBuildPatchServicesModule::GetCloudDirectory(), NextGuid );
+		}
+		else if( bIsChunkData )
+		{
+			uint64 ChunkHash;
+			const bool bFound = InstallManifest->GetChunkHash( NextGuid, ChunkHash );
+			// Should be impossible to not exist
+			check( bFound );
+			DownloadUrl = FBuildPatchUtils::GetChunkNewFilename( InstallManifest->GetManifestVersion(), FBuildPatchServicesModule::GetCloudDirectory(), NextGuid, ChunkHash );
+		}
+		else
+		{
+			FSHAHash FileHash;
+			const bool bFound = InstallManifest->GetFileDataHash( NextGuid, FileHash );
+			// Should be impossible to not exist
+			check( bFound );
+			DownloadUrl = FBuildPatchUtils::GetFileNewFilename( InstallManifest->GetManifestVersion(), FBuildPatchServicesModule::GetCloudDirectory(), NextGuid, FileHash );
+		}
+
+		const bool bIsHTTPRequest = DownloadUrl.Contains( TEXT( "http" ), ESearchCase::IgnoreCase );
+
+		// Start the download
+		int32 HttpRequestID = INDEX_NONE;
+		if( bIsHTTPRequest )
+		{
+			// Kick off the download
+#if EXTRA_DOWNLOAD_LOGGING
+			GWarn->Logf(TEXT("BuildPatchDownloader: Queuing chunk HTTP download %s"), *NextGuid.ToString());
+#endif
+			HttpRequestID = FBuildPatchHTTP::QueueHttpRequest( DownloadUrl, FHttpRequestCompleteDelegate::CreateThreadSafeSP( this, &FBuildPatchDownloader::HttpRequestComplete ), FHttpRequestProgressDelegate::CreateThreadSafeSP( this, &FBuildPatchDownloader::HttpRequestProgress ) );
+		}
+		else
+		{
+			// Load file from drive/network
+			TArray< uint8 > FileDataArray;
+			FArchive* Reader = IFileManager::Get().CreateFileReader( *DownloadUrl );
+			bool bSuccess = Reader != NULL;
+			if( Reader )
+			{
+				const int64 BytesPerCall = 256*1024;
+				const int64 FileSize = Reader->TotalSize();
+				FileDataArray.Reset();
+				FileDataArray.AddUninitialized( FileSize );
+				int64 BytesRead = 0;
+				while ( BytesRead < FileSize )
+				{
+					const int64 ReadLen = FMath::Min<int64>( BytesPerCall, FileSize - BytesRead );
+					Reader->Serialize( FileDataArray.GetTypedData() + BytesRead, ReadLen );
+					BytesRead += ReadLen;
+					OnDownloadProgress( NextGuid, BytesRead );
+				}
+				bSuccess = Reader->Close();
+				delete Reader;
+			}
+			OnDownloadComplete( NextGuid, DownloadUrl, FileDataArray, bSuccess, INDEX_NONE );
+		}
+	}
+	SetIdle( true );
+	SetRunning( false );
+	return 0;
+}
+
+bool FBuildPatchDownloader::IsComplete()
+{
+	FScopeLock Lock( &FlagsLock );
+	return !bIsRunning && bIsInited;
+}
+
+bool FBuildPatchDownloader::IsIdle()
+{
+	FScopeLock Lock( &FlagsLock );
+	return bIsIdle;
+}
+
+TArray< FBuildPatchDownloadRecord > FBuildPatchDownloader::GetDownloadRecordings()
+{
+	FScopeLock Lock( &DownloadRecordsLock );
+	return DownloadRecords;
+}
+
+void FBuildPatchDownloader::AddChunkToDownload( const FGuid& Guid, const bool bInsertAtFront /*= false */ )
+{
+	check( Guid.IsValid() );
+	FScopeLock Lock( &DataToDownloadLock );
+	if( bInsertAtFront )
+	{
+		DataToDownload.Insert( Guid, 0 );
+	}
+	else
+	{
+		DataToDownload.Add( Guid );
+	}
+	DataToDownloadTotalBytes += InstallManifest->GetDataSize(Guid);
+}
+
+void FBuildPatchDownloader::AddChunksToDownload( const TArray< FGuid >& Guids, const bool bInsertAtFront /*= false */ )
+{
+	FScopeLock Lock( &DataToDownloadLock );
+	for ( FGuid Guid: Guids )
+	{
+		check( Guid.IsValid() );
+		if( bInsertAtFront )
+		{
+			DataToDownload.Insert( Guid, 0 );
+		}
+		else
+		{
+			DataToDownload.Add( Guid );
+		}
+		DataToDownloadTotalBytes += InstallManifest->GetDataSize(Guid);
+	}
+}
+
+void FBuildPatchDownloader::ClearDownloadJobs()
+{
+	FScopeLock Lock( &DataToDownloadLock );
+	DataToDownload.Empty();
+}
+
+void FBuildPatchDownloader::NotifyNoMoreChunksToAdd()
+{
+	FScopeLock Lock( &FlagsLock );
+	bWaitingForJobs = false;
+}
+
+int32 FBuildPatchDownloader::GetNumChunksLeft()
+{
+	int32 Count = 0;
+	{
+		FScopeLock Lock( &DataToDownloadLock );
+		Count += DataToDownload.Num();
+	}
+	{
+		FScopeLock Lock( &InFlightDownloadsLock );
+		Count += InFlightDownloads.Num();
+	}
+	return Count;
+}
+
+int64 FBuildPatchDownloader::GetNumBytesLeft()
+{
+	int64 Count = 0;
+	// Count the bytes for each file in the queue
+	{
+		FScopeLock Lock(&DataToDownloadLock);
+		for (auto DataToDownloadIt = DataToDownload.CreateConstIterator(); DataToDownloadIt; ++DataToDownloadIt)
+		{
+			const FGuid& DataGuid = *DataToDownloadIt;
+			Count += InstallManifest->GetDataSize(DataGuid);
+		}
+	}
+	// Count the bytes left for each in-flight file
+	{
+		FScopeLock Lock(&InFlightDownloadsLock);
+		for (auto InFlightDownloadsIt = InFlightDownloads.CreateConstIterator(); InFlightDownloadsIt; ++InFlightDownloadsIt)
+		{
+			const FGuid& InFlightKey = InFlightDownloadsIt.Key();
+			const FDownloadJob& InFlightJob = InFlightDownloadsIt.Value();
+			Count += FMath::Max<int64>(InstallManifest->GetDataSize(InFlightKey) - InFlightJob.DownloadRecord.DownloadSize, 0);
+		}
+	}
+	return Count;
+}
+
+int64 FBuildPatchDownloader::GetNumBytesQueued()
+{
+	FScopeLock Lock(&DataToDownloadLock);
+	return DataToDownloadTotalBytes;
+}
+
+int32 FBuildPatchDownloader::GetByteDownloadCountReset()
+{
+	return ByteDownloadCount.Reset();
+}
+
+void FBuildPatchDownloader::IncrementByteDownloadCount( const int32& NumBytes )
+{
+	ByteDownloadCount.Add( NumBytes );
+}
+
+void FBuildPatchDownloader::SetRunning( bool bRunning )
+{
+	FScopeLock Lock( &FlagsLock );
+	bIsRunning = bRunning;
+}
+
+void FBuildPatchDownloader::SetInited( bool bInited )
+{
+	FScopeLock Lock( &FlagsLock );
+	bIsInited = bInited;
+}
+
+void FBuildPatchDownloader::SetIdle( bool bIdle )
+{
+	FScopeLock Lock( &FlagsLock );
+	bIsIdle = bIdle;
+}
+
+bool FBuildPatchDownloader::ShouldBeRunning()
+{
+	FScopeLock Lock( &FlagsLock );
+	return bWaitingForJobs || GetNumChunksLeft() > 0;
+}
+
+void FBuildPatchDownloader::HttpRequestProgress( FHttpRequestPtr Request, int32 BytesSoFar )
+{
+	if ( !FBuildPatchInstallError::HasFatalError() )
+	{
+		FGuid DataGUID;
+		FBuildPatchUtils::GetGUIDFromFilename( Request->GetURL(), DataGUID );
+		OnDownloadProgress( DataGUID, BytesSoFar );
+	}
+}
+
+void FBuildPatchDownloader::HttpRequestComplete( FHttpRequestPtr Request, FHttpResponsePtr Response, bool bSucceeded )
+{
+	if ( !FBuildPatchInstallError::HasFatalError() )
+	{
+		FGuid DataGUID;
+		FBuildPatchUtils::GetGUIDFromFilename( Request->GetURL(), DataGUID );
+		if( bSucceeded && Response.IsValid() && EHttpResponseCodes::IsOk( Response->GetResponseCode() ) )
+		{
+			OnDownloadComplete( DataGUID, Request->GetURL(), Response->GetContent(), true, Response->GetResponseCode() );
+		}
+		else
+		{
+			OnDownloadComplete( DataGUID, Request->GetURL(), TArray< uint8 >(), false, Response.IsValid() ? Response->GetResponseCode() : INDEX_NONE );
+		}
+	}
+}
+
+void FBuildPatchDownloader::OnDownloadProgress( const FGuid& Guid, const int32& BytesSoFar )
+{
+	if ( !FBuildPatchInstallError::HasFatalError() )
+	{
+		FScopeLock ScopeLock( &InFlightDownloadsLock );
+		FDownloadJob& CurrentJob = InFlightDownloads[ Guid ];
+		const int32 DeltaBytes = BytesSoFar - CurrentJob.DownloadRecord.DownloadSize;
+		CurrentJob.DownloadRecord.DownloadSize = BytesSoFar;
+		IncrementByteDownloadCount( DeltaBytes );
+#if EXTRA_DOWNLOAD_LOGGING
+		// Log slow downloads
+		const double tNow = FPlatformTime::Seconds();
+		const double tDelta = tNow - CurrentJob.DownloadRecord.EndTime;
+		const double tRunning = tNow - CurrentJob.DownloadRecord.StartTime;
+		const double currSpeed = DeltaBytes / tDelta;
+		const double avSpeed = BytesSoFar / tRunning;
+		if (tRunning > 60 && (tNow - CurrentJob.TimeLastLogged > 5))
+		{
+			CurrentJob.TimeLastLogged = tNow;
+			GWarn->Logf(TEXT("BuildPatchDownloader: WARNING: Slow download %s Running:%.2fsec\tDownloaded:%d\tSpeed:%.0f\tAverageSpeed:%.0f (B/s)"), *Guid.ToString(), tRunning, BytesSoFar, currSpeed, avSpeed);
+		}
+		CurrentJob.DownloadRecord.EndTime = tNow;
+#endif
+	}
+}
+
+void FBuildPatchDownloader::OnDownloadComplete( const FGuid& Guid, const FString& DownloadUrl, const TArray< uint8 >& DataArray, bool bSucceeded, int32 ResponseCode )
+{
+#if EXTRA_DOWNLOAD_LOGGING
+	GWarn->Logf(TEXT("BuildPatchDownloader: OnDownloadComplete %s %s %d"), *Guid.ToString(), bSucceeded?TEXT("SUCCESS"):TEXT("FAIL"), ResponseCode);
+#endif
+	if (!FBuildPatchInstallError::HasFatalError())
+	{
+		FScopeLock ScopeLock(&InFlightDownloadsLock);
+		FDownloadJob& CurrentJob = InFlightDownloads[Guid];
+
+		// Safety check, if this fails, we will be causing thread access issues with Run() processing data on complete downloads.
+		// The StateFlag acts as a thread-safe gate between access here and access there, InFlightDownloadsLock must be a short lived lock to avoid
+		// stalling the main thread here
+		check(CurrentJob.StateFlag.GetValue() == EDownloadState::DownloadRunning);
+
+		const int32 RequestSize = DataArray.Num();
+		const int32 DeltaBytes = RequestSize - CurrentJob.DownloadRecord.DownloadSize;
+		IncrementByteDownloadCount(DeltaBytes);
+		CurrentJob.DataArray.Empty(RequestSize);
+		CurrentJob.DataArray.Append(DataArray);
+		CurrentJob.DownloadRecord.DownloadSize = RequestSize;
+		CurrentJob.DownloadRecord.EndTime = FPlatformTime::Seconds();
+		CurrentJob.StateFlag.Set(bSucceeded ? EDownloadState::DownloadSuccess : EDownloadState::DownloadFail);
+		CurrentJob.ResponseCode = ResponseCode;
+		CurrentJob.DownloadUrl = DownloadUrl;
+	}
+}
+
+bool FBuildPatchDownloader::GetNextDownload( FGuid& Job )
+{
+	// Check ChunkCache is ready
+	const bool bIsFileData = InstallManifest->IsFileDataManifest();
+	if( FBuildPatchInstallError::HasFatalError() || ( !bIsFileData && !FBuildPatchChunkCache::IsAvailable() ) )
+	{
+		return false;
+	}
+	// Check retry list, or number of running downloads
+	bool bAllowedDownload = DataToRetry.Num() > 0;
+	if( !bAllowedDownload )
+	{
+		FScopeLock ScopeLock( &InFlightDownloadsLock );
+		bAllowedDownload = InFlightDownloads.Num() < NUM_DOWNLOAD_THREADS;
+	}
+	// Find a guid that the chache wants
+	if( bAllowedDownload )
+	{
+		// Check retrying a download
+		if( DataToRetry.Num() > 0 )
+		{
+			FScopeLock ScopeLock( &InFlightDownloadsLock );
+			FDownloadJob& InFlightJob = InFlightDownloads[ DataToRetry[0] ];
+			const double TimeSinceFail = FPlatformTime::Seconds() - InFlightJob.DownloadRecord.EndTime;
+			bAllowedDownload = TimeSinceFail >= CHUNK_RETRY_TIME;
+			if( bAllowedDownload )
+			{
+				Job = DataToRetry.Pop();
+			}
+		}
+		else
+		{
+			bAllowedDownload = false;
+			DataToDownloadLock.Lock();
+			for( int32 DataToDownloadIdx = 0; DataToDownloadIdx < DataToDownload.Num(); ++DataToDownloadIdx )
+			{
+				const FGuid& DownloadJob = DataToDownload[ DataToDownloadIdx ];
+				check( DownloadJob.IsValid() );
+				bAllowedDownload = bIsFileData || FBuildPatchChunkCache::Get().ReserveChunkInventorySlot( DownloadJob );
+				if( bAllowedDownload )
+				{
+					Job = DownloadJob;
+					DataToDownload.RemoveAt( DataToDownloadIdx );
+					break;
+				}
+				// Allow other threads chance
+				int32 PrevNum = DataToDownload.Num();
+				DataToDownloadLock.Unlock();
+				FPlatformProcess::Sleep( 0.0f );
+				DataToDownloadLock.Lock();
+				// Start over if array changed
+				if( PrevNum != DataToDownload.Num() )
+				{
+					DataToDownloadIdx = 0;
+				}
+			}
+			DataToDownloadLock.Unlock();
+		}
+	}
+	// If allowed to download, add to InFlight
+	if( bAllowedDownload )
+	{
+		check( Job.IsValid() );
+		FScopeLock ScopeLock( &InFlightDownloadsLock );
+		FDownloadJob& InFlightJob = InFlightDownloads.FindOrAdd( Job );
+		InFlightJob.Guid = Job;
+		InFlightJob.DataArray.Empty();
+		InFlightJob.ResponseCode = 0;
+		InFlightJob.StateFlag.Set( EDownloadState::DownloadRunning );
+		InFlightJob.DownloadRecord.DownloadSize = 0;
+		InFlightJob.DownloadRecord.StartTime = FPlatformTime::Seconds();
+	}
+	return bAllowedDownload;
+}
+
+/* FBuildPatchChunkCache system singleton setup
+*****************************************************************************/
+TSharedPtr< FBuildPatchDownloader, ESPMode::ThreadSafe > FBuildPatchDownloader::SingletonInstance = NULL;
+
+void FBuildPatchDownloader::Create( const FString& InSaveDirectory, const FBuildPatchAppManifestRef& InInstallManifest, FBuildPatchProgress* InBuildProgress )
+{
+	// We won't allow misuse of these functions
+	check( !SingletonInstance.IsValid() );
+	SingletonInstance = MakeShareable( new FBuildPatchDownloader( InSaveDirectory, InInstallManifest, InBuildProgress ) );
+}
+
+FBuildPatchDownloader& FBuildPatchDownloader::Get()
+{
+	// We won't allow misuse of these functions
+	check( SingletonInstance.IsValid() );
+	return *SingletonInstance.Get();
+}
+
+void FBuildPatchDownloader::Shutdown()
+{
+	// We won't allow misuse of these functions
+	check( SingletonInstance.IsValid() );
+	SingletonInstance.Reset();
+}
diff --git a/Engine/Source/Runtime/Online/BuildPatchServices/Private/BuildPatchDownloader.h b/Engine/Source/Runtime/Online/BuildPatchServices/Private/BuildPatchDownloader.h
new file mode 100644
--- /dev/null
+++ b/Engine/Source/Runtime/Online/BuildPatchServices/Private/BuildPatchDownloader.h
@@ -0,0 +1,335 @@
+// Copyright 1998-2014 Epic Games, Inc. All Rights Reserved.
+
+/*=============================================================================
+	FBuildPatchDownloader.h: Declares the BuildPatchChunkDownloader
+	that runs a thread to download chunks in it's queue.
+=============================================================================*/
+
+/**
+ * FChunkDownloadRecord
+ * This is a simple struct to record downloads
+ */
+struct FBuildPatchDownloadRecord
+{
+	// The start time of the download
+	double StartTime;
+	// The end time of the download
+	double EndTime;
+	// The size of the download
+	int64 DownloadSize;
+	/**
+	 * Default constructor
+	 */
+	FBuildPatchDownloadRecord()
+		: StartTime( 0 )
+		, EndTime( 0 )
+		, DownloadSize( 0 ) {}
+	/**
+	 * Less-than operator to allow sorting by start time.
+	 * @param Other	 The record to compare against.
+	 * @return true if this download started first, otherwise false.
+	 */
+	FORCEINLINE bool operator<( const FBuildPatchDownloadRecord& Other ) const
+	{
+		return StartTime < Other.StartTime;
+	}
+};
+
+/**
+ * FBuildPatchDownloader
+ * This class controls a thread that downloads chunks from a static chunk list
+ */
+class FBuildPatchDownloader
+	: public FRunnable
+	, public TSharedFromThis< FBuildPatchDownloader, ESPMode::ThreadSafe >
+{
+private:
+	// Enum to list download status flags
+	struct EDownloadState
+	{
+		enum Type
+		{
+			// The download is initialising
+			DownloadInit = 0,
+			// The download is running
+			DownloadRunning,
+			// The download failed
+			DownloadFail,
+			// The download completed successfully
+			DownloadSuccess
+		};
+	};
+
+	// A struct to store a guid with it's data type
+	struct FDownloadJob
+	{
+		// The data guid
+		FGuid Guid;
+		// The download url for this job
+		FString DownloadUrl;
+		// A thread safe counter that will count retries.
+		FThreadSafeCounter RetryCount;
+		// A thread safe counter that we will use a simple check for download completion
+		FThreadSafeCounter StateFlag;
+		// The downloaded data
+		TArray< uint8 > DataArray;
+		// Track download time
+		FBuildPatchDownloadRecord DownloadRecord;
+		// The response code, for HTTP downloads
+		int32 ResponseCode;
+		// A time we last logged info for this download, to help not overly spam
+		double TimeLastLogged;
+
+		/**
+		 * Default constructor
+		 */
+		FDownloadJob()
+			: DownloadUrl( TEXT( "" ) )
+			, RetryCount( 0 )
+			, StateFlag( EDownloadState::DownloadInit )
+			, ResponseCode( 0 )
+			, TimeLastLogged( 0 )
+		{}
+	};
+
+private:
+
+	// The directory to save data to
+	const FString SaveDirectory;
+
+	// The Manifest that is being installed
+	const FBuildPatchAppManifestRef InstallManifest;
+
+	// Hold a pointer to my thread for deleting
+	FRunnableThread* Thread;
+
+	// A flag marking that we a running
+	bool bIsRunning;
+
+	// A flag marking that we initialized correctly
+	bool bIsInited;
+
+	// A flag marking whether the thread is idle waiting for jobs
+	bool bIsIdle;
+
+	// A flag that says whether more chunks could still be queued
+	bool bWaitingForJobs;
+
+	// A critical section to protect the flags
+	FCriticalSection FlagsLock;
+
+	// Store a record of each download we made for info
+	TArray< FBuildPatchDownloadRecord > DownloadRecords;
+
+	// Critical section to protect the download list
+	FCriticalSection DownloadRecordsLock;
+
+	// List of download jobs to be processed
+	TArray< FGuid > DataToDownload;
+
+	// List of download jobs to be retried
+	TArray< FGuid > DataToRetry;
+
+	// Count of number of bytes total added to download queue
+	int64 DataToDownloadTotalBytes;
+
+	// Critical section to protect the download list
+	FCriticalSection DataToDownloadLock;
+
+	// List of download jobs retry counts
+	TMap< FGuid, FDownloadJob > InFlightDownloads;
+
+	// Critical section to protect the retry count list
+	FCriticalSection InFlightDownloadsLock;
+
+	// Counter for tracking number of bytes downloaded per delta
+	FThreadSafeCounter ByteDownloadCount;
+
+	// Pointer to the build progress tracker
+	FBuildPatchProgress* BuildProgress;
+
+public:
+
+	/**
+	 * Constructor
+	 * @param	InSaveDirectory		The chunk save to directory
+	 * @param	InInstallManifest	The manifest being installed so we can pull info from it
+	 * @param	InBuildProgress		Pointer to the progress tracker, for setting progress and getting pause state
+	 */
+	FBuildPatchDownloader( const FString& InSaveDirectory, const FBuildPatchAppManifestRef& InInstallManifest, FBuildPatchProgress* InBuildProgress );
+
+	/**
+	 * Default Destructor, deletes allocated Thread
+	 */
+	~FBuildPatchDownloader();
+
+	// Begin FRunnable
+	virtual bool Init() override;
+	virtual uint32 Run() override;
+	// End FRunnable
+
+	/**
+	 * Get whether the thread has finished working
+	 * @return	true if the thread completed
+	 */
+	bool IsComplete();
+
+	/**
+	 * Get whether the thread is idle waiting for chunks to download
+	 * @return	true if the thread completed
+	 */
+	bool IsIdle();
+
+	/**
+	 * Gets the array of download recordings. Should not be polled, only call when the thread has finished to gather data
+	 * @return	An array of download records
+	 */
+	TArray< FBuildPatchDownloadRecord > GetDownloadRecordings();
+
+	/**
+	 * Adds a chunk GUID to the download list
+	 * @param Guid				The GUID for the data
+	 * @param bInsertAtFront	If this should be queued at the front
+	 */
+	void AddChunkToDownload( const FGuid& Guid, const bool bInsertAtFront = false );
+
+	/**
+	 * Adds an array of chunk GUIDs to the download list
+	 * @param Guid				The GUID array for the data
+	 * @param bInsertAtFront	If this should be queued at the front
+	 */
+	void AddChunksToDownload( const TArray< FGuid >& Guids, const bool bInsertAtFront = false );
+
+	/**
+	 * Clear out the download list
+	 */
+	void ClearDownloadJobs();
+
+	/**
+	 * A notification for when there is no possibility of more jobs to be queued (so the thread can end)
+	 */
+	void NotifyNoMoreChunksToAdd();
+
+	/**
+	 * Get how many chunks are in the list
+	 * @return The number of chunks in the list
+	 */
+	int32 GetNumChunksLeft();
+
+	/**
+	 * Get the number of bytes left to download.
+	 * @return The number of byte left to download
+	 */
+	int64 GetNumBytesLeft();
+
+	/**
+	 * Get the number of bytes that have been queued
+	 * @return The total number of bytes queued
+	 */
+	int64 GetNumBytesQueued();
+
+	/**
+	 * Get how many bytes have been downloaded, resetting the counter
+	 * @return The number of bytes downloaded since the last call
+	 */
+	int32 GetByteDownloadCountReset();
+
+	/**
+	 * Add some bytes to the download counter
+	 * @param	NumBytes The number of bytes to add
+	 */
+	void IncrementByteDownloadCount( const int32& NumBytes );
+
+private:
+
+	/**
+	 * Sets the bIsRunning flag
+	 * @param bRunning	Whether the thread is running
+	 */
+	void SetRunning( bool bRunning );
+
+	/**
+	 * Sets the bIsInited flag
+	 * @param bInited	Whether the thread successfully initialized
+	 */
+	void SetInited( bool bInited );
+
+	/**
+	 * Sets the bIsIdle flag
+	 * @param bIdle	Whether the thread is waiting for jobs
+	 */
+	void SetIdle( bool bIdle );
+
+	/**
+	 * Gets whether the thread still has work to do
+	 * @return	false if the thread can quit
+	 */
+	bool ShouldBeRunning();
+
+	/**
+	 * Receives the HTTP request bytes downloaded when downloading a chunk.
+	 * @param Request		The request that was made
+	 * @param BytesSoFar	Whether the request successfully completed
+	 */
+	void HttpRequestProgress( FHttpRequestPtr Request, int32 BytesSoFar );
+
+	/**
+	 * Receives the HTTP response when downloading a chunk.
+	 * @param Request		The request that was made
+	 * @param Response		The request's response
+	 * @param bSucceeded	Whether the request successfully completed
+	 */
+	void HttpRequestComplete( FHttpRequestPtr Request, FHttpResponsePtr Response, bool bSucceeded );
+
+	/**
+	 * Progress update for download of a chunk, from HTTP or Network/HDD.
+	 * @param Guid			The guid for the data
+	 * @param BytesSoFar	The number of bytes downloaded so far
+	 */
+	void OnDownloadProgress( const FGuid& Guid, const int32& BytesSoFar );
+
+	/**
+	 * Completes a download of a chunk, from HTTP or Network/HDD.
+	 * @param Guid			The guid for the data
+	 * @param DownloadUrl	The url used to download the data
+	 * @param DataArray		The array of bytes for the data
+	 * @param bSucceeded	Whether the download was successful
+	 * @param ResponseCode	The HTTP response code
+	 */
+	void OnDownloadComplete( const FGuid& Guid, const FString& DownloadUrl, const TArray< uint8 >& DataArray, bool bSucceeded, int32 ResponseCode );
+
+	/**
+	 * Get the next GUID to download data with
+	 * @param Guid		Receives the download data GUID
+	 * @return true if there was a job in the list
+	 */
+	bool GetNextDownload( FGuid& Guid );
+
+	/* Here we have static access for the singleton
+	*****************************************************************************/
+public:
+
+	/**
+	 * Creates the singleton downloader system. Takes the arguments required for construction.
+	 * @param	InSaveDirectory		The chunk save to directory
+	 * @param	InInstallManifest	The manifest being installed so we can pull info from it
+	 * @param	InBuildProgress		Pointer to the progress tracker, for setting progress and getting pause state
+	 */
+	static void Create( const FString& InSaveDirectory, const FBuildPatchAppManifestRef& InInstallManifest, FBuildPatchProgress* InBuildProgress );
+
+	/**
+	 * Get the singleton class object. Must only be called between Create and Shutdown calls otherwise the code will assert.
+	 * @return		reference the the downloader system
+	 */
+	static FBuildPatchDownloader& Get();
+
+	/**
+	 * Shuts down and cleans up the downloader system, deallocating immediately.
+	 */
+	static void Shutdown();
+
+private:
+
+	// The threadsafe shared pointer holding the singleton download system
+	static TSharedPtr< FBuildPatchDownloader, ESPMode::ThreadSafe > SingletonInstance;
+};
diff --git a/Engine/Source/Runtime/Online/BuildPatchServices/Private/BuildPatchError.cpp b/Engine/Source/Runtime/Online/BuildPatchServices/Private/BuildPatchError.cpp
new file mode 100644
--- /dev/null
+++ b/Engine/Source/Runtime/Online/BuildPatchServices/Private/BuildPatchError.cpp
@@ -0,0 +1,149 @@
+// Copyright 1998-2014 Epic Games, Inc. All Rights Reserved.
+
+/*=============================================================================
+	BuildPatchError.cpp: Implements classes involved setting and getting error information.
+=============================================================================*/
+
+#include "BuildPatchServicesPrivatePCH.h"
+
+#define LOCTEXT_NAMESPACE "BuildPatchInstallError"
+
+/* FBuildPatchInstallError implementation
+*****************************************************************************/
+void FBuildPatchInstallError::Reset()
+{
+	FScopeLock ScopeLock( &ThreadLock );
+	ErrorState = EBuildPatchInstallError::NoError;
+}
+
+bool FBuildPatchInstallError::HasFatalError()
+{
+	FScopeLock ScopeLock( &ThreadLock );
+	return ErrorState != EBuildPatchInstallError::NoError;
+}
+
+EBuildPatchInstallError::Type FBuildPatchInstallError::GetErrorState()
+{
+	FScopeLock ScopeLock( &ThreadLock );
+	return ErrorState;
+}
+
+bool FBuildPatchInstallError::IsInstallationCancelled()
+{
+	FScopeLock ScopeLock( &ThreadLock );
+	return ErrorState == EBuildPatchInstallError::UserCanceled || ErrorState == EBuildPatchInstallError::ApplicationClosing;
+}
+
+bool FBuildPatchInstallError::IsNoRetryError()
+{
+	FScopeLock ScopeLock( &ThreadLock );
+	return ErrorState == EBuildPatchInstallError::DownloadError || ErrorState == EBuildPatchInstallError::MoveFileToInstall;
+}
+
+FString FBuildPatchInstallError::GetErrorString()
+{
+	FScopeLock ScopeLock( &ThreadLock );
+	FString ReturnString = ToString();
+	if ( HasFatalError() )
+	{
+		ReturnString += TEXT( " " );
+		ReturnString += ErrorString;
+	}
+	return ReturnString;
+}
+
+const FText& FBuildPatchInstallError::GetErrorText()
+{
+	static const FText NoError( LOCTEXT( "BuildPatchInstallError_NoError", "The operation was successful." ) );
+	static const FText DownloadError( LOCTEXT( "BuildPatchInstallError_DownloadError", "Could not download patch data. Please try again later." ) );
+	static const FText FileConstructionFail( LOCTEXT( "BuildPatchInstallError_FileConstructionFail", "A file corruption has occurred. Please try again." ) );
+	static const FText MoveFileToInstall( LOCTEXT( "BuildPatchInstallError_MoveFileToInstall", "A file access error has occurred. Please check your running processes." ) );
+	static const FText BuildVerifyFail( LOCTEXT( "BuildPatchInstallError_BuildCorrupt", "The installation is corrupt." ) );
+	static const FText ApplicationClosing( LOCTEXT( "BuildPatchInstallError_ApplicationClosing", "The application is closing." ) );
+	static const FText ApplicationError( LOCTEXT( "BuildPatchInstallError_ApplicationError", "Patching service could not start." ) );
+	static const FText UserCanceled( LOCTEXT( "BuildPatchInstallError_UserCanceled", "User cancelled." ) );
+	static const FText PrerequisiteError( LOCTEXT( "BuildPatchInstallError_PrerequisiteError", "Prerequisites install failed.") );
+	static const FText InvalidOrMax( LOCTEXT( "BuildPatchInstallError_InvalidOrMax", "An unknown error ocurred." ) );
+
+	FScopeLock ScopeLock( &ThreadLock );
+	switch( ErrorState )
+	{
+		case EBuildPatchInstallError::NoError: return NoError;
+		case EBuildPatchInstallError::DownloadError: return DownloadError;
+		case EBuildPatchInstallError::FileConstructionFail: return FileConstructionFail;
+		case EBuildPatchInstallError::MoveFileToInstall: return MoveFileToInstall;
+		case EBuildPatchInstallError::BuildVerifyFail: return BuildVerifyFail;
+		case EBuildPatchInstallError::ApplicationClosing: return ApplicationClosing;
+		case EBuildPatchInstallError::ApplicationError: return ApplicationError;
+		case EBuildPatchInstallError::UserCanceled: return UserCanceled;
+		case EBuildPatchInstallError::PrerequisiteError: return PrerequisiteError;
+		default: return InvalidOrMax;
+	}
+}
+
+void FBuildPatchInstallError::SetFatalError( const EBuildPatchInstallError::Type& ErrorType, const FString& ErrorLog )
+{
+	FScopeLock ScopeLock( &ThreadLock );
+	// Only accept the first error
+	if( HasFatalError() )
+	{
+		return;
+	}
+	if( ErrorType != EBuildPatchInstallError::NoError)
+	{
+		// Log first as the fatal error
+		GWarn->Logf( TEXT( "BuildDataGenerator: FATAL ERROR: %s, %s" ), *ToString( ErrorType ), *ErrorLog );
+	}
+	else
+	{
+		// Log any that follow, but they are likely caused by the first
+		GWarn->Logf( TEXT( "BuildDataGenerator: Secondary error: %s, %s" ), *ToString( ErrorType ), *ErrorLog );
+	}
+	ErrorState = ErrorType;
+	ErrorString = ErrorLog;
+}
+
+const FString& FBuildPatchInstallError::ToString( const EBuildPatchInstallError::Type& ErrorType )
+{
+	// Const enum strings, special case no error.
+	static const FString NoError( "SUCCESS" );
+	static const FString DownloadError( "EBuildPatchInstallError::DownloadError" );
+	static const FString FileConstructionFail( "EBuildPatchInstallError::FileConstructionFail" );
+	static const FString MoveFileToInstall( "EBuildPatchInstallError::MoveFileToInstall" );
+	static const FString BuildVerifyFail( "EBuildPatchInstallError::BuildVerifyFail" );
+	static const FString ApplicationClosing( "EBuildPatchInstallError::ApplicationClosing" );
+	static const FString ApplicationError( "EBuildPatchInstallError::ApplicationError" );
+	static const FString UserCanceled( "EBuildPatchInstallError::UserCanceled" );
+	static const FString PrerequisiteError( "EBuildPatchInstallError::PrerequisiteError" );
+	static const FString InvalidOrMax( "EBuildPatchInstallError::InvalidOrMax" );
+
+	FScopeLock ScopeLock( &ThreadLock );
+	switch( ErrorType )
+	{
+		case EBuildPatchInstallError::NoError: return NoError;
+		case EBuildPatchInstallError::DownloadError: return DownloadError;
+		case EBuildPatchInstallError::FileConstructionFail: return FileConstructionFail;
+		case EBuildPatchInstallError::MoveFileToInstall: return MoveFileToInstall;
+		case EBuildPatchInstallError::BuildVerifyFail: return BuildVerifyFail;
+		case EBuildPatchInstallError::ApplicationClosing: return ApplicationClosing;
+		case EBuildPatchInstallError::ApplicationError: return ApplicationError;
+		case EBuildPatchInstallError::UserCanceled: return UserCanceled;
+		case EBuildPatchInstallError::PrerequisiteError: return PrerequisiteError;	
+		default: return InvalidOrMax;
+	}
+}
+
+const FString& FBuildPatchInstallError::ToString()
+{
+	return ToString( GetErrorState() );
+}
+
+/**
+ * Static FBuildPatchInstallError variables
+ */
+FCriticalSection FBuildPatchInstallError::ThreadLock;
+EBuildPatchInstallError::Type FBuildPatchInstallError::ErrorState = EBuildPatchInstallError::NoError;
+FString FBuildPatchInstallError::ErrorString;
+
+#undef LOCTEXT_NAMESPACE
+
diff --git a/Engine/Source/Runtime/Online/BuildPatchServices/Private/BuildPatchError.h b/Engine/Source/Runtime/Online/BuildPatchServices/Private/BuildPatchError.h
new file mode 100644
--- /dev/null
+++ b/Engine/Source/Runtime/Online/BuildPatchServices/Private/BuildPatchError.h
@@ -0,0 +1,129 @@
+// Copyright 1998-2014 Epic Games, Inc. All Rights Reserved.
+
+/*=============================================================================
+	BuildPatchError.h: Declares classes involved setting and getting error information.
+=============================================================================*/
+
+#ifndef __BuildPatchError_h__
+#define __BuildPatchError_h__
+
+#pragma once
+
+/**
+ * Declares the error type enum for use with the fatal error system
+ */
+namespace EBuildPatchInstallError
+{
+	enum Type
+	{
+		// There has been no registered error
+		NoError = 0,
+
+		// A download registered a fatal error
+		DownloadError = 1,
+
+		// A file failed to construct properly
+		FileConstructionFail = 2,
+
+		// An error occurred trying to move the file to the install location
+		MoveFileToInstall = 3,
+
+		// The installed build failed to verify
+		BuildVerifyFail = 4,
+
+		// The user or some process has closed the application
+		ApplicationClosing = 5,
+
+		// An application error, such as module fail to load.
+		ApplicationError = 6,
+
+		// User canceled download
+		UserCanceled = 7,
+
+		// A prerequisites installer failed
+		PrerequisiteError = 8
+	};
+};
+
+/**
+ * A struct to hold static fatal error flagging
+ */
+struct FBuildPatchInstallError
+{
+
+private:
+	// Use thread safe access to the variables
+	static FCriticalSection ThreadLock;
+
+	// The current error state
+	static EBuildPatchInstallError::Type ErrorState;
+
+	// The current error string
+	static FString ErrorString;
+
+public:
+	/**
+	 * Static function to reset the state
+	 */
+	static void Reset();
+
+	/**
+	 * Static function to get if there has been a fatal error
+	 * @return True if a fatal error has been set
+	 */
+	static bool HasFatalError();
+
+	/**
+	 * Static function to get if error state indicates installation should cancel
+	 * @return True if the installation should cancel without retry
+	 */
+	static bool IsInstallationCancelled();
+
+	/**
+	 * Static function to get if error state indicates a network related error. It's not helpful to keep retrying
+	 * the installation stage if this is the case.
+	 * @return True if the installation should cancel without retry
+	 */
+	static bool IsNoRetryError();
+
+	/**
+	 * Static function to get fatal error
+	 * @return The fatal error type
+	 */
+	static EBuildPatchInstallError::Type GetErrorState();
+
+	/**
+	 * Static function to get the error string for logging
+	 * @return The error string
+	 */
+	static FString GetErrorString();
+
+	/**
+	 * Static function to get the error text for UI
+	 * @return The error text
+	 */
+	static const FText& GetErrorText();
+
+	/**
+	 * Static function allowing any worker to set fatal error
+	 * @param ErrorType		The error type value.
+	 * @param ErrorLog		An optional string with extra information.
+	 */
+	static void SetFatalError( const EBuildPatchInstallError::Type& ErrorType, const FString& ErrorLog = TEXT( "" ) );
+
+private:
+	/**
+	 * Returns the string representation of the specified FBuildPatchInstallError value. Used for logging only.
+	 * @param ErrorType - The error type value.
+	 * @return A string value.
+	 */
+	static const FString& ToString( const EBuildPatchInstallError::Type& ErrorType );
+
+	/**
+	 * Returns the string representation of the current FBuildPatchInstallError value. Used for logging only.
+	 * @return A string value.
+	 */
+	static const FString& ToString();
+};
+
+#endif // __BuildPatchChunk_h__
diff --git a/Engine/Source/Runtime/Online/BuildPatchServices/Private/BuildPatchFileConstructor.cpp b/Engine/Source/Runtime/Online/BuildPatchServices/Private/BuildPatchFileConstructor.cpp
new file mode 100644
--- /dev/null
+++ b/Engine/Source/Runtime/Online/BuildPatchServices/Private/BuildPatchFileConstructor.cpp
@@ -0,0 +1,630 @@
+// Copyright 1998-2014 Epic Games, Inc. All Rights Reserved.
+
+/*=============================================================================
+	BuildPatchFileConstructor.cpp: Implements the BuildPatchFileConstructor class
+	that handles creating files in a manifest from the chunks that make it.
+=============================================================================*/
+
+#include "BuildPatchServicesPrivatePCH.h"
+
+// This define the number of bytes on a half-finished file that we ignore from the end
+// incase of previous partial write.
+#define NUM_BYTES_RESUME_IGNORE     1024
+
+/**
+ * This struct handles loading and saving of simple resume information, that will allow us to decide which
+ * files should be resumed from. It will also check that we are creating the same version and app as we expect to be.
+ */
+struct FResumeData
+{
+public:
+	// Save the staging directory
+	const FString StagingDir;
+
+	// The filename to the resume data information
+	const FString ResumeDataFile;
+
+	// A string determining the app and version we are installing
+	const FString PatchVersion;
+
+	// An array of files that were started
+	TArray< FString > FilesStarted;
+
+	// An array of files that were completed, determined by expected filesize
+	TArray< FString > FilesCompleted;
+
+	// The manifest for the app we are installing
+	FBuildPatchAppManifestRef BuildManifest;
+
+	// Whether we have resume data for this install
+	bool bHasResumeData;
+
+public:
+	/**
+	 * Constructor - reads in the resume data
+	 * @param InStagingDir      The install staging directory
+	 * @param InBuildManifest   The manifest we are installing from
+	 */
+	FResumeData( const FString& InStagingDir, const FBuildPatchAppManifestRef& InBuildManifest )
+		: StagingDir( InStagingDir )
+		, ResumeDataFile( InStagingDir / TEXT( "$resumeData" ) )
+		, PatchVersion( InBuildManifest->GetAppName() + InBuildManifest->GetVersionString() )
+		, BuildManifest( InBuildManifest )
+	{
+		// Load data from previous resume file
+		bHasResumeData = FPaths::FileExists( ResumeDataFile );
+		if( bHasResumeData )
+		{
+			FString PrevResumeData;
+			TArray< FString > ResumeDataLines;
+			FFileHelper::LoadFileToString( PrevResumeData, *ResumeDataFile );
+			PrevResumeData.ParseIntoArray( &ResumeDataLines, TEXT( "\n" ), true );
+			// Line 1 will be the previously attempted version
+			bHasResumeData = ( ResumeDataLines.Num() >= 1 ) && ResumeDataLines[0] == PatchVersion;
+			if( bHasResumeData )
+			{
+				// The rest of the lines will be the filenames that were started
+				for( int32 LineIdx = 1; LineIdx < ResumeDataLines.Num(); ++LineIdx )
+				{
+					const FString& Filename = ResumeDataLines[ LineIdx ];
+					const FString FullFilename = InStagingDir / Filename;
+					FilesStarted.Add( Filename );
+				}
+			}
+		}
+	}
+
+	/**
+	 * Saves out the resume data
+	 */
+	void SaveOut()
+	{
+		// Save out the data
+		FString DataOut = PatchVersion + TEXT("\n");
+		for( auto FilesStartedIt = FilesStarted.CreateConstIterator(); FilesStartedIt; ++FilesStartedIt )
+		{
+			const FString& Filename = *FilesStartedIt;
+			DataOut += Filename  + TEXT("\n");
+		}
+		FFileHelper::SaveStringToFile( DataOut, *ResumeDataFile );
+	}
+
+	/**
+	 * Checks whether the file was completed during last install attempt and adds it to FilesCompleted if so
+	 * @param Filename    The filename to check
+	 */
+	void CheckFile( const FString& Filename )
+	{
+		// If we had resume data, check file size is correct
+		if( bHasResumeData && FilesStarted.Contains( Filename ) )
+		{
+			const FString FullFilename = StagingDir / Filename;
+			const int64 DiskFileSize = IFileManager::Get().FileSize( *FullFilename );
+			const int64 CompleteFileSize = BuildManifest->GetFileSize( Filename );
+			if( DiskFileSize == CompleteFileSize )
+			{
+				FilesCompleted.Add( Filename );
+			}
+			// Sanity check, if file is larger than we expect, that's bad
+			else if( DiskFileSize > CompleteFileSize )
+			{
+				FilesStarted.Remove( Filename );
+			}
+		}
+	}
+};
+
+/* FBuildPatchFileConstructor implementation
+ *****************************************************************************/
+FBuildPatchFileConstructor::FBuildPatchFileConstructor( FBuildPatchAppManifestPtr InInstalledManifest, FBuildPatchAppManifestRef InBuildManifest, const FString& InInstallDirectory, const FString& InStageDirectory, const TArray< FString >& InConstructList, FBuildPatchProgress* InBuildProgress )
+	: Thread( NULL )
+	, bIsRunning( false )
+	, bIsInited( false )
+	, bIsDownloadStarted( false )
+	, ThreadLock()
+	, InstalledManifest( InInstalledManifest )
+	, BuildManifest( InBuildManifest )
+	, BuildProgress( InBuildProgress )
+	, InstallDirectory( InInstallDirectory )
+	, StagingDirectory( InStageDirectory )
+	, TotalJobSize( 0 )
+	, ByteProcessed( 0 )
+	, FilesToConstruct( InConstructList )
+{
+	// Init progress
+	BuildProgress->SetStateProgress( EBuildPatchProgress::Installing, 0.0f );
+	// Count initial job size
+	for( auto ConstructIt = InConstructList.CreateConstIterator(); ConstructIt; ++ConstructIt )
+	{
+		TotalJobSize += InBuildManifest->GetFileSize( *ConstructIt );
+	}
+	// Start thread!
+	const TCHAR* ThreadName = TEXT( "FileConstructorThread" );
+	Thread = FRunnableThread::Create(this, ThreadName);
+}
+
+FBuildPatchFileConstructor::~FBuildPatchFileConstructor()
+{
+	// Wait for and deallocate the thread
+	if( Thread != NULL )
+	{
+		Thread->WaitForCompletion();
+		delete Thread;
+		Thread = NULL;
+	}
+}
+
+bool FBuildPatchFileConstructor::Init()
+{
+	// We are ready to go if our delegates are bound and directories successfully created
+	bool bReady = IFileManager::Get().DirectoryExists( *StagingDirectory );
+	bReady = bReady && IFileManager::Get().DirectoryExists( *InstallDirectory );
+	return bReady;
+}
+
+uint32 FBuildPatchFileConstructor::Run()
+{
+	SetRunning( true );
+	SetInited( true );
+	const bool bIsFileData = BuildManifest->IsFileDataManifest();
+
+	// Save the list of completed files
+	TArray< FString > ConstructedFiles;
+
+	// Check for resume data
+	FResumeData ResumeData( StagingDirectory, BuildManifest );
+
+	// Start resume progress at zero or one
+	BuildProgress->SetStateProgress( EBuildPatchProgress::Resuming, ResumeData.bHasResumeData ? 0.0f : 1.0f );
+
+	// While we have files to construct, run
+	FString FileToConstruct;
+	while( GetFileToConstruct( FileToConstruct ) && !FBuildPatchInstallError::HasFatalError() )
+	{
+		// Check resume status, currently we are only supporting sequential resume, so once we start downloading, we can't resume any more.
+		// this only comes up if the resume data has been changed externally.
+		ResumeData.CheckFile( FileToConstruct );
+		const bool bFilePreviouslyComplete = !bIsDownloadStarted && ResumeData.FilesCompleted.Contains( FileToConstruct );
+		const bool bFilePreviouslyStarted = !bIsDownloadStarted && ResumeData.FilesStarted.Contains( FileToConstruct );
+
+		// Construct or skip the file
+		bool bFileSuccess;
+		if( bFilePreviouslyComplete )
+		{
+			bFileSuccess = true;
+			CountBytesProcessed( BuildManifest->GetFileSize( FileToConstruct ) );
+			// Inform the chunk cache of the chunk skip
+			if( !bIsFileData )
+			{
+				FBuildPatchChunkCache::Get().SkipFile( FileToConstruct );
+			}
+		}
+		else
+		{
+			bFileSuccess = ConstructFileFromChunks( FileToConstruct, bFilePreviouslyStarted );
+			// Add to resume data if successful or failure was not file construction fail
+			if( bFileSuccess || FBuildPatchInstallError::GetErrorState() != EBuildPatchInstallError::FileConstructionFail )
+			{
+				ResumeData.FilesStarted.AddUnique( FileToConstruct );
+			}
+		}
+
+		// If the file succeeded, add to lists
+		if( bFileSuccess )
+		{
+			ConstructedFiles.Add( FileToConstruct );
+		}
+		else
+		{
+			GWarn->Logf( TEXT( "BuildPatchServices: ERROR: Failed to construct file %s" ), *FPaths::GetCleanFilename( FileToConstruct ) );
+			FBuildPatchInstallError::SetFatalError( EBuildPatchInstallError::FileConstructionFail );
+		}
+
+		// Pause
+		BuildProgress->WaitWhilePaused();
+	}
+
+	// Save resume data
+	ResumeData.SaveOut();
+	BuildProgress->SetStateProgress(EBuildPatchProgress::Resuming, 1.0f);
+
+	// Set constructed files
+	ThreadLock.Lock();
+	FilesConstructed.Empty();
+	FilesConstructed.Append( ConstructedFiles );
+	ThreadLock.Unlock();
+
+	SetRunning( false );
+	return 0;
+}
+
+void FBuildPatchFileConstructor::Wait()
+{
+	if( Thread != NULL )
+	{
+		Thread->WaitForCompletion();
+	}
+}
+
+bool FBuildPatchFileConstructor::IsComplete()
+{
+	FScopeLock Lock( &ThreadLock );
+	return !bIsRunning && bIsInited;
+}
+
+void FBuildPatchFileConstructor::GetFilesConstructed( TArray< FString >& ConstructedFiles )
+{
+	FScopeLock Lock( &ThreadLock );
+	ConstructedFiles.Empty();
+	ConstructedFiles.Append( FilesConstructed );
+}
+
+void FBuildPatchFileConstructor::AddFileDataToInventory( const FGuid& FileGuid, const FString& Filename )
+{
+	FScopeLock Lock( &FileDataAvailabilityLock );
+	FileDataAvailability.Add( FileGuid, Filename );
+}
+
+void FBuildPatchFileConstructor::PurgeFileDataInventory()
+{
+	FScopeLock Lock( &FileDataAvailabilityLock );
+	FileDataAvailability.Empty();
+}
+
+const bool FBuildPatchFileConstructor::IsFileDataAvailable( const FGuid& FileGuid ) const
+{
+	// If we should be pausing, return false so that the construction waits
+	if( BuildProgress->GetPauseState() )
+	{
+		return false;
+	}
+	// If we have a fatal error, we can say true, which will cause the constructor to error out
+	bool bRtn = FBuildPatchInstallError::HasFatalError();
+	if ( bRtn == false )
+	{
+		FScopeLock Lock( &FileDataAvailabilityLock );
+		bRtn = FileDataAvailability.Contains( FileGuid );
+	}
+	return bRtn;
+}
+
+const FString FBuildPatchFileConstructor::GetFileDataFilename( const FGuid& ChunkGuid ) const
+{
+	FString rtn = TEXT( "" );
+	{
+		FScopeLock Lock( &FileDataAvailabilityLock );
+		rtn = FileDataAvailability.FindRef( ChunkGuid );
+	}
+	return rtn;
+}
+
+void FBuildPatchFileConstructor::SetRunning( bool bRunning )
+{
+	FScopeLock Lock( &ThreadLock );
+	bIsRunning = bRunning;
+}
+
+void FBuildPatchFileConstructor::SetInited( bool bInited )
+{
+	FScopeLock Lock( &ThreadLock );
+	bIsInited = bInited;
+}
+
+void FBuildPatchFileConstructor::CountBytesProcessed( const int64& ByteCount )
+{
+	ByteProcessed += ByteCount;
+	const double Total = TotalJobSize;
+	const double Current = ByteProcessed;
+	BuildProgress->SetStateProgress( EBuildPatchProgress::Installing, Current / Total );
+}
+
+bool FBuildPatchFileConstructor::GetFileToConstruct(FString& Filename)
+{
+	FScopeLock Lock(&ThreadLock);
+	const bool bFileAvailable = FilesToConstruct.IsValidIndex(0);
+	if (bFileAvailable)
+	{
+		Filename = FilesToConstruct[0];
+		FilesToConstruct.RemoveAt(0);
+	}
+	return bFileAvailable;
+}
+
+bool FBuildPatchFileConstructor::ConstructFileFromChunks( const FString& Filename, bool bResumeExisting )
+{
+	const bool bIsFileData = BuildManifest->IsFileDataManifest();
+	bResumeExisting = bResumeExisting && !bIsFileData;
+	bool bSuccess = true;
+	FString ErrorString;
+	FString NewFilename = StagingDirectory / Filename;
+
+	// Calculate the hash as we write the data
+	FSHA1 HashState;
+	FSHAHash HashValue;
+
+	// First make sure we can get the file manifest
+	TSharedPtr< FBuildPatchFileManifest > FileManifest = BuildManifest->GetFileManifest( Filename );
+	bSuccess = FileManifest.IsValid();
+	if( bSuccess )
+	{
+#if PLATFORM_MAC
+		if( !FileManifest->SymlinkTarget.IsEmpty() )
+		{
+			bSuccess = symlink(TCHAR_TO_UTF8(*FileManifest->SymlinkTarget), TCHAR_TO_UTF8(*NewFilename)) == 0;
+			return bSuccess;
+		}
+#endif
+
+		// Check for resuming of existing file
+		int64 StartPosition = 0;
+		int32 StartChunkPart = 0;
+		if( bResumeExisting )
+		{
+			// We have to read in the existing file so that the hash check can still be done.
+			FArchive* NewFileReader = IFileManager::Get().CreateFileReader( *NewFilename );
+			if( NewFileReader != NULL )
+			{
+				// Read buffer
+				uint8* ReadBuffer = new uint8[ FBuildPatchData::ChunkDataSize ];
+				// Reuse a certain amount of the file
+				StartPosition = FMath::Max<int64>( 0, NewFileReader->TotalSize() - NUM_BYTES_RESUME_IGNORE );
+				// We'll also find the correct chunkpart to start writing from
+				int64 ByteCounter = 0;
+				for( int32 ChunkPartIdx = StartChunkPart; ChunkPartIdx < FileManifest->FileChunkParts.Num() && !FBuildPatchInstallError::HasFatalError(); ++ChunkPartIdx )
+				{
+					const FChunkPart& ChunkPart = FileManifest->FileChunkParts[ ChunkPartIdx ];
+					const int64 NextBytePosition = ByteCounter + ChunkPart.Size;
+					if( NextBytePosition <= StartPosition )
+					{
+						// Read data for hash check
+						NewFileReader->Serialize( ReadBuffer, ChunkPart.Size );
+						HashState.Update( ReadBuffer, ChunkPart.Size );
+						// Count bytes read from file
+						ByteCounter = NextBytePosition;
+						// Set to resume from next chunk part
+						StartChunkPart = ChunkPartIdx + 1;
+						// Inform the chunk cache of the chunk part skip
+						FBuildPatchChunkCache::Get().SkipChunkPart( ChunkPart );
+						// Wait if paused
+						BuildProgress->WaitWhilePaused();
+					}
+					else
+					{
+						// No more parts on disk
+						break;
+					}
+				}
+				// Set start position to the byte we got up to
+				StartPosition = ByteCounter;
+				// Clean read buffer
+				delete[] ReadBuffer;
+				// Close file
+				NewFileReader->Close();
+				delete NewFileReader;
+			}
+		}
+
+		// Now we can make sure the chunk cache knows to start downloading chunks
+		if( !bIsFileData && !bIsDownloadStarted && !FBuildPatchInstallError::HasFatalError() )
+		{
+			bIsDownloadStarted = true;
+			FBuildPatchChunkCache::Get().BeginDownloads();
+		}
+
+		// Attempt to create the file
+		FArchive* NewFile = IFileManager::Get().CreateFileWriter( *NewFilename, bResumeExisting ? EFileWrite::FILEWRITE_Append : 0 );
+		bSuccess = NewFile != NULL;
+		if( bSuccess )
+		{
+			// Whenever we start writing again, there's no more resuming to be done
+			BuildProgress->SetStateProgress( EBuildPatchProgress::Resuming, 1.0f );
+
+			// Seek to file write position
+			NewFile->Seek( StartPosition );
+
+			// For each chunk, load it, and place it's data into the file
+			for( int32 ChunkPartIdx = StartChunkPart; ChunkPartIdx < FileManifest->FileChunkParts.Num() && bSuccess && !FBuildPatchInstallError::HasFatalError(); ++ChunkPartIdx )
+			{
+				const FChunkPart& ChunkPart = FileManifest->FileChunkParts[ ChunkPartIdx ];
+				if( bIsFileData )
+				{
+					bSuccess = InsertFileData( ChunkPart, *NewFile, HashState );
+				}
+				else
+				{
+					bSuccess = InsertChunkData( ChunkPart, *NewFile, HashState );
+				}
+				if( bSuccess )
+				{
+					CountBytesProcessed( ChunkPart.Size );
+					// Wait if paused
+					BuildProgress->WaitWhilePaused();
+				}
+				else
+				{
+					ErrorString = TEXT( "Failed to construct file " );
+					ErrorString += Filename;
+					ErrorString += TEXT( " because of chunk " );
+					ErrorString += ChunkPart.Guid.ToString();
+					GWarn->Logf( TEXT( "BuildPatchFileConstructor: ERROR: %s" ), *ErrorString );
+					FBuildPatchInstallError::SetFatalError( EBuildPatchInstallError::FileConstructionFail, ErrorString );
+				}
+			}
+
+			// Close the file writer
+			NewFile->Close();
+			delete NewFile;
+		}
+		else
+		{
+			FBuildPatchAnalytics::RecordConstructionError( Filename, FPlatformMisc::GetLastError(), TEXT( "Could Not Create File" ) );
+			ErrorString = TEXT( "Could not create new file " );
+			ErrorString += Filename;
+			GWarn->Logf( TEXT( "BuildPatchFileConstructor: ERROR: %s" ), *ErrorString );
+			FBuildPatchInstallError::SetFatalError( EBuildPatchInstallError::FileConstructionFail, ErrorString );
+		}
+	}
+	else
+	{
+		FBuildPatchAnalytics::RecordConstructionError( Filename, INDEX_NONE, TEXT( "Missing File Manifest" ) );
+		ErrorString = TEXT( "Build manifest does not contain a file manifest for " );
+		ErrorString += Filename;
+		FBuildPatchInstallError::SetFatalError( EBuildPatchInstallError::FileConstructionFail, ErrorString );
+	}
+
+	// Verify the hash for the file that we created
+	if( bSuccess )
+	{
+		HashState.Final();
+		HashState.GetHash( HashValue.Hash );
+		bSuccess = HashValue == FileManifest->FileHash;
+		if( !bSuccess )
+		{
+			FBuildPatchAnalytics::RecordConstructionError( Filename, INDEX_NONE, TEXT( "Serialised Verify Fail" ) );
+			ErrorString = TEXT( "Verify failed after constructing file " );
+			ErrorString += Filename;
+			GWarn->Logf( TEXT( "BuildDataGenerator: ERROR: %s" ), *ErrorString );
+			FBuildPatchInstallError::SetFatalError( EBuildPatchInstallError::FileConstructionFail, ErrorString );
+		}
+	}
+
+#if PLATFORM_MAC
+	if( bSuccess && FileManifest->bIsUnixExecutable )
+	{
+		// Enable executable permission bit
+		struct stat FileInfo;
+		if (stat(TCHAR_TO_UTF8(*NewFilename), &FileInfo) == 0)
+		{
+			bSuccess = chmod(TCHAR_TO_UTF8(*NewFilename), FileInfo.st_mode | S_IXUSR | S_IXGRP | S_IXOTH) == 0;
+		}
+	}
+#endif
+	
+	// Delete the staging file if unsuccessful by means of construction fail (i.e. keep if canceled or download issue)
+	if( !bSuccess && FBuildPatchInstallError::GetErrorState() == EBuildPatchInstallError::FileConstructionFail )
+	{
+		IFileManager::Get().Delete( *NewFilename, false, true );
+	}
+
+	return bSuccess;
+}
+
+bool FBuildPatchFileConstructor::InsertFileData( const FChunkPart& ChunkPart, FArchive& DestinationFile, FSHA1& HashState )
+{
+	bool bSuccess = false;
+
+	// Wait for the file data to be available
+	while( IsFileDataAvailable( ChunkPart.Guid ) == false )
+	{
+		FPlatformProcess::Sleep( 0.1f );
+	}
+	FString Filename = GetFileDataFilename( ChunkPart.Guid );
+
+	// Read the file
+	FArchive* FileReader = IFileManager::Get().CreateFileReader( *Filename );
+	bSuccess = FileReader != NULL;
+	if( bSuccess )
+	{
+		const int64 FileSize = FileReader->TotalSize();
+		// Read the header
+		FChunkHeader Header;
+		*FileReader << Header;
+		// Check header magic
+		if ( !Header.IsValidMagic() )
+		{
+			bSuccess = false;
+			FBuildPatchAnalytics::RecordConstructionError( Filename, INDEX_NONE, TEXT( "File Data Magic Fail" ) );
+			GLog->Logf( TEXT( "BuildPatchFileConstructor: ERROR: InsertFileData: magic failed for %s" ), *FPaths::GetCleanFilename( Filename ) );
+		}
+		// Check we have the right kind of hash; file data must have SHA1
+		if ( bSuccess && ( Header.HashType != FChunkHeader::HASH_SHA1 ) )
+		{
+			bSuccess = false;
+			FBuildPatchAnalytics::RecordConstructionError( Filename, INDEX_NONE, TEXT( "File Data Incorrect Hash" ) );
+			GLog->Logf( TEXT( "BuildPatchFileConstructor: ERROR: InsertFileData: incorrect hashtype info for %s" ), *FPaths::GetCleanFilename( Filename ) );
+		}
+		// Check GUID matches if provided
+		if ( bSuccess && ChunkPart.Guid.IsValid() && ( ChunkPart.Guid != Header.Guid ) )
+		{
+			bSuccess = false;
+			FBuildPatchAnalytics::RecordConstructionError( Filename, INDEX_NONE, TEXT( "File Data Mismatch GUID" ) );
+			GLog->Logf( TEXT( "BuildPatchFileConstructor: ERROR: InsertFileData: mismatch guid for %s" ), *FPaths::GetCleanFilename( Filename ) );
+		}
+		// Continue if all was fine
+		if( bSuccess )
+		{
+			// Load the data to check
+			uint8* FileReadBuffer = new uint8[ FileBufferSize ];
+			int64 DataRead = 0;
+			switch ( Header.StoredAs )
+			{
+			case FChunkHeader::STORED_RAW:
+				{
+					// Check we are able to get the chunk part
+					const int64 StartOfPartPos = Header.HeaderSize + ChunkPart.Offset;
+					const int64 EndOfPartPos = StartOfPartPos + ChunkPart.Size;
+					bSuccess = EndOfPartPos <= FileSize;
+					if ( !bSuccess )
+					{
+						FBuildPatchAnalytics::RecordConstructionError( Filename, INDEX_NONE, TEXT( "File Data Part OOB" ) );
+						GLog->Logf( TEXT( "BuildPatchFileConstructor: ERROR: InsertFileData: part out of bounds for %s" ), *FPaths::GetCleanFilename( Filename ) );
+					}
+					else
+					{
+						// Seek to and copy the part across
+						FileReader->Seek( StartOfPartPos );
+						while( DataRead < ChunkPart.Size )
+						{
+							const int64 SizeLeft = ChunkPart.Size - DataRead;
+							const int64 ReadLen = FMath::Min< int64 >( FileBufferSize, SizeLeft );
+							FileReader->Serialize( FileReadBuffer, ReadLen );
+							HashState.Update( FileReadBuffer, ReadLen );
+							DestinationFile.Serialize( FileReadBuffer, ReadLen );
+							DataRead += ReadLen;
+						}
+					}
+				}
+				break;
+			default:
+				check( false ); // @TODO LSwift: Implement other storage methods!
+				bSuccess = false;
+				break;
+			}
+			delete[] FileReadBuffer;
+		}
+
+		FileReader->Close();
+		delete FileReader;
+	}
+	else
+	{
+		FBuildPatchAnalytics::RecordConstructionError( Filename, FPlatformMisc::GetLastError(), TEXT( "Missing File Data" ) );
+		GLog->Logf( TEXT( "BuildPatchFileConstructor: ERROR: InsertFileData could not open data file %s" ), *Filename );
+	}
+	return bSuccess;
+}
+
+bool FBuildPatchFileConstructor::InsertChunkData( const FChunkPart& ChunkPart, FArchive& DestinationFile, FSHA1& HashState )
+{
+	uint8* Data;
+	uint8* DataStart;
+	FChunkFile* ChunkFile = FBuildPatchChunkCache::Get().GetChunkFile( ChunkPart.Guid );
+	if( ChunkFile != NULL && !FBuildPatchInstallError::HasFatalError() )
+	{
+		ChunkFile->GetDataLock( &Data, NULL );
+		DataStart = &Data[ ChunkPart.Offset ];
+		HashState.Update( DataStart, ChunkPart.Size );
+		DestinationFile.Serialize( DataStart, ChunkPart.Size );
+		ChunkFile->Dereference();
+		ChunkFile->ReleaseDataLock();
+		return true;
+	}
+	return false;
+}
+
+/**
+ * Static FBuildPatchFileConstructor variables
+ */
+TMap< FGuid, FString > FBuildPatchFileConstructor::FileDataAvailability;
+FCriticalSection FBuildPatchFileConstructor::FileDataAvailabilityLock;
+
diff --git a/Engine/Source/Runtime/Online/BuildPatchServices/Private/BuildPatchFileConstructor.h b/Engine/Source/Runtime/Online/BuildPatchServices/Private/BuildPatchFileConstructor.h
new file mode 100644
--- /dev/null
+++ b/Engine/Source/Runtime/Online/BuildPatchServices/Private/BuildPatchFileConstructor.h
@@ -0,0 +1,188 @@
+// Copyright 1998-2014 Epic Games, Inc. All Rights Reserved.
+
+/*=============================================================================
+	BuildPatchFileConstructor.h: Declares the BuildPatchFileConstructor class
+	that handles creating files in a manifest from the chunks that make it.
+=============================================================================*/
+
+// Forward declarations
+class FBuildPatchAppManifest;
+struct FChunkPart;
+
+/**
+ * FBuildPatchFileConstructor
+ * This class controls a thread that constructs files from a file list, given install details, and chunk availability notifications
+ */
+class FBuildPatchFileConstructor
+	: public FRunnable
+{
+private:
+	// Hold a pointer to my thread for easier deleting
+	FRunnableThread* Thread;
+
+	// A flag marking that we a running
+	bool bIsRunning;
+
+	// A flag marking that we initialized correctly
+	bool bIsInited;
+
+	// A flag marking that we told the chunk cache to queue required downloads
+	bool bIsDownloadStarted;
+
+	// A critical section to protect the flags and variables
+	FCriticalSection ThreadLock;
+
+	// The build manifest for the version already installed
+	FBuildPatchAppManifestPtr InstalledManifest;
+
+	// The build manifest for the app we are installing
+	FBuildPatchAppManifestRef BuildManifest;
+
+	// Pointer to the build progress tracker
+	FBuildPatchProgress* BuildProgress;
+
+	// The directory for the installed build
+	FString InstallDirectory;
+
+	// The directory for staging files
+	FString StagingDirectory;
+
+	// Total job size for tracking progress
+	int64 TotalJobSize;
+
+	// Byte processed so far for tracking progress
+	int64 ByteProcessed;
+
+	// A list of filenames for files that need to be constructed in this build
+	TArray< FString > FilesToConstruct;
+
+	// A list of filenames for files that have been constructed
+	TArray< FString > FilesConstructed;
+
+	// A static map of GUIDs to filenames, listing each piece of file data that has been downloaded
+	static TMap< FGuid, FString > FileDataAvailability;
+
+	// A static critical section to protect the file availability map
+	static FCriticalSection FileDataAvailabilityLock;
+
+public:
+
+	/**
+	 * Constructor
+	 * @param InInstalledManifest	The Manifest for the build that is currently install. Invalid for fresh installs.
+	 * @param InBuildManifest		The Manifest for the build we are installing
+	 * @param InInstallDirectory	The location for the installation.
+	 * @param InSaveDirectory		The location where we will store temporary files
+	 * @param InConstructList		The list of files to be constructed, filename paths should match those contained in manifest.
+	 * @param InBuildProgress		Pointer to the progress tracking class
+	 */
+	FBuildPatchFileConstructor( FBuildPatchAppManifestPtr InInstalledManifest, FBuildPatchAppManifestRef InBuildManifest, const FString& InInstallDirectory, const FString& InStageDirectory, const TArray< FString >& InConstructList, FBuildPatchProgress* InBuildProgress );
+
+	/**
+	 * Default Destructor, will delete the allocated Thread
+	 */
+	~FBuildPatchFileConstructor();
+
+	// Begin FRunnable
+	virtual bool Init() override;
+	virtual uint32 Run() override;
+	// End FRunnable
+
+	/**
+	 * Blocks the calling thread until this one has completed
+	 */
+	void Wait();
+
+	/**
+	 * Get whether the thread has finished working
+	 * @return	true if the thread completed
+	 */
+	bool IsComplete();
+	
+	/**
+	 * Get list of constructed files - only valid when complete
+	 * @param		OUT		Populated with the list of files
+	 */
+	void GetFilesConstructed( TArray< FString >& ConstructedFiles );
+
+	/**
+	 * Static function for registering a file download that has been successfully acquired
+	 * @param FileGuid		The GUID for the chunk
+	 * @param Filename		The filename to the saved file data
+	 */
+	static void AddFileDataToInventory( const FGuid& FileGuid, const FString& Filename );
+
+	/**
+	 * Static function to clear the inventory of available chunks
+	 */
+	static void PurgeFileDataInventory();
+
+private:
+
+	/**
+	 * Sets the bIsRunning flag
+	 * @param bRunning	Whether the thread is running
+	 */
+	void SetRunning( bool bRunning );
+
+	/**
+	 * Sets the bIsInited flag
+	 * @param bInited	Whether the thread successfully initialized
+	 */
+	void SetInited( bool bInited );
+
+	/**
+	 * Count additional bytes processed, and set new install progress value
+	 * @param ByteCount		Number of bytes to increment by
+	 */
+	void CountBytesProcessed( const int64& ByteCount );
+
+	/**
+	 * Function to fetch a chunk from the download list
+	 * @param Filename		Receives the filename for the file to construct from the manifest
+	 * @return true if there was a chunk guid in the list
+	 */
+	bool GetFileToConstruct( FString& Filename );
+
+	/**
+	 * Get whether a piece of file data is currently available
+	 * @param FileGuid		The Guid of the file data in question
+	 * @return	true if the file data can be loaded from disk
+	 */
+	const bool IsFileDataAvailable( const FGuid& FileGuid ) const;
+
+	/**
+	 * Gets the full path the the requested file data if available
+	 * @param FileGuid		The Guid of the file data in question
+	 * @return	the path to the file data, or empty string if unavailable
+	 */
+	const FString GetFileDataFilename( const FGuid& FileGuid ) const;
+
+	/**
+	 * Constructs a particular file referenced by the given BuildManifest. The function takes an interface to a class that can provide availability information of chunks so that this
+	 * file construction process can be ran alongside chunk acquisition threads. It will Sleep while waiting for chunks that it needs.
+	 * @param Filename			The Filename for the file to construct, that matches an entry in the BuildManifest.
+	 * @param bResumeExisting	Whether we should resume from an existing file
+	 * @return	true if no file errors occurred
+	 */
+	bool ConstructFileFromChunks( const FString& Filename, bool bResumeExisting );
+
+	/**
+	 * Loads the file data and inserts it into a destination file according to the chunk part info. The chunk part info is mostly the whole file, but is still used for splitting support
+	 * in future.
+	 * @param ChunkPart			The chunk part details.
+	 * @param DestinationFile	The archive for the file being constructed.
+	 * @param HashState			An FSHA1 hash state to update with the data going into the destination file.
+	 * @return	true if no file errors occurred
+	 */
+	bool InsertFileData( const FChunkPart& ChunkPart, FArchive& DestinationFile, FSHA1& HashState );
+
+	/**
+	 * Inserts the data data from a chunk into the destination file according to the chunk part info
+	 * @param ChunkPart			The chunk part details.
+	 * @param DestinationFile	The Filename for the file being constructed.
+	 * @param HashState			An FSHA1 hash state to update with the data going into the destination file.
+	 * @return true if no errors were detected
+	 */
+	bool InsertChunkData( const FChunkPart& ChunkPart, FArchive& DestinationFile, FSHA1& HashState );
+};
