Moving Build Patch Services from NoRedist to Runtime

diff --git a/Engine/Source/Runtime/Online/BuildPatchServices/Private/BuildPatchGeneration.cpp b/Engine/Source/Runtime/Online/BuildPatchServices/Private/BuildPatchGeneration.cpp
new file mode 100644
--- /dev/null
+++ b/Engine/Source/Runtime/Online/BuildPatchServices/Private/BuildPatchGeneration.cpp
@@ -0,0 +1,1793 @@
+// Copyright 1998-2014 Epic Games, Inc. All Rights Reserved.
+
+/*=============================================================================
+	BuildPatchGeneration.cpp: Implements the classes that control build
+	installation, and the generation of chunks and manifests from a build image.
+=============================================================================*/
+
+#include "BuildPatchServicesPrivatePCH.h"
+
+#if WITH_BUILDPATCHGENERATION
+
+/**
+ * Creates a hash value from a FSHAHash. This is so that a TMap can be keyed
+ * using a FSHAHash
+ * @param Hash	The sha hash
+ * @return The crc for the data
+ */
+FORCEINLINE uint32 GetTypeHash( const FSHAHash& Hash )
+{
+	return FCrc::MemCrc_DEPRECATED( &Hash, sizeof( FSHAHash ) );
+}
+
+#if PLATFORM_MAC
+static bool IsUnixExecutable(const TCHAR* Filename)
+{
+	struct stat FileInfo;
+	if (stat(TCHAR_TO_UTF8(Filename), &FileInfo) == 0)
+	{
+		return (FileInfo.st_mode & S_IXUSR) != 0;
+	}
+	return false;
+}
+
+static FString GetSymlinkTarget(const TCHAR* Filename)
+{
+	ANSICHAR SymlinkTarget[MAX_PATH] = { 0 };
+	if (readlink(TCHAR_TO_UTF8(Filename), SymlinkTarget, MAX_PATH) != -1)
+	{
+		return UTF8_TO_TCHAR(SymlinkTarget);
+	}
+	return TEXT("");
+}
+#endif
+
+/* FBuildStreamReader implementation
+*****************************************************************************/
+FBuildStream::FBuildStreamReader::FBuildStreamReader()
+	: DepotDirectory(TEXT(""))
+	, BuildStream(NULL)
+	, Thread(NULL)
+{
+}
+
+FBuildStream::FBuildStreamReader::~FBuildStreamReader()
+{
+	if (Thread)
+	{
+		Thread->Kill();
+		delete Thread;
+		Thread = NULL;
+	}
+}
+
+bool FBuildStream::FBuildStreamReader::Init() 
+{
+	return BuildStream != NULL && IFileManager::Get().DirectoryExists( *DepotDirectory );
+}
+
+uint32 FBuildStream::FBuildStreamReader::Run()
+{
+	// Clear the build stream
+	BuildStream->Clear();
+
+	TArray< FString > AllFiles;
+	IFileManager::Get().FindFilesRecursive( AllFiles, *DepotDirectory, TEXT("*.*"), true, false );
+	AllFiles.Sort();
+
+	// Remove the files that appear in an ignore list
+	FBuildDataGenerator::StripIgnoredFiles( AllFiles, DepotDirectory, IgnoreListFile );
+
+	// Allocate our file read buffer
+	uint8* FileReadBuffer = new uint8[ FileBufferSize ];
+
+	for (auto& SourceFile : AllFiles)
+	{
+		// Read the file
+		FArchive* FileReader = IFileManager::Get().CreateFileReader( *SourceFile );
+		const bool bBuildFileOpenSuccess = FileReader != NULL;
+		if( bBuildFileOpenSuccess )
+		{
+			// Make SourceFile the format we want it in and start a new file
+			FPaths::MakePathRelativeTo( SourceFile, *( DepotDirectory + TEXT( "/" ) ) );
+			int64 FileSize = FileReader->TotalSize();
+			// Process files that have bytes
+			if( FileSize > 0 )
+			{
+				BuildStream->BeginNewFile( SourceFile, FileSize );
+				while( !FileReader->AtEnd() )
+				{
+					const int64 SizeLeft = FileSize - FileReader->Tell();
+					const uint32 ReadLen = FMath::Min< int64 >( FileBufferSize, SizeLeft );
+					FileReader->Serialize( FileReadBuffer, ReadLen );
+					// Copy into data stream
+					BuildStream->EnqueueData( FileReadBuffer, ReadLen );
+				}
+			}
+			// Special case zero byte files
+			else if( FileSize == 0 )
+			{
+				BuildStream->AddEmptyFile( SourceFile );
+			}
+			FileReader->Close();
+			delete FileReader;
+		}
+		else
+		{
+			// Not being able to load a required file from the build would be fatal, hard fault.
+			GLog->Logf( TEXT( "FBuildStreamReader: Could not open file from build! %s" ), *SourceFile );
+			GLog->PanicFlushThreadedLogs();
+			// Use bool variable for easier to understand assert message
+			check( bBuildFileOpenSuccess );
+		}
+	}
+
+	// Mark end of build
+	BuildStream->EndOfBuild();
+
+	// Deallocate our file read buffer
+	delete[] FileReadBuffer;
+
+	return 0;
+}
+
+void FBuildStream::FBuildStreamReader::StartThread()
+{
+	Thread = FRunnableThread::Create(this, TEXT("BuildStreamReaderThread"));
+}
+
+/* FBuildStream implementation
+*****************************************************************************/
+FBuildStream::FBuildStream( const FString& RootDirectory, const FString& IgnoreListFile )
+{
+	BuildStreamReader.DepotDirectory = RootDirectory;
+	BuildStreamReader.IgnoreListFile = IgnoreListFile;
+	BuildStreamReader.BuildStream = this;
+	BuildStreamReader.StartThread();
+}
+
+FBuildStream::~FBuildStream()
+{
+}
+
+bool FBuildStream::GetFileSpan( const uint64& StartingIdx, FString& Filename, uint64& FileSize )
+{
+	bool bFound = false;
+	FilesListsCS.Lock();
+	// Find the filename
+	FFileSpan* FileSpan = FilesParsed.Find( StartingIdx );
+	if( FileSpan != NULL )
+	{
+		Filename = FileSpan->Filename;
+		FileSize = FileSpan->Size;
+		bFound = true;
+	}
+	FilesListsCS.Unlock();
+	return bFound;
+}
+
+uint32 FBuildStream::DequeueData( uint8* Buffer, const uint32& ReqSize, const bool WaitForData )
+{
+	// Wait for data
+	if( WaitForData )
+	{
+		while ( DataAvailable() < ReqSize && !IsEndOfBuild() )
+		{
+			FPlatformProcess::Sleep( 0.01f );
+		}
+	}
+
+	BuildDataStreamCS.Lock();
+	uint32 ReadLen = FMath::Min< int64 >( ReqSize, DataAvailable() );
+	ReadLen = BuildDataStream.Dequeue( Buffer, ReadLen );
+	BuildDataStreamCS.Unlock();
+
+	return ReadLen;
+}
+
+const TArray< FString > FBuildStream::GetEmptyFiles()
+{
+	FScopeLock ScopeLock( &FilesListsCS );
+	return EmptyFiles;
+}
+
+bool FBuildStream::IsEndOfBuild()
+{
+	bool rtn;
+	NoMoreDataCS.Lock();
+	rtn = bNoMoreData;
+	NoMoreDataCS.Unlock();
+	return rtn;
+}
+
+bool FBuildStream::IsEndOfData()
+{
+	bool rtn;
+	NoMoreDataCS.Lock();
+	rtn = bNoMoreData;
+	NoMoreDataCS.Unlock();
+	BuildDataStreamCS.Lock();
+	rtn &= BuildDataStream.RingDataUsage() == 0;
+	BuildDataStreamCS.Unlock();
+	return rtn;
+}
+
+void FBuildStream::Clear()
+{
+	EndOfBuild( false );
+
+	BuildDataStreamCS.Lock();
+	BuildDataStream.Empty();
+	BuildDataStreamCS.Unlock();
+
+	FilesListsCS.Lock();
+	FilesParsed.Empty();
+	FilesListsCS.Unlock();
+}
+
+uint32 FBuildStream::SpaceLeft()
+{
+	uint32 rtn;
+	BuildDataStreamCS.Lock();
+	rtn = BuildDataStream.RingDataSize() - BuildDataStream.RingDataUsage();
+	BuildDataStreamCS.Unlock();
+	return rtn;
+}
+
+uint32 FBuildStream::DataAvailable()
+{
+	uint32 rtn;
+	BuildDataStreamCS.Lock();
+	rtn = BuildDataStream.RingDataUsage();
+	BuildDataStreamCS.Unlock();
+	return rtn;
+}
+
+void FBuildStream::BeginNewFile( const FString& Filename, const uint64& FileSize )
+{
+	FFileSpan FileSpan;
+	FileSpan.Filename = Filename;
+	FileSpan.Size = FileSize;
+	BuildDataStreamCS.Lock();
+	FileSpan.StartIdx = BuildDataStream.TotalDataPushed();
+	BuildDataStreamCS.Unlock();
+	FilesListsCS.Lock();
+	FilesParsed.Add( FileSpan.StartIdx, FileSpan );
+	FilesListsCS.Unlock();
+}
+
+void FBuildStream::AddEmptyFile( const FString& Filename )
+{
+	FilesListsCS.Lock();
+	EmptyFiles.Add( Filename );
+	FilesListsCS.Unlock();
+}
+
+void FBuildStream::EnqueueData( const uint8* Buffer, const uint32& Len )
+{
+	// Wait for space
+	while ( SpaceLeft() < Len )
+	{
+		FPlatformProcess::Sleep( 0.01f );
+	}
+	BuildDataStreamCS.Lock();
+	BuildDataStream.Enqueue( Buffer, Len );
+	BuildDataStreamCS.Unlock();
+}
+
+void FBuildStream::EndOfBuild( bool bIsEnd )
+{
+	NoMoreDataCS.Lock();
+	bNoMoreData = bIsEnd;
+	NoMoreDataCS.Unlock();
+}
+
+/* FBuildDataProcessor implementation
+*****************************************************************************/
+FBuildDataChunkProcessor::FBuildDataChunkProcessor( FBuildPatchAppManifestRef InBuildManifest, const FString& InBuildRoot )
+	: NumNewChunks( 0 )
+	, NumKnownChunks( 0 )
+	, BuildRoot( InBuildRoot )
+	, ChunkIsPushed( false )
+	, IsProcessingChunk( false )
+	, IsProcessingChunkPart( false )
+	, BackupProcessingChunk( false )
+	, BackupProcessingChunkPart( false )
+	, CurrentChunkBufferPos( 0 )
+	, BackupChunkBufferPos( 0 )
+	, IsProcessingFile( false )
+	, CurrentFile( NULL )
+	, ChunkWriter( FBuildPatchServicesModule::GetCloudDirectory() )
+{
+	BuildManifest = InBuildManifest;
+	CurrentChunkGuid.Invalidate();
+	BackupChunkGuid.Invalidate();
+	CurrentChunkBuffer = new uint8[ FBuildPatchData::ChunkDataSize ];
+}
+
+FBuildDataChunkProcessor::~FBuildDataChunkProcessor()
+{
+	delete[] CurrentChunkBuffer;
+}
+
+void FBuildDataChunkProcessor::BeginNewChunk( const bool& bZeroBuffer )
+{
+	check( IsProcessingChunk == false );
+	check( IsProcessingChunkPart == false );
+	IsProcessingChunk = true;
+
+	// NB: If you change this function, make sure you make required changes to 
+	// Backup of chunk if need be at the end of a recognized chunk
+
+	// Erase the chunk buffer data ready for new chunk parts
+	if( bZeroBuffer )
+	{
+		FMemory::Memzero( CurrentChunkBuffer, FBuildPatchData::ChunkDataSize );
+		CurrentChunkBufferPos = 0;
+	}
+
+	// Get a new GUID for this chunk for identification
+	CurrentChunkGuid = FGuid::NewGuid();
+}
+
+void FBuildDataChunkProcessor::BeginNewChunkPart()
+{
+	check( IsProcessingChunkPart == false );
+	check( IsProcessingChunk == true );
+	IsProcessingChunkPart = true;
+
+	// The current file should have a new chunk part setup
+	if( CurrentFile.IsValid() )
+	{
+		FChunkPart NewPart;
+		NewPart.Guid = CurrentChunkGuid;
+		NewPart.Offset = CurrentChunkBufferPos;
+		CurrentFile->FileChunkParts.Add( NewPart );
+	}
+}
+
+void FBuildDataChunkProcessor::EndNewChunkPart()
+{
+	check( IsProcessingChunkPart == true );
+	check( IsProcessingChunk == true );
+	IsProcessingChunkPart = false;
+
+	// Current file should have it's last chunk part updated with the size value finalized.
+	if( CurrentFile.IsValid() )
+	{
+		FChunkPart* ChunkPart = &CurrentFile->FileChunkParts.Top();
+		ChunkPart->Size = CurrentChunkBufferPos - ChunkPart->Offset;
+		check( ChunkPart->Size != 0 );
+	}
+}
+
+void FBuildDataChunkProcessor::EndNewChunk( const uint64& ChunkHash, const uint8* ChunkData, const FGuid& ChunkGuid )
+{
+	check( IsProcessingChunk == true );
+	check( IsProcessingChunkPart == false );
+	IsProcessingChunk = false;
+
+	// A bool that will state whether we should update and log chunking progress
+	bool bLogProgress = false;
+
+	// If the new chunk was recognized, then we will have got a different Guid through, so fix up guid for files using this chunk
+	if( CurrentChunkGuid != ChunkGuid )
+	{
+		bLogProgress = true;
+		++NumKnownChunks;
+		if( CurrentFile.IsValid() )
+		{
+			const uint32 NumFileChunkParts = CurrentFile->FileChunkParts.Num();
+			for( uint32 ChunkIdx = 0; ChunkIdx < NumFileChunkParts; ++ChunkIdx)
+			{
+				FChunkPart& ChunkPart = CurrentFile->FileChunkParts[ChunkIdx];
+				if( ChunkPart.Guid == CurrentChunkGuid )
+				{
+					ChunkPart.Guid = ChunkGuid;
+				}
+			}
+		}
+		for( auto FileIt = BuildManifest->FileManifestList.CreateConstIterator(); FileIt; ++FileIt)
+		{
+			TSharedPtr< FBuildPatchFileManifest > ManifestFile = FileIt.Value();
+
+			const uint32 NumFileChunkParts = ManifestFile->FileChunkParts.Num();
+			for( uint32 ChunkIdx = 0; ChunkIdx < NumFileChunkParts; ++ChunkIdx)
+			{
+				FChunkPart& ChunkPart = ManifestFile->FileChunkParts[ChunkIdx];
+				if( ChunkPart.Guid == CurrentChunkGuid )
+				{
+					ChunkPart.Guid = ChunkGuid;
+				}
+			}
+		}
+	}
+	// If the chunk is new, count it if we were passed data
+	else if( ChunkData )
+	{
+		bLogProgress = true;
+		++NumNewChunks;
+	}
+
+	// Always queue the chunk if passed data, it will be skipped automatically if existing already, and we need to make
+	// sure the latest version is saved out when recognising an older version
+	if( ChunkData )
+	{
+		ChunkWriter.QueueChunk( ChunkData, ChunkGuid, ChunkHash );
+	}
+
+	if( bLogProgress )
+	{
+		// Output to log for builder info
+		GLog->Logf( TEXT( "%s %s [%d:%d]" ), *BuildManifest->GetAppName(), *BuildManifest->GetVersionString(), NumNewChunks, NumKnownChunks );
+	}
+}
+
+void FBuildDataChunkProcessor::PushChunk()
+{
+	check( !ChunkIsPushed );
+	ChunkIsPushed = true;
+
+	// Then we can skip over this entire chunk using it for current files in processing
+	BackupProcessingChunk = IsProcessingChunk;
+	BackupProcessingChunkPart = IsProcessingChunkPart;
+	BackupChunkGuid = CurrentChunkGuid;
+	BackupChunkBufferPos = CurrentChunkBufferPos;
+	if( BackupProcessingChunkPart )
+	{
+		EndNewChunkPart();
+	}
+	if( BackupProcessingChunk )
+	{
+		// null ChunkData will stop the incomplete chunk being saved out
+		EndNewChunk( 0, NULL, CurrentChunkGuid );
+	}
+
+	// Start this chunk, must begin from 0
+	CurrentChunkBufferPos = 0;
+	BeginNewChunk( false );
+	BeginNewChunkPart();
+}
+
+void FBuildDataChunkProcessor::PopChunk( const uint64& ChunkHash, const uint8* ChunkData, const FGuid& ChunkGuid )
+{
+	check( ChunkIsPushed );
+	ChunkIsPushed = false;
+
+	EndNewChunkPart();
+	EndNewChunk( ChunkHash, ChunkData, ChunkGuid );
+
+	BuildManifest->ChunkHashList.Add( ChunkGuid, ChunkHash );
+	BuildManifest->DataGroupList.Add( ChunkGuid, FCrc::MemCrc32( &ChunkGuid, sizeof( FGuid ) ) % 100 );
+
+	// Backup data for previous part chunk
+	IsProcessingChunk = BackupProcessingChunk;
+	CurrentChunkGuid = BackupChunkGuid;
+	CurrentChunkBufferPos = BackupChunkBufferPos;
+}
+
+void FBuildDataChunkProcessor::FinalChunk()
+{
+	// If we're processing a chunk part, then we MUST be processing a chunk too
+	check(IsProcessingChunkPart ? IsProcessingChunk : true);
+
+	if( IsProcessingChunk )
+	{
+		// The there is no more data and the last file is processed.
+		// The final chunk should be finished
+		uint64 NewChunkHash = FRollingHash< FBuildPatchData::ChunkDataSize >::GetHashForDataSet( CurrentChunkBuffer );
+		FGuid ChunkGuid = CurrentChunkGuid;
+		FBuildDataGenerator::FindExistingChunkData( NewChunkHash, CurrentChunkBuffer, ChunkGuid );
+		if( IsProcessingChunkPart )
+		{
+			EndNewChunkPart();
+		}
+		EndNewChunk( NewChunkHash, CurrentChunkBuffer, ChunkGuid );
+		BuildManifest->ChunkHashList.Add( ChunkGuid, NewChunkHash );
+		BuildManifest->DataGroupList.Add( ChunkGuid, FCrc::MemCrc32( &ChunkGuid, sizeof( FGuid ) ) % 100 );
+	}
+
+	// Wait for the chunk writer to finish
+	ChunkWriter.NoMoreChunks();
+	ChunkWriter.WaitForThread();
+}
+
+void FBuildDataChunkProcessor::BeginFile( const FString& FileName )
+{
+	check( IsProcessingFile == false );
+	IsProcessingFile = true;
+
+	// We should start a new file, which begins with the current new chunk and current chunk offset
+	CurrentFile = MakeShareable( new FBuildPatchFileManifest() );
+	CurrentFile->Filename = FileName;
+#if PLATFORM_MAC
+	CurrentFile->bIsUnixExecutable = IsUnixExecutable(*(BuildRoot / FileName));
+	CurrentFile->SymlinkTarget = GetSymlinkTarget(*(BuildRoot / FileName));
+#endif
+
+	// Setup for current chunk part
+	if( IsProcessingChunkPart )
+	{
+		FChunkPart NewPart;
+		NewPart.Guid = CurrentChunkGuid;
+		NewPart.Offset = CurrentChunkBufferPos;
+		CurrentFile->FileChunkParts.Add( NewPart );
+	}
+}
+
+void FBuildDataChunkProcessor::EndFile()
+{
+	check( IsProcessingFile == true );
+	IsProcessingFile = false;
+
+	if( CurrentFile.IsValid() )
+	{
+		if( IsProcessingChunkPart )
+		{
+			// Current file should have it's last chunk part updated with the size value finalized.
+			FChunkPart* ChunkPart = &CurrentFile->FileChunkParts.Top();
+			ChunkPart->Size = CurrentChunkBufferPos - ChunkPart->Offset;
+		}
+
+		// And add this file to our list of file manifests
+		BuildManifest->FileManifestList.Add( CurrentFile->Filename, CurrentFile );
+	}
+
+	CurrentFile.Reset();
+}
+
+void FBuildDataChunkProcessor::SkipKnownByte( const uint8& NextByte, const bool& bStartOfFile, const bool& bEndOfFile, const FString& Filename )
+{
+	// Check for start of new file
+	if( bStartOfFile )
+	{
+		BeginFile( Filename );
+		FileHash.Reset();
+	}
+
+	// Increment position between start and end of file
+	++CurrentChunkBufferPos;
+	FileHash.Update( &NextByte, 1 );
+
+	// Check for end of file
+	if( bEndOfFile )
+	{
+		FileHash.Final();
+		FileHash.GetHash( CurrentFile->FileHash.Hash );
+		EndFile();
+	}
+}
+
+void FBuildDataChunkProcessor::ProcessNewByte( const uint8& NewByte, const bool& bStartOfFile, const bool& bEndOfFile, const FString& Filename )
+{
+	// If we finished a chunk, we begin a new chunk and new chunk part
+	if( !IsProcessingChunk )
+	{
+		BeginNewChunk();
+		BeginNewChunkPart();
+	}
+	// Or if we finished a chunk part (by recognizing a chunk hash), we begin a new one
+	else if( !IsProcessingChunkPart )
+	{
+		BeginNewChunkPart();
+	}
+
+	// Check for start of new file
+	if( bStartOfFile )
+	{
+		BeginFile( Filename );
+		FileHash.Reset();
+	}
+
+	// We add the old byte to our new chunk which will be used as part of the file it belonged to
+	CurrentChunkBuffer[ CurrentChunkBufferPos++ ] = NewByte;
+	FileHash.Update( &NewByte, 1 );
+
+	// Check for end of file
+	if( bEndOfFile )
+	{
+		FileHash.Final();
+		FileHash.GetHash( CurrentFile->FileHash.Hash );
+		EndFile();
+	}
+
+	// Do we have a full new chunk?
+	check( CurrentChunkBufferPos <= FBuildPatchData::ChunkDataSize );
+	if( CurrentChunkBufferPos == FBuildPatchData::ChunkDataSize )
+	{
+		uint64 NewChunkHash = FRollingHash< FBuildPatchData::ChunkDataSize >::GetHashForDataSet( CurrentChunkBuffer );
+		FGuid ChunkGuid = CurrentChunkGuid;
+		FBuildDataGenerator::FindExistingChunkData( NewChunkHash, CurrentChunkBuffer, ChunkGuid );
+		EndNewChunkPart();
+		EndNewChunk( NewChunkHash, CurrentChunkBuffer, ChunkGuid );
+		BuildManifest->ChunkHashList.Add( ChunkGuid, NewChunkHash );
+		BuildManifest->DataGroupList.Add( ChunkGuid, FCrc::MemCrc32( &ChunkGuid, sizeof( FGuid ) ) % 100 );
+	}
+}
+
+void FBuildDataChunkProcessor::GetChunkStats( uint32& OutNewFiles, uint32& OutKnownFiles )
+{
+	OutNewFiles = NumNewChunks;
+	OutKnownFiles = NumKnownChunks;
+}
+
+const TMap<FGuid, int64>& FBuildDataChunkProcessor::GetChunkFilesizes()
+{
+	ChunkWriter.GetChunkFilesizes(ChunkFileSizes);
+	return ChunkFileSizes;
+}
+
+/* FBuildDataFileProcessor implementation
+*****************************************************************************/
+FBuildDataFileProcessor::FBuildDataFileProcessor( FBuildPatchAppManifestRef InBuildManifest, const FString& InBuildRoot )
+	: NumNewFiles( 0 )
+	, NumKnownFiles( 0 )
+	, BuildManifest( InBuildManifest )
+	, BuildRoot( InBuildRoot )
+	, CurrentFile( NULL )
+	, FileHash()
+	, FileSize( 0 )
+	, IsProcessingFile( false )
+{
+}
+
+void FBuildDataFileProcessor::BeginFile( const FString& InFileName )
+{
+	check( !IsProcessingFile );
+	IsProcessingFile = true;
+
+	// Create the new file
+	CurrentFile = MakeShareable( new FBuildPatchFileManifest() );
+	CurrentFile->Filename = InFileName;
+#if PLATFORM_MAC
+	CurrentFile->bIsUnixExecutable = IsUnixExecutable(*(BuildRoot / InFileName));
+	CurrentFile->SymlinkTarget = GetSymlinkTarget(*(BuildRoot / InFileName));
+#endif
+
+	// Add the 'chunk part' which will just be the whole file
+	FChunkPart FilePart;
+	CurrentFile->FileChunkParts.Add( FilePart );
+
+	// Reset the hash and size counter
+	FileHash.Reset();
+	FileSize = 0;
+}
+
+void FBuildDataFileProcessor::ProcessFileData( const uint8* Data, const uint32& DataLen )
+{
+	check( IsProcessingFile );
+
+	// Update the hash
+	FileHash.Update( Data, DataLen );
+
+	// Count size
+	FileSize += DataLen;
+}
+
+void FBuildDataFileProcessor::EndFile()
+{
+	check( IsProcessingFile );
+	check( CurrentFile.IsValid() );
+	check( CurrentFile->FileChunkParts.Num() == 1 );
+	IsProcessingFile = false;
+
+	// Finalise the hash
+	FileHash.Final();
+	FileHash.GetHash( CurrentFile->FileHash.Hash );
+
+	// Use hash and full file name to find out if we have a file match
+	FString FullPath = FPaths::Combine( *BuildRoot, *CurrentFile->Filename );
+	FGuid FileGuid = FGuid::NewGuid();
+	bool bFoundSameFile = FBuildDataGenerator::FindExistingFileData( FullPath, CurrentFile->FileHash, FileGuid );
+
+	// Fill the 'chunk part' info
+	FChunkPart& FilePart = CurrentFile->FileChunkParts[0];
+	FilePart.Guid = FileGuid;
+	FilePart.Offset = 0;
+	FilePart.Size = FileSize;
+
+	// Add it to the manifest
+	BuildManifest->FileManifestList.Add( CurrentFile->Filename, CurrentFile );
+
+	// Set the data group
+	BuildManifest->DataGroupList.Add( FileGuid, FCrc::MemCrc32( &FileGuid, sizeof( FGuid ) ) % 100 );
+
+	// Always call save to ensure new versions exist when recognising old ones, the save will be skipped if required copies already exist.
+	FBuildDataGenerator::SaveOutFileData( FullPath, CurrentFile->FileHash, FileGuid );
+
+	// Count new/known stats
+	if( bFoundSameFile == false )
+	{
+		++NumNewFiles;
+	}
+	else
+	{
+		++NumKnownFiles;
+	}
+
+	// Output to log for builder info
+	GLog->Logf( TEXT( "%s %s [%d:%d]" ), *BuildManifest->GetAppName(), *BuildManifest->GetVersionString(), NumNewFiles, NumKnownFiles );
+}
+
+void FBuildDataFileProcessor::GetFileStats( uint32& OutNewFiles, uint32& OutKnownFiles )
+{
+	OutNewFiles = NumNewFiles;
+	OutKnownFiles = NumKnownFiles;
+}
+
+/* FBuildSimpleChunkCache::FChunkReader implementation
+*****************************************************************************/
+FBuildGenerationChunkCache::FChunkReader::FChunkReader( const FString& InChunkFilePath, TSharedRef< FChunkFile > InChunkFile, uint32* InBytesRead )
+	: ChunkFilePath( InChunkFilePath )
+	, ChunkFileReader( NULL )
+	, ChunkFile( InChunkFile )
+	, FileBytesRead( InBytesRead )
+	, MemoryBytesRead( 0 )
+{
+	ChunkFile->GetDataLock( &ChunkData, &ChunkHeader );
+}
+
+FBuildGenerationChunkCache::FChunkReader::~FChunkReader()
+{
+	ChunkFile->ReleaseDataLock();
+
+	// Close file handle
+	if( ChunkFileReader != NULL )
+	{
+		ChunkFileReader->Close();
+		delete ChunkFileReader;
+		ChunkFileReader = NULL;
+	}
+}
+
+FArchive* FBuildGenerationChunkCache::FChunkReader::GetArchive()
+{
+	// Open file handle?
+	if( ChunkFileReader == NULL )
+	{
+		ChunkFileReader = IFileManager::Get().CreateFileReader( *ChunkFilePath );
+		if( ChunkFileReader == NULL )
+		{
+			// Break the magic to mark as invalid
+			ChunkHeader->Magic = 0;
+			GLog->Logf( TEXT( "WARNING: Skipped missing chunk file %s" ), *ChunkFilePath );
+			return NULL;
+		}
+		// Read Header?
+		if( ChunkHeader->Guid.IsValid() == false )
+		{
+			*ChunkFileReader << *ChunkHeader;
+		}
+		// Check we can seek otherwise bad chunk
+		const int64 ExpectedFileSize = ChunkHeader->DataSize + ChunkHeader->HeaderSize;
+		const int64 ChunkFileSize = ChunkFileReader->TotalSize();
+		const int64 NextByte = ChunkHeader->HeaderSize + *FileBytesRead;
+		if( ChunkFileSize == ExpectedFileSize
+		 && NextByte < ChunkFileSize )
+		{
+			// Seek to next byte
+			ChunkFileReader->Seek( NextByte );
+			// Break the magic to mark as invalid if archive errored, this chunk will get ignored
+			if( ChunkFileReader->GetError() )
+			{
+				ChunkHeader->Magic = 0;
+			}
+		}
+		else
+		{
+			// Break the magic to mark as invalid
+			ChunkHeader->Magic = 0;
+		}
+		// If this chunk is valid and compressed, we must read the entire file and decompress to memory now if we have not already
+		// as we cannot compare to compressed data
+		if( *FileBytesRead == 0 && ChunkHeader->StoredAs & FChunkHeader::STORED_COMPRESSED )
+		{
+			// Load the compressed chunk data
+			TArray< uint8 > CompressedData;
+			CompressedData.Empty( ChunkHeader->DataSize );
+			CompressedData.AddUninitialized( ChunkHeader->DataSize );
+			ChunkFileReader->Serialize( CompressedData.GetData(), ChunkHeader->DataSize );
+			// Uncompress
+			bool bSuceess = FCompression::UncompressMemory(
+				static_cast< ECompressionFlags >( COMPRESS_ZLIB | COMPRESS_BiasMemory ),
+				ChunkData,
+				FBuildPatchData::ChunkDataSize,
+				CompressedData.GetData(),
+				ChunkHeader->DataSize );
+			// Mark that we have fully read decompressed data and unpate the chunkfile's data size as we are expanding it
+			*FileBytesRead = FBuildPatchData::ChunkDataSize;
+			ChunkHeader->DataSize = FBuildPatchData::ChunkDataSize;
+			// Check uncompression was OK
+			if( !bSuceess )
+			{
+				ChunkHeader->Magic = 0;
+			}
+		}
+	}
+	return ChunkFileReader;
+}
+
+const bool FBuildGenerationChunkCache::FChunkReader::IsValidChunk()
+{
+	if( ChunkHeader->Guid.IsValid() == false )
+	{
+		GetArchive();
+	}
+	// Check magic, and current support etc.
+	const bool bValidHeader = ChunkHeader->IsValidMagic();
+	const bool bValidChunkGuid = ChunkHeader->Guid.IsValid();
+	const bool bSupportedFormat = ChunkHeader->StoredAs == FChunkHeader::STORED_RAW || ChunkHeader->StoredAs == FChunkHeader::STORED_COMPRESSED;
+	return bValidHeader && bValidChunkGuid && bSupportedFormat;
+}
+
+const FGuid& FBuildGenerationChunkCache::FChunkReader::GetChunkGuid()
+{
+	if( ChunkHeader->Guid.IsValid() == false )
+	{
+		GetArchive();
+	}
+	return ChunkHeader->Guid;
+}
+
+void FBuildGenerationChunkCache::FChunkReader::ReadNextBytes( uint8** OutDataBuffer, const uint32& ReadLength )
+{
+	uint8* BufferNextByte = &ChunkData[ MemoryBytesRead ];
+	// Do we need to load from disk?
+	if( ( MemoryBytesRead + ReadLength ) > *FileBytesRead )
+	{
+		FArchive* Reader = GetArchive();
+		// Do not allow incorrect usage
+		check( Reader );
+		check( ReadLength <= BytesLeft() );
+		// Read the number of bytes extra we need
+		const int32 NumFileBytesRead = *FileBytesRead;
+		const int32 NextMemoryBytesRead = MemoryBytesRead + ReadLength;
+		const uint32 FileReadLen = FMath::Max<int32>( 0, NextMemoryBytesRead - NumFileBytesRead );
+		*FileBytesRead += FileReadLen;
+		Reader->Serialize( BufferNextByte, FileReadLen );
+		// Assert if read error, if theres some problem accessing chunks then continuing would cause bad patch
+		// ratios, so it's better to hard fault.
+		const bool bChunkReadOK = !Reader->GetError();
+		if( !bChunkReadOK )
+		{
+			// Print something helpful
+			GLog->Logf( TEXT( "FATAL ERROR: Could not read from chunk FArchive %s" ), *ChunkFilePath );
+			// Check with bool variable so that output will be readable
+			check( bChunkReadOK );
+		}
+	}
+	MemoryBytesRead += ReadLength;
+	(*OutDataBuffer) = BufferNextByte;
+}
+
+const uint32 FBuildGenerationChunkCache::FChunkReader::BytesLeft()
+{
+	if( ChunkHeader->Guid.IsValid() == false )
+	{
+		GetArchive();
+	}
+	return FBuildPatchData::ChunkDataSize - MemoryBytesRead;
+}
+
+/* FBuildSimpleChunkCache implementation
+*****************************************************************************/
+TSharedRef< FBuildGenerationChunkCache::FChunkReader > FBuildGenerationChunkCache::GetChunkReader( const FString& ChunkFilePath )
+{
+	if( ChunkCache.Contains( ChunkFilePath ) == false )
+	{
+		// Remove oldest access from cache?
+		if( ChunkCache.Num() >= NumChunksToCache )
+		{
+			FString OldestAccessChunk = TEXT( "" );
+			double OldestAccessTime = FPlatformTime::Seconds();
+			for( auto ChunkCacheIt = ChunkCache.CreateConstIterator(); ChunkCacheIt; ++ChunkCacheIt )
+			{
+				const FString& ChunkFilePath = ChunkCacheIt.Key();
+				const FChunkFile& ChunkFile = ChunkCacheIt.Value().Get();
+				if( ChunkFile.GetLastAccessTime() < OldestAccessTime )
+				{
+					OldestAccessTime = ChunkFile.GetLastAccessTime();
+					OldestAccessChunk = ChunkFilePath;
+				}
+			}
+			ChunkCache.Remove( OldestAccessChunk );
+			delete BytesReadPerChunk[ OldestAccessChunk ];
+			BytesReadPerChunk.Remove( OldestAccessChunk );
+		}
+		// Add the chunk to cache
+		ChunkCache.Add( ChunkFilePath, MakeShareable( new FChunkFile( 1, true ) ) );
+		BytesReadPerChunk.Add( ChunkFilePath, new uint32( 0 ) );
+	}
+	return MakeShareable( new FChunkReader( ChunkFilePath, ChunkCache[ ChunkFilePath ], BytesReadPerChunk[ ChunkFilePath ] ) );
+}
+
+void FBuildGenerationChunkCache::Cleanup()
+{
+	ChunkCache.Empty( 0 );
+	for( auto BytesReadPerChunkIt = BytesReadPerChunk.CreateConstIterator(); BytesReadPerChunkIt; ++BytesReadPerChunkIt )
+	{
+		delete BytesReadPerChunkIt.Value();
+	}
+	BytesReadPerChunk.Empty( 0 );
+}
+
+/* FBuildSimpleChunkCache system singleton setup
+*****************************************************************************/
+TSharedPtr< FBuildGenerationChunkCache > FBuildGenerationChunkCache::SingletonInstance = NULL;
+
+void FBuildGenerationChunkCache::Init()
+{
+	// We won't allow misuse of these functions
+	check( !SingletonInstance.IsValid() );
+	SingletonInstance = MakeShareable( new FBuildGenerationChunkCache() );
+}
+
+FBuildGenerationChunkCache& FBuildGenerationChunkCache::Get()
+{
+	// We won't allow misuse of these functions
+	check( SingletonInstance.IsValid() );
+	return *SingletonInstance.Get();
+}
+
+void FBuildGenerationChunkCache::Shutdown()
+{
+	// We won't allow misuse of these functions
+	check( SingletonInstance.IsValid() );
+	SingletonInstance->Cleanup();
+	SingletonInstance.Reset();
+}
+
+/* FBuildDataGenerator static variables
+*****************************************************************************/
+TMap< FGuid, FString > FBuildDataGenerator::ExistingChunkGuidInventory;
+TMap< uint64, TArray< FGuid > > FBuildDataGenerator::ExistingChunkHashInventory;
+bool FBuildDataGenerator::ExistingChunksEnumerated = false;
+TMap< FSHAHash, TArray< FString > > FBuildDataGenerator::ExistingFileInventory;
+bool FBuildDataGenerator::ExistingFilesEnumerated = false;
+FCriticalSection FBuildDataGenerator::SingleConcurrentBuildCS;
+
+static void AddCustomFieldsToBuildManifest(const TMap<FString, FVariant>& CustomFields, IBuildManifestPtr BuildManifest)
+{
+	for (const auto& CustomField : CustomFields)
+	{
+		int32 VarType = CustomField.Value.GetType();
+		if (VarType == EVariantTypes::Float || VarType == EVariantTypes::Double)
+		{
+			BuildManifest->SetCustomField(CustomField.Key, (double)CustomField.Value);
+		}
+		else if (VarType == EVariantTypes::Int8 || VarType == EVariantTypes::Int16 || VarType == EVariantTypes::Int32 || VarType == EVariantTypes::Int64 ||
+			VarType == EVariantTypes::UInt8 || VarType == EVariantTypes::UInt16 || VarType == EVariantTypes::UInt32 || VarType == EVariantTypes::UInt64)
+		{
+			BuildManifest->SetCustomField(CustomField.Key, (int64)CustomField.Value);
+		}
+		else if (VarType == EVariantTypes::String)
+		{
+			BuildManifest->SetCustomField(CustomField.Key, CustomField.Value.GetValue<FString>());
+		}
+	}
+}
+
+/* FBuildDataGenerator implementation
+*****************************************************************************/
+bool FBuildDataGenerator::GenerateChunksManifestFromDirectory( const FBuildPatchSettings& Settings )
+{
+	const FString& RootDirectory = Settings.RootDirectory;
+	const uint32& InAppID = Settings.InAppID;
+	const FString& AppName = Settings.AppName;
+	const FString& BuildVersion = Settings.BuildVersion;
+	const FString& LaunchExe = Settings.LaunchExe;
+	const FString& LaunchCommand = Settings.LaunchCommand;
+	const FString& IgnoreListFile = Settings.IgnoreListFile;
+	const FString& PrereqName = Settings.PrereqName;
+	const FString& PrereqPath = Settings.PrereqPath;
+	const FString& prereqArgs = Settings.PrereqArgs;
+
+	// Output to log for builder info
+	GLog->Logf( TEXT( "Running Chunks Patch Generation for: %u:%s %s" ), InAppID, *AppName, *BuildVersion );
+
+	// Take the build CS
+	FScopeLock SingleConcurrentBuild( &SingleConcurrentBuildCS );
+
+	// Create our chunk cache
+	FBuildGenerationChunkCache::Init();
+
+	// Create a manifest
+	FBuildPatchAppManifestRef BuildManifest = MakeShareable( new FBuildPatchAppManifest() );
+
+	// Reset chunk inventory
+	ExistingChunksEnumerated = false;
+	ExistingChunkGuidInventory.Empty();
+	ExistingChunkHashInventory.Empty();
+
+	// Clear current data
+	BuildManifest->DestroyData(); 
+
+	// Declare a build processor
+	FBuildDataChunkProcessor DataProcessor( BuildManifest, RootDirectory );
+
+	// Create a build streamer
+	FBuildStream* BuildStream = new FBuildStream( RootDirectory, IgnoreListFile );
+
+	// Set the App details
+	BuildManifest->AppID = InAppID;
+	BuildManifest->AppNameString = AppName;
+	BuildManifest->BuildVersionString = BuildVersion;
+	BuildManifest->LaunchExeString = LaunchExe;
+	BuildManifest->LaunchCommandString = LaunchCommand;
+	BuildManifest->bIsFileData = false;
+
+	// Create a data buffer
+	const uint32 DataBufferSize = FBuildPatchData::ChunkDataSize;
+	uint8* DataBuffer = new uint8[ DataBufferSize ];
+
+	// We'll need a rolling hash for chunking
+	FRollingHash< FBuildPatchData::ChunkDataSize >* RollingHash = new FRollingHash< FBuildPatchData::ChunkDataSize >();
+	
+	// Refers to how much data has been processed (into the FBuildDataProcessor)
+	uint64 ProcessPos = 0;
+
+	// Records the current file we are processing
+	FString FileName;
+
+	// And the current file's data left to process
+	uint64 FileDataCount = 0;
+
+	// Used to store data read lengths
+	uint32 ReadLen = 0;
+
+	// The last time we logged out data processed
+	double LastProgressLog = FPlatformTime::Seconds();
+	const double TimeGenStarted = LastProgressLog;
+
+	// Loop through all data
+	while ( !BuildStream->IsEndOfData() )
+	{
+		// Grab some data from the build stream
+		ReadLen = BuildStream->DequeueData( DataBuffer, DataBufferSize );
+
+		// A bool says if there's no more data to come from the Build Stream
+		const bool bNoMoreData = BuildStream->IsEndOfData();
+
+		// Refers to how much data from DataBuffer has been passed into the rolling hash
+		uint32 DataBufferPos = 0;
+
+		// Count how many times we pad the rolling hash with zero
+		uint32 PaddedZeros = 0;
+
+		// Process data while we have more
+		while ( ( DataBufferPos < ReadLen ) || ( bNoMoreData && PaddedZeros < RollingHash->GetWindowSize() ) )
+		{
+			// Prime the rolling hash
+			if( RollingHash->GetNumDataNeeded() > 0 )
+			{
+				if( DataBufferPos < ReadLen )
+				{
+					RollingHash->ConsumeByte( DataBuffer[ DataBufferPos++ ] );
+				}
+				else
+				{
+					RollingHash->ConsumeByte( 0 );
+					++PaddedZeros;
+				}
+				// Keep looping until primed
+				continue;
+			}
+
+			// Check if we recognized a chunk
+			FGuid ChunkGuid;
+			const uint64 WindowHash = RollingHash->GetWindowHash();
+			const TRingBuffer< uint8, FBuildPatchData::ChunkDataSize >& WindowData = RollingHash->GetWindowData();
+			bool ChunkRecognised = FindExistingChunkData( WindowHash, WindowData, ChunkGuid );
+			if( ChunkRecognised )
+			{
+				// Process all bytes
+				DataProcessor.PushChunk();
+				const uint32 WindowDataSize = RollingHash->GetWindowSize() - PaddedZeros;
+				for( uint32 i = 0; i < WindowDataSize; ++i )
+				{
+					const bool bStartOfFile = BuildStream->GetFileSpan( ProcessPos, FileName, FileDataCount );
+					const bool bEndOfFile = FileDataCount <= 1;
+					check( FileDataCount > 0 );// If FileDataCount is ever 0, it means this piece of data belongs to no file, so something is wrong
+					DataProcessor.SkipKnownByte( WindowData[i], bStartOfFile, bEndOfFile, FileName );
+					++ProcessPos;
+					--FileDataCount;
+				}
+				uint8* SerialWindowData = new uint8[ FBuildPatchData::ChunkDataSize ];
+				WindowData.Serialize( SerialWindowData );
+				DataProcessor.PopChunk( WindowHash, SerialWindowData, ChunkGuid );
+				delete[] SerialWindowData;
+
+				// Clear
+				RollingHash->Clear();
+			}
+			else
+			{
+				// Process one byte
+				const bool bStartOfFile = BuildStream->GetFileSpan( ProcessPos, FileName, FileDataCount );
+				const bool bEndOfFile = FileDataCount <= 1;
+				DataProcessor.ProcessNewByte( WindowData.Bottom(), bStartOfFile, bEndOfFile, FileName );
+				++ProcessPos;
+				--FileDataCount;
+
+				// Roll
+				if( DataBufferPos < ReadLen )
+				{
+					RollingHash->RollForward( DataBuffer[ DataBufferPos++ ] );
+				}
+				else if( bNoMoreData )
+				{
+					RollingHash->RollForward( 0 );
+					++PaddedZeros;
+				}
+			}
+
+			// Log processed data
+			if( ( FPlatformTime::Seconds() - LastProgressLog ) >= 10.0 )
+			{
+				LastProgressLog = FPlatformTime::Seconds();
+				GLog->Logf( TEXT( "Processed %lld bytes." ), ProcessPos );
+			}
+		}
+	}
+
+	// The final chunk if any should be finished.
+	// This also triggers the chunk writer thread to exit.
+	DataProcessor.FinalChunk();
+
+	// Handle empty files
+	FSHA1 EmptyHasher;
+	EmptyHasher.Final();
+	const TArray< FString >& EmptyFileList = BuildStream->GetEmptyFiles();
+	for (const auto& EmptyFile : EmptyFileList)
+	{
+		FBuildPatchFileManifest* EmptyFileManifest = new FBuildPatchFileManifest();
+		EmptyHasher.GetHash( EmptyFileManifest->FileHash.Hash );
+		EmptyFileManifest->Filename = EmptyFile;
+		BuildManifest->FileManifestList.Add( EmptyFileManifest->Filename, MakeShareable( EmptyFileManifest ) );
+	}
+
+	// Add chunk sizes
+	const TMap<FGuid, int64>& ChunkFilesizes = DataProcessor.GetChunkFilesizes();
+	BuildManifest->ChunkFilesizeList.Empty(ChunkFilesizes.Num());
+	BuildManifest->ChunkFilesizeList.Append(ChunkFilesizes);
+
+	// Clean up memory
+	delete[] DataBuffer;
+	delete BuildStream;
+	delete RollingHash;
+	FBuildGenerationChunkCache::Shutdown();
+
+	// Save manifest into the cloud directory
+	FString ManifestFilename = FBuildPatchServicesModule::GetCloudDirectory();
+	ManifestFilename /= FDefaultValueHelper::RemoveWhitespaces( BuildManifest->AppNameString + BuildManifest->BuildVersionString ) + TEXT(".manifest");
+	AddCustomFieldsToBuildManifest(Settings.CustomFields, BuildManifest);
+	BuildManifest->SaveToFile(ManifestFilename);
+
+	// Output to log for builder info
+	GLog->Logf( TEXT( "Saved manifest to %s" ), *ManifestFilename );
+	GLog->Logf( TEXT( "Generation took %s" ), *FPlatformTime::PrettyTime( FPlatformTime::Seconds() - TimeGenStarted ) );
+
+	// @TODO LSwift: Detect errors and return false on failure
+	return true;
+}
+
+bool FBuildDataGenerator::GenerateFilesManifestFromDirectory( const FBuildPatchSettings& Settings )
+{
+	const FString& RootDirectory = Settings.RootDirectory;
+	const uint32& InAppID = Settings.InAppID;
+	const FString& AppName = Settings.AppName;
+	const FString& BuildVersion = Settings.BuildVersion;
+	const FString& LaunchExe = Settings.LaunchExe;
+	const FString& LaunchCommand = Settings.LaunchCommand;
+	const FString& IgnoreListFile = Settings.IgnoreListFile;
+
+	// Output to log for builder info
+	GLog->Logf( TEXT( "Running Files Patch Generation for: %u:%s %s" ), InAppID, *AppName, *BuildVersion );
+
+	// Take the build CS
+	FScopeLock SingleConcurrentBuild( &SingleConcurrentBuildCS );
+
+	// Create a manifest
+	FBuildPatchAppManifestRef BuildManifest = MakeShareable( new FBuildPatchAppManifest() );
+
+	// Reset file inventory
+	ExistingFilesEnumerated = false;
+	ExistingFileInventory.Empty();
+
+	// Clear current data
+	BuildManifest->DestroyData(); 
+
+	// Declare a build processor
+	FBuildDataFileProcessor DataProcessor( BuildManifest, RootDirectory );
+
+	// Set the App details
+	BuildManifest->AppID = InAppID;
+	BuildManifest->AppNameString = AppName;
+	BuildManifest->BuildVersionString = BuildVersion;
+	BuildManifest->LaunchExeString = LaunchExe;
+	BuildManifest->LaunchCommandString = LaunchCommand;
+	BuildManifest->bIsFileData = true;
+
+	// Create a data buffer
+	uint8* FileReadBuffer = new uint8[ FileBufferSize ];
+	
+	// Refers to how much data has been processed (into the FBuildDataFileProcessor)
+	uint64 ProcessPos = 0;
+
+	// Find all files
+	TArray< FString > AllFiles;
+	IFileManager::Get().FindFilesRecursive( AllFiles, *RootDirectory, TEXT("*.*"), true, false );
+	AllFiles.Sort();
+
+	// Remove the files that appear in an ignore list
+	FBuildDataGenerator::StripIgnoredFiles( AllFiles, RootDirectory, IgnoreListFile );
+
+	// Loop through all files
+	for( auto FileIt = AllFiles.CreateConstIterator(); FileIt; ++FileIt )
+	{
+		const FString& FileName = *FileIt;
+		// Read the file
+		FArchive* FileReader = IFileManager::Get().CreateFileReader( *FileName );
+		if( FileReader != NULL )
+		{
+			// Make SourceFile the format we want it in and start a new file
+			FString SourceFile = FileName;
+			FPaths::MakePathRelativeTo( SourceFile, *( RootDirectory + TEXT( "/" ) ) );
+			int64 FileSize = FileReader->TotalSize();
+			if( FileSize < 0 )
+			{
+				// Skip potential error ( INDEX_NONE == -1 )
+				continue;
+			}
+			DataProcessor.BeginFile( SourceFile );
+			while( !FileReader->AtEnd() )
+			{
+				const int64 SizeLeft = FileSize - FileReader->Tell();
+				const uint32 ReadLen = FMath::Min< int64 >( FileBufferSize, SizeLeft );
+				ProcessPos += ReadLen;
+				FileReader->Serialize( FileReadBuffer, ReadLen );
+				// Copy into data stream
+				DataProcessor.ProcessFileData( FileReadBuffer, ReadLen );
+			}
+			FileReader->Close();
+			delete FileReader;
+			DataProcessor.EndFile();
+		}
+		else
+		{
+			// @TODO LSwift: Handle File error?
+		}
+	}
+
+	// Clean up memory
+	delete[] FileReadBuffer;
+
+	// Save manifest into the cloud directory
+	FString ManifestFilename = FBuildPatchServicesModule::GetCloudDirectory();
+	ManifestFilename /= FDefaultValueHelper::RemoveWhitespaces( BuildManifest->AppNameString + BuildManifest->BuildVersionString ) + TEXT(".manifest");
+	AddCustomFieldsToBuildManifest(Settings.CustomFields, BuildManifest);
+	BuildManifest->SaveToFile(ManifestFilename);
+
+	// Output to log for builder info
+	GLog->Logf( TEXT( "Saved manifest to %s" ), *ManifestFilename );
+
+	// @TODO LSwift: Detect errors and return false on failure
+	return true;
+}
+
+bool FBuildDataGenerator::FindExistingChunkData( const uint64& ChunkHash, const uint8* ChunkData, FGuid& ChunkGuid )
+{
+	// Quick code hack
+	TRingBuffer< uint8, FBuildPatchData::ChunkDataSize > ChunkDataRing;
+	const uint32 ChunkDataLen = FBuildPatchData::ChunkDataSize;
+	ChunkDataRing.Enqueue( ChunkData, ChunkDataLen );
+	return FindExistingChunkData( ChunkHash, ChunkDataRing, ChunkGuid );
+}
+
+bool FBuildDataGenerator::FindExistingChunkData( const uint64& ChunkHash, const TRingBuffer< uint8, FBuildPatchData::ChunkDataSize >& ChunkData, FGuid& ChunkGuid )
+{
+	bool bFoundMatchingChunk = false;
+
+	// Perform an inventory on Cloud chunks if not already done
+	if( ExistingChunksEnumerated == false )
+	{
+		IFileManager& FileManager = IFileManager::Get();
+		FString JSONOutput;
+		TSharedRef< TJsonWriter< TCHAR, TPrettyJsonPrintPolicy< TCHAR > > > DebugWriter = TJsonWriterFactory< TCHAR, TPrettyJsonPrintPolicy< TCHAR > >::Create( &JSONOutput );
+		DebugWriter->WriteObjectStart();
+
+		// Find all manifest files
+		const FString CloudDir = FBuildPatchServicesModule::GetCloudDirectory();
+		if (FileManager.DirectoryExists(*CloudDir))
+		{
+			const double StartEnumerate = FPlatformTime::Seconds();
+			TArray<FString> AllManifests;
+			GLog->Logf(TEXT("BuildDataGenerator: Enumerating Manifests from %s"), *CloudDir);
+			FileManager.FindFiles(AllManifests, *(CloudDir / TEXT("*.manifest")), true, false);
+			const double EnumerateTime = FPlatformTime::Seconds() - StartEnumerate;
+			GLog->Logf(TEXT("BuildDataGenerator: Found %d manifests in %.1f seconds"), AllManifests.Num(), EnumerateTime);
+
+			// Load all manifest files
+			uint64 NumChunksFound = 0;
+			const double StartLoadAllManifest = FPlatformTime::Seconds();
+			for (const auto& ManifestFile : AllManifests)
+			{
+				// Determine chunks from manifest file
+				const FString ManifestFilename = CloudDir / ManifestFile;
+				FBuildPatchAppManifestRef BuildManifest = MakeShareable(new FBuildPatchAppManifest());
+				const double StartLoadManifest = FPlatformTime::Seconds();
+				if (BuildManifest->LoadFromFile(ManifestFilename))
+				{
+					const double LoadManifestTime = FPlatformTime::Seconds() - StartLoadManifest;
+					GLog->Logf(TEXT("BuildDataGenerator: Loaded %s in %.1f seconds"), *ManifestFile, LoadManifestTime);
+					if(!BuildManifest->IsFileDataManifest())
+					{
+						TArray<FGuid> ChunksReferenced;
+						BuildManifest->GetDataList(ChunksReferenced);
+						for (const auto& ChunkGuid : ChunksReferenced)
+						{
+							uint64 ChunkHash;
+							if (BuildManifest->GetChunkHash(ChunkGuid, ChunkHash))
+							{
+								if (ChunkHash != 0)
+								{
+									TArray< FGuid >& HashChunkList = ExistingChunkHashInventory.FindOrAdd(ChunkHash);
+									if (!HashChunkList.Contains(ChunkGuid))
+									{
+										++NumChunksFound;
+										HashChunkList.Add(ChunkGuid);
+									}
+								}
+								else
+								{
+									GLog->Logf(TEXT("BuildDataGenerator: WARNING: Ignored an existing chunk %s with a failed hash value of zero to avoid performance problems while chunking"), *ChunkGuid.ToString());
+								}
+							}
+							else
+							{
+								GLog->Logf(TEXT("BuildDataGenerator: WARNING: Missing chunk hash for %s in manifest %s"), *ChunkGuid.ToString(), *ManifestFile);
+							}
+						}
+					}
+					else
+					{
+						GLog->Logf(TEXT("BuildDataGenerator: INFO: Ignoring non-chunked manifest %s"), *ManifestFilename);
+					}
+				}
+				else
+				{
+					GLog->Logf(TEXT("BuildDataGenerator: WARNING: Could not read Manifest file. Data recognition will suffer (%s)"), *ManifestFilename);
+				}
+			}
+			const double LoadAllManifestTime = FPlatformTime::Seconds() - StartLoadAllManifest;
+			GLog->Logf(TEXT("BuildDataGenerator: Used %d manifests to enumerate %llu chunks in %.1f seconds"), AllManifests.Num(), NumChunksFound, LoadAllManifestTime);
+		}
+		else
+		{
+			GLog->Logf(TEXT("BuildDataGenerator: Cloud directory does not exist: %s"), *CloudDir);
+		}
+
+		ExistingChunksEnumerated = true;
+	}
+
+	// Do we have a chunk matching this data?
+	if( ExistingChunkHashInventory.Num() > 0 )
+	{
+		TArray< FGuid >* ChunkList = ExistingChunkHashInventory.Find( ChunkHash );
+		if( ChunkList != NULL )
+		{
+			// We need to load each chunk in this list and compare data
+			for( auto ChunkIt = ChunkList->CreateConstIterator(); ChunkIt && !bFoundMatchingChunk ; ++ChunkIt)
+			{
+				FGuid Guid = *ChunkIt;
+				if (!ExistingChunkGuidInventory.Contains(Guid))
+				{
+					ExistingChunkGuidInventory.Add(Guid, DiscoverChunkFilename(Guid, ChunkHash));
+				}
+				const FString& SourceFile = ExistingChunkGuidInventory[ Guid ];
+				// Read the file
+				uint8* TempChunkData = new uint8[ FBuildPatchData::ChunkDataSize ];
+				ChunkData.Serialize( TempChunkData );
+				bool bChunkIsUsable = true;
+				if( CompareDataToChunk( SourceFile, TempChunkData, Guid, bChunkIsUsable ) )
+				{
+					// We have a chunk match!!
+					bFoundMatchingChunk = true;
+					ChunkGuid = Guid;
+				}
+				// Check if this chunk should be dumped
+				if(!bChunkIsUsable)
+				{
+					GLog->Logf(TEXT("BuildDataGenerator: Chunk %s unusable, removed from inventory."), *ChunkGuid.ToString());
+					ChunkList->Remove(Guid);
+					--ChunkIt;
+				}
+				delete[] TempChunkData;
+			}
+		}
+	}
+
+	return bFoundMatchingChunk;
+}
+
+FString FBuildDataGenerator::DiscoverChunkFilename(const FGuid& ChunkGuid, const uint64& ChunkHash)
+{
+	static double AccumTime = 0.0;
+	const double StartDiscovery = FPlatformTime::Seconds();
+	const FString CloudDir = FBuildPatchServicesModule::GetCloudDirectory();
+	FString ChunkFilenameChecked;
+	for (EBuildPatchAppManifestVersion::Type VersionCounter = static_cast<EBuildPatchAppManifestVersion::Type>(EBuildPatchAppManifestVersion::LatestPlusOne - 1);
+		VersionCounter >= EBuildPatchAppManifestVersion::Original;
+		VersionCounter = static_cast<EBuildPatchAppManifestVersion::Type>(VersionCounter - 1))
+	{
+		const FString ChunkFilename = VersionCounter < EBuildPatchAppManifestVersion::DataFileRenames ?
+			FBuildPatchUtils::GetChunkOldFilename(CloudDir, ChunkGuid) :
+			FBuildPatchUtils::GetChunkNewFilename(VersionCounter, CloudDir, ChunkGuid, ChunkHash);
+		if (ChunkFilenameChecked != ChunkFilename)
+		{
+			ChunkFilenameChecked = ChunkFilename;
+			if (FPaths::FileExists(ChunkFilename))
+			{
+				const double DiscoveryTime = FPlatformTime::Seconds() - StartDiscovery;
+				AccumTime += DiscoveryTime;
+				GLog->Logf(TEXT("BuildDataGenerator: DiscoverChunkFilename: Chunk %s found in %f secs. Accum time %.2f secs"), *ChunkGuid.ToString(), DiscoveryTime, AccumTime);
+				return ChunkFilename;
+			}
+		}
+	}
+	GLog->Logf(TEXT("BuildDataGenerator: DiscoverChunkFilename: Chunk %s not found"), *ChunkGuid.ToString());
+	return TEXT("");
+}
+
+bool FBuildDataGenerator::FindExistingFileData( const FString& InSourceFile, const FSHAHash& InFileHash, FGuid& OutFileGuid )
+{
+	bool bFoundMatchingFile = false;
+
+	// Perform an inventory on Cloud files if not already done
+	if( ExistingFilesEnumerated == false )
+	{
+		IFileManager& FileManager = IFileManager::Get();
+		FString JSONOutput;
+		TSharedRef< TJsonWriter< TCHAR, TPrettyJsonPrintPolicy< TCHAR > > > DebugWriter = TJsonWriterFactory< TCHAR, TPrettyJsonPrintPolicy< TCHAR > >::Create( &JSONOutput );
+		DebugWriter->WriteObjectStart();
+
+		TArray< FGuid > FoundFiles;
+
+		// The directory containing old filename version files
+		const FString CloudFileDir = FBuildPatchServicesModule::GetCloudDirectory() / TEXT( "Files" );
+		// The directory containing new filename version files
+		const FString CloudFile2Dir = FBuildPatchServicesModule::GetCloudDirectory() / TEXT( "FilesV2" );
+
+		GLog->Logf( TEXT( "BuildDataGenerator: Enumerating Files from %s" ), *CloudFile2Dir );
+		if( FileManager.DirectoryExists( *CloudFile2Dir ) )
+		{
+			const double StartEnumerate = FPlatformTime::Seconds();
+
+			// Find all files
+			TArray<FString> AllFiles;
+			FileManager.FindFilesRecursive( AllFiles, *CloudFile2Dir, TEXT("*.file"), true, false );
+
+			for( auto FileIt = AllFiles.CreateConstIterator(); FileIt; ++FileIt )
+			{
+				FString SourceFile = *FileIt;
+				FSHAHash FoundHash;
+				FGuid FoundGuid;
+				FBuildPatchUtils::GetFileDetailFromNewFilename( SourceFile, FoundGuid, FoundHash );
+
+				// Add to inventory
+				FoundFiles.Add( FoundGuid );
+				TArray< FString >& FileList = ExistingFileInventory.FindOrAdd( FoundHash );
+				FileList.Add( SourceFile );
+
+				const void* HashBuffer = &FoundHash;
+				DebugWriter->WriteValue( FoundGuid.ToString(), FString::FromBlob( static_cast<const uint8*>( HashBuffer ), sizeof( FoundHash ) ) );
+			}
+
+			const double EnumerateTime = FPlatformTime::Seconds() - StartEnumerate;
+			GLog->Logf( TEXT( "BuildDataGenerator: Found %d new name files in %.1f seconds" ), AllFiles.Num(), EnumerateTime );
+		}
+		else
+		{
+			GLog->Logf( TEXT( "BuildDataGenerator: Cloud directory does not exist: %s" ), *CloudFile2Dir );
+		}
+
+		GLog->Logf( TEXT( "BuildDataGenerator: Enumerating Files From %s" ), *CloudFileDir );
+		if( IFileManager::Get().DirectoryExists( *CloudFileDir ) )
+		{
+			const double StartEnumerate = FPlatformTime::Seconds();
+			const int32 PreviousNumFiles = FoundFiles.Num();
+
+			// Find all files
+			TArray<FString> AllFiles;
+			IFileManager::Get().FindFilesRecursive( AllFiles, *CloudFileDir, TEXT("*.file"), true, false );
+
+			for( auto FileIt = AllFiles.CreateConstIterator(); FileIt; ++FileIt )
+			{
+				FString SourceFile = *FileIt;
+
+				// Skip any for GUID we already found
+				FGuid FoundGuid;
+				FBuildPatchUtils::GetGUIDFromFilename( SourceFile, FoundGuid );
+				if( FoundFiles.Contains( FoundGuid ) )
+				{
+					continue;
+				}
+
+				// Read the file
+				FArchive* FileReader = IFileManager::Get().CreateFileReader( *SourceFile );
+				if( FileReader != NULL )
+				{
+					// Read the header
+					FChunkHeader Header;
+					*FileReader << Header;
+
+					// Check magic
+					if( Header.IsValidMagic() && Header.HashType == FChunkHeader::HASH_SHA1 )
+					{
+						TArray< FString >& FileList = ExistingFileInventory.FindOrAdd( Header.SHAHash );
+						FileList.Add( SourceFile );
+						FoundFiles.Add( Header.Guid );
+
+						const void* HashBuffer = &Header.SHAHash;
+						DebugWriter->WriteValue( Header.Guid.ToString(), FString::FromBlob( static_cast<const uint8*>( HashBuffer ), sizeof( Header.SHAHash ) ) );
+					}
+					else
+					{
+						GLog->Logf( TEXT( "BuildDataGenerator: Failed magic/hashtype check on file [%d:%d] %s" ), Header.Magic, Header.HashType, *SourceFile );
+					}
+
+					FileReader->Close();
+					delete FileReader;
+				}
+				else
+				{
+					GLog->Logf( TEXT( "BuildDataGenerator: Failed to read chunk %s" ), *SourceFile );
+				}
+			}
+
+			const double EnumerateTime = FPlatformTime::Seconds() - StartEnumerate;
+			const int32 NewNumFiles = FoundFiles.Num();
+			GLog->Logf( TEXT( "BuildDataGenerator: Found %d extra old files in %.1f seconds" ), NewNumFiles - PreviousNumFiles, EnumerateTime );
+		}
+		else
+		{
+			GLog->Logf( TEXT( "BuildDataGenerator: Cloud directory does not exist: %s" ), *CloudFileDir );
+		}
+
+		DebugWriter->WriteObjectEnd();
+		DebugWriter->Close();
+
+		FArchive* FileOut = IFileManager::Get().CreateFileWriter( *( CloudFileDir + TEXT( "DebugFileList.txt" ) ) );
+		if( FileOut != NULL )
+		{
+			FileOut->Serialize(TCHAR_TO_ANSI(*JSONOutput), JSONOutput.Len());
+			FileOut->Close();
+			delete FileOut;
+		}
+
+		ExistingFilesEnumerated = true;
+	}
+
+	// Do we have a file matching this data?
+	if( ExistingFileInventory.Num() > 0 )
+	{
+		TArray< FString >* FileList = ExistingFileInventory.Find( InFileHash );
+		if( FileList != NULL )
+		{
+			// We need to load each chunk in this list and compare data
+			for( auto FileIt = FileList->CreateConstIterator(); FileIt && !bFoundMatchingFile ; ++FileIt)
+			{
+				FString CloudFilename = *FileIt;
+				// Compare the files
+				FArchive* SourceFile = IFileManager::Get().CreateFileReader( *InSourceFile );
+				FArchive* FoundFile = IFileManager::Get().CreateFileReader( *CloudFilename );
+				if( SourceFile != NULL && FoundFile != NULL)
+				{
+					FChunkHeader FoundHeader;
+					*FoundFile << FoundHeader;
+					const int64 SourceFileSize = SourceFile->TotalSize();
+					const int64 FoundFileSize = FoundFile->TotalSize();
+					if( SourceFileSize == FoundHeader.DataSize )
+					{
+						// Currently only support stored raw!
+						check( FoundHeader.StoredAs == FChunkHeader::STORED_RAW );
+						// Makes no sense here for the sizes to not add up
+						check( FoundFileSize == ( FoundHeader.DataSize + FoundHeader.HeaderSize ) );
+						// Move FoundFile to start of file data
+						FoundFile->Seek( FoundHeader.HeaderSize );
+						// Compare
+						bool bSameData = true;
+						uint8* TempSourceBuffer = new uint8[ FileBufferSize ];
+						uint8* TempFoundBuffer = new uint8[ FileBufferSize ];
+						while ( bSameData && ( ( SourceFile->AtEnd() || FoundFile->AtEnd() ) == false ) )
+						{
+							const int64 SizeLeft = SourceFileSize - SourceFile->Tell();
+							const uint32 ReadLen = FMath::Min< int64 >( FileBufferSize, SizeLeft );
+							SourceFile->Serialize( TempSourceBuffer, ReadLen );
+							FoundFile->Serialize( TempFoundBuffer, ReadLen );
+							bSameData = FMemory::Memcmp( TempSourceBuffer, TempFoundBuffer, ReadLen ) == 0;
+						}
+						delete[] TempSourceBuffer;
+						delete[] TempFoundBuffer;
+						// Did we match?
+						if( bSameData && SourceFile->AtEnd() && FoundFile->AtEnd() )
+						{
+							// Yes we did!
+							bFoundMatchingFile = true;
+							OutFileGuid = FoundHeader.Guid;
+						}
+					}
+				}
+				if( SourceFile != NULL )
+				{
+					SourceFile->Close();
+					delete SourceFile;
+				}
+				if( FoundFile != NULL )
+				{
+					FoundFile->Close();
+					delete FoundFile;
+				}
+			}
+		}
+	}
+
+	return bFoundMatchingFile;
+}
+
+bool FBuildDataGenerator::SaveOutFileData( const FString& SourceFile, const FSHAHash& FileHash, const FGuid& FileGuid )
+{
+	bool bAlreadySaved = false;
+	bool bSuccess = false;
+	IFileManager& FileManager = IFileManager::Get();
+
+	const FString NewFilename = FBuildPatchUtils::GetFileNewFilename( EBuildPatchAppManifestVersion::GetLatestVersion(), FBuildPatchServicesModule::GetCloudDirectory(), FileGuid, FileHash );
+	bAlreadySaved = FPaths::FileExists( NewFilename );
+
+#if SAVE_OLD_FILEDATA_FILENAMES
+	const FString OldFilename = FBuildPatchUtils::GetFileOldFilename( FBuildPatchServicesModule::GetCloudDirectory(), FileGuid );
+	bAlreadySaved = bAlreadySaved && FPaths::FileExists( OldFilename );
+#endif
+
+	if( !bAlreadySaved )
+	{
+		FArchive* FileOut = FileManager.CreateFileWriter( *NewFilename );
+		FArchive* FileIn = FileManager.CreateFileReader( *SourceFile );
+
+		bSuccess = FileOut != NULL && FileIn != NULL;
+
+#if SAVE_OLD_FILEDATA_FILENAMES
+		FArchive* OldFileOut = FileManager.CreateFileWriter( *OldFilename );
+		bSuccess = bSuccess && OldFileOut != NULL;
+#endif
+
+		if( bSuccess )
+		{
+			const int64 FileSize = FileIn->TotalSize();
+
+			// LSwift: No support for too large files currently
+			check( FileSize <= 0xFFFFFFFF );
+
+			// Setup Header
+			FChunkHeader Header;
+			*FileOut << Header;
+			Header.HeaderSize = FileOut->Tell();
+			Header.DataSize = FileSize;
+			Header.Guid = FileGuid;
+			Header.HashType = FChunkHeader::HASH_SHA1;
+			Header.SHAHash = FileHash;
+			Header.StoredAs = FChunkHeader::STORED_RAW;
+
+			// Write out file
+			FileOut->Seek( 0 );
+			*FileOut << Header;
+
+#if SAVE_OLD_FILEDATA_FILENAMES
+			*OldFileOut << Header;
+#endif
+
+			uint8* FileDataBuffer = new uint8[ FileBufferSize ];
+			while( !FileIn->AtEnd() )
+			{
+				const int64 SizeLeft = FileSize - FileIn->Tell();
+				const uint32 ReadLen = FMath::Min< int64 >( FileBufferSize, SizeLeft );
+				FileIn->Serialize( FileDataBuffer, ReadLen );
+				FileOut->Serialize( FileDataBuffer, ReadLen );
+
+#if SAVE_OLD_FILEDATA_FILENAMES
+				OldFileOut->Serialize( FileDataBuffer, ReadLen );
+#endif
+
+			}
+			delete[] FileDataBuffer;
+		}
+
+		// Close files
+		if( FileIn != NULL )
+		{
+			FileIn->Close();
+			delete FileIn;
+		}
+		if( FileOut != NULL )
+		{
+			FileOut->Close();
+			delete FileOut;
+		}
+
+#if SAVE_OLD_FILEDATA_FILENAMES
+		if( OldFileOut != NULL )
+		{
+			OldFileOut->Close();
+			delete OldFileOut;
+		}
+#endif
+
+	}
+
+	return bSuccess;
+}
+
+bool FBuildDataGenerator::CompareDataToChunk( const FString& ChunkFilePath, uint8* ChunkData, FGuid& ChunkGuid, bool& OutSourceChunkIsValid )
+{
+	bool bMatching = false;
+
+	// Read the file
+	TSharedRef< FBuildGenerationChunkCache::FChunkReader > ChunkReader = FBuildGenerationChunkCache::Get().GetChunkReader( ChunkFilePath );
+	if( ChunkReader->IsValidChunk() )
+	{
+		ChunkGuid = ChunkReader->GetChunkGuid();
+		// Default true
+		bMatching = true;
+		// Compare per small block (for early outing!)
+		const uint32 CompareSize = 64;
+		uint8* ReadBuffer;
+		uint32 NumCompared = 0;
+		while( bMatching && ChunkReader->BytesLeft() > 0 && NumCompared < FBuildPatchData::ChunkDataSize )
+		{
+			const uint32 ReadLen = FMath::Min< uint32 >( CompareSize, ChunkReader->BytesLeft() );
+			ChunkReader->ReadNextBytes( &ReadBuffer, ReadLen );
+			bMatching = FMemory::Memcmp( &ChunkData[ NumCompared ], ReadBuffer, ReadLen ) == 0;
+			NumCompared += ReadLen;
+		}
+	}
+
+	// Set chunk valid state after loading in case reading discovered bad or ignorable data
+	OutSourceChunkIsValid = ChunkReader->IsValidChunk();
+
+	return bMatching;
+}
+
+void FBuildDataGenerator::StripIgnoredFiles( TArray< FString >& AllFiles, const FString& DepotDirectory, const FString& IgnoreListFile )
+{
+	const int32 OriginalNumFiles = AllFiles.Num();
+	FString IgnoreFileList = TEXT( "" );
+	FFileHelper::LoadFileToString( IgnoreFileList, *IgnoreListFile );
+	TArray< FString > IgnoreFiles;
+	IgnoreFileList.ParseIntoArray( &IgnoreFiles, TEXT( "\r\n" ), true );
+	struct FRemoveMatchingStrings
+	{ 
+		const FString* MatchingString;
+		FRemoveMatchingStrings( const FString* InMatch )
+			: MatchingString(InMatch) {}
+
+		bool operator()(const FString& RemovalCandidate) const
+		{
+			FString PathA = RemovalCandidate;
+			FString PathB = *MatchingString;
+			FPaths::NormalizeFilename( PathA );
+			FPaths::NormalizeFilename( PathB );
+			return PathA == PathB;
+		}
+	};
+
+	for( int32 IgnoreIdx = 0; IgnoreIdx < IgnoreFiles.Num(); ++IgnoreIdx )
+	{
+		const FString& IgnoreFile = IgnoreFiles[IgnoreIdx];
+		const FString FullIgnorePath = DepotDirectory / IgnoreFile;
+		AllFiles.RemoveAll( FRemoveMatchingStrings( &FullIgnorePath ) );
+	}
+	const int32 NewNumFiles = AllFiles.Num();
+	GLog->Logf( TEXT( "Stripped %d ignorable file(s)" ), ( OriginalNumFiles - NewNumFiles ) );
+}
+
+#endif //WITH_BUILDPATCHGENERATION
diff --git a/Engine/Source/Runtime/Online/BuildPatchServices/Private/BuildPatchGeneration.h b/Engine/Source/Runtime/Online/BuildPatchServices/Private/BuildPatchGeneration.h
new file mode 100644
--- /dev/null
+++ b/Engine/Source/Runtime/Online/BuildPatchServices/Private/BuildPatchGeneration.h
@@ -0,0 +1,666 @@
+// Copyright 1998-2014 Epic Games, Inc. All Rights Reserved.
+
+/*=============================================================================
+	BuildPatchGeneration.h: Declares classes involved with producing build data.
+=============================================================================*/
+
+#ifndef __BuildPatchGeneration_h__
+#define __BuildPatchGeneration_h__
+
+#pragma once
+
+#if WITH_BUILDPATCHGENERATION
+
+/**
+ * A threaded class that holds Build stream data as a FIFO data buffer of given size. Will read until full then wait
+ * for data to be read out before continuing.
+ */
+class FBuildStream
+{
+private:
+	// A private runnable class that reads in data
+	class FBuildStreamReader : public FRunnable
+	{
+	public:
+		// The directory that we want to read from
+		FString DepotDirectory;
+
+		// The path to a text file containing \r\n separated depot relative files to ignore
+		FString IgnoreListFile;
+
+		// Pointer to FBuildStream
+		FBuildStream* BuildStream;
+
+		// Default Constructor
+		FBuildStreamReader();
+
+		// Destructor
+		~FBuildStreamReader();
+
+		// Start FRunnable interface
+		virtual bool Init() override;
+		virtual uint32 Run() override;
+		// End FRunnable interface
+
+		void StartThread();
+
+		FRunnableThread * Thread;
+
+	} BuildStreamReader;
+
+	// A private struct for storing the data span of a file in the build stream
+	struct FFileSpan
+	{
+		// The file's filename
+		FString Filename;
+		// The file data start index
+		uint64 StartIdx;
+		// The size of the file
+		uint64 Size;
+		
+		/**
+		 * Default constructor
+		 */
+		FFileSpan()
+			: StartIdx( 0 )
+			, Size( 0 )
+		{
+		}
+	};
+
+public:
+	/**
+	 * Constructor
+	 */
+	FBuildStream( const FString& RootDirectory, const FString& IgnoreListFile );
+
+	/**
+	 * Default destructor
+	 */
+	~FBuildStream();
+
+	/**
+	 * Retrieves the file details for a specific start index.
+	 * @param	IN	StartingIdx		The data index into the build image.
+	 * @param	OUT Filename		The filename for the file starting at StartingIdx. NB: The variable remains UNCHANGED if no file starts here.
+	 * @param	OUT FileSize		The size of the file starting at StartingIdx. NB: The variable remains UNCHANGED if no file starts here.
+	 * @return	true if the data byte at StartingIdx is the start of a file
+	 */
+	bool GetFileSpan( const uint64& StartingIdx, FString& Filename, uint64& FileSize );
+
+	/**
+	 * Fetches some data from the buffer, also removing it.
+	 * @param	IN	Buffer			Pointer to buffer to receive the data.
+	 * @param	IN	ReqSize			The amount of data to attempt to retrieve.
+	 * @param	IN	WaitForData		Optional: Default true. Whether to wait until there is enough data in the buffer.
+	 * @return	The amount of data retrieved
+	 */
+	uint32 DequeueData( uint8* Buffer, const uint32& ReqSize, const bool WaitForData = true );
+
+	/**
+	 * Gets a list of empty files that the build contains.
+	 * @return	Array of empty files in the build
+	 */
+	const TArray< FString > GetEmptyFiles();
+
+	/**
+	 * Whether the read thread has finished reading the build image.
+	 * @return	true if there is no more data coming into the buffer
+	 */
+	bool IsEndOfBuild();
+
+	/**
+	 * Whether there is any more data available to dequeue from the buffer
+	 * @return	true if there is no more data coming in, and the internal buffer is also empty.
+	 */
+	bool IsEndOfData();
+
+private:
+	/**
+	 * Hide the default constructor so that the params must be given
+	 */
+	FBuildStream(){};
+
+	//////////////////////////////////////////////////////////////////////////
+	// Private functions for manipulating the data store
+private:
+	/**
+	 * Wipes out all data
+	 */
+	void Clear();
+
+	/**
+	 * How much more data the internal buffer can hold before being full
+	 * @return	the number of bytes free in the buffer
+	 */
+	uint32 SpaceLeft();
+
+	/**
+	 * How much data the internal buffer is currently holding
+	 * @return	the number of bytes held in the buffer
+	 */
+	uint32 DataAvailable();
+
+	/**
+	 * Mark the beginning of a new file
+	 * @param	IN	Filename	The local build filename to the file
+	 * @param	IN	FileSize	The size of the file
+	 */
+	void BeginNewFile( const FString& Filename, const uint64& FileSize );
+
+	/**
+	 * We must special case empty files because we don't want them effecting chunk data
+	 * @param	IN	Filename	The local build filename to the file
+	 */
+	void AddEmptyFile( const FString& Filename );
+
+	/**
+	 * Enqueue a buffer of data into the internal buffer.
+	 * @param	IN	Buffer		Pointer to buffer of data.
+	 * @param	IN	Len			The length of the data buffer.
+	 */
+	void EnqueueData( const uint8* Buffer, const uint32& Len );
+
+	/**
+	 * Mark the end of the build image
+	 * @param	IN	bIsEnd		Optional: Default true. Whether to set the build end as true or false
+	 */
+	void EndOfBuild( bool bIsEnd = true );
+
+	//////////////////////////////////////////////////////////////////////////
+	// Private variables
+private:
+	// Holds the filenames referred to index positions in the data stream
+	TMap< uint64, FFileSpan > FilesParsed;
+
+	// Holds a list of empty files
+	TArray< FString > EmptyFiles;
+
+	// A critical section for accessing the FilesParsed and EmptyFiles data.
+	FCriticalSection FilesListsCS;
+
+	// Holds the ring buffer for loaded data code words
+	TRingBuffer< uint8, StreamBufferSize > BuildDataStream;
+
+	// A critical section for accessing the data stream, BuildDataStream
+	FCriticalSection BuildDataStreamCS;
+
+	// Whether the build has been completely read from disk
+	bool bNoMoreData;
+
+	// A critical section for accessing the no more data flag, NoMoreData.
+	FCriticalSection NoMoreDataCS;
+};
+
+/**
+ * A class that controls processing of build data to produce the manifest with chunks
+ */
+class FBuildDataChunkProcessor
+{
+public:
+
+	/**
+	 * Constructor
+	 * @param	InBuildManifest 	a build manifest that receives the manifest data.
+	 * @param	InBuildRoot			The root directory of the build
+	 */
+	FBuildDataChunkProcessor( FBuildPatchAppManifestRef InBuildManifest, const FString& InBuildRoot );
+	
+	/**
+	 * Default destructor
+	 */
+	~FBuildDataChunkProcessor();
+
+	/**
+	 * Start a new chunk
+	 * @param	bZeroBuffer 	Optional: Default true. Whether to clear the chunk buffer with zero
+	 */
+	void BeginNewChunk( const bool& bZeroBuffer = true );
+
+	/**
+	 * Start a new chunk part
+	 */
+	void BeginNewChunkPart();
+	
+	/**
+	 * End the current chunk part
+	 */
+	void EndNewChunkPart();
+	
+	/**
+	 * End the current chunk
+	 * @param	ChunkHash		The chunk's hash value
+	 * @param	ChunkData		Pointer to the chunk data, must be of FBuildPatchData::ChunkDataSize length
+	 * @param	ChunkGuid		The chunk GUID of this chunk. If it matches CurrentChunkGuid then this is a new chunk, otherwise it is the guid of the recognized chunk
+	 */
+	void EndNewChunk( const uint64& ChunkHash, const uint8* ChunkData, const FGuid& ChunkGuid );
+	
+	/**
+	 * Begins insertion of recognized chunk
+	 */
+	void PushChunk();
+
+	/**
+	 * Ends the insertion of recognized chunk
+	 * @param	ChunkHash		The hash value for the recognised chunk
+	 * @param	ChunkData		The data for the recognised chunk
+	 * @param	ChunkGuid		The GUID for the recognised chunk
+	 */
+	void PopChunk( const uint64& ChunkHash, const uint8* ChunkData, const FGuid& ChunkGuid );
+
+	/**
+	 * When the build image has been fully covered, call to end any current chunk we have (which will be zero padded)
+	 */
+	void FinalChunk();
+
+	/**
+	 * Mark the beginning of a new file
+	 * @param	FileName	The file's local filename
+	 */
+	void BeginFile( const FString& FileName );
+
+	/**
+	 * Mark the end of the current file
+	 */
+	void EndFile();
+
+	/**
+	 * Skips over the next byte from a recognized chunk
+	 * @param	NextByte		The next byte
+	 * @param	bStartOfFile	Whether this byte is the beginning of a file
+	 * @param	bEndOfFile		Whether this byte is the end of a file
+	 * @param	Filename		The filename of the file, only required if start of new file
+	 */
+	void SkipKnownByte( const uint8& NextByte, const bool& bStartOfFile, const bool& bEndOfFile, const FString& Filename );
+
+	/**
+	 * Process a new byte into the current new chunk
+	 * @param	NewByte			The new byte
+	 * @param	bStartOfFile	Whether this byte is the beginning of a file
+	 * @param	bEndOfFile		Whether this byte is the end of a file
+	 * @param	Filename		The filename of the file, only required if start of new file
+	 */
+	void ProcessNewByte( const uint8& NewByte, const bool& bStartOfFile, const bool& bEndOfFile, const FString& Filename );
+
+	/**
+	 * Get the new/know chunks stats
+	 * @param	OutNewChunks	Receives the number of new chunks found
+	 * @param	OutKnownChunks	Receives the number of recognised chunks found
+	 */
+	void GetChunkStats( uint32& OutNewChunks, uint32& OutKnownChunks );
+
+	/**
+	 * Get the chunk filesizes that have been saved out
+	 * @return ref to the Map of GUID to filesize
+	 */
+	const TMap<FGuid, int64>& GetChunkFilesizes();
+
+private:
+
+	// Keep track of the number of new chunks generated
+	uint32 NumNewChunks;
+
+	// Keep track of the number of recognised chunks
+	uint32 NumKnownChunks;
+
+	// Holds the build manifest that is receiving the information
+	FBuildPatchAppManifestPtr BuildManifest;
+
+	// Holds the build manifest that is receiving the information
+	FString BuildRoot;
+
+	// The chunk writer system
+	FChunkWriter ChunkWriter;
+
+	// A GUID for the current chunk
+	FGuid CurrentChunkGuid;
+
+	// The current chunk data
+	uint8* CurrentChunkBuffer;
+
+	// The current index into the chunk data
+	uint32 CurrentChunkBufferPos;
+
+	// Holds the current file being processed. Receives the chunk parts that we pass over.
+	TSharedPtr< class FBuildPatchFileManifest > CurrentFile;
+
+	// Holds an SHA hash calculator
+	FSHA1 FileHash;
+
+	// Are we currently processing a chunk
+	bool IsProcessingChunk;
+
+	// Are we currently processing a chunk part
+	bool IsProcessingChunkPart;
+
+	// Are we currently processing a file
+	bool IsProcessingFile;
+
+	// Have we pushed to a recognized chunk
+	bool ChunkIsPushed;
+
+	// Backup of CurrentChunkGuid when pushing a chunk
+	FGuid BackupChunkGuid;
+
+	// Backup of CurrentChunkBufferPos when pushing a chunk
+	uint32 BackupChunkBufferPos;
+
+	// Backup of IsProcessingChunk when pushing a chunk
+	bool BackupProcessingChunk;
+
+	// Backup of IsProcessingChunkPart when pushing a chunk
+	bool BackupProcessingChunkPart;
+
+	// Cached map of chunk file sizes
+	TMap<FGuid, int64> ChunkFileSizes;
+};
+
+/**
+ * A class that controls processing of build data to produce the manifest with files
+ */
+class FBuildDataFileProcessor
+{
+private:
+
+	/**
+	 * Make the default constructor private so that we must have the required params passed in.
+	 */
+	FBuildDataFileProcessor() { check(false); }
+
+public:
+
+	/**
+	 * Constructor
+	 * @param	InBuildManifest		The build manifest that receives the manifest data.
+	 * @param	InBuildRoot			The root directory of the build
+	 */
+	FBuildDataFileProcessor( FBuildPatchAppManifestRef InBuildManifest, const FString& InBuildRoot );
+
+	/**
+	 * Mark the beginning of a new file
+	 * @param	InFileName		The file's local filename
+	 */
+	void BeginFile( const FString& InFileName );
+
+	/**
+	 * Process a set of data for a file
+	 * @param	Data		The data memory
+	 * @param	DataLen		The length of the data
+	 */
+	void ProcessFileData( const uint8* Data, const uint32& DataLen );
+
+	/**
+	 * Mark the end of the current file
+	 */
+	void EndFile();
+
+	/**
+	 * Get the new/know file stats
+	 * @param	OutNewFiles		Receives the number of new files found
+	 * @param	OutKnownFiles	Receives the number of recognised files found
+	 */
+	void GetFileStats( uint32& OutNewFiles, uint32& OutKnownFiles );
+
+private:
+
+	// Keep track of the number of new files
+	uint32 NumNewFiles;
+
+	// Keep track of the number of recognised files
+	uint32 NumKnownFiles;
+
+	// Holds the build manifest that is receiving the information
+	FBuildPatchAppManifestPtr BuildManifest;
+
+	// Holds the build manifest that is receiving the information
+	FString BuildRoot;
+
+	// Holds the current file being processed. Receives the chunk parts that we pass over.
+	TSharedPtr< class FBuildPatchFileManifest > CurrentFile;
+
+	// Holds an SHA hash calculator
+	FSHA1 FileHash;
+
+	// Counts the size of the file
+	uint64 FileSize;
+
+	// Are we currently processing a file
+	bool IsProcessingFile;
+};
+
+/**
+ * A class that can hold recently used chunks in memory so that when hash collisions for particular hashes are common
+ * we are not having to load existing chunks from disk/network every time. This is especially notable for when source data
+ * can have repeated cycles of bytes for long periods.
+ */
+class FBuildGenerationChunkCache
+{
+public:
+	/**
+	 * A child class that is return to provide an interface to chunk data, wraps up reading from RAM or file
+	 */
+	class FChunkReader
+	{
+	private:
+		// The path to the file for this chunk
+		const FString ChunkFilePath;
+
+		// The archive for this chunk if loading from file
+		FArchive* ChunkFileReader;
+
+		// Shared ref to the FChunkFile that holds chunk data
+		TSharedRef< FChunkFile > ChunkFile;
+
+		// Pointer to an unsigned int that stores how many bytes were read from disk so far
+		uint32* FileBytesRead;
+
+		// Pointer to the header data from our FChunkFile
+		FChunkHeader* ChunkHeader;
+
+		// Pointer to the chunk data from our FChunkFile
+		uint8* ChunkData;
+
+		// How many bytes we have read out of this chunk, used to calculate when we need more data from file
+		uint32 MemoryBytesRead;
+
+	public:
+		/**
+		 * Constructs the reader from required data
+		 * @param	InChunkFilePath		The file path of the source file for the chunk
+		 * @param	InChunkFile			The data structure for this chunk
+		 * @param	InBytesRead			Pointer to an unsigned int to keep track of how much data is read from file
+		 */
+		FChunkReader( const FString& InChunkFilePath, TSharedRef< FChunkFile > InChunkFile, uint32* InBytesRead );
+
+		/**
+		 * Destructor. Cleans up and releases data locks.
+		 */
+		~FChunkReader();
+
+		/**
+		 * Whether the chunk is valid according to header information.
+		 * @return		Returns true if the chunk file checked out ok to use
+		 */
+		const bool IsValidChunk();
+
+		/**
+		 * Get the GUID for this chunk
+		 * @return		The chunk GUID
+		 */
+		const FGuid& GetChunkGuid();
+
+		/**
+		 * Reads the next block of bytes from the chunk data, setting the pointer to a memory location of the first byte requested
+		 * @param OutDataBuffer		Pointer to a pointer which gets set to location of next byte
+		 * @param ReadLength		How many bytes should be available to read from the pointer given. This must be <= BytesLeft
+		 */
+		void ReadNextBytes( uint8** OutDataBuffer, const uint32& ReadLength );
+
+		/**
+		 * Get the number of bytes left to read from this chunk data
+		 * @return		The number of bytes left available to read from ReadNextBytes
+		 */
+		const uint32 BytesLeft();
+
+	private:
+		/**
+		 * Returns the FArchive for the chunk source file. Wraps up opening the file, grabbing the header, and seeking to the
+		 * next byte that hasn't been read yet, which all happens only first time we need the file per reader.
+		 * @return		Archive for reading chunk from disk
+		 */
+		FArchive* GetArchive();
+	};
+
+private:
+	// The number of chunks to cache in memory
+	const static int32 NumChunksToCache = 256;
+
+	// The map of chunks we have cached
+	TMap< FString, TSharedRef< FChunkFile > > ChunkCache;
+
+	// The store for how much file data has been read for each chunk
+	TMap< FString, uint32* > BytesReadPerChunk;
+
+	// The singleton instance for this class
+	static TSharedPtr< FBuildGenerationChunkCache > SingletonInstance;
+
+public:
+	/**
+	 * Get a reader class for a chunk file. The reader class is the interface to chunk data which will either come from RAM or file.
+	 * @return		Shared ref to a reader class, in order to allow memory freeing, you should not keep hold of references when finished with chunk
+	 */
+	TSharedRef< FChunkReader > GetChunkReader( const FString& ChunkFilePath );
+
+	/**
+	 * Call to initialise static singleton so that the cache system can be used
+	 */
+	static void Init();
+
+	/**
+	 * Get a reference to the chunk system. Only call between Init and Shutdown calls.
+	 */
+	static FBuildGenerationChunkCache& Get();
+
+	/**
+	 * Call to clean up all data and release memory allocations for the class.
+	 */
+	static void Shutdown();
+
+private:
+	/**
+	 * Clear out and clean up all memory allocations
+	 */
+	void Cleanup();
+};
+
+/**
+ * A class that controls the process of generating manifests and chunk data from a build image, and also constructing builds from manifests and chunks
+ */
+class FBuildDataGenerator
+{
+	friend class FBuildDataChunkProcessor;
+	friend class FBuildDataFileProcessor;
+public:
+
+	/**
+	 * Processes a Build Image to determine new chunks and produce a chunk based manifest, all saved to the cloud.
+	 * NOTE: This function is blocking and will not return until finished. Don't run on main thread.
+	 * @param Settings				Specifies the settings for the operation.  See FBuildPatchSettings documentation.
+	 * @return		true if no file errors occurred
+	 */
+	//static bool GenerateChunksManifestFromDirectory( const FString& RootDirectory, const uint32& InAppID, const FString& AppName, const FString& BuildVersion, const FString& LaunchExe, const FString& LaunchCommand, const FString& IgnoreListFile );
+	static bool GenerateChunksManifestFromDirectory( const FBuildPatchSettings& Settings );
+
+	/**
+	 * Processes a Build Image to determine new raw files and produce a raw file based manifest, all saved to the cloud.
+	 * NOTE: This function is blocking and will not return until finished. Don't run on main thread.
+	 * @param Settings				Specifies the settings for the operation.  See FBuildPatchSettings documentation.
+	 * @return		true if no file errors occurred
+	 */
+	//static bool GenerateFilesManifestFromDirectory( const FString& RootDirectory, const uint32& InAppID, const FString& AppName, const FString& BuildVersion, const FString& LaunchExe, const FString& LaunchCommand, const FString& IgnoreListFile );
+	static bool GenerateFilesManifestFromDirectory( const FBuildPatchSettings& Settings );
+
+	/**
+	 * Checks to see if a given chunk already exists in the database of known chunks
+	 * NOTE: This function is blocking and will not return until finished. Don't run on main thread.
+	 * @param ChunkHash		IN		The hash for this chunk.
+	 * @param ChunkData		IN		The actual chunk data.
+	 * @param ChunkGuid		OUT		Will be set to the existing chunk guid if found.
+	 * @return		Whether this chunk is new or existing.
+	 */
+	static bool FindExistingChunkData( const uint64& ChunkHash, const TRingBuffer< uint8, FBuildPatchData::ChunkDataSize >& ChunkData, FGuid& ChunkGuid );
+	static bool FindExistingChunkData( const uint64& ChunkHash, const uint8* ChunkData, FGuid& ChunkGuid );
+
+	/**
+	 * Given a Guid and Hash for a chunk file, try to find it from a cloud directory, starting from the latest
+	 * version of the naming convention.
+	 * NOTE: This function is blocking and will not return until finished. Don't run on main thread.
+	 * @param ChunkGuid		The guid of the chunk data
+	 * @param ChunkHash		The chunk's hash value
+	 * @return		The full patch to the found chunk, otherwise empty string.
+	 */
+	static FString DiscoverChunkFilename(const FGuid& ChunkGuid, const uint64& ChunkHash);
+
+	/**
+	 * Checks to see if a given file already exists in the database of known files
+	 * NOTE: This function is blocking and will not return until finished. Don't run on main thread.
+	 * @param InSourceFile	IN		The filename of the file on disk.
+	 * @param InFileHash	IN		The hash for the file.
+	 * @param OutFileGuid	OUT		Will be set to the existing file guid if found.
+	 * @return		Whether this file is new or existing.
+	 */
+	static bool FindExistingFileData( const FString& InSourceFile, const FSHAHash& InFileHash, FGuid& OutFileGuid );
+
+	/**
+	 * Saves out the file data for use with file based patch manifests
+	 * NOTE: This function is blocking and will not return until finished. Don't run on main thread.
+	 * @param SourceFile	IN		The filename of the source file on disk.
+	 * @param FileHash		IN		The hash for the file.
+	 * @param FileGuid		OUT		The GUID for the file
+	 * @return		Whether this file is new or existing.
+	 */
+	static bool SaveOutFileData( const FString& SourceFile, const FSHAHash& FileHash, const FGuid& FileGuid );
+
+	/**
+	 * Strips out files from the provided array that appear in the ignore list file
+	 * @param AllFiles			The array of files to be stripped
+	 * @param DepotDirectory	The directory that is the root of the ignore file list contents
+	 * @param IgnoreListFile	Path to a text file containing list of ignorable files
+	 */
+	static void StripIgnoredFiles( TArray< FString >& AllFiles, const FString& DepotDirectory, const FString& IgnoreListFile );
+
+private:
+
+	/**
+	 * Helper function to compare given data to an existing chunk on file
+	 * @param ChunkFilePath       IN    The path to the chunk file.
+	 * @param ChunkData           IN    The data to compare with
+	 * @param ChunkGuid           OUT   This will receive the GUID from the chunk header.
+	 * @param SourceChunkIsValid  OUT   This will be set with whether the source chunk specified is valid
+	 * @return		true if no file errors occurred and the data is identical
+	 */
+	static bool CompareDataToChunk(const FString& ChunkFilePath, uint8* ChunkData, FGuid& ChunkGuid, bool& SourceChunkIsValid);
+
+private:
+	// For Chunk Manifest generation, the existing chunks rolling hash lookup
+	static TMap< uint64, TArray< FGuid > > ExistingChunkHashInventory;
+
+	// For Chunk Manifest generation, the existing chunks GUID lookup
+	static TMap< FGuid, FString > ExistingChunkGuidInventory;
+
+	// For Chunk Manifest generation, have we enumerated the existing chunks yet
+	static bool ExistingChunksEnumerated;
+
+	// For File Manifest generation, the existing files
+	static TMap< FSHAHash, TArray< FString > > ExistingFileInventory;
+
+	// For File Manifest generation, have we enumerated the existing files yet
+	static bool ExistingFilesEnumerated;
+
+	// A critical section to enforce only one build at a time
+	static FCriticalSection SingleConcurrentBuildCS;
+};
+
+#endif //WITH_BUILDPATCHGENERATION
+
+#endif // __BuildPatchGeneration_h__
diff --git a/Engine/Source/Runtime/Online/BuildPatchServices/Private/BuildPatchHTTP.cpp b/Engine/Source/Runtime/Online/BuildPatchServices/Private/BuildPatchHTTP.cpp
new file mode 100644
--- /dev/null
+++ b/Engine/Source/Runtime/Online/BuildPatchServices/Private/BuildPatchHTTP.cpp
@@ -0,0 +1,150 @@
+// Copyright 1998-2014 Epic Games, Inc. All Rights Reserved.
+
+/*=============================================================================
+	BuildPatchHTTP.cpp: Implements a simple class to wrap queuing HTTP requests
+	in order to insure the HTTP module is only used from the main thread.
+=============================================================================*/
+
+#include "BuildPatchServicesPrivatePCH.h"
+
+/* FBuildPatchHTTP implementation
+ *****************************************************************************/
+bool FBuildPatchHTTP::Tick( float Delta )
+{
+	// Process the http request and cancel queues
+	{
+		FScopeLock ScopeLock( &HttpRequestQueueCS );
+		if( FBuildPatchInstallError::HasFatalError() )
+		{
+			// Make sure all requests get canceled if we get fatal error
+			TArray< int32 > RequestKeys;
+			HttpRequestMap.GetKeys( RequestKeys );
+			CancelRequestQueue.Append( RequestKeys );
+			while( HttpRequestQueue.Num() > 0 )
+			{
+				FHttpRequestInfo HttpRequestInfo = HttpRequestQueue.Pop();
+				HttpRequestInfo.OnCompleteDelegate.ExecuteIfBound( NULL, NULL, false );
+			}
+		}
+		else
+		{
+			while( HttpRequestQueue.Num() > 0 )
+			{
+				FHttpRequestInfo HttpRequestInfo = HttpRequestQueue.Pop();
+				TSharedRef< IHttpRequest > HttpRequest = FHttpModule::Get().CreateRequest();
+				HttpRequest->OnProcessRequestComplete() = HttpRequestInfo.OnCompleteDelegate;
+				HttpRequest->OnRequestProgress() = HttpRequestInfo.OnProgressDelegate;
+				HttpRequest->SetURL( HttpRequestInfo.UrlRequest );
+				HttpRequest->ProcessRequest();
+				HttpRequestMap.Add( HttpRequestInfo.RequestID, HttpRequest );
+			}
+		}
+		while( CancelRequestQueue.Num() > 0 )
+		{
+			const int32 CancelID = CancelRequestQueue.Pop();
+			TWeakPtr< IHttpRequest > HttpRequest = HttpRequestMap.FindRef( CancelID );
+			HttpRequestMap.Remove( CancelID );
+			TSharedPtr< IHttpRequest > HttpRequestPinned = HttpRequest.Pin();
+			if( HttpRequestPinned.IsValid() )
+			{
+				HttpRequestPinned->CancelRequest();
+			}
+		}
+	}
+
+	// Process the Analytics Event queue
+	{
+		FScopeLock ScopeLock( &AnalyticsEventQueueCS );
+		while( AnalyticsEventQueue.Num() > 0 )
+		{
+			const FAnalyticsEventInfo& EventInfo = AnalyticsEventQueue[ AnalyticsEventQueue.Num() -1 ];
+			FBuildPatchAnalytics::RecordEvent( EventInfo.EventName, EventInfo.Attributes );
+			AnalyticsEventQueue.Pop();
+		}
+	}
+
+	return true;
+}
+
+int32 FBuildPatchHTTP::QueueHttpRequest( const FString& UrlRequest, const FHttpRequestCompleteDelegate& OnCompleteDelegate, const FHttpRequestProgressDelegate& OnProgressDelegate )
+{
+	if( !FBuildPatchInstallError::HasFatalError() )
+	{
+		FBuildPatchHTTP& BuildPatchHTTP = Get();
+		// Queue the request creation
+		{
+			FScopeLock ScopeLock( &BuildPatchHTTP.HttpRequestQueueCS );
+			FHttpRequestInfo RequestInfo;
+			RequestInfo.OnCompleteDelegate = OnCompleteDelegate;
+			RequestInfo.OnProgressDelegate = OnProgressDelegate;
+			RequestInfo.UrlRequest = UrlRequest;
+			RequestInfo.RequestID = BuildPatchHTTP.HttpRequestIDCount++;
+			BuildPatchHTTP.HttpRequestQueue.Add( RequestInfo );
+			return RequestInfo.RequestID;
+		}
+	}
+	else
+	{
+		OnCompleteDelegate.ExecuteIfBound( NULL, NULL, false );
+		return INDEX_NONE;
+	}
+}
+
+void FBuildPatchHTTP::CancelHttpRequest( const int32& RequestID )
+{
+	FBuildPatchHTTP& BuildPatchHTTP = Get();
+	// Queue the request cancellation
+	{
+		FScopeLock ScopeLock( &BuildPatchHTTP.HttpRequestQueueCS );
+		BuildPatchHTTP.CancelRequestQueue.Add( RequestID );
+	}
+}
+
+void FBuildPatchHTTP::CancelAllHttpRequests()
+{
+	FBuildPatchHTTP& BuildPatchHTTP = Get();
+	// Queue the request cancellation
+	{
+		FScopeLock ScopeLock( &BuildPatchHTTP.HttpRequestQueueCS );
+		TArray< int32 > RequestKeys;
+		BuildPatchHTTP.HttpRequestMap.GetKeys( RequestKeys );
+		BuildPatchHTTP.CancelRequestQueue.Append( RequestKeys );
+		BuildPatchHTTP.Tick(0);
+	}
+}
+
+void FBuildPatchHTTP::QueueAnalyticsEvent( const FString& EventName, const TArray< FAnalyticsEventAttribute >& Attributes )
+{
+	FBuildPatchHTTP& BuildPatchHTTP = Get();
+	// Queue the event recording
+	{
+		FScopeLock ScopeLock( &BuildPatchHTTP.AnalyticsEventQueueCS );
+		FAnalyticsEventInfo EventInfo;
+		EventInfo.EventName = EventName;
+		EventInfo.Attributes = Attributes;
+		BuildPatchHTTP.AnalyticsEventQueue.Add( EventInfo );
+	}
+}
+
+/* Static initialization
+ *****************************************************************************/
+FBuildPatchHTTP* FBuildPatchHTTP::SingletonInstance = NULL;
+FBuildPatchHTTP& FBuildPatchHTTP::Get()
+{
+	if( SingletonInstance == NULL )
+	{
+		SingletonInstance = new FBuildPatchHTTP();
+		FTicker::GetCoreTicker().AddTicker( FTickerDelegate::CreateRaw( SingletonInstance, &FBuildPatchHTTP::Tick ) );
+	}
+	return *SingletonInstance;
+}
+void FBuildPatchHTTP::OnShutdown()
+{
+	if( SingletonInstance != NULL )
+	{
+		SingletonInstance->Tick( 0.0f );
+		FTicker::GetCoreTicker().RemoveTicker( FTickerDelegate::CreateRaw( SingletonInstance, &FBuildPatchHTTP::Tick ) );
+		delete SingletonInstance;
+		SingletonInstance = NULL;
+	}
+}
diff --git a/Engine/Source/Runtime/Online/BuildPatchServices/Private/BuildPatchHTTP.h b/Engine/Source/Runtime/Online/BuildPatchServices/Private/BuildPatchHTTP.h
new file mode 100644
--- /dev/null
+++ b/Engine/Source/Runtime/Online/BuildPatchServices/Private/BuildPatchHTTP.h
@@ -0,0 +1,124 @@
+// Copyright 1998-2014 Epic Games, Inc. All Rights Reserved.
+
+/*=============================================================================
+	BuildPatchHTTP.h: Declares a simple class to wrap queuing HTTP requests
+	in order to insure the HTTP module is only used from the main thread
+=============================================================================*/
+
+#ifndef __BuildPatchHTTP_h__
+#define __BuildPatchHTTP_h__
+
+#pragma once
+
+/**
+ * Implements the BuildPatchHTTP
+ */
+class FBuildPatchHTTP : private FNoncopyable
+{
+private:
+
+	/**
+	 * A simple struct to hold details required to set off a HTTP request
+	 */
+	struct FHttpRequestInfo
+	{
+		// Holds the delegate that will get called when the request completes
+		FHttpRequestCompleteDelegate OnCompleteDelegate;
+
+		// Holds the delegate that will get called for the request progress
+		FHttpRequestProgressDelegate OnProgressDelegate;
+
+		// The url for the request
+		FString UrlRequest;
+
+		// An ID for this request
+		int32 RequestID;
+	};
+
+	/**
+	 * A simple struct to hold details required to record an analytics event
+	 */
+	struct FAnalyticsEventInfo
+	{
+		// The analytics event name
+		FString EventName;
+
+		// The list of attributes
+		TArray< FAnalyticsEventAttribute > Attributes;
+	};
+
+public:
+
+	/**
+	 * A helper to queue up a HTTP request, which must be created on the main thread.
+	 * @param UrlRequest				The url for the request
+	 * @param OnCompletionDelegate		The delegate to be called on request complete
+	 * @param OnProgressDelegate		The delegate to be called for request progress
+	 * @return	An ID for this request
+	 */
+	static int32 QueueHttpRequest( const FString& UrlRequest, const FHttpRequestCompleteDelegate& OnCompleteDelegate, const FHttpRequestProgressDelegate& OnProgressDelegate );
+
+	/**
+	 * Cancel a HTTP request started with QueueHttpRequest.
+	 * @param RequestID		The ID that was returned from QueueHttpRequest for the request.
+	 */
+	static void CancelHttpRequest( const int32& RequestID );
+
+	/**
+	 * Cancels all HTTP requests started with QueueHttpRequest.
+	 */
+	static void CancelAllHttpRequests();
+
+	/**
+	 * A helper to queue up an Analytics Event, which must be created on the main thread.
+	 * @param EventName		The name of the event
+	 * @param Attributes	The list of attributes
+	 */
+	static void QueueAnalyticsEvent( const FString& EventName, const TArray< FAnalyticsEventAttribute >& Attributes );
+
+	/**
+	 * Clears any remaining queue and deletes the singleton. NB: Be weary of timing, if one of the static functions above is called after this
+	 * then the singleton will be re-created
+	 */
+	static void OnShutdown();
+
+private:
+
+	// A queue of HTTP requests that need to be started
+	TArray< FHttpRequestInfo > HttpRequestQueue;
+
+	// A queue of requests that have been canceled
+	TArray< int32 > CancelRequestQueue;
+
+	// A map of HTTP requests that have been started
+	TMap< int32, TWeakPtr< IHttpRequest > > HttpRequestMap;
+
+	// A counter for giving out HttpRequest IDs
+	int32 HttpRequestIDCount;
+
+	// A critical section to protect access to HttpRequestQueue
+	FCriticalSection HttpRequestQueueCS;
+
+	// A queue of Analytics Events that need to be recorded
+	TArray< FAnalyticsEventInfo > AnalyticsEventQueue;
+
+	// A critical section to protect access to AnalyticsEventQueue
+	FCriticalSection AnalyticsEventQueueCS;
+
+private:
+
+	/**
+	 * A ticker used to perform actions on the main thread.
+	 * @param Delta		The tick delta
+	 * @return	true to continue ticking
+	 */
+	bool Tick( float Delta );
+
+/* Here we have private static access for the singleton
+*****************************************************************************/
+private:
+	static FBuildPatchHTTP& Get();
+	static FBuildPatchHTTP* SingletonInstance;
+};
+
+#endif // __BuildPatchHTTP_h__
