Moving Build Patch Services from NoRedist to Runtime

diff --git a/Engine/Source/Runtime/Online/BuildPatchServices/Private/BuildPatchHash.h b/Engine/Source/Runtime/Online/BuildPatchServices/Private/BuildPatchHash.h
new file mode 100644
--- /dev/null
+++ b/Engine/Source/Runtime/Online/BuildPatchServices/Private/BuildPatchHash.h
@@ -0,0 +1,603 @@
+// Copyright 1998-2014 Epic Games, Inc. All Rights Reserved.
+
+/*=============================================================================
+	BuildPatchHash.h: Declares and implements template classes for use with the 
+	Build and Patching hash functionality.
+=============================================================================*/
+
+#ifndef __BuildPatchHash_h__
+#define __BuildPatchHash_h__
+
+#pragma once
+
+/**
+ * A macro for barrel rolling a 64 bit value n times to the left.
+ * @param Value		The value to shift
+ * @param Shifts	The number of times to shift
+ */
+#define ROTLEFT_64B( Value, Shifts ) Value = ( ( ( Value ) << ( ( Shifts ) % 64 ) ) | ( ( Value ) >> ( ( 64 - ( ( Shifts ) % 64 ) ) % 64 ) ) )
+
+/**
+ * A ring buffer template class.
+ */
+template< typename DataType, uint32 BufferDataSize >
+class TRingBuffer
+{
+public:
+	/**
+	 * Default Constructor
+	 */
+	TRingBuffer();
+
+	/**
+	 * Default Destructor
+	 */
+	~TRingBuffer();
+
+	/**
+	 * Clears memory and indexes
+	 */
+	void Empty();
+
+	/**
+	 * Pushes a data word to the end of the FIFO. WILL OVERWRITE OLDEST if full.
+	 * @param Val	The data word to push
+	 */
+	void Enqueue( const DataType& Val );
+
+	/**
+	 * Pushes a buffer of data words to the end of the FIFO. WILL OVERWRITE OLDEST if full.
+	 * @param ValBuf		The buffer pointer
+	 * @param BufLen		The length to copy
+	 */
+	void Enqueue( const DataType* ValBuf, const uint32& BufLen );
+
+	/**
+	 * Take the next data word from the FIFO buffer.
+	 * @param OUT	ValOut	The variable to receive the data word
+	 * @return	true if the buffer was not empty, false otherwise.
+	 */
+	bool Dequeue( DataType& ValOut );
+
+	/**
+	 * Take the next set of data words from the FIFO buffer.
+	 * @param ValBuf		The buffer to receive the data
+	 * @param BufLen		The number of words requested
+	 * @return	the number of words actually copied
+	 */
+	uint32 Dequeue( DataType* ValBuf, const uint32& BufLen );
+
+	/**
+	 * Get the next data word from the FIFO buffer without removing it.
+	 * @param OUT	ValOut	The variable to receive the data word
+	 * @return	true if the buffer was not empty, false otherwise.
+	 */
+	bool Peek( DataType& ValOut ) const;
+
+	/**
+	 * Get the next set of data words from the FIFO buffer without removing them.
+	 * @param ValBuf		The buffer to receive the data
+	 * @param BufLen		The number of words requested
+	 * @return	the number of words actually copied
+	 */
+	uint32 Peek( DataType* ValBuf, const uint32& BufLen ) const;
+
+	/**
+	 * Compare the memory in the FIFO to the memory in the given buffer
+	 * @param SerialBuffer	The buffer containing data to compare
+	 * @param CompareLen		The number of words to compare
+	 * @return	< 0 if data < SerialBuffer.. 0 if data == SerialBuffer... > 0 if data > SerialBuffer
+	 */
+	int32 SerialCompare( const DataType* SerialBuffer, uint32 CompareLen ) const;
+
+	/**
+	 * Serializes the internal buffer into the given buffer
+	 * @param SerialBuffer	A preallocated buffer to copy data into
+	 */
+	void Serialize( DataType* SerialBuffer ) const;
+
+	/**
+	 * Square bracket operators for accessing data in the buffer by FIFO index. [0] returns the next entry in the FIFO that would get provided by Dequeue or Peek.
+	 * @param Index	The index into the FIFO buffer
+	 * @return	The data word
+	 */
+	FORCEINLINE const DataType& operator[]( const int32& Index ) const;
+	FORCEINLINE DataType& operator[]( const int32& Index );
+
+	/**
+	 * Gets the last data word in the FIFO (i.e. most recently pushed).
+	 * @return	The last data word
+	 */
+	FORCEINLINE const DataType& Top() const;
+	FORCEINLINE DataType& Top();
+
+	/**
+	 * Gets the first data word in the FIFO (i.e. oldest).
+	 * @return	The first data word
+	 */
+	FORCEINLINE const DataType& Bottom() const;
+	FORCEINLINE DataType& Bottom();
+
+	/**
+	 * Gets the buffer index that the last data word is stored in.
+	 * @return	The index of Top()
+	 */
+	FORCEINLINE const uint32 TopIndex() const;
+
+	/**
+	 * Gets the buffer index that the first data word is stored in.
+	 * @return	The index of Bottom()
+	 */
+	FORCEINLINE const uint32 BottomIndex() const;
+
+	/**
+	 * Gets the buffer index that the next enqueued word will get stored in.
+	 * @return	The index that the next enqueued word will have
+	 */
+	FORCEINLINE const uint32 NextIndex() const;
+
+	/**
+	 * Gets the size of the data buffer.
+	 * @return	Template arg BufferDataSize
+	 */
+	FORCEINLINE const uint32 RingDataSize() const;
+
+	/**
+	 * Gets the number of words currently in the FIFO.
+	 * @return	data size
+	 */
+	FORCEINLINE const uint32 RingDataUsage() const;
+
+	/**
+	 * Gets the total number of words that have been pushed through this buffer since clearing
+	 * @return	total length of data passed through
+	 */
+	FORCEINLINE const uint64 TotalDataPushed() const;
+
+
+private:
+	// The data memory
+	DataType* Data;
+
+	// Value keeping track of free space
+	uint32 NumDataAvailable;
+
+	// Value keeping track of the next data index
+	uint32 DataIndex;
+
+	// Value to keep track of total amount of data Enqueued
+	uint64 TotalNumDataPushed;
+};
+
+/**
+ * A static struct containing 64bit polynomial constant and hash table lookup for use with FRollingHash.
+ */
+struct FRollingHashConst
+{
+	// The lookup hash table
+	static uint64 HashTable[ 256 ];
+	
+	/**
+	 * Builds the hash table for use when hashing. Must be called before using FRollingHash.
+	 */
+	static void Init();
+};
+
+/**
+ * A class that performs a rolling hash
+ */
+template< uint32 WindowSize >
+class FRollingHash
+{
+	// A typedef for our data ring buffer
+	typedef TRingBuffer< uint8, WindowSize > HashRingBuffer;
+
+public:
+	/**
+	 * Constructor
+	 */
+	FRollingHash();
+
+	/**
+	 * Pass this function the initial data set to start the Rolling Hash with.
+	 * @param NewByte		The byte to add to the hash state
+	 */
+	void ConsumeByte( const uint8& NewByte );
+
+	/**
+	 * Helper to consume from a byte array
+	 * @param NewBytes		The byte array
+	 * @param NumBytes		The number of bytes to consume
+	 */
+	void ConsumeBytes( const uint8* NewBytes, const uint32& NumBytes );
+
+	/**
+	 * Rolls the window by one byte forwards.
+	 * @param NewByte		The byte that will be added to the front of the window
+	 */
+	void RollForward( const uint8& NewByte );
+
+	/**
+	 * Clears all data ready for a new entire data set
+	 */
+	void Clear();
+
+	/**
+	 * Get the HashState for the current window
+	 * @return		The hash state
+	 */
+	const uint64 GetWindowHash() const;
+
+	/**
+	 * Get the Ring Buffer for the current window
+	 * @return		The ring buffer
+	 */
+	const HashRingBuffer& GetWindowData() const;
+
+	/**
+	 * Get how many DataValueType values we still need to consume until our window is full 
+	 * @return		The number of DataValueType we still need
+	 */
+	const uint32 GetNumDataNeeded() const;
+
+	/**
+	 * Get the size of our window 
+	 * @return		Template arg WindowSize
+	 */
+	const uint32 GetWindowSize() const;
+
+	/**
+	 * Static function to simply return the hash for a given data range.
+	 * @param DataSet	The buffer to the data. This must be a buffer of length == WindowSize
+	 * @return			The hash state for the provided data
+	 */
+	static uint64 GetHashForDataSet( const uint8* DataSet );
+
+private:
+	// The current hash value
+	uint64 HashState;
+	// The number of bytes we have consumed so far, used in hash function and to check validity of calls
+	uint32 NumBytesConsumed;
+	// Store the data to make access and rolling easier
+	HashRingBuffer WindowData;
+};
+
+
+/* TRingBuffer implementation
+*****************************************************************************/
+template< typename DataType, uint32 BufferDataSize >
+TRingBuffer< DataType, BufferDataSize >::TRingBuffer()
+{
+	Data = new DataType[ BufferDataSize ];
+	Empty();
+}
+
+template< typename DataType, uint32 BufferDataSize >
+TRingBuffer< DataType, BufferDataSize >::~TRingBuffer()
+{
+	delete[] Data;
+}
+
+template< typename DataType, uint32 BufferDataSize >
+void TRingBuffer< DataType, BufferDataSize >::Empty()
+{
+	FMemory::Memzero( Data, sizeof( DataType ) * BufferDataSize );
+	NumDataAvailable = 0;
+	DataIndex = 0;
+	TotalNumDataPushed = 0;
+}
+
+template< typename DataType, uint32 BufferDataSize >
+void TRingBuffer< DataType, BufferDataSize >::Enqueue( const DataType& Val )
+{
+	Data[ DataIndex++ ] = Val;
+	DataIndex %= BufferDataSize;
+	++TotalNumDataPushed;
+	NumDataAvailable = FMath::Clamp< uint32 >( NumDataAvailable + 1, 0, BufferDataSize );
+}
+
+template< typename DataType, uint32 BufferDataSize >
+void TRingBuffer< DataType, BufferDataSize >::Enqueue( const DataType* ValBuf, const uint32& BufLen )
+{
+	for( uint32 i = 0; i < BufLen; ++i )
+	{
+		Enqueue( ValBuf[i] );
+	}
+}
+
+template< typename DataType, uint32 BufferDataSize >
+bool TRingBuffer< DataType, BufferDataSize >::Dequeue( DataType& ValOut )
+{
+	if( NumDataAvailable > 0 )
+	{
+		ValOut = Bottom();
+		NumDataAvailable = FMath::Clamp< uint32 >( NumDataAvailable - 1, 0, BufferDataSize );
+		return true;
+	}
+	else
+	{
+		return false;
+	}
+}
+
+template< typename DataType, uint32 BufferDataSize >
+uint32 TRingBuffer< DataType, BufferDataSize >::Dequeue( DataType* ValBuf, const uint32& BufLen )
+{
+	uint32 Count = 0;
+	for( uint32 i = 0; i < BufLen; ++i )
+	{
+		if( Dequeue( ValBuf[i] ) )
+		{
+			++Count;
+		}
+		else
+		{
+			break;
+		}
+	}
+	return Count;
+}
+
+template< typename DataType, uint32 BufferDataSize >
+bool TRingBuffer< DataType, BufferDataSize >::Peek( DataType& ValOut ) const
+{
+	if( NumDataAvailable > 0 )
+	{
+		ValOut = Bottom();
+		return true;
+	}
+	else
+	{
+		return false;
+	}
+}
+
+template< typename DataType, uint32 BufferDataSize >
+uint32 TRingBuffer< DataType, BufferDataSize >::Peek( DataType* ValBuf, const uint32& BufLen ) const
+{
+	uint32 Count = 0;
+	for( uint32 i = 0; i < BufLen; ++i )
+	{
+		if( Peek( ValBuf[i] ) )
+		{
+			++Count;
+		}
+		else
+		{
+			break;
+		}
+	}
+	return Count;
+}
+
+template< typename DataType, uint32 BufferDataSize >
+int32 TRingBuffer< DataType, BufferDataSize >::SerialCompare( const DataType* SerialBuffer, uint32 CompareLen ) const
+{
+	check( CompareLen <= BufferDataSize);
+	const uint32 FirstPartLen = BufferDataSize - DataIndex;
+	int32 FirstCmp = FMemory::Memcmp( Data + BottomIndex(), SerialBuffer, sizeof( DataType ) * FMath::Min( FirstPartLen, CompareLen ) );
+	if( FirstCmp != 0 || FirstPartLen >= CompareLen)
+	{
+		return FirstCmp;
+	}
+	return FMemory::Memcmp( Data, &SerialBuffer[ FirstPartLen ], sizeof( DataType ) * ( CompareLen - FirstPartLen ) );
+}
+
+template< typename DataType, uint32 BufferDataSize >
+void TRingBuffer< DataType, BufferDataSize >::Serialize( DataType* SerialBuffer ) const
+{
+	const uint32 FirstPartLen = BufferDataSize - DataIndex;
+ 	FMemory::Memcpy( SerialBuffer, Data + BottomIndex(), sizeof( DataType ) * FirstPartLen );
+ 	FMemory::Memcpy( &SerialBuffer[ FirstPartLen ], Data, sizeof( DataType ) * ( BufferDataSize - FirstPartLen ) );
+}
+
+template< typename DataType, uint32 BufferDataSize >
+FORCEINLINE const DataType& TRingBuffer< DataType, BufferDataSize >::operator[]( const int32& Index ) const
+{
+	return Data[ ( BottomIndex() + Index ) % BufferDataSize ];
+}
+
+template< typename DataType, uint32 BufferDataSize >
+FORCEINLINE DataType& TRingBuffer< DataType, BufferDataSize >::operator[]( const int32& Index )
+{
+	return Data[ ( BottomIndex() + Index ) % BufferDataSize ];
+}
+
+template< typename DataType, uint32 BufferDataSize >
+FORCEINLINE const DataType& TRingBuffer< DataType, BufferDataSize >::Top() const
+{
+	return Data[ TopIndex() ];
+}
+
+template< typename DataType, uint32 BufferDataSize >
+FORCEINLINE DataType& TRingBuffer< DataType, BufferDataSize >::Top()
+{
+	return Data[ TopIndex() ];
+}
+
+template< typename DataType, uint32 BufferDataSize >
+FORCEINLINE const DataType& TRingBuffer< DataType, BufferDataSize >::Bottom() const
+{
+	return Data[ BottomIndex() ];
+}
+
+template< typename DataType, uint32 BufferDataSize >
+FORCEINLINE DataType& TRingBuffer< DataType, BufferDataSize >::Bottom()
+{
+	return Data[ BottomIndex() ];
+}
+
+template< typename DataType, uint32 BufferDataSize >
+FORCEINLINE const uint32 TRingBuffer< DataType, BufferDataSize >::TopIndex() const
+{
+	return ( DataIndex + BufferDataSize - 1 ) % BufferDataSize;
+}
+
+template< typename DataType, uint32 BufferDataSize >
+FORCEINLINE const uint32 TRingBuffer< DataType, BufferDataSize >::BottomIndex() const
+{
+	return ( DataIndex + BufferDataSize - NumDataAvailable ) % BufferDataSize;
+}
+
+template< typename DataType, uint32 BufferDataSize >
+FORCEINLINE const uint32 TRingBuffer< DataType, BufferDataSize >::NextIndex() const
+{
+	return DataIndex;
+}
+
+template< typename DataType, uint32 BufferDataSize >
+FORCEINLINE const uint32 TRingBuffer< DataType, BufferDataSize >::RingDataSize() const
+{
+	return BufferDataSize;
+}
+
+template< typename DataType, uint32 BufferDataSize >
+FORCEINLINE const uint32 TRingBuffer< DataType, BufferDataSize >::RingDataUsage() const
+{
+	return NumDataAvailable;
+}
+
+template< typename DataType, uint32 BufferDataSize >
+FORCEINLINE const uint64 TRingBuffer< DataType, BufferDataSize >::TotalDataPushed() const
+{
+	return TotalNumDataPushed;
+}
+
+/* FRollingHash implementation
+*****************************************************************************/
+template< uint32 WindowSize >
+FRollingHash< WindowSize >::FRollingHash()
+	: HashState( 0 )
+	, NumBytesConsumed( 0 )
+{
+}
+
+template< uint32 WindowSize >
+void FRollingHash< WindowSize >::ConsumeByte( const uint8& NewByte )
+{
+	// We must be setup by consuming the correct amount of bytes
+	check( NumBytesConsumed < WindowSize );
+	++NumBytesConsumed;
+
+	// Add the byte to our buffer
+	WindowData.Enqueue( NewByte );
+	// Add to our HashState
+	ROTLEFT_64B( HashState, 1 );
+	HashState ^= FRollingHashConst::HashTable[ NewByte ];
+}
+
+template< uint32 WindowSize >
+void FRollingHash< WindowSize >::ConsumeBytes( const uint8* NewBytes, const uint32& NumBytes )
+{
+	for ( uint32 i = 0; i < NumBytes; ++i )
+	{
+		ConsumeByte( NewBytes[i] );
+	}
+}
+
+template< uint32 WindowSize >
+const uint32 FRollingHash< WindowSize >::GetNumDataNeeded() const
+{
+	return WindowSize - NumBytesConsumed;
+}
+
+template< uint32 WindowSize >
+const uint32 FRollingHash< WindowSize >::GetWindowSize() const
+{
+	return WindowSize;
+}
+
+template< uint32 WindowSize >
+void FRollingHash< WindowSize >::RollForward( const uint8& NewByte )
+{
+	// We must have consumed enough bytes to function correctly
+	check( NumBytesConsumed == WindowSize );
+	uint8 OldByte;
+	WindowData.Dequeue( OldByte );
+	WindowData.Enqueue( NewByte );
+	// Update our HashState
+	uint64 OldTerm = FRollingHashConst::HashTable[ OldByte ];
+	ROTLEFT_64B( OldTerm, WindowSize );
+	ROTLEFT_64B( HashState, 1 );
+	HashState ^= OldTerm;
+	HashState ^= FRollingHashConst::HashTable[ NewByte ];
+}
+
+template< uint32 WindowSize >
+void FRollingHash< WindowSize >::Clear()
+{
+	HashState = 0;
+	NumBytesConsumed = 0;
+	WindowData.Empty();
+}
+
+template< uint32 WindowSize >
+const uint64 FRollingHash< WindowSize >::GetWindowHash() const
+{
+	// We must have consumed enough bytes to function correctly
+	check( NumBytesConsumed == WindowSize );
+	return HashState;
+}
+
+template< uint32 WindowSize >
+const TRingBuffer< uint8, WindowSize >& FRollingHash< WindowSize >::GetWindowData() const
+{
+	return WindowData;
+}
+
+template< uint32 WindowSize >
+uint64 FRollingHash< WindowSize >::GetHashForDataSet( const uint8* DataSet )
+{
+	uint64 HashState = 0;
+	for ( uint32 i = 0; i < WindowSize; ++i )
+	{
+		ROTLEFT_64B( HashState, 1 );
+		HashState ^= FRollingHashConst::HashTable[ DataSet[i] ];
+	}
+	return HashState;
+}
+
+/**
+* A static function to perform sanity checks on the rolling hash class.
+*/
+static bool CheckRollingHashAlgorithm()
+{
+	bool bCheckOk = true;
+	// Sanity Check the RollingHash code!!
+	FString IndivWords[6];
+	IndivWords[0] = "123456";	IndivWords[1] = "7890-=";	IndivWords[2] = "qwerty";	IndivWords[3] = "uiop[]";	IndivWords[4] = "asdfgh";	IndivWords[5] = "jkl;'#";
+	FString DataToRollOver = FString::Printf( TEXT( "%s%s%s%s%s%s" ), *IndivWords[0], *IndivWords[1], *IndivWords[2], *IndivWords[3], *IndivWords[4], *IndivWords[5] );
+	uint64 IndivHashes[6];
+	for (uint32 idx = 0; idx < 6; ++idx )
+	{
+		uint8 Converted[6];
+		for (uint32 iChar = 0; iChar < 6; ++iChar )
+			Converted[iChar] = IndivWords[idx][iChar];
+		IndivHashes[idx] = FRollingHash< 6 >::GetHashForDataSet( Converted );
+	}
+
+	FRollingHash< 6 > RollingHash;
+	uint32 StrIdx = 0;
+	for (uint32 k=0; k<6; ++k)
+		RollingHash.ConsumeByte( DataToRollOver[ StrIdx++ ] );
+	bCheckOk = bCheckOk && IndivHashes[0] == RollingHash.GetWindowHash();
+	for (uint32 k=0; k<6; ++k)
+		RollingHash.RollForward( DataToRollOver[ StrIdx++ ] );
+	bCheckOk = bCheckOk && IndivHashes[1] == RollingHash.GetWindowHash();
+	for (uint32 k=0; k<6; ++k)
+		RollingHash.RollForward( DataToRollOver[ StrIdx++ ] );
+	bCheckOk = bCheckOk && IndivHashes[2] == RollingHash.GetWindowHash();
+	for (uint32 k=0; k<6; ++k)
+		RollingHash.RollForward( DataToRollOver[ StrIdx++ ] );
+	bCheckOk = bCheckOk && IndivHashes[3] == RollingHash.GetWindowHash();
+	for (uint32 k=0; k<6; ++k)
+		RollingHash.RollForward( DataToRollOver[ StrIdx++ ] );
+	bCheckOk = bCheckOk && IndivHashes[4] == RollingHash.GetWindowHash();
+	for (uint32 k=0; k<6; ++k)
+		RollingHash.RollForward( DataToRollOver[ StrIdx++ ] );
+	bCheckOk = bCheckOk && IndivHashes[5] == RollingHash.GetWindowHash();
+
+	return bCheckOk;
+}
+
+#endif // __BuildPatchHash_h__
diff --git a/Engine/Source/Runtime/Online/BuildPatchServices/Private/BuildPatchInstaller.cpp b/Engine/Source/Runtime/Online/BuildPatchServices/Private/BuildPatchInstaller.cpp
new file mode 100644
--- /dev/null
+++ b/Engine/Source/Runtime/Online/BuildPatchServices/Private/BuildPatchInstaller.cpp
@@ -0,0 +1,963 @@
+// Copyright 1998-2014 Epic Games, Inc. All Rights Reserved.
+
+/*=============================================================================
+	BuildPatchInstaller.cpp: Implements the FBuildPatchInstaller class which
+	controls the process of installing a build described by a build manifest.
+=============================================================================*/
+
+#include "BuildPatchServicesPrivatePCH.h"
+
+#define LOCTEXT_NAMESPACE "BuildPatchInstaller"
+
+#define NUM_DOWNLOAD_READINGS	5
+
+#define TIME_PER_READING		0.5
+
+#define MEGABYTE_VALUE			1048576
+
+#define KILOBYTE_VALUE			1024
+
+#define NUM_FILE_MOVE_RETRIES	5
+
+/* FBuildPatchInstaller implementation
+*****************************************************************************/
+FBuildPatchInstaller::FBuildPatchInstaller( FBuildPatchBoolManifestDelegate InOnCompleteDelegate, FBuildPatchAppManifestPtr CurrentManifest, FBuildPatchAppManifestRef InstallManifest, const FString& InInstallDirectory, const FString& InStagingDirectory, FBuildPatchInstallationInfo& InstallationInfoRef )
+	: Thread( NULL )
+	, OnCompleteDelegate( InOnCompleteDelegate )
+	, CurrentBuildManifest( CurrentManifest )
+	, NewBuildManifest( InstallManifest )
+	, InstallDirectory( InInstallDirectory )
+	, StagingDirectory( InStagingDirectory )
+	, DataStagingDir( InStagingDirectory / TEXT( "PatchData" ) )
+	, InstallStagingDir( InStagingDirectory / TEXT( "Install" ) )
+	, PreviousMoveMarker( InstallDirectory / TEXT( "$movedMarker" ) )
+	, ThreadLock()
+	, bIsFileData( InstallManifest->IsFileDataManifest() )
+	, bIsChunkData( !bIsFileData )
+	, bIsRepairing( CurrentManifest == InstallManifest )
+	, bSuccess( true )
+	, bIsRunning( false )
+	, bIsInited( false )
+	, DownloadSpeedValue( 0.0 )
+	, DownloadBytesLeft( 0 )
+	, BuildStats()
+	, BuildProgress()
+	, InitialNumChunkDownloads( 0.0f )
+	, InitialNumChunkConstructions( 0.0f )
+	, TotalInitialDownloadSize( 0 )
+	, TimePausedAt( 0.0 )
+	, InstallationInfo( InstallationInfoRef )
+{
+	// Start thread!
+	const TCHAR* ThreadName = TEXT( "BuildPatchInstallerThread" );
+	Thread = FRunnableThread::Create(this, ThreadName);
+}
+
+FBuildPatchInstaller::~FBuildPatchInstaller()
+{
+	if( Thread != NULL )
+	{
+		Thread->WaitForCompletion();
+		delete Thread;
+		Thread = NULL;
+	}
+}
+
+bool FBuildPatchInstaller::Init()
+{
+	// Make sure the installation directory exists
+	IFileManager::Get().MakeDirectory( *InstallDirectory, true );
+
+	// Init build stats that count
+	BuildStats.ProcessPausedTime = 0.0f;
+
+	// We are ready to go if our delegates are bound and directories successfully created
+	bool bInstallerInitSuccess = OnCompleteDelegate.IsBound();
+	bInstallerInitSuccess &= IFileManager::Get().DirectoryExists( *InstallDirectory );
+
+	// Currently we don't handle init failures, so make sure we are not missing them
+	check( bInstallerInitSuccess );
+	return bInstallerInitSuccess;
+}
+
+uint32 FBuildPatchInstaller::Run()
+{
+	// Make sure this function can never be parallelized
+	static FCriticalSection SingletonFunctionLockCS;
+	FScopeLock SingletonFunctionLock(&SingletonFunctionLockCS);
+	FBuildPatchInstallError::Reset();
+
+	SetRunning(true);
+	SetInited(true);
+	SetDownloadSpeed(-1);
+	UpdateDownloadProgressInfo(true);
+
+	// Register the current manifest with the installation info, to make sure we pull from it
+	if (CurrentBuildManifest.IsValid())
+	{
+		InstallationInfo.RegisterAppInstallation(CurrentBuildManifest.ToSharedRef(), InstallDirectory);
+	}
+
+	// Keep track of files that failed verify
+	TArray<FString> CorruptFiles;
+
+	// Init prereqs progress value
+	const bool bInstallPrereqs = !CurrentBuildManifest.IsValid() && !NewBuildManifest->GetPrereqPath().IsEmpty();
+
+	// Get the start time
+	double StartTime = FPlatformTime::Seconds();
+	double CleanUpTime = 0;
+
+	// Keep retrying the install while it is not canceled, or caused by download error
+	bool bProcessSuccess = false;
+	bool bCanRetry = true;
+	int32 InstallRetries = 5;
+	while (!bProcessSuccess && bCanRetry)
+	{
+		// Run the install
+		bool bInstallSuccess = RunInstallation(CorruptFiles);
+		BuildProgress.SetStateProgress(EBuildPatchProgress::PrerequisitesInstall, bInstallPrereqs ? 0.0f : 1.0f);
+		if (bInstallSuccess)
+		{
+			BuildProgress.SetStateProgress(EBuildPatchProgress::Downloading, 1.0f);
+			BuildProgress.SetStateProgress(EBuildPatchProgress::Installing, 1.0f);
+		}
+
+		// Backup local changes then move generated files
+		bInstallSuccess = bInstallSuccess && RunBackupAndMove();
+
+		// Run Verification
+		CorruptFiles.Empty();
+		BuildProgress.SetStateProgress(EBuildPatchProgress::Initializing, 1.0f);
+		bProcessSuccess = bInstallSuccess && RunVerification(CorruptFiles);
+
+		// Clean staging if INSTALL success
+		if (bInstallSuccess)
+		{
+			GLog->Logf(TEXT("BuildPatchServices: Deleting staging area"));
+			CleanUpTime = FPlatformTime::Seconds();
+			IFileManager::Get().DeleteDirectory(*StagingDirectory, false, true);
+			CleanUpTime = FPlatformTime::Seconds() - CleanUpTime;
+		}
+		BuildProgress.SetStateProgress(EBuildPatchProgress::CleanUp, 1.0f);
+
+		// Set if we can retry
+		--InstallRetries;
+		bCanRetry = InstallRetries > 0 && !FBuildPatchInstallError::IsInstallationCancelled() && !FBuildPatchInstallError::IsNoRetryError();
+
+		// If successful or we will retry, remove the moved files marker
+		if (bProcessSuccess || bCanRetry)
+		{
+			GLog->Logf(TEXT("BuildPatchServices: Reset MM"));
+			IFileManager::Get().Delete(*PreviousMoveMarker, false, true);
+		}
+	}
+
+	if (bProcessSuccess)
+	{
+		// Run the prerequisites installer if this is our first install and the manifest has prerequisites info
+		if (bInstallPrereqs)
+		{
+			// @TODO: We also want to trigger prereq install if this is an update and the prereq installer differs in the update
+			bProcessSuccess &= RunPrereqInstaller();
+		}
+	}
+
+	// Set final stat values and log out results
+	{
+		FScopeLock Lock(&ThreadLock);
+		bSuccess = bProcessSuccess;
+		BuildStats.ProcessSuccess = bProcessSuccess;
+		BuildStats.ProcessExecuteTime = (FPlatformTime::Seconds() - StartTime) - BuildStats.ProcessPausedTime;
+		BuildStats.FailureReason = FBuildPatchInstallError::GetErrorString();
+		BuildStats.FailureReasonText = FBuildPatchInstallError::GetErrorText();
+		BuildStats.CleanUpTime = CleanUpTime;
+
+		// Log stats
+		GLog->Logf(TEXT("BuildPatchServices: Build Stat: AppName: %s"), *BuildStats.AppName);
+		GLog->Logf(TEXT("BuildPatchServices: Build Stat: AppInstalledVersion: %s"), *BuildStats.AppInstalledVersion);
+		GLog->Logf(TEXT("BuildPatchServices: Build Stat: AppPatchVersion: %s"), *BuildStats.AppPatchVersion);
+		GLog->Logf(TEXT("BuildPatchServices: Build Stat: CloudDirectory: %s"), *BuildStats.CloudDirectory);
+		GLog->Logf(TEXT("BuildPatchServices: Build Stat: NumFilesInBuild: %u"), BuildStats.NumFilesInBuild);
+		GLog->Logf(TEXT("BuildPatchServices: Build Stat: NumFilesOutdated: %u"), BuildStats.NumFilesOutdated);
+		GLog->Logf(TEXT("BuildPatchServices: Build Stat: NumFilesToRemove: %u"), BuildStats.NumFilesToRemove);
+		GLog->Logf(TEXT("BuildPatchServices: Build Stat: NumChunksRequired: %u"), BuildStats.NumChunksRequired);
+		GLog->Logf(TEXT("BuildPatchServices: Build Stat: ChunksQueuedForDownload: %u"), BuildStats.ChunksQueuedForDownload);
+		GLog->Logf(TEXT("BuildPatchServices: Build Stat: ChunksLocallyAvailable: %u"), BuildStats.ChunksLocallyAvailable);
+		GLog->Logf(TEXT("BuildPatchServices: Build Stat: NumChunksDownloaded: %u"), BuildStats.NumChunksDownloaded);
+		GLog->Logf(TEXT("BuildPatchServices: Build Stat: NumChunksRecycled: %u"), BuildStats.NumChunksRecycled);
+		GLog->Logf(TEXT("BuildPatchServices: Build Stat: NumChunksCacheBooted: %u"), BuildStats.NumChunksCacheBooted);
+		GLog->Logf(TEXT("BuildPatchServices: Build Stat: NumDriveCacheChunkLoads: %u"), BuildStats.NumDriveCacheChunkLoads);
+		GLog->Logf(TEXT("BuildPatchServices: Build Stat: NumRecycleFailures: %u"), BuildStats.NumRecycleFailures);
+		GLog->Logf(TEXT("BuildPatchServices: Build Stat: NumDriveCacheLoadFailures: %u"), BuildStats.NumDriveCacheLoadFailures);
+		GLog->Logf(TEXT("BuildPatchServices: Build Stat: TotalDownloadedData: %lld"), BuildStats.TotalDownloadedData);
+		GLog->Logf(TEXT("BuildPatchServices: Build Stat: AverageDownloadSpeed: %.3f MB/sec"), BuildStats.AverageDownloadSpeed / 1024.0 / 1024.0);
+		GLog->Logf(TEXT("BuildPatchServices: Build Stat: TheoreticalDownloadTime: %s"), *FPlatformTime::PrettyTime(BuildStats.TheoreticalDownloadTime));
+		GLog->Logf(TEXT("BuildPatchServices: Build Stat: VerifyTime: %s"), *FPlatformTime::PrettyTime(BuildStats.VerifyTime));
+		GLog->Logf(TEXT("BuildPatchServices: Build Stat: CleanUpTime: %s"), *FPlatformTime::PrettyTime(BuildStats.CleanUpTime));
+		GLog->Logf(TEXT("BuildPatchServices: Build Stat: ProcessExecuteTime: %s"), *FPlatformTime::PrettyTime(BuildStats.ProcessExecuteTime));
+		GLog->Logf(TEXT("BuildPatchServices: Build Stat: ProcessPausedTime: %.1f sec"), BuildStats.ProcessPausedTime);
+		GLog->Logf(TEXT("BuildPatchServices: Build Stat: ProcessSuccess: %s"), BuildStats.ProcessSuccess ? TEXT("TRUE") : TEXT("FALSE"));
+		GLog->Logf(TEXT("BuildPatchServices: Build Stat: FailureReason: %s"), *BuildStats.FailureReason);
+		GLog->Logf(TEXT("BuildPatchServices: Build Stat: FailureReasonText: %s"), *BuildStats.FailureReasonText.ToString());
+	}
+
+	// Mark that we are done
+	SetRunning(false);
+
+	return bSuccess ? 0 : 1;
+}
+
+bool FBuildPatchInstaller::RunInstallation(TArray<FString>& CorruptFiles)
+{
+	GLog->Logf(TEXT("BuildPatchServices: Starting Installation"));
+	// Save the staging directories
+	FPaths::NormalizeDirectoryName(DataStagingDir);
+	FPaths::NormalizeDirectoryName(InstallStagingDir);
+
+	// Make sure staging directories exist
+	IFileManager::Get().MakeDirectory(*DataStagingDir, true);
+	IFileManager::Get().MakeDirectory(*InstallStagingDir, true);
+
+	// Reset any error from a previous install
+	FBuildPatchInstallError::Reset();
+	FBuildPatchAnalytics::ResetCounters();
+	BuildProgress.Reset();
+	BuildProgress.SetStateProgress(EBuildPatchProgress::Initializing, 0.01f);
+	BuildProgress.SetStateProgress(EBuildPatchProgress::CleanUp, 0.0f);
+
+	// Remove any inventory
+	FBuildPatchFileConstructor::PurgeFileDataInventory();
+
+	// Check if we should skip out of this process
+	bool bPreviousStagingCompleted = FPaths::FileExists(PreviousMoveMarker);
+	if (bPreviousStagingCompleted)
+	{
+		GLog->Logf(TEXT("BuildPatchServices: Detected previous staging completed"));
+		// Set weights for verify only
+		BuildProgress.SetStateWeight(EBuildPatchProgress::Downloading, 0.0f);
+		BuildProgress.SetStateWeight(EBuildPatchProgress::Installing, 0.0f);
+		BuildProgress.SetStateWeight(EBuildPatchProgress::MovingToInstall, 0.0f);
+		BuildProgress.SetStateWeight(EBuildPatchProgress::BuildVerification, 1.0f);
+		// Mark all installation steps complete
+		BuildProgress.SetStateProgress(EBuildPatchProgress::Initializing, 1.0f);
+		BuildProgress.SetStateProgress(EBuildPatchProgress::Resuming, 1.0f);
+		BuildProgress.SetStateProgress(EBuildPatchProgress::Downloading, 1.0f);
+		BuildProgress.SetStateProgress(EBuildPatchProgress::Installing, 1.0f);
+		BuildProgress.SetStateProgress(EBuildPatchProgress::MovingToInstall, 1.0f);
+		return true;
+	}
+
+	// Get the list of files needing construction
+	TArray< FString > FilesToConstruct;
+	if (CorruptFiles.Num() > 0)
+	{
+		FilesToConstruct.Append(CorruptFiles);
+	}
+	else
+	{
+		FBuildPatchAppManifest::GetOutdatedFiles(CurrentBuildManifest, NewBuildManifest, InstallDirectory, FilesToConstruct);
+	}
+	GLog->Logf(TEXT("BuildPatchServices: Requiring %d files"), FilesToConstruct.Num());
+
+	// Create the downloader
+	FBuildPatchDownloader::Create(DataStagingDir, NewBuildManifest, &BuildProgress);
+
+	// Create chunk cache
+	if (bIsChunkData)
+	{
+		FBuildPatchChunkCache::Init(NewBuildManifest, CurrentBuildManifest, DataStagingDir, InstallDirectory, &BuildProgress, FilesToConstruct, InstallationInfo);
+	}
+
+	// Hold the file constructor thread
+	FBuildPatchFileConstructor* FileConstructor = NULL;
+
+	// Store some totals
+	const uint32 NumFilesInBuild = NewBuildManifest->GetNumFiles();
+
+	// Stats for build
+	const uint32 NumFilesToConstruct = bIsFileData ? NumFilesInBuild : FBuildPatchChunkCache::Get().GetStatNumFilesToConstruct();
+	const uint32 NumRequiredChunks = bIsFileData ? NumFilesInBuild : FBuildPatchChunkCache::Get().GetStatNumRequiredChunks();
+	const uint32 NumChunksToDownload = bIsFileData ? NumFilesInBuild : FBuildPatchChunkCache::Get().GetStatNumChunksToDownload();
+	const uint32 NumChunksToConstruct = bIsFileData ? 0 : FBuildPatchChunkCache::Get().GetStatNumChunksToRecycle();
+	TotalInitialDownloadSize = bIsFileData ? NewBuildManifest->GetFileSize(FilesToConstruct) : FBuildPatchChunkCache::Get().GetStatTotalChunkDownloadSize();
+
+	// Save stats
+	{
+		FScopeLock Lock(&ThreadLock);
+		BuildStats.AppName = NewBuildManifest->GetAppName();
+		BuildStats.AppPatchVersion = NewBuildManifest->GetVersionString();
+		BuildStats.AppInstalledVersion = CurrentBuildManifest.IsValid() ? CurrentBuildManifest->GetVersionString() : TEXT("NONE");
+		BuildStats.CloudDirectory = FBuildPatchServicesModule::GetCloudDirectory();
+		BuildStats.NumFilesInBuild = NumFilesInBuild;
+		BuildStats.NumFilesOutdated = NumFilesToConstruct;
+		BuildStats.NumChunksRequired = NumRequiredChunks;
+		BuildStats.ChunksQueuedForDownload = NumChunksToDownload;
+		BuildStats.ChunksLocallyAvailable = NumChunksToConstruct;
+	}
+
+	// Save initial counts as float for use with progress updates
+	InitialNumChunkDownloads = NumChunksToDownload;
+	InitialNumChunkConstructions = NumChunksToConstruct;
+
+	// Setup some weightings for the progress tracking
+	const float NumRequiredChunksFloat = NumRequiredChunks;
+	BuildProgress.SetStateWeight(EBuildPatchProgress::Downloading, NumRequiredChunksFloat > 0.0f ? InitialNumChunkDownloads / NumRequiredChunksFloat : 0.0f);
+	BuildProgress.SetStateWeight(EBuildPatchProgress::Installing, NumRequiredChunksFloat > 0.0f ? 0.1f + (InitialNumChunkConstructions / NumRequiredChunksFloat) : 0.0f);
+	BuildProgress.SetStateWeight(EBuildPatchProgress::MovingToInstall, NumFilesToConstruct > 0 ? 0.05f : 0.0f);
+	// A verify weight of 1 / 9 will make it 10% of the total progress
+	BuildProgress.SetStateWeight(EBuildPatchProgress::BuildVerification, 1.1f / 9.0f);
+
+	// If this is a repair operation, start off with install and download complete
+	if (bIsRepairing)
+	{
+		GLog->Logf(TEXT("BuildPatchServices: Performing a repair operation"));
+		BuildProgress.SetStateProgress(EBuildPatchProgress::Downloading, 1.0f);
+		BuildProgress.SetStateProgress(EBuildPatchProgress::Installing, 1.0f);
+		BuildProgress.SetStateProgress(EBuildPatchProgress::MovingToInstall, 1.0f);
+	}
+
+	// Start the file constructor
+	GLog->Logf(TEXT("BuildPatchServices: Starting file contruction worker"));
+	FileConstructor = new FBuildPatchFileConstructor(CurrentBuildManifest, NewBuildManifest, InstallDirectory, InstallStagingDir, FilesToConstruct, &BuildProgress);
+
+	// Initializing is now complete if we are constructing files
+	BuildProgress.SetStateProgress(EBuildPatchProgress::Initializing, NumFilesToConstruct > 0 ? 1.0f : 0.0f);
+
+	// If this is file data, queue the download list
+	if (bIsFileData)
+	{
+		TArray< FGuid > RequiredFileData;
+		NewBuildManifest->GetChunksRequiredForFiles(FilesToConstruct, RequiredFileData);
+		FBuildPatchDownloader::Get().AddChunksToDownload(RequiredFileData);
+	}
+
+	// Wait for the file constructor to complete
+	while (FileConstructor->IsComplete() == false)
+	{
+		UpdateDownloadProgressInfo();
+		FPlatformProcess::Sleep(0.1f);
+	}
+	FileConstructor->Wait();
+	delete FileConstructor;
+	FileConstructor = NULL;
+	GLog->Logf(TEXT("BuildPatchServices: File construction complete"));
+
+	// Wait for downloader to complete
+	FBuildPatchDownloader::Get().NotifyNoMoreChunksToAdd();
+	while (FBuildPatchDownloader::Get().IsComplete() == false)
+	{
+		UpdateDownloadProgressInfo();
+		FPlatformProcess::Sleep(0.0f);
+	}
+	TArray< FBuildPatchDownloadRecord > AllChunkDownloads = FBuildPatchDownloader::Get().GetDownloadRecordings();
+	SetDownloadSpeed(-1);
+
+	// Calculate the average download speed from the recordings
+	// NB: Because we are threading several downloads at once this is not simply averaging every download. We have to know about
+	//     how much data is being received simultaneously too. We also need to ignore download pauses.
+	int64 TotalDownloadedBytes = 0;
+	double TotalTimeDownloading = 0;
+	double RecoredEndTime = 0;
+	if (AllChunkDownloads.Num() > 0)
+	{
+		// Sort by start time
+		AllChunkDownloads.Sort();
+		// Start with first record
+		TotalTimeDownloading = AllChunkDownloads[0].EndTime - AllChunkDownloads[0].StartTime;
+		TotalDownloadedBytes = AllChunkDownloads[0].DownloadSize;
+		RecoredEndTime = AllChunkDownloads[0].EndTime;
+		// For every other record..
+		for (int32 RecordIdx = 1; RecordIdx < AllChunkDownloads.Num(); ++RecordIdx)
+		{
+			// Do we have some time to count
+			if (RecoredEndTime < AllChunkDownloads[RecordIdx].EndTime)
+			{
+				// Was there a break in downloading
+				if (AllChunkDownloads[RecordIdx].StartTime > RecoredEndTime)
+				{
+					TotalTimeDownloading += AllChunkDownloads[RecordIdx].EndTime - AllChunkDownloads[RecordIdx].StartTime;
+				}
+				// Otherwise don't count time overlap
+				else
+				{
+					TotalTimeDownloading += AllChunkDownloads[RecordIdx].EndTime - RecoredEndTime;
+				}
+				RecoredEndTime = AllChunkDownloads[RecordIdx].EndTime;
+			}
+			// Count all bytes
+			TotalDownloadedBytes += AllChunkDownloads[RecordIdx].DownloadSize;
+		}
+	}
+
+	// Set final stats
+	{
+		FScopeLock Lock(&ThreadLock);
+		BuildStats.TotalDownloadedData = TotalDownloadedBytes;
+		BuildStats.NumChunksDownloaded = AllChunkDownloads.Num();
+		const double TotalDownloadedBytesDouble = TotalDownloadedBytes;
+		BuildStats.AverageDownloadSpeed = TotalTimeDownloading > 0 ? TotalDownloadedBytesDouble / TotalTimeDownloading : 0;
+		BuildStats.TheoreticalDownloadTime = TotalTimeDownloading;
+		BuildStats.NumChunksRecycled = bIsFileData ? 0 : FBuildPatchChunkCache::Get().GetCounterChunksRecycled();
+		BuildStats.NumChunksCacheBooted = bIsFileData ? 0 : FBuildPatchChunkCache::Get().GetCounterChunksCacheBooted();
+		BuildStats.NumDriveCacheChunkLoads = bIsFileData ? 0 : FBuildPatchChunkCache::Get().GetCounterDriveCacheChunkLoads();
+		BuildStats.NumRecycleFailures = bIsFileData ? 0 : FBuildPatchChunkCache::Get().GetCounterRecycleFailures();
+		BuildStats.NumDriveCacheLoadFailures = bIsFileData ? 0 : FBuildPatchChunkCache::Get().GetCounterDriveCacheLoadFailures();
+	}
+
+	// Perform static cleanup
+	if (bIsChunkData)
+	{
+		FBuildPatchChunkCache::Shutdown();
+	}
+	FBuildPatchDownloader::Shutdown();
+	FBuildPatchFileConstructor::PurgeFileDataInventory();
+
+	GLog->Logf(TEXT("BuildPatchServices: Staged install complete"));
+
+	return !FBuildPatchInstallError::HasFatalError();
+}
+
+bool FBuildPatchInstaller::RunBackupAndMove()
+{
+	GLog->Logf(TEXT("BuildPatchServices: Running backup and stage relocation"));
+	// If there's no error, move all complete files
+	bool bMoveSuccess = FBuildPatchInstallError::HasFatalError() == false;
+	if (bMoveSuccess)
+	{
+		// First handle files that should be removed for patching
+		TArray< FString > FilesToRemove;
+		if (CurrentBuildManifest.IsValid())
+		{
+			FBuildPatchAppManifest::GetRemovableFiles(CurrentBuildManifest.ToSharedRef(), NewBuildManifest, FilesToRemove);
+		}
+		// Add to build stats
+		ThreadLock.Lock();
+		BuildStats.NumFilesToRemove = FilesToRemove.Num();
+		ThreadLock.Unlock();
+		for (const FString& OldFilename : FilesToRemove)
+		{
+			BackupFileIfNecessary(OldFilename);
+			const bool bDeleteSuccess = IFileManager::Get().Delete(*(InstallDirectory / OldFilename), false, true, true);
+			const uint32 LastError = FPlatformMisc::GetLastError();
+			GLog->Logf(TEXT("BuildPatchServices: Removed (%u,%u) %s"), bDeleteSuccess ? 1 : 0, LastError, *OldFilename);
+		}
+
+		// Now handle files that have been constructed
+		bool bSavedMoveMarkerFile = false;
+		TArray< FString > ConstructionFiles;
+		NewBuildManifest->GetFileList(ConstructionFiles);
+		BuildProgress.SetStateProgress(EBuildPatchProgress::MovingToInstall, 0.0f);
+		const float NumConstructionFilesFloat = ConstructionFiles.Num();
+		for (auto ConstructionFilesIt = ConstructionFiles.CreateConstIterator(); ConstructionFilesIt && bMoveSuccess && !FBuildPatchInstallError::HasFatalError(); ++ConstructionFilesIt)
+		{
+			const FString& ConstructionFile = *ConstructionFilesIt;
+			const FString SrcFilename = InstallStagingDir / ConstructionFile;
+			const FString DestFilename = InstallDirectory / ConstructionFile;
+			const float FileIndexFloat = ConstructionFilesIt.GetIndex();
+			// Skip files not constructed
+			if (!FPaths::FileExists(SrcFilename))
+			{
+				BuildProgress.SetStateProgress(EBuildPatchProgress::MovingToInstall, FileIndexFloat / NumConstructionFilesFloat);
+				continue;
+			}
+			// Create the move marker file
+			if (!bSavedMoveMarkerFile)
+			{
+				bSavedMoveMarkerFile = true;
+				GLog->Logf(TEXT("BuildPatchServices: Create MM"));
+				FArchive* MoveMarkerFile = IFileManager::Get().CreateFileWriter(*PreviousMoveMarker, FILEWRITE_EvenIfReadOnly);
+				if (MoveMarkerFile != NULL)
+				{
+					MoveMarkerFile->Close();
+					delete MoveMarkerFile;
+				}
+				// Make sure we have some progress if we do some work
+				if (BuildProgress.GetStateWeight(EBuildPatchProgress::MovingToInstall) <= 0.0f)
+				{
+					BuildProgress.SetStateWeight(EBuildPatchProgress::MovingToInstall, 0.1f);
+				}
+			}
+			// Backup file if need be
+			BackupFileIfNecessary(ConstructionFile);
+			// Move the file to the installation directory
+			int32 MoveRetries = NUM_FILE_MOVE_RETRIES;
+			bMoveSuccess = IFileManager::Get().Move(*DestFilename, *SrcFilename, true, true);
+			uint32 ErrorCode = FPlatformMisc::GetLastError();
+			while (!bMoveSuccess && MoveRetries > 0)
+			{
+				--MoveRetries;
+				FBuildPatchAnalytics::RecordConstructionError(ConstructionFile, ErrorCode, TEXT("Failed To Move"));
+				GWarn->Logf(TEXT("BuildPatchServices: ERROR: Failed to move file %s (%d), trying copy"), *ConstructionFile, ErrorCode);
+				bMoveSuccess = IFileManager::Get().Copy(*DestFilename, *SrcFilename, true, true) == COPY_OK;
+				ErrorCode = FPlatformMisc::GetLastError();
+				if (!bMoveSuccess)
+				{
+					GWarn->Logf(TEXT("BuildPatchServices: ERROR: Failed to copy file %s (%d), retying after 0.5 sec"), *ConstructionFile, ErrorCode);
+					FPlatformProcess::Sleep(0.5f);
+					--MoveRetries;
+					bMoveSuccess = IFileManager::Get().Move(*DestFilename, *SrcFilename, true, true);
+					ErrorCode = FPlatformMisc::GetLastError();
+				}
+				else
+				{
+					IFileManager::Get().Delete(*SrcFilename, false, true, false);
+				}
+			}
+			if (!bMoveSuccess)
+			{
+				GWarn->Logf(TEXT("BuildPatchServices: ERROR: Failed to move file %s"), *FPaths::GetCleanFilename(ConstructionFile));
+				FBuildPatchInstallError::SetFatalError(EBuildPatchInstallError::MoveFileToInstall);
+			}
+			else
+			{
+				FilesInstalled.Add(ConstructionFile);
+				BuildProgress.SetStateProgress(EBuildPatchProgress::MovingToInstall, FileIndexFloat / NumConstructionFilesFloat);
+			}
+		}
+
+		bMoveSuccess = bMoveSuccess && (FBuildPatchInstallError::HasFatalError() == false);
+		if (bMoveSuccess)
+		{
+			BuildProgress.SetStateProgress(EBuildPatchProgress::MovingToInstall, 1.0f);
+		}
+	}
+	GLog->Logf(TEXT("BuildPatchServices: Relocation complete %d"), bMoveSuccess?1:0);
+	return bMoveSuccess;
+}
+
+bool FBuildPatchInstaller::RunVerification(TArray< FString >& CorruptFiles)
+{
+	// Make sure this function can never be parallelized
+	static FCriticalSection SingletonFunctionLockCS;
+	FScopeLock SingletonFunctionLock(&SingletonFunctionLockCS);
+
+	BuildProgress.SetStateProgress(EBuildPatchProgress::BuildVerification, 0.0f);
+
+	// Verify the installation
+	double VerifyTime = 0;
+	double VerifyPauseTime;
+	GLog->Logf(TEXT("BuildPatchServices: Verifying install"));
+	CorruptFiles.Empty();
+	VerifyTime = FPlatformTime::Seconds();
+	bool bVerifySuccess = NewBuildManifest->VerifyAgainstDirectory(InstallDirectory, CorruptFiles, FBuildPatchFloatDelegate::CreateRaw(this, &FBuildPatchInstaller::UpdateVerificationProgress), FBuildPatchBoolRetDelegate::CreateRaw(this, &FBuildPatchInstaller::IsPaused), VerifyPauseTime);
+	VerifyTime = FPlatformTime::Seconds() - VerifyTime - VerifyPauseTime;
+	if (!bVerifySuccess)
+	{
+		FString ErrorString = TEXT("Build verification failed on ");
+		ErrorString += FString::Printf(TEXT("%d"), CorruptFiles.Num());
+		ErrorString += TEXT(" file(s)");
+		FBuildPatchInstallError::SetFatalError(EBuildPatchInstallError::BuildVerifyFail, ErrorString);
+	}
+
+	ThreadLock.Lock();
+	BuildStats.VerifyTime = VerifyTime;
+	ThreadLock.Unlock();
+
+	BuildProgress.SetStateProgress(EBuildPatchProgress::BuildVerification, 1.0f);
+
+	// Delete/Backup any incorrect files if failure was not cancellation
+	if (!FBuildPatchInstallError::IsInstallationCancelled())
+	{
+		for (const FString& CorruptFile : CorruptFiles)
+		{
+			BackupFileIfNecessary(CorruptFile, true);
+			IFileManager::Get().Delete(*(InstallDirectory / CorruptFile), false, true);
+			IFileManager::Get().Delete(*(InstallStagingDir / CorruptFile), false, true);
+		}
+	}
+
+	GLog->Logf(TEXT("BuildPatchServices: Verify stage complete %d"), bVerifySuccess ? 1 : 0);
+
+	return bVerifySuccess;
+}
+
+bool FBuildPatchInstaller::BackupFileIfNecessary(const FString& Filename, bool bDiscoveredByVerification /*= false */)
+{
+	const FString BackupDirectory = FBuildPatchServicesModule::GetBackupDirectory();
+	const FString InstalledFilename = InstallDirectory / Filename;
+	const FString BackupFilename = BackupDirectory / Filename;
+	const bool bBackupOriginals = !BackupDirectory.IsEmpty();
+	// Skip if not doing backups
+	if (!bBackupOriginals)
+	{
+		return true;
+	}
+	// Skip if no file to backup
+	const bool bInstalledFileExists = FPaths::FileExists(InstalledFilename);
+	if (!bInstalledFileExists)
+	{
+		return true;
+	}
+	// Skip if already backed up
+	const bool bAlreadyBackedUp = FPaths::FileExists(BackupFilename);
+	if (bAlreadyBackedUp)
+	{
+		return true;
+	}
+	// Skip if the target file was already copied to the installation
+	const bool bAlreadyInstalled = FilesInstalled.Contains(Filename);
+	if (bAlreadyInstalled)
+	{
+		return true;
+	}
+	// If discovered by verification, but the patching system did not touch the file, we know it must be backed up.
+	// If patching system touched the file it would already have been backed up
+	if (bDiscoveredByVerification && CurrentBuildManifest.IsValid() && !FBuildPatchAppManifest::IsFileOutdated(CurrentBuildManifest.ToSharedRef(), NewBuildManifest, Filename))
+	{
+		return IFileManager::Get().Move(*BackupFilename, *InstalledFilename, true, true, true);
+	}
+	// There are lots of checks to do in order to figure out whether a file needs backing up.
+	TSharedPtr< FBuildPatchFileManifest > OldFileManifest = CurrentBuildManifest.IsValid() ? CurrentBuildManifest->GetFileManifest(Filename) : NULL;
+	TSharedPtr< FBuildPatchFileManifest > NewFileManifest = NewBuildManifest->GetFileManifest(Filename);
+	const bool bCurrentlyInstalled = CurrentBuildManifest.IsValid();
+	const bool bPerformingRepair = CurrentBuildManifest == NewBuildManifest;
+	const bool bPerformingPatch = bCurrentlyInstalled && !bPerformingRepair;
+	const bool bFreshInstall = !bCurrentlyInstalled;
+	const bool bFileIsNewAddition = OldFileManifest.IsValid() == false;
+	const bool bFileRemovedByPatch = bPerformingPatch && !NewFileManifest.IsValid();
+	bool bUserEditedFile = bDiscoveredByVerification;
+	const bool bCheckFileChanges = !bDiscoveredByVerification && ((!bPerformingRepair && !bFileIsNewAddition) || bFileRemovedByPatch);
+	if (bCheckFileChanges)
+	{
+		const int64 InstalledFilesize = IFileManager::Get().FileSize(*InstalledFilename);
+		const int64 OriginalFileSize = OldFileManifest.IsValid() ? OldFileManifest->GetFileSize() : INDEX_NONE;
+		const int64 NewFileSize = NewFileManifest.IsValid() ? NewFileManifest->GetFileSize() : INDEX_NONE;
+		const FSHAHash HashZero;
+		const FSHAHash& HashOld = OldFileManifest.IsValid() ? OldFileManifest->FileHash : HashZero;
+		const FSHAHash& HashNew = NewFileManifest.IsValid() ? NewFileManifest->FileHash : HashZero;
+		const bool bFileSizeDiffers = OriginalFileSize != InstalledFilesize && NewFileSize != InstalledFilesize;
+		bUserEditedFile = bFileSizeDiffers || FBuildPatchUtils::VerifyFile(InstalledFilename, HashOld, HashNew) == 0;
+	}
+	// Finally, use the above logic to determine if we must do the backup
+	const bool bNeedBackup = bFileIsNewAddition || bPerformingRepair || bUserEditedFile;
+	bool bBackupSuccess = true;
+	if (bNeedBackup)
+	{
+		GLog->Logf(TEXT("BuildPatchServices: Backing up (%d,%d,%d) %s"), bFileIsNewAddition?1:0, bPerformingRepair?1:0, bUserEditedFile?1:0, *Filename);
+		bBackupSuccess = IFileManager::Get().Move(*BackupFilename, *InstalledFilename, true, true, true);
+	}
+	return bBackupSuccess;
+}
+
+bool FBuildPatchInstaller::RunPrereqInstaller()
+{
+	FString PrereqPath = InstallDirectory / NewBuildManifest->GetPrereqPath();
+	PrereqPath = FPaths::ConvertRelativePathToFull(PrereqPath);
+	FString PrereqCommandline = NewBuildManifest->GetPrereqArgs();
+
+	GLog->Logf(TEXT("BuildPatchServices: Running prerequisites installer %s %s"), *PrereqPath, *PrereqCommandline);
+
+	BuildProgress.SetStateProgress(EBuildPatchProgress::PrerequisitesInstall, 0.0f);
+
+	// @TODO: Tell our installer to run with no UI since we will have BuildPatchProgress
+	// @TODO: Pass in params to the installer that tell it to fire up the portal/launcher after install if it has to perform a restart.
+	FProcHandle ProcessHandle = FPlatformProcess::CreateProc(*PrereqPath, *PrereqCommandline, true, false, false, NULL, 0, *FPaths::GetPath(PrereqPath), NULL);
+	bool bPrereqInstallSuccessful = true;
+
+	if (!ProcessHandle.IsValid())
+	{
+		GLog->Logf(TEXT("BuildPatchServices: ERROR: Failed to start the prerequisites install process."));
+		FBuildPatchAnalytics::RecordPrereqInstallnError(PrereqPath, PrereqCommandline, -1, TEXT("Failed to start installer"));
+		// @TODO: Do we need to do anything special to fail?
+		bPrereqInstallSuccessful = false;
+	}
+
+	int32 ReturnCode;
+	if (bPrereqInstallSuccessful)
+	{
+		FPlatformProcess::WaitForProc(ProcessHandle);
+		FPlatformProcess::GetProcReturnCode(ProcessHandle, &ReturnCode);
+
+		ProcessHandle.Close();
+
+		if (ReturnCode != 0)
+		{
+			if (ReturnCode == 3010)
+			{
+				GLog->Logf(TEXT("BuildPatchServices: Prerequisites executable returned restart required code %d"), ReturnCode);
+				// @TODO: Inform app that system restart is required
+			}
+			else
+			{
+				GLog->Logf(TEXT("BuildPatchServices: ERROR: Prerequisites executable failed with code %d"), ReturnCode);
+				FBuildPatchAnalytics::RecordPrereqInstallnError(PrereqPath, PrereqCommandline, ReturnCode, TEXT("Failed to install"));
+				bPrereqInstallSuccessful = false;
+			}
+		}
+	}
+
+	if (bPrereqInstallSuccessful)
+	{
+		BuildProgress.SetStateProgress(EBuildPatchProgress::PrerequisitesInstall, 1.0f);
+	}
+	else
+	{
+		FBuildPatchInstallError::SetFatalError(EBuildPatchInstallError::PrerequisiteError);
+	}
+
+	return bPrereqInstallSuccessful;
+}
+
+void FBuildPatchInstaller::SetRunning( bool bRunning )
+{
+	FScopeLock Lock( &ThreadLock );
+	bIsRunning = bRunning;
+}
+
+void FBuildPatchInstaller::SetInited( bool bInited )
+{
+	FScopeLock Lock( &ThreadLock );
+	bIsInited = bInited;
+}
+
+void FBuildPatchInstaller::SetDownloadSpeed( double ByteSpeed )
+{
+	FScopeLock Lock( &ThreadLock );
+	DownloadSpeedValue = ByteSpeed;
+}
+
+void FBuildPatchInstaller::SetDownloadBytesLeft( const int64& BytesLeft )
+{
+	FScopeLock Lock( &ThreadLock );
+	DownloadBytesLeft = BytesLeft;
+}
+
+void FBuildPatchInstaller::UpdateDownloadProgressInfo( bool bReset )
+{
+	// Static variables for persistent values
+	static double LastTime = FPlatformTime::Seconds();
+	static double NowTime = 0;
+	static double DeltaTime = 0;
+	static double LastReadingTime = 0;
+	static double LastDataReadings[NUM_DOWNLOAD_READINGS] = { 0 };
+	static double LastTimeReadings[NUM_DOWNLOAD_READINGS] = { 0 };
+	static uint32 ReadingIdx = 0;
+	static bool bProgressIsDownload = true;
+	static double AverageDownloadSpeed = 0;
+
+	// Reset internals?
+	if( bReset )
+	{
+		LastTime = FPlatformTime::Seconds();
+		LastReadingTime = LastTime;
+		NowTime = 0;
+		DeltaTime = 0;
+		for (int32 i = 0; i < NUM_DOWNLOAD_READINGS; ++i)
+		{
+			LastDataReadings[i] = 0;
+			LastTimeReadings[i] = 0;
+		}
+		ReadingIdx = 0;
+		bProgressIsDownload = true;
+		AverageDownloadSpeed = 0;
+		return;
+	}
+
+	// Return if not downloading yet
+	if( !NewBuildManifest->IsFileDataManifest() && !FBuildPatchChunkCache::Get().HaveDownloadsStarted() )
+	{
+		return;
+	}
+
+	// Calculate percentage complete based on number of chunks
+	const int64 DownloadNumBytesLeft = FBuildPatchDownloader::Get().GetNumBytesLeft();
+	const float DownloadSizeFloat = TotalInitialDownloadSize;
+	const float DownloadBytesLeftFloat = DownloadNumBytesLeft;
+	const float DownloadProgress = 1.0f - (TotalInitialDownloadSize > 0 ? DownloadBytesLeftFloat / DownloadSizeFloat : 0.0f);
+	BuildProgress.SetStateProgress( EBuildPatchProgress::Downloading, DownloadProgress );
+
+	// Calculate the average download speed
+	NowTime = FPlatformTime::Seconds();
+	DeltaTime += NowTime - LastTime;
+	if( DeltaTime > TIME_PER_READING )
+	{
+		const double BytesDownloaded = FBuildPatchDownloader::Get().GetByteDownloadCountReset();
+		const double TimeSinceLastReading = NowTime - LastReadingTime;
+		LastReadingTime = NowTime;
+		LastDataReadings[ReadingIdx] = BytesDownloaded;
+		LastTimeReadings[ReadingIdx] = TimeSinceLastReading;
+		ReadingIdx = (ReadingIdx + 1) % NUM_DOWNLOAD_READINGS;
+		DeltaTime = 0;
+		double TotalData = 0;
+		double TotalTime = 0;
+		for (uint32 i = 0; i < NUM_DOWNLOAD_READINGS; ++i)
+		{
+			TotalData += LastDataReadings[i];
+			TotalTime += LastTimeReadings[i];
+		}
+		AverageDownloadSpeed = TotalData / TotalTime;
+	}
+
+	// Set download values
+	SetDownloadSpeed( DownloadProgress < 1.0f ? AverageDownloadSpeed : -1.0f );
+	SetDownloadBytesLeft( DownloadNumBytesLeft );
+
+	// Set last time
+	LastTime = NowTime;
+}
+
+//@todo this is deprecated and shouldn't be used anymore [6/4/2014 justin.sargent]
+FText FBuildPatchInstaller::GetDownloadSpeedText()
+{
+	static const FText DownloadSpeedFormat = LOCTEXT("BuildPatchInstaller_DownloadSpeedFormat", "{Current} / {Total} ({Speed}/sec)");
+
+	FScopeLock Lock(&ThreadLock);
+	FText SpeedDisplayedText;
+	if (DownloadSpeedValue >= 0)
+	{
+		FNumberFormattingOptions SpeedFormattingOptions;
+		SpeedFormattingOptions.MaximumFractionalDigits = 1;
+		SpeedFormattingOptions.MinimumFractionalDigits = 0;
+
+		FNumberFormattingOptions SizeFormattingOptions;
+		SizeFormattingOptions.MaximumFractionalDigits = 1;
+		SizeFormattingOptions.MinimumFractionalDigits = 1;
+
+		FFormatNamedArguments Args;
+		Args.Add(TEXT("Speed"), FText::AsMemory(DownloadSpeedValue, &SpeedFormattingOptions));
+		Args.Add(TEXT("Total"), FText::AsMemory(TotalInitialDownloadSize, &SizeFormattingOptions));
+		Args.Add(TEXT("Current"), FText::AsMemory(TotalInitialDownloadSize - DownloadBytesLeft, &SizeFormattingOptions));
+		return FText::Format(DownloadSpeedFormat, Args);
+	}
+
+	return FText();
+}
+
+double FBuildPatchInstaller::GetDownloadSpeed() const
+{
+	FScopeLock Lock(&ThreadLock);
+	return DownloadSpeedValue;
+}
+
+int64 FBuildPatchInstaller::GetInitialDownloadSize() const
+{
+	FScopeLock Lock(&ThreadLock);
+	return TotalInitialDownloadSize;
+}
+
+int64 FBuildPatchInstaller::GetTotalDownloaded() const
+{
+	FScopeLock Lock(&ThreadLock);
+	return TotalInitialDownloadSize - DownloadBytesLeft;
+}
+
+bool FBuildPatchInstaller::IsComplete()
+{
+	FScopeLock Lock( &ThreadLock );
+	return !bIsRunning && bIsInited;
+}
+
+bool FBuildPatchInstaller::IsCanceled()
+{
+	FScopeLock Lock( &ThreadLock );
+	return FBuildPatchInstallError::GetErrorState() == EBuildPatchInstallError::UserCanceled;
+}
+
+bool FBuildPatchInstaller::IsPaused()
+{
+	return BuildProgress.GetPauseState();
+}
+
+bool FBuildPatchInstaller::HasError()
+{
+	FScopeLock Lock( &ThreadLock );
+	if( FBuildPatchInstallError::GetErrorState() == EBuildPatchInstallError::UserCanceled )
+	{
+		return false;
+	}
+	return !BuildStats.ProcessSuccess;
+}
+
+//@todo this is deprecated and shouldn't be used anymore [6/4/2014 justin.sargent]
+FText FBuildPatchInstaller::GetPercentageText()
+{
+	static const FText PleaseWait = LOCTEXT( "BuildPatchInstaller_GenericProgress", "Please Wait" );
+
+	FScopeLock Lock( &ThreadLock );
+
+	float Progress = GetUpdateProgress() * 100.0f;
+	if( Progress <= 0.0f )
+	{
+		return PleaseWait;
+	}
+
+	FNumberFormattingOptions PercentFormattingOptions;
+	PercentFormattingOptions.MaximumFractionalDigits = 0;
+	PercentFormattingOptions.MinimumFractionalDigits = 0;
+
+	return FText::AsPercent(GetUpdateProgress(), &PercentFormattingOptions);
+}
+
+FText FBuildPatchInstaller::GetStatusText()
+{
+	return BuildProgress.GetStateText();
+}
+
+float FBuildPatchInstaller::GetUpdateProgress()
+{
+	return BuildProgress.GetProgress();
+}
+
+FBuildInstallStats FBuildPatchInstaller::GetBuildStatistics()
+{
+	FScopeLock Lock( &ThreadLock );
+	return BuildStats;
+}
+
+FText FBuildPatchInstaller::GetErrorText()
+{
+	FScopeLock Lock( &ThreadLock );
+	return BuildStats.FailureReasonText;
+}
+
+void FBuildPatchInstaller::CancelInstall()
+{
+	FBuildPatchInstallError::SetFatalError( EBuildPatchInstallError::UserCanceled );
+	FBuildPatchHTTP::CancelAllHttpRequests();
+	// Make sure we are not paused
+	if( IsPaused() )
+	{
+		TogglePauseInstall();
+	}
+}
+
+bool FBuildPatchInstaller::TogglePauseInstall()
+{
+	if( IsPaused() )
+	{
+		// We are now resuming, so record time paused for
+		FScopeLock Lock( &ThreadLock );
+		const double PausedForSec = FPlatformTime::Seconds() - TimePausedAt;
+		BuildStats.ProcessPausedTime += PausedForSec;
+	}
+	else
+	{
+		// If there is an error, we don't allow the pause
+		if( FBuildPatchInstallError::HasFatalError() )
+		{
+			return false;
+		}
+		// Set the time we pause at
+		TimePausedAt = FPlatformTime::Seconds();
+	}
+	return BuildProgress.TogglePauseState();
+}
+
+void FBuildPatchInstaller::UpdateVerificationProgress( float Percent )
+{
+	BuildProgress.SetStateProgress( EBuildPatchProgress::BuildVerification, Percent );
+}
+
+void FBuildPatchInstaller::ExecuteCompleteDelegate()
+{
+	// Should be executed in main thread, and already be complete
+	check( IsInGameThread() );
+	check( IsComplete() );
+	// Call the complete delegate
+	OnCompleteDelegate.Execute( bSuccess, NewBuildManifest );
+}
+
+
+
+#undef LOCTEXT_NAMESPACE
diff --git a/Engine/Source/Runtime/Online/BuildPatchServices/Private/BuildPatchInstaller.h b/Engine/Source/Runtime/Online/BuildPatchServices/Private/BuildPatchInstaller.h
new file mode 100644
--- /dev/null
+++ b/Engine/Source/Runtime/Online/BuildPatchServices/Private/BuildPatchInstaller.h
@@ -0,0 +1,218 @@
+// Copyright 1998-2014 Epic Games, Inc. All Rights Reserved.
+
+/*=============================================================================
+	BuildPatchInstaller.h: Declares the FBuildPatchInstaller class which
+	controls the process of installing a build described by a build manifest.
+=============================================================================*/
+
+typedef TSharedPtr< class FBuildPatchInstaller, ESPMode::ThreadSafe > FBuildPatchInstallerPtr;
+typedef TSharedRef< class FBuildPatchInstaller, ESPMode::ThreadSafe > FBuildPatchInstallerRef;
+typedef TWeakPtr< class FBuildPatchInstaller, ESPMode::ThreadSafe > FBuildPatchInstallerWeakPtr;
+
+/**
+ * FBuildPatchInstaller
+ * This class controls a thread that wraps the code to install/patch an app from manifests
+ */
+class FBuildPatchInstaller
+	: public IBuildInstaller
+	, public FRunnable
+{
+private:
+	// Hold a pointer to my thread for easier deleting
+	FRunnableThread* Thread;
+
+	// The delegates that we will be calling
+	FBuildPatchBoolManifestDelegate OnCompleteDelegate;
+
+	// The manifest for the build we have installed (if applicable)
+	FBuildPatchAppManifestPtr CurrentBuildManifest;
+
+	// The manifest for the build we want to install
+	FBuildPatchAppManifestRef NewBuildManifest;
+
+	// The Destination directory that we want to install to
+	FString InstallDirectory;
+
+	// The Staging directory used for construction
+	FString StagingDirectory;
+
+	// The directory created in staging, to store local patch data
+	FString DataStagingDir;
+
+	// The directory created in staging to construct install files to
+	FString InstallStagingDir;
+
+	// The filename used to mark a previous install that did not complete but moved staged files into the install directory
+	FString PreviousMoveMarker;
+
+	// A critical section to protect variables
+	mutable FCriticalSection ThreadLock;
+
+	// A flag to store if we are installing file data
+	const bool bIsFileData;
+
+	// A flag to store if we are installing chunk data (to help readability)
+	const bool bIsChunkData;
+
+	// A flag to store if we are performing a repair
+	const bool bIsRepairing;
+
+	// A flag storing whether the process was a success
+	bool bSuccess;
+
+	// A flag marking that we a running
+	bool bIsRunning;
+
+	// A flag marking that we initialized correctly
+	bool bIsInited;
+
+	// The download speed value
+	double DownloadSpeedValue;
+
+	// The number of bytes left to download
+	int64 DownloadBytesLeft;
+
+	// Keep track of build stats
+	FBuildInstallStats BuildStats;
+
+	// Keep track of install progress
+	FBuildPatchProgress BuildProgress;
+
+	// Cache the float value for number of chunk downloads for progress calculations
+	float InitialNumChunkDownloads;
+
+	// Cache the float value for number of chunk constructions for progress calculations
+	float InitialNumChunkConstructions;
+
+	// Cache the number of bytes required to download
+	int64 TotalInitialDownloadSize;
+
+	// Track the time the installer was paused so we can record pause time
+	double TimePausedAt;
+
+	// Holds a list of files that have been placed into the install directory
+	TArray< FString > FilesInstalled;
+
+	// Referecne to the module's installation info
+	FBuildPatchInstallationInfo& InstallationInfo;
+
+public:
+	/**
+	 * Constructor takes the required arguments, CurrentManifest can be invalid for fresh install.
+	 * @param OnCompleteDelegate	Delegate for when the process has completed
+	 * @param CurrentManifest		The manifest for the currently installed build. Can be invalid if fresh install.
+	 * @param InstallManifest		The manifest for the build to be installed
+	 * @param InstallDirectory		The directory where the build will be installed
+	 * @param StagingDirectory		The directory for storing the intermediate files
+	 * @param InstallationInfoRef	Reference to the module's installation info that keeps record of locally installed apps for use as chunk sources
+	 */
+	FBuildPatchInstaller( FBuildPatchBoolManifestDelegate OnCompleteDelegate, FBuildPatchAppManifestPtr CurrentManifest, FBuildPatchAppManifestRef InstallManifest, const FString& InstallDirectory, const FString& StagingDirectory, FBuildPatchInstallationInfo& InstallationInfoRef );
+
+	/**
+	 * Default Destructor, will delete the allocated Thread
+	 */
+	~FBuildPatchInstaller();
+
+	// Begin FRunnable
+	bool Init() override;
+	uint32 Run() override;
+	// End FRunnable
+
+	// Begin IBuildInstaller
+	virtual bool IsComplete() override;
+	virtual bool IsCanceled() override;
+	virtual bool IsPaused() override;
+	virtual bool HasError() override;
+	//@todo this is deprecated and shouldn't be used anymore [6/4/2014 justin.sargent]
+	virtual FText GetPercentageText() override;
+	//@todo this is deprecated and shouldn't be used anymore [6/4/2014 justin.sargent]
+	virtual FText GetDownloadSpeedText() override;
+	virtual double GetDownloadSpeed() const override;
+	virtual int64 GetInitialDownloadSize() const override;
+	virtual int64 GetTotalDownloaded() const override;
+	virtual FText GetStatusText() override;
+	virtual float GetUpdateProgress() override;
+	virtual FBuildInstallStats GetBuildStatistics() override;
+	virtual FText GetErrorText() override;
+	virtual void CancelInstall() override;
+	virtual bool TogglePauseInstall() override;
+	// End IBuildInstaller interface
+
+	/**
+	 * Executes the on complete delegate. This should only be called when completed, and is separated out
+	 * to allow control to make this call on the main thread.
+	 */
+	void ExecuteCompleteDelegate();
+
+private:
+
+	/**
+	 * Runs the installation process
+	 * @param	CorruptFiles	A list of files that were corrupt, to only install those.
+	 * @return    Returns true if there were no errors blocking installation
+	 */
+	bool RunInstallation(TArray<FString>& CorruptFiles);
+
+	/**
+	 * Runs the backup process for locally changed files, and then moves new files into installation directory
+	 * @return    Returns true if there were no errors
+	 */
+	bool RunBackupAndMove();
+
+	/**
+	 * Runs the verification process
+	 * @param CorruptFiles  OUT     Receives the list of files that failed verification
+	 * @return    Returns true if there were no corrupt files
+	 */
+	bool RunVerification( TArray< FString >& CorruptFiles );
+
+	/**
+	 * Checks a particular file in the install directory to see if it needs backing up, and does so if necessary
+	 * @param Filename                      The filename to check, which should match a filename in a manifest
+	 * @param bDiscoveredByVerification     Optional, whether the file was detected changed already by verification stage. Default: false.
+	 * @return    Returns true if there were no errors
+	 */
+	bool BackupFileIfNecessary( const FString& Filename, bool bDiscoveredByVerification = false );
+
+	/**
+	 * Runs any prerequisites associated with the installation
+	 * @return    Returns true if the prerequisites installer succeeded, false otherwise
+	 */
+	bool RunPrereqInstaller();
+	
+	/**
+	 * Sets the bIsRunning flag
+	 * @param bRunning	Whether the thread is running
+	 */
+	void SetRunning( bool bRunning );
+
+	/**
+	 * Sets the bIsInited flag
+	 * @param bInited	Whether the thread successfully initialized
+	 */
+	void SetInited( bool bInited );
+
+	/**
+	 * Sets the current byte download speed
+	 * @param ByteSpeed	The download speed
+	 */
+	void SetDownloadSpeed( double ByteSpeed );
+
+	/**
+	 * Sets the bytes left to download
+	 * @param BytesLeft	The number of bytes left
+	 */
+	void SetDownloadBytesLeft( const int64& BytesLeft );
+
+	/**
+	 * Helper to calculate new chunk progress values
+	 * @param bReset	Resets internals without updating
+	 */
+	void UpdateDownloadProgressInfo( bool bReset = false );
+
+	/**
+	 * Callback for verification process to set progress
+	 * @param Percent	The current process percentage
+	 */
+	void UpdateVerificationProgress( float Percent );
+};
