Moving Build Patch Services from NoRedist to Runtime

diff --git a/Engine/Source/Runtime/Online/BuildPatchServices/Private/BuildPatchManifest.cpp b/Engine/Source/Runtime/Online/BuildPatchServices/Private/BuildPatchManifest.cpp
new file mode 100644
--- /dev/null
+++ b/Engine/Source/Runtime/Online/BuildPatchServices/Private/BuildPatchManifest.cpp
@@ -0,0 +1,1002 @@
+// Copyright 1998-2014 Epic Games, Inc. All Rights Reserved.
+
+/*=============================================================================
+	BuildPatchManifest.cpp: Implements the manifest classes.
+=============================================================================*/
+
+#include "BuildPatchServicesPrivatePCH.h"
+
+#define LOCTEXT_NAMESPACE "BuildPatchManifest"
+
+/**
+ * Helper functions that convert generic types to and from string blobs for use with JSON parsing.
+ * It's kind of horrible but guarantees no loss of data as the JSON reader/writer only supports float functionality 
+ * which would result in data loss with high int32 values, and we'll be using uint64.
+ */
+template< typename DataType >
+bool FromStringBlob( const FString& StringBlob, DataType& ValueOut )
+{
+	void* AsBuffer = &ValueOut;
+	return FString::ToBlob( StringBlob, static_cast< uint8* >( AsBuffer ), sizeof( DataType ) );
+}
+template< typename DataType >
+FString ToStringBlob( const DataType& DataVal )
+{
+	const void* AsBuffer = &DataVal;
+	return FString::FromBlob( static_cast<const uint8*>( AsBuffer ), sizeof( DataType ) );
+}
+
+/* EBuildPatchAppManifestVersion implementation
+*****************************************************************************/
+const EBuildPatchAppManifestVersion::Type EBuildPatchAppManifestVersion::GetLatestVersion()
+{
+	return static_cast<EBuildPatchAppManifestVersion::Type>(LatestPlusOne - 1);
+}
+
+const FString& EBuildPatchAppManifestVersion::GetChunkSubdir(const EBuildPatchAppManifestVersion::Type ManifestVersion)
+{
+	static const FString ChunksV1 = TEXT("Chunks");
+	static const FString ChunksV2 = TEXT("ChunksV2");
+	static const FString ChunksV3 = TEXT("ChunksV3");
+	return ManifestVersion < DataFileRenames ? ChunksV1
+		: ManifestVersion < ChunkCompressionSupport ? ChunksV2
+		: ChunksV3;
+}
+
+const FString& EBuildPatchAppManifestVersion::GetFileSubdir(const EBuildPatchAppManifestVersion::Type ManifestVersion)
+{
+	static const FString FilesV1 = TEXT("Files");
+	static const FString FilesV2 = TEXT("FilesV2");
+	return ManifestVersion < DataFileRenames ? FilesV1
+		: FilesV2;
+}
+
+/* FBuildPatchCustomField implementation
+*****************************************************************************/
+FBuildPatchCustomField::FBuildPatchCustomField( const FString& Value )
+	: CustomValue( Value )
+{
+}
+
+const FString FBuildPatchCustomField::AsString() const
+{
+	return CustomValue;
+}
+
+const double FBuildPatchCustomField::AsDouble() const
+{
+	// The Json parser currently only supports float so we have to decode string blob instead
+	double Rtn;
+	if( FromStringBlob( CustomValue, Rtn ) )
+	{
+		return Rtn;
+	}
+	return 0;
+}
+
+const int64 FBuildPatchCustomField::AsInteger() const
+{
+	// The Json parser currently only supports float so we have to decode string blob instead
+	int64 Rtn;
+	if( FromStringBlob( CustomValue, Rtn ) )
+	{
+		return Rtn;
+	}
+	return 0;
+}
+
+/* FBuildPatchFileManifest implementation
+*****************************************************************************/
+
+FBuildPatchFileManifest::FBuildPatchFileManifest()
+	: Filename( TEXT( "" ) )
+	, FileSize( INDEX_NONE )
+#if PLATFORM_MAC
+	, bIsUnixExecutable( false )
+	, SymlinkTarget( TEXT( "" ) )
+#endif
+{
+}
+
+const int64& FBuildPatchFileManifest::GetFileSize()
+{
+	if( FileSize == INDEX_NONE )
+	{
+		FileSize = 0;
+		for( auto FileChunkPartIt = FileChunkParts.CreateConstIterator(); FileChunkPartIt; ++FileChunkPartIt)
+		{
+			FileSize += (*FileChunkPartIt).Size;
+		}
+	}
+	return FileSize;
+}
+
+/* FBuildPatchAppManifest implementation
+*****************************************************************************/
+
+FBuildPatchAppManifest::FBuildPatchAppManifest()
+	: ManifestFileVersion( EBuildPatchAppManifestVersion::GetLatestVersion() )
+{
+}
+
+FBuildPatchAppManifest::FBuildPatchAppManifest( const uint32& InAppID, const FString& AppName )
+	: ManifestFileVersion( EBuildPatchAppManifestVersion::GetLatestVersion() )
+	, AppID( InAppID )
+	, AppNameString( AppName )
+{
+}
+
+FBuildPatchAppManifest::~FBuildPatchAppManifest()
+{
+	// Make sure the hash data is cleaned up
+	DestroyData();
+}
+
+void FBuildPatchAppManifest::DestroyData()
+{
+	// Clear Manifest data
+	AppNameString = "";
+	BuildVersionString = "";
+	LaunchExeString = "";
+	LaunchCommandString = "";
+	PrereqNameString = "";
+	PrereqPathString = "";
+	PrereqArgsString = "";
+	FileManifestList.Empty();
+	ChunkHashList.Empty();
+}
+
+bool FBuildPatchAppManifest::SaveToFile(FString Filename)
+{
+	bool bSuccess = true;
+	FArchive* FileOut = IFileManager::Get().CreateFileWriter(*Filename);
+	bSuccess &= FileOut != NULL;
+	if (bSuccess)
+	{
+		FString JSONOutput;
+		SerializeToJSON(JSONOutput);
+		FileOut->Serialize(TCHAR_TO_ANSI(*JSONOutput), JSONOutput.Len());
+		FileOut->Close();
+		delete FileOut;
+	}
+	return bSuccess;
+}
+
+bool FBuildPatchAppManifest::LoadFromFile( FString Filename )
+{
+	FString JSONString;
+	bool bSuccess = FFileHelper::LoadFileToString( JSONString, *Filename );
+	if( bSuccess )
+	{
+		bSuccess &= DeserializeFromJSON( JSONString );
+	}
+	return bSuccess;
+}
+
+void FBuildPatchAppManifest::SerializeToJSON( FString& JSONOutput )
+{
+#if UE_BUILD_DEBUG // We'll use this to switch between human readable JSON
+	TSharedRef< TJsonWriter< TCHAR, TPrettyJsonPrintPolicy< TCHAR > > > Writer = TJsonWriterFactory< TCHAR, TPrettyJsonPrintPolicy< TCHAR > >::Create(&JSONOutput);
+#else
+	TSharedRef< TJsonWriter< TCHAR, TCondensedJsonPrintPolicy< TCHAR > > > Writer = TJsonWriterFactory< TCHAR, TCondensedJsonPrintPolicy< TCHAR > >::Create( &JSONOutput );
+#endif //ALLOW_DEBUG_FILES
+
+	Writer->WriteObjectStart();
+	{
+		// Write general data
+		Writer->WriteValue( TEXT("ManifestFileVersion"), ToStringBlob( static_cast< int32 >( ManifestFileVersion ) ) );
+		Writer->WriteValue( TEXT("bIsFileData"), bIsFileData );
+		Writer->WriteValue( TEXT("AppID"), ToStringBlob( AppID ) );
+		Writer->WriteValue( TEXT("AppNameString"), AppNameString );
+		Writer->WriteValue( TEXT("BuildVersionString"), BuildVersionString );
+		Writer->WriteValue( TEXT("LaunchExeString"), LaunchExeString );
+		Writer->WriteValue( TEXT("LaunchCommand"), LaunchCommandString );
+		Writer->WriteValue( TEXT("PrereqName"), PrereqNameString );
+		Writer->WriteValue( TEXT("PrereqPath"), PrereqPathString );
+		Writer->WriteValue( TEXT("PrereqArgs"), PrereqArgsString );
+		// Write file manifest data
+		Writer->WriteArrayStart( TEXT("FileManifestList") );
+		for( auto FileManifestIt = FileManifestList.CreateConstIterator(); FileManifestIt; ++FileManifestIt )
+		{
+			const TSharedPtr< FBuildPatchFileManifest > FileManifest = FileManifestIt.Value();
+			Writer->WriteObjectStart();
+			{
+				Writer->WriteValue( TEXT("Filename"), FileManifest->Filename );
+				Writer->WriteValue( TEXT("FileHash"), FString::FromBlob( FileManifestIt.Value()->FileHash.Hash, FSHA1::DigestSize ) );
+#if PLATFORM_MAC
+				if( FileManifest->bIsUnixExecutable )
+				{
+					Writer->WriteValue( TEXT("bIsUnixExecutable"), FileManifest->bIsUnixExecutable );
+				}
+				const bool bIsSymlink = !FileManifest->SymlinkTarget.IsEmpty();
+				if( bIsSymlink )
+				{
+					Writer->WriteValue( TEXT("SymlinkTarget"), FileManifest->SymlinkTarget );
+				}
+				else
+#endif
+				{
+					Writer->WriteArrayStart( TEXT("FileChunkParts") );
+					{
+						for( auto FileChunkPartIt = FileManifestIt.Value()->FileChunkParts.CreateConstIterator(); FileChunkPartIt; ++FileChunkPartIt )
+						{
+							const FChunkPart& FileChunkPart = *FileChunkPartIt;
+							Writer->WriteObjectStart();
+							{
+								Writer->WriteValue( TEXT("Guid"), FileChunkPart.Guid.ToString() );
+								Writer->WriteValue( TEXT("Offset"), ToStringBlob( FileChunkPart.Offset ) );
+								Writer->WriteValue( TEXT("Size"), ToStringBlob( FileChunkPart.Size ) );
+							}
+							Writer->WriteObjectEnd();
+						}
+					}
+					Writer->WriteArrayEnd();
+				}
+			}
+			Writer->WriteObjectEnd();
+		}
+		Writer->WriteArrayEnd();
+		// Write chunk hash list
+		Writer->WriteObjectStart( TEXT("ChunkHashList") );
+		for( auto ChunkHashListIt = ChunkHashList.CreateConstIterator(); ChunkHashListIt; ++ChunkHashListIt )
+		{
+			const FGuid& ChunkGuid = ChunkHashListIt.Key();
+			const uint64& ChunkHash = ChunkHashListIt.Value();
+			Writer->WriteValue( ChunkGuid.ToString(), ToStringBlob( ChunkHash ) );
+		}
+		Writer->WriteObjectEnd();
+		// Write data group list
+		Writer->WriteObjectStart( TEXT("DataGroupList") );
+		for( auto DataGroupListIt = DataGroupList.CreateConstIterator(); DataGroupListIt; ++DataGroupListIt )
+		{
+			const FGuid& DataGuid = DataGroupListIt.Key();
+			const uint8& DataGroup = DataGroupListIt.Value();
+			Writer->WriteValue( DataGuid.ToString(), ToStringBlob( DataGroup ) );
+		}
+		Writer->WriteObjectEnd();
+		// Write chunk size list
+		Writer->WriteObjectStart(TEXT("ChunkFilesizeList"));
+		for (auto ChunkFilesizeListIt = ChunkFilesizeList.CreateConstIterator(); ChunkFilesizeListIt; ++ChunkFilesizeListIt)
+		{
+			const FGuid& ChunkGuid = ChunkFilesizeListIt.Key();
+			const int64& ChunkSize = ChunkFilesizeListIt.Value();
+			Writer->WriteValue(ChunkGuid.ToString(), ToStringBlob(ChunkSize));
+		}
+		Writer->WriteObjectEnd();
+		// Write custom fields
+		Writer->WriteObjectStart( TEXT("CustomFields") );
+		for( auto CustomFieldsIt = CustomFields.CreateConstIterator(); CustomFieldsIt; ++CustomFieldsIt )
+		{
+			Writer->WriteValue( CustomFieldsIt.Key(), CustomFieldsIt.Value() );
+		}
+		Writer->WriteObjectEnd();
+	}
+	Writer->WriteObjectEnd();
+
+	Writer->Close();
+}
+
+// @TODO LSwift: Perhaps replace FromBlob and ToBlob usage with hexadecimal notation instead
+bool FBuildPatchAppManifest::DeserializeFromJSON( const FString& JSONInput )
+{
+	bool bSuccess = true;
+	TSharedPtr< FJsonObject > JSONManifestObject;
+	TSharedRef< TJsonReader< TCHAR > > Reader = TJsonReaderFactory< TCHAR >::Create( JSONInput );
+
+	// Clear current data
+	DestroyData();
+
+	// Attempt to deserialize JSON
+	if( !FJsonSerializer::Deserialize( Reader, JSONManifestObject ) || !JSONManifestObject.IsValid() )
+	{
+		return false;
+	}
+
+	// Get the values map
+	TMap< FString, TSharedPtr< FJsonValue > >& JsonValueMap = JSONManifestObject->Values;
+
+	// Get the app and version strings
+	TSharedPtr< FJsonValue > JsonAppID = JsonValueMap.FindRef( TEXT("AppID") );
+	TSharedPtr< FJsonValue > JsonAppNameString = JsonValueMap.FindRef( TEXT("AppNameString") );
+	TSharedPtr< FJsonValue > JsonBuildVersionString = JsonValueMap.FindRef( TEXT("BuildVersionString") );
+	TSharedPtr< FJsonValue > JsonLaunchExe = JsonValueMap.FindRef( TEXT("LaunchExeString") );
+	TSharedPtr< FJsonValue > JsonLaunchCommand = JsonValueMap.FindRef( TEXT("LaunchCommand") );
+	TSharedPtr< FJsonValue > JsonPrereqName = JsonValueMap.FindRef( TEXT("PrereqName") );
+	TSharedPtr< FJsonValue > JsonPrereqPath = JsonValueMap.FindRef( TEXT("PrereqPath") );
+	TSharedPtr< FJsonValue > JsonPrereqArgs = JsonValueMap.FindRef( TEXT("PrereqArgs") );
+	bSuccess = bSuccess && JsonAppID.IsValid();
+	if( bSuccess )
+	{
+		bSuccess = bSuccess && FromStringBlob( JsonAppID->AsString(), AppID );
+	}
+	bSuccess = bSuccess && JsonAppNameString.IsValid();
+	if( bSuccess )
+	{
+		AppNameString = JsonAppNameString->AsString();
+	}
+	bSuccess = bSuccess && JsonBuildVersionString.IsValid();
+	if( bSuccess )
+	{
+		BuildVersionString = JsonBuildVersionString->AsString();
+	}
+	bSuccess = bSuccess && JsonLaunchExe.IsValid();
+	if( bSuccess )
+	{
+		LaunchExeString = JsonLaunchExe->AsString();
+	}
+	bSuccess = bSuccess && JsonLaunchCommand.IsValid();
+	if( bSuccess )
+	{
+		LaunchCommandString = JsonLaunchCommand->AsString();
+	}
+
+	// Get the prerequisites installer info.  These are optional entries.
+	PrereqNameString = JsonPrereqName.IsValid() ? JsonPrereqName->AsString() : FString();
+	PrereqPathString = JsonPrereqPath.IsValid() ? JsonPrereqPath->AsString() : FString();
+	PrereqArgsString = JsonPrereqArgs.IsValid() ? JsonPrereqArgs->AsString() : FString();
+
+	// Manifest version did not always exist
+	int32 ManifestFileVersionInt = 0;
+	TSharedPtr< FJsonValue > JsonManifestFileVersion = JsonValueMap.FindRef( TEXT("ManifestFileVersion") );
+	if( JsonManifestFileVersion.IsValid() && FromStringBlob( JsonManifestFileVersion->AsString(), ManifestFileVersionInt ) )
+	{
+		ManifestFileVersion = static_cast< EBuildPatchAppManifestVersion::Type >( ManifestFileVersionInt );
+	}
+	else
+	{
+		// Then we presume version just before we started outputting the version
+		ManifestFileVersion = static_cast< EBuildPatchAppManifestVersion::Type >( EBuildPatchAppManifestVersion::StartStoringVersion - 1 );
+	}
+
+	// Store a list of all data GUID for later use
+	TSet< FGuid > AllDataGuids;
+
+	// Get the FileManifestList
+	TSharedPtr< FJsonValue > JsonFileManifestList = JsonValueMap.FindRef( TEXT("FileManifestList") );
+	bSuccess = bSuccess && JsonFileManifestList.IsValid();
+	if( bSuccess )
+	{
+		TArray< TSharedPtr< FJsonValue > > JsonFileManifestArray = JsonFileManifestList->AsArray();
+		for ( auto JsonFileManifestIt = JsonFileManifestArray.CreateConstIterator(); JsonFileManifestIt && bSuccess; ++JsonFileManifestIt )
+		{
+			TSharedPtr< FJsonObject > JsonFileManifest = (*JsonFileManifestIt)->AsObject();
+			TSharedRef< FBuildPatchFileManifest > FileManifest = MakeShareable( new FBuildPatchFileManifest() );
+			FileManifest->Filename = JsonFileManifest->GetStringField( TEXT("Filename") );
+#if PLATFORM_MAC
+			FileManifest->bIsUnixExecutable = JsonFileManifest->HasField( TEXT("bIsUnixExecutable") ) && JsonFileManifest->GetBoolField( TEXT("bIsUnixExecutable") );
+			FileManifest->SymlinkTarget = JsonFileManifest->HasField( TEXT("SymlinkTarget") ) ? JsonFileManifest->GetStringField( TEXT("SymlinkTarget") ) : TEXT( "" );
+#endif
+			FileManifestList.Add( FileManifest->Filename, FileManifest );
+			bSuccess = bSuccess && FString::ToBlob( JsonFileManifest->GetStringField( TEXT("FileHash") ), FileManifest->FileHash.Hash, FSHA1::DigestSize );
+			TArray< TSharedPtr< FJsonValue > > JsonChunkPartArray = JsonFileManifest->GetArrayField( TEXT("FileChunkParts") );
+			for ( auto JsonChunkPartIt = JsonChunkPartArray.CreateConstIterator(); JsonChunkPartIt && bSuccess; ++JsonChunkPartIt )
+			{
+				FileManifest->FileChunkParts.Add( FChunkPart() );
+				TSharedPtr< FJsonObject > JsonChunkPart = (*JsonChunkPartIt)->AsObject();
+				FChunkPart& FileChunkPart = FileManifest->FileChunkParts.Top();
+				bSuccess = bSuccess && FGuid::Parse( JsonChunkPart->GetStringField( TEXT("Guid") ), FileChunkPart.Guid );
+				bSuccess = bSuccess && FromStringBlob( JsonChunkPart->GetStringField( TEXT("Offset") ), FileChunkPart.Offset );
+				bSuccess = bSuccess && FromStringBlob( JsonChunkPart->GetStringField( TEXT("Size") ), FileChunkPart.Size );
+				AllDataGuids.Add( FileChunkPart.Guid );
+			}
+		}
+	}
+
+	// Get the ChunkHashList
+	TSharedPtr< FJsonValue > JsonChunkHashList = JsonValueMap.FindRef( TEXT("ChunkHashList") );
+	bSuccess = bSuccess && JsonChunkHashList.IsValid();
+	if( bSuccess )
+	{
+		TSharedPtr< FJsonObject > JsonChunkHashListObj = JsonChunkHashList->AsObject();
+		for( auto ChunkHashIt = JsonChunkHashListObj->Values.CreateConstIterator(); ChunkHashIt && bSuccess; ++ChunkHashIt )
+		{
+			FGuid ChunkGuid;
+			uint64 ChunkHash = 0;
+			bSuccess = bSuccess && FGuid::Parse( ChunkHashIt.Key(), ChunkGuid );
+			bSuccess = bSuccess && FromStringBlob( ChunkHashIt.Value()->AsString(), ChunkHash );
+			ChunkHashList.Add( ChunkGuid, ChunkHash );
+			AllDataGuids.Add( ChunkGuid );
+		}
+	}
+
+	// Get the bIsFileData value. The variable will exist in versions of StoresIfChunkOrFileData or later, otherwise the previous method is to check
+	// if ChunkHashList is empty.
+	TSharedPtr< FJsonValue > JsonIsFileData = JsonValueMap.FindRef( TEXT("bIsFileData") );
+	if( JsonIsFileData.IsValid() && JsonIsFileData->Type == EJson::Boolean )
+	{
+		bIsFileData = JsonIsFileData->AsBool();
+	}
+	else
+	{
+		bIsFileData = ChunkHashList.Num() == 0;
+	}
+
+	// Get the DataGroupList
+	TSharedPtr< FJsonValue > JsonDataGroupList = JsonValueMap.FindRef( TEXT("DataGroupList") );
+	if( JsonDataGroupList.IsValid() )
+	{
+		TSharedPtr< FJsonObject > JsonDataGroupListObj = JsonDataGroupList->AsObject();
+		for( auto DataGroupIt = JsonDataGroupListObj->Values.CreateConstIterator(); DataGroupIt && bSuccess; ++DataGroupIt )
+		{
+			FGuid DataGuid;
+			uint8 DataGroup;
+			// If the list exists, we must be able to parse it ok otherwise error
+			bSuccess = bSuccess && FGuid::Parse( DataGroupIt.Key(), DataGuid );
+			bSuccess = bSuccess && FromStringBlob( DataGroupIt.Value()->AsString(), DataGroup );
+			DataGroupList.Add( DataGuid, DataGroup );
+		}
+	}
+	else if( bSuccess )
+	{
+		// If the list did not exist in the manifest then the grouping is the deprecated crc functionality, as long
+		// as there are no previous parsing errors we can build the group list from the AllDataGuids list.
+		for( auto AllDataGuidsIt = AllDataGuids.CreateConstIterator(); AllDataGuidsIt; ++AllDataGuidsIt )
+		{
+			const FGuid& DataGuid = *AllDataGuidsIt;
+			uint8 DataGroup = FCrc::MemCrc_DEPRECATED( &DataGuid, sizeof( FGuid ) ) % 100;
+			DataGroupList.Add( DataGuid, DataGroup );
+		}
+	}
+
+	// Get the ChunkFilesizeList
+	TSharedPtr< FJsonValue > JsonChunkFilesizeList = JsonValueMap.FindRef(TEXT("ChunkFilesizeList"));
+	if (JsonChunkFilesizeList.IsValid())
+	{
+		TSharedPtr< FJsonObject > JsonChunkFilesizeListObj = JsonChunkFilesizeList->AsObject();
+		for (auto ChunkFilesizeIt = JsonChunkFilesizeListObj->Values.CreateConstIterator(); ChunkFilesizeIt; ++ChunkFilesizeIt)
+		{
+			FGuid ChunkGuid;
+			int64 ChunkSize = 0;
+			if (FGuid::Parse(ChunkFilesizeIt.Key(), ChunkGuid))
+			{
+				FromStringBlob(ChunkFilesizeIt.Value()->AsString(), ChunkSize);
+				ChunkFilesizeList.Add(ChunkGuid, ChunkSize);
+			}
+		}
+	}
+
+	// Get the custom fields. This is optional, and should not fail if it does not exist
+	TSharedPtr< FJsonValue > JsonCustomFields = JsonValueMap.FindRef( TEXT( "CustomFields" ) );
+	if( JsonCustomFields.IsValid() )
+	{
+		TSharedPtr< FJsonObject > JsonCustomFieldsObj = JsonCustomFields->AsObject();
+		for( auto CustomFieldIt = JsonCustomFieldsObj->Values.CreateConstIterator(); CustomFieldIt && bSuccess; ++CustomFieldIt )
+		{
+			CustomFields.Add( CustomFieldIt.Key(), CustomFieldIt.Value()->AsString() );
+		}
+	}
+
+	// If this is file data, fill out the guid to filename lookup
+	if( bIsFileData )
+	{
+		for( auto FileManifestIt = FileManifestList.CreateConstIterator(); FileManifestIt; ++FileManifestIt )
+		{
+			const TSharedPtr< FBuildPatchFileManifest > FileManifest = FileManifestIt.Value();
+			check( FileManifest->FileChunkParts.Num() == 1 );
+			FileNameReference.Add( FileManifest->FileChunkParts[ 0 ].Guid, FileManifest->Filename );
+		}
+	}
+
+	// Make sure we don't have any half loaded data
+	if( !bSuccess )
+	{
+		DestroyData();
+	}
+
+	return bSuccess;
+}
+
+const EBuildPatchAppManifestVersion::Type FBuildPatchAppManifest::GetManifestVersion() const
+{
+	return ManifestFileVersion;
+}
+
+void FBuildPatchAppManifest::GetChunksRequiredForFiles( const TArray< FString >& FileList, TArray< FGuid >& RequiredChunks, bool bAddUnique ) const
+{
+	for( auto FileIt = FileList.CreateConstIterator(); FileIt ; ++FileIt )
+	{
+		TSharedPtr< FBuildPatchFileManifest > FileManifest = FileManifestList.FindRef( *FileIt );
+		if( FileManifest.IsValid() )
+		{
+			for( auto ChunkPartIt = FileManifest->FileChunkParts.CreateConstIterator(); ChunkPartIt; ++ChunkPartIt )
+			{
+				const FChunkPart& ChunkPart = *ChunkPartIt;
+				if( bAddUnique )
+				{
+					RequiredChunks.AddUnique( ChunkPart.Guid );
+				}
+				else
+				{
+					RequiredChunks.Add( ChunkPart.Guid );
+				}
+			}
+		}
+	}
+}
+
+const int64 FBuildPatchAppManifest::GetDownloadSize() const
+{
+	// If this is not a chunked build, then return build size
+	if (bIsFileData)
+	{
+		return GetBuildSize();
+	}
+	else
+	{
+		TArray<FGuid> ChunkList;
+		ChunkHashList.GetKeys(ChunkList);
+		return GetDataSize(ChunkList);
+	}
+}
+
+const int64 FBuildPatchAppManifest::GetBuildSize() const
+{
+	int64 TotalSize = 0;
+	for ( auto& File : FileManifestList )
+	{
+		if ( File.Value.IsValid() )
+		{
+			TotalSize += File.Value->GetFileSize();
+		}
+	}
+
+	return TotalSize;
+}
+
+const int64 FBuildPatchAppManifest::GetFileSize( const TArray< FString >& Filenames ) const
+{
+	int64 TotalSize = 0;
+	for( int32 FileIdx = 0; FileIdx < Filenames.Num(); ++FileIdx )
+	{
+		TSharedPtr< FBuildPatchFileManifest > FileManifest = FileManifestList.FindRef( Filenames[ FileIdx ] );
+		if( FileManifest.IsValid() )
+		{
+			TotalSize += FileManifest->GetFileSize();
+		}
+	}
+	return TotalSize;
+}
+
+const int64 FBuildPatchAppManifest::GetFileSize( const FString& Filename ) const
+{
+	TSharedPtr< FBuildPatchFileManifest > FileManifest = FileManifestList.FindRef( Filename );
+	if( FileManifest.IsValid() )
+	{
+		return FileManifest->GetFileSize();
+	}
+	return 0;
+}
+
+const int64 FBuildPatchAppManifest::GetDataSize(const FGuid& DataGuid) const
+{
+	if (bIsFileData)
+	{
+		// For file data, the file must exist in the list
+		check(FileNameReference.Contains(DataGuid));
+		return GetFileSize(FileNameReference[DataGuid]);
+	}
+	else if (ChunkFilesizeList.Contains(DataGuid))
+	{
+		// Chunk data sizes are stored in our map
+		return ChunkFilesizeList[DataGuid];
+	}
+	else
+	{
+		// Default chunk size to be the data size. Inaccurate, but represents original behavior.
+		return FBuildPatchData::ChunkDataSize;
+	}
+}
+
+const int64 FBuildPatchAppManifest::GetDataSize(const TArray< FGuid >& DataGuids) const
+{
+	int64 TotalSize = 0;
+	for (auto DataGuidsIt = DataGuids.CreateConstIterator(); DataGuidsIt; ++DataGuidsIt)
+	{
+		TotalSize += GetDataSize(*DataGuidsIt);
+	}
+	return TotalSize;
+}
+
+const uint32 FBuildPatchAppManifest::GetNumFiles() const
+{
+	return FileManifestList.Num();
+}
+
+void FBuildPatchAppManifest::GetFileList( TArray< FString >& Filenames ) const
+{
+	FileManifestList.GetKeys( Filenames );
+}
+
+void FBuildPatchAppManifest::GetDataList(TArray< FGuid >& DataGuids) const
+{
+	ChunkHashList.GetKeys( DataGuids );
+}
+
+const TSharedPtr< FBuildPatchFileManifest > FBuildPatchAppManifest::GetFileManifest( const FString& Filename )
+{
+	return FileManifestList.FindRef( Filename );
+}
+
+const bool FBuildPatchAppManifest::IsFileDataManifest() const
+{
+	return bIsFileData;
+}
+
+const bool FBuildPatchAppManifest::GetChunkHash( const FGuid& ChunkGuid, uint64& OutHash ) const
+{
+	const uint64* ChunkHash = ChunkHashList.Find( ChunkGuid );
+	if( ChunkHash != NULL )
+	{
+		OutHash = *ChunkHash;
+		return true;
+	}
+	return false;
+}
+
+const bool FBuildPatchAppManifest::GetFileDataHash( const FGuid& FileGuid, FSHAHash& OutHash ) const
+{
+	// File data hash is the same as the hash for the file it holds, so we can grab it from there
+	for( auto FileManifestIt = FileManifestList.CreateConstIterator(); FileManifestIt; ++FileManifestIt )
+	{
+		if( FileManifestIt.Value().IsValid() )
+		{
+			const FBuildPatchFileManifest& FileManifest = *FileManifestIt.Value().Get();
+			for( auto FileChunkPartsIt = FileManifest.FileChunkParts.CreateConstIterator(); FileChunkPartsIt; ++FileChunkPartsIt )
+			{
+				const FChunkPart& FileDataPart = *FileChunkPartsIt;
+				if( FileDataPart.Guid == FileGuid )
+				{
+					OutHash = FileManifest.FileHash;
+					return true;
+				}
+			}
+		}
+	}
+	return false;
+}
+
+const uint32& FBuildPatchAppManifest::GetAppID() const
+{
+	return AppID;
+}
+
+const FString& FBuildPatchAppManifest::GetAppName() const
+{
+	return AppNameString;
+}
+
+const FString& FBuildPatchAppManifest::GetVersionString() const
+{
+	return BuildVersionString;
+}
+
+const FString& FBuildPatchAppManifest::GetLaunchExe() const
+{
+	return LaunchExeString;
+}
+
+const FString& FBuildPatchAppManifest::GetLaunchCommand() const
+{
+	return LaunchCommandString;
+}
+
+const FString& FBuildPatchAppManifest::GetPrereqName() const
+{
+	return PrereqNameString;
+}
+
+const FString& FBuildPatchAppManifest::GetPrereqPath() const
+{
+	return PrereqPathString;
+}
+
+const FString& FBuildPatchAppManifest::GetPrereqArgs() const
+{
+	return PrereqArgsString;
+}
+
+IBuildManifestRef FBuildPatchAppManifest::Duplicate() const
+{
+	return MakeShareable(new FBuildPatchAppManifest(*this));
+}
+
+const IManifestFieldPtr FBuildPatchAppManifest::GetCustomField( const FString& FieldName ) const
+{
+	IManifestFieldPtr Return;
+	if( CustomFields.Contains( FieldName ) )
+	{
+		Return = MakeShareable( new FBuildPatchCustomField( CustomFields.FindRef( FieldName ) ) );
+	}
+	return Return;
+}
+
+const IManifestFieldPtr FBuildPatchAppManifest::SetCustomField( const FString& FieldName, const FString& Value )
+{
+	return MakeShareable( new FBuildPatchCustomField( CustomFields.Add( FieldName, Value ) ) );
+}
+
+const IManifestFieldPtr FBuildPatchAppManifest::SetCustomField( const FString& FieldName, const double& Value )
+{
+	return MakeShareable( new FBuildPatchCustomField( CustomFields.Add( FieldName, ToStringBlob( Value ) ) ) );
+}
+
+const IManifestFieldPtr FBuildPatchAppManifest::SetCustomField( const FString& FieldName, const int64& Value )
+{
+	return MakeShareable( new FBuildPatchCustomField( CustomFields.Add( FieldName, ToStringBlob( Value ) ) ) );
+}
+
+void FBuildPatchAppManifest::EnumerateProducibleChunks( const FString& InstallDirectory, const TArray< FGuid >& ChunksRequired, TArray< FGuid >& ChunksAvailable ) const
+{
+	// A struct that will store byte ranges
+	struct FChunkRange
+	{
+		// The inclusive min byte
+		uint32 Min;
+		// The inclusive max byte
+		uint32 Max;
+	};
+	// A struct that will sort an FChunkRange array by Min
+	struct FCompareFChunkRangeByMin
+	{
+		FORCEINLINE bool operator()( const FChunkRange& A, const FChunkRange& B ) const
+		{
+			return A.Min < B.Min;
+		}
+	};
+	// The first thing we need is an inventory of FFileChunkParts that we can get which refer to our required chunks
+	TMap< FGuid, TArray< FFileChunkPart > > ChunkPartInventory;
+	EnumerateChunkPartInventory( ChunksRequired, ChunkPartInventory );
+
+	// For each chunk that we found a part for, check that the union of chunk parts we have for it, contains all bytes of the chunk
+	for( auto ChunkPartInventoryIt = ChunkPartInventory.CreateConstIterator(); ChunkPartInventoryIt; ++ChunkPartInventoryIt )
+	{
+		// Create an array of FChunkRanges that describes the chunk data that we have available
+		const FGuid& ChunkGuid = ChunkPartInventoryIt.Key();
+		const TArray< FFileChunkPart >& ChunkParts = ChunkPartInventoryIt.Value();
+		if( ChunksAvailable.Contains( ChunkGuid ) )
+		{
+			continue;
+		}
+		TArray< FChunkRange > ChunkRanges;
+		for( auto ChunkPartIt = ChunkParts.CreateConstIterator(); ChunkPartIt; ++ChunkPartIt )
+		{
+			const FFileChunkPart& FileChunkPart = *ChunkPartIt;
+			const int64 SourceFilesize = IFileManager::Get().FileSize( *(InstallDirectory / FileChunkPart.Filename) );
+			const int64 LastRequiredByte = FileChunkPart.FileOffset + FileChunkPart.ChunkPart.Size;
+			if( SourceFilesize == GetFileSize( FileChunkPart.Filename ) && SourceFilesize >= LastRequiredByte )
+			{
+				const uint32 Min = FileChunkPart.ChunkPart.Offset;
+				const uint32 Max = Min + FileChunkPart.ChunkPart.Size - 1;
+				// We will store our ranges using inclusive values
+				FChunkRange NextRange;
+				NextRange.Min = Min;
+				NextRange.Max = Max;
+				ChunkRanges.Add( NextRange );
+			}
+		}
+		ChunkRanges.Sort( FCompareFChunkRangeByMin() );
+		// Now we have a sorted array of ChunkPart ranges, we walk this array to find a gap in available data
+		// which would mean we cannot generate this chunk
+		uint32 ByteCount = 0;
+		for( auto ChunkRangeIt = ChunkRanges.CreateConstIterator(); ChunkRangeIt; ++ChunkRangeIt )
+		{
+			const FChunkRange& ChunkRange = *ChunkRangeIt;
+			if( ChunkRange.Min <= ByteCount && ChunkRange.Max >= ByteCount )
+			{
+				ByteCount = ChunkRange.Max;
+			}
+			else
+			{
+				break;
+			}
+		}
+		// If we can make the chunk, add it to the list
+		const bool bCanMakeChunk = ByteCount == ( FBuildPatchData::ChunkDataSize - 1 );
+		if( bCanMakeChunk )
+		{
+			ChunksAvailable.AddUnique( ChunkGuid );
+		}
+	}
+}
+
+bool FBuildPatchAppManifest::VerifyAgainstDirectory( const FString& VerifyDirectory, TArray < FString >& OutDatedFiles, FBuildPatchFloatDelegate ProgressDelegate, FBuildPatchBoolRetDelegate ShouldPauseDelegate, double& TimeSpentPaused )
+{
+	bool bAllCorrect = true;
+	OutDatedFiles.Empty();
+	TimeSpentPaused = 0;
+
+	// Setup progress tracking
+	double TotalBuildSize = GetBuildSize();
+	double ProcessedBytes = 0;
+	struct FLocalProgress
+	{
+		double CurrentBuildPercentage;
+		double CurrentFileWeight;
+		FBuildPatchFloatDelegate ProgressDelegate;
+		void PerFileProgress( float Progress )
+		{
+			const double PercentComplete = CurrentBuildPercentage + ( Progress * CurrentFileWeight );
+			ProgressDelegate.ExecuteIfBound( PercentComplete );
+		}
+	} LocalProgress;
+	LocalProgress.CurrentBuildPercentage = 0;
+	LocalProgress.ProgressDelegate = ProgressDelegate;
+
+	// For all files in FileManifestList, check that they produce the correct SHA1 hash, adding any that don't to the list
+	for( auto FileManifestIt = FileManifestList.CreateConstIterator(); FileManifestIt && !FBuildPatchInstallError::HasFatalError(); ++FileManifestIt )
+	{
+		// Pause if necessary
+		const double PrePauseTime = FPlatformTime::Seconds();
+		double PostPauseTime = PrePauseTime;
+		bool bShouldPause = ShouldPauseDelegate.IsBound() && ShouldPauseDelegate.Execute();
+		while( bShouldPause )
+		{
+			FPlatformProcess::Sleep( 0.1f );
+			bShouldPause = ShouldPauseDelegate.Execute();
+			PostPauseTime = FPlatformTime::Seconds();
+		}
+		// Count up pause time
+		TimeSpentPaused += PostPauseTime - PrePauseTime;
+		// Get file details and construct
+		TSharedPtr< FBuildPatchFileManifest > FileManifest = FileManifestIt.Value();
+		FString FullFilename = VerifyDirectory / FileManifest->Filename;
+		const int64 CurrentFileSize = FileManifest->GetFileSize();
+		LocalProgress.CurrentFileWeight = CurrentFileSize / TotalBuildSize;
+		const bool bSameFile = FBuildPatchUtils::VerifyFile(FullFilename, FileManifest->FileHash, FileManifest->FileHash, FBuildPatchFloatDelegate::CreateRaw(&LocalProgress, &FLocalProgress::PerFileProgress), ShouldPauseDelegate, TimeSpentPaused) != 0;
+		ProcessedBytes += CurrentFileSize;
+		LocalProgress.CurrentBuildPercentage = ProcessedBytes / TotalBuildSize;
+		bAllCorrect &= bSameFile;
+		if( !bSameFile )
+		{
+			OutDatedFiles.Add( FileManifest->Filename );
+		}
+	}
+
+	return bAllCorrect && !FBuildPatchInstallError::HasFatalError();
+}
+
+void FBuildPatchAppManifest::EnumerateChunkPartInventory( const TArray< FGuid >& ChunksRequired, TMap< FGuid, TArray< FFileChunkPart > >& ChunkPartsAvailable ) const
+{
+	ChunkPartsAvailable.Empty();
+	// For each file in the manifest, check what chunks it is made out of, and grab details for the ones in ChunksRequired
+	for( auto FileManifestIt = FileManifestList.CreateConstIterator(); FileManifestIt; ++FileManifestIt )
+	{
+		const TSharedPtr< FBuildPatchFileManifest > FileManifest = FileManifestIt.Value();
+		uint64 FileOffset = 0;
+		for( auto ChunkPartIt = FileManifest->FileChunkParts.CreateConstIterator(); ChunkPartIt; ++ChunkPartIt )
+		{
+			const FChunkPart& ChunkPart = *ChunkPartIt;
+			if( ChunksRequired.Contains( ChunkPart.Guid ) )
+			{
+				TArray< FFileChunkPart >& FileChunkParts = ChunkPartsAvailable.FindOrAdd( ChunkPart.Guid );
+				FFileChunkPart FileChunkPart;
+				FileChunkPart.Filename = FileManifest->Filename;
+				FileChunkPart.ChunkPart = ChunkPart;
+				FileChunkPart.FileOffset = FileOffset;
+				FileChunkParts.Add( FileChunkPart );
+			}
+			FileOffset += ChunkPart.Size;
+		}
+	}
+}
+
+void FBuildPatchAppManifest::GetRemovableFiles(FBuildPatchAppManifestRef OldManifest, FBuildPatchAppManifestRef NewManifest, TArray< FString >& RemovableFiles)
+{
+	NewManifest->GetRemovableFiles(OldManifest, RemovableFiles);
+}
+
+void FBuildPatchAppManifest::GetRemovableFiles(IBuildManifestRef InOldManifest, TArray< FString >& RemovableFiles) const
+{
+	// Cast manifest parameters
+	FBuildPatchAppManifestRef OldManifest = StaticCastSharedRef< FBuildPatchAppManifest >(InOldManifest);
+	// Simply put, any files that exist in the OldManifest file list, but do not in this manifest's file list, are assumed
+	// to be files no longer required by the build
+	for (auto OldFilesIt = OldManifest->FileManifestList.CreateConstIterator(); OldFilesIt; ++OldFilesIt)
+	{
+		const FString& OldFile = OldFilesIt.Key();
+		if (!this->FileManifestList.Contains(OldFile))
+		{
+			RemovableFiles.Add(OldFile);
+		}
+	}
+}
+
+void FBuildPatchAppManifest::GetRemovableFiles(const TCHAR* InstallPath, TArray< FString >& RemovableFiles) const
+{
+	TArray<FString> AllFiles;
+	IFileManager::Get().FindFilesRecursive(AllFiles, InstallPath, TEXT("*"), true, false);
+	
+	FString BasePath = InstallPath;
+
+#if PLATFORM_MAC
+	// On Mac paths in manifest start with app bundle name
+	if (BasePath.EndsWith(".app"))
+	{
+		BasePath = FPaths::GetPath(BasePath) + "/";
+	}
+#endif
+	
+	for (int32 FileIndex = 0; FileIndex < AllFiles.Num(); ++FileIndex)
+	{
+		const FString& FileName = AllFiles[FileIndex].RightChop(BasePath.Len());
+		if (!FileManifestList.Contains(*FileName))
+		{
+			RemovableFiles.Add(AllFiles[FileIndex]);
+		}
+	}
+}
+
+void FBuildPatchAppManifest::GetOutdatedFiles( FBuildPatchAppManifestPtr OldManifest, FBuildPatchAppManifestRef NewManifest, const FString& InstallDirectory, TArray< FString >& OutDatedFiles )
+{
+	if( !OldManifest.IsValid() )
+	{
+		// All files are outdated if no OldManifest
+		NewManifest->FileManifestList.GetKeys( OutDatedFiles );
+	}
+	else
+	{
+		// Enumerate files in the NewManifest file list, that do not exist, or have different hashes in the OldManifest
+		// to be files no longer required by the build
+		for( auto NewFilesIt = NewManifest->FileManifestList.CreateConstIterator(); NewFilesIt; ++NewFilesIt )
+		{
+			const FString& NewFile = NewFilesIt.Key();
+			const int64 ExistingFileSize = IFileManager::Get().FileSize( *(InstallDirectory / NewFile) );
+			const TSharedPtr< FBuildPatchFileManifest > NewFileManifest = NewFilesIt.Value();
+			const TSharedPtr< FBuildPatchFileManifest > OldFileManifest = OldManifest->FileManifestList.FindRef( NewFile );
+			// Check changed
+			if( IsFileOutdated( OldManifest.ToSharedRef(), NewManifest, NewFile ) )
+			{
+				OutDatedFiles.Add( NewFile );
+			}
+			// Double check an unchanged file is not missing (size will be -1) or is incorrect size
+			else if( ExistingFileSize != NewFileManifest->GetFileSize() )
+			{
+				OutDatedFiles.Add( NewFile );
+			}
+		}
+	}
+}
+
+bool FBuildPatchAppManifest::IsFileOutdated( FBuildPatchAppManifestRef OldManifest, FBuildPatchAppManifestRef NewManifest, const FString& Filename )
+{
+	// If both app manifests are the same, return false as only repair would touch the file.
+	if( OldManifest == NewManifest )
+	{
+		return false;
+	}
+	// Get file manifests
+	TSharedPtr< FBuildPatchFileManifest > OldFile = OldManifest->GetFileManifest( Filename );
+	TSharedPtr< FBuildPatchFileManifest > NewFile = NewManifest->GetFileManifest( Filename );
+	// Out of date if not in either manifest
+	if( !OldFile.IsValid() || !NewFile.IsValid() )
+	{
+		return true;
+	}
+	// Different hash means different file
+	if( OldFile->FileHash != NewFile->FileHash )
+	{
+		return true;
+	}
+	return false;
+}
+
+const uint32 FBuildPatchAppManifest::GetNumberOfChunkReferences( const FGuid& ChunkGuid ) const
+{
+	uint32 RefCount = 0;
+	// For each file in the manifest, check if it has references to this chunk
+	for( auto FileManifestIt = FileManifestList.CreateConstIterator(); FileManifestIt; ++FileManifestIt )
+	{
+		const TSharedPtr< FBuildPatchFileManifest >& FileManifest = FileManifestIt.Value();
+		for( auto ChunkPartIt = FileManifest->FileChunkParts.CreateConstIterator(); ChunkPartIt; ++ChunkPartIt )
+		{
+			const FChunkPart& ChunkPart = *ChunkPartIt;
+			if( ChunkPart.Guid == ChunkGuid )
+			{
+				++RefCount;
+			}
+		}
+	}
+	return RefCount;
+}
+
+#undef LOCTEXT_NAMESPACE
diff --git a/Engine/Source/Runtime/Online/BuildPatchServices/Private/BuildPatchManifest.h b/Engine/Source/Runtime/Online/BuildPatchServices/Private/BuildPatchManifest.h
new file mode 100644
--- /dev/null
+++ b/Engine/Source/Runtime/Online/BuildPatchServices/Private/BuildPatchManifest.h
@@ -0,0 +1,442 @@
+// Copyright 1998-2014 Epic Games, Inc. All Rights Reserved.
+
+/*=============================================================================
+	BuildPatchManifest.h: Declares the manifest classes.
+=============================================================================*/
+
+#pragma once
+
+// A delegate returning a bool. Used to pass a paused state
+DECLARE_DELEGATE_RetVal( bool, FBuildPatchBoolRetDelegate );
+
+typedef TSharedPtr< class FBuildPatchCustomField, ESPMode::ThreadSafe > FBuildPatchCustomFieldPtr;
+typedef TSharedRef< class FBuildPatchCustomField, ESPMode::ThreadSafe > FBuildPatchCustomFieldRef;
+typedef TSharedPtr< class FBuildPatchAppManifest, ESPMode::ThreadSafe > FBuildPatchAppManifestPtr;
+typedef TSharedRef< class FBuildPatchAppManifest, ESPMode::ThreadSafe > FBuildPatchAppManifestRef;
+
+/**
+ * An enum type to describe supported features of a certain manifest
+ */
+namespace EBuildPatchAppManifestVersion
+{
+	enum Type
+	{
+		// The original version
+		Original = 0,
+		// Support for custom fields
+		CustomFields,
+		// Started storing the version number
+		StartStoringVersion,
+		// Made after data files where renamed to include the hash value, these chunks now go to ChunksV2
+		DataFileRenames,
+		// Manifest stores whether build was constructed with chunk or file data
+		StoresIfChunkOrFileData,
+		// Manifest stores group number for each chunk/file data for reference so that external readers don't need to know how to calculate them
+		StoresDataGroupNumbers,
+		// Added support for chunk compression, these chunks now go to ChunkV3. NB: Not File Data Compression yet
+		ChunkCompressionSupport,
+		// Manifest stores product prerequisites info
+		StoresPrerequisitesInfo,
+		// Manifest stores chunk download sizes
+		StoresChunkFileSizes,
+
+		// Always last, signifies the latest version plus 1 to allow initialization simplicity
+		LatestPlusOne
+	};
+
+	/** @return The last known manifest feature of this codebase. Handy for manifest constructor */
+	const Type GetLatestVersion();
+
+	/**
+	 * Get the chunk subdirectory for used for a specific manifest version, e.g. Chunks, ChunksV2 etc
+	 * @param ManifestVersion     The version of the manifest
+	 * @return The subdirectory name that this manifest version will access
+	 */
+	const FString& GetChunkSubdir( const Type ManifestVersion );
+
+	/**
+	 * Get the file data subdirectory for used for a specific manifest version, e.g. Files, FilesV2 etc
+	 * @param ManifestVersion     The version of the manifest
+	 * @return The subdirectory name that this manifest version will access
+	 */
+	const FString& GetFileSubdir( const Type ManifestVersion );
+};
+
+/**
+ * Declare the FBuildPatchCustomField object class
+ */
+class FBuildPatchCustomField
+	: public IManifestField
+{
+private:
+	// Holds the underlying value
+	FString CustomValue;
+
+	/**
+	 * Hide the default constructor
+	 */
+	FBuildPatchCustomField(){}
+public:
+	
+	/**
+	 * Constructor taking the custom value
+	 */
+	FBuildPatchCustomField( const FString& Value );
+
+	// START IBuildManifest Interface
+	virtual const FString AsString() const override;
+	virtual const double AsDouble() const override;
+	virtual const int64 AsInteger() const override;
+	// END IBuildManifest Interface
+};
+
+/**
+ * Declare the FBuildPatchFileManifest object class
+ */
+class FBuildPatchFileManifest
+{
+public:
+
+	// The filename of the file, local to a build's root
+	FString Filename;
+
+	// The Integrity Verification hash
+	FSHAHash FileHash;
+
+	// The sequential list of chunk parts this file is made up of
+	TArray< FChunkPart > FileChunkParts;
+
+#if PLATFORM_MAC
+	// Additional properties for Unix files
+	bool bIsUnixExecutable;
+	FString SymlinkTarget;
+#endif
+
+public:
+	// Default constructor
+	FBuildPatchFileManifest();
+	
+	/**
+	 * Get the total size for this file
+	 * @return		The size in bytes.
+	 */
+	const int64& GetFileSize();
+
+private:
+	// Holds the file size, calculated from the FileChunkParts
+	int64 FileSize;
+};
+
+/**
+ * Declare the FBuildPatchAppManifest object class. This will contain FBuildPatchFileManifest and other build information.
+ */
+class FBuildPatchAppManifest
+	: public IBuildManifest
+{
+	// Allow access to build processor classes
+	friend class FBuildDataGenerator;
+	friend class FBuildDataChunkProcessor;
+	friend class FBuildDataFileProcessor;
+	friend class FBuildPatchInstaller;
+public:
+	
+	/**
+	 * Default constructor
+	 */
+	FBuildPatchAppManifest();
+
+	/**
+	 * Basic details constructor
+	 */
+	FBuildPatchAppManifest( const uint32& InAppID, const FString& AppName );
+
+	/**
+	 * Default destructor
+	 */
+	~FBuildPatchAppManifest();
+
+	// START IBuildManifest Interface
+	virtual const uint32&  GetAppID() const override;
+	virtual const FString& GetAppName() const override;
+	virtual const FString& GetVersionString() const override;
+	virtual const FString& GetLaunchExe() const override;
+	virtual const FString& GetLaunchCommand() const override;
+	virtual const FString& GetPrereqName() const override;
+	virtual const FString& GetPrereqPath() const override;
+	virtual const FString& GetPrereqArgs() const override;
+	virtual const int64 GetDownloadSize() const override;
+	virtual const int64 GetBuildSize() const override;
+	virtual void GetRemovableFiles(IBuildManifestRef OldManifest, TArray< FString >& RemovableFiles) const override;
+	virtual void GetRemovableFiles(const TCHAR* InstallPath, TArray< FString >& RemovableFiles) const override;
+	virtual IBuildManifestRef Duplicate() const override;
+	virtual const IManifestFieldPtr GetCustomField(const FString& FieldName) const override;
+	virtual const IManifestFieldPtr SetCustomField( const FString& FieldName, const FString& Value ) override;
+	virtual const IManifestFieldPtr SetCustomField( const FString& FieldName, const double& Value ) override;
+	virtual const IManifestFieldPtr SetCustomField( const FString& FieldName, const int64& Value ) override;
+	// END IBuildManifest Interface
+
+	/**
+	 * Sets up the internal map from a file
+	 * @param Filename		The file to load JSON from
+	 * @return		True if successful.
+	 */
+	bool LoadFromFile( FString Filename );
+
+	/**
+	 * Sets up the object from the passed in JSON string
+	 * @param JSONInput		The JSON string to deserialize from
+	 * @return		True if successful.
+	 */
+	bool DeserializeFromJSON( const FString& JSONInput );
+
+	/**
+	 * Saves out the manifest information
+	 * @param Filename		The file to save a JSON representation to
+	 * @return		True if successful.
+	 */
+	bool SaveToFile( FString Filename );
+
+	/**
+	 * Creates the object in JSON format
+	 * @param JSONOutput		A string to receive the JSON representation
+	 */
+	void SerializeToJSON( FString& JSONOutput );
+
+	/**
+	 * Gets the version for this manifest. Useful for manifests that were loaded from JSON.
+	 * @return		The highest available feature support
+	 */
+	const EBuildPatchAppManifestVersion::Type GetManifestVersion() const;
+
+	/**
+	 * Provides a list of chunks required to produce the list of given files.
+	 * @param FileList			IN		The list of files.
+	 * @param RequiredChunks	OUT		The list of chunk GUIDs needed for those files.
+	 */
+	void GetChunksRequiredForFiles( const TArray< FString >& FileList, TArray< FGuid >& RequiredChunks, bool bAddUnique = true ) const;
+
+	/**
+	 * Get the number of times a chunks is referenced in this manifest
+	 * @param ChunkGuid		The chunk GUID
+	 * @return	The number of references to this chunk
+	 */
+	const uint32 GetNumberOfChunkReferences( const FGuid& ChunkGuid ) const;
+
+	/**
+	 * Returns the size of a particular data file by it's GUID.
+	 * @param DataGuid		The GUID for the data
+	 * @return		File size.
+	 */
+	const int64 GetDataSize( const FGuid& DataGuid ) const;
+
+	/**
+	 * Returns the total size of all data files in it's list.
+	 * @param DataGuids		The GUID array for the data
+	 * @return		File size.
+	 */
+	const int64 GetDataSize( const TArray< FGuid >& DataGuids ) const;
+
+	/**
+	 * Returns the size of a particular file in the build
+	 * VALID FOR ANY MANIFEST
+	 * @param Filename		The file.
+	 * @return		File size.
+	 */
+	const int64 GetFileSize( const FString& Filename ) const;
+
+	/**
+	 * Returns the total size of all files in the array
+	 * VALID FOR ANY MANIFEST
+	 * @param Filenames		The array of files.
+	 * @return		Total size of files in array.
+	 */
+	const int64 GetFileSize( const TArray< FString >& Filenames ) const;
+
+	/**
+	 * Returns the number of files in this build.
+	 * @return		The number of files.
+	 */
+	const uint32 GetNumFiles() const;
+
+	/**
+	 * Get the list of files described by this manifest
+	 * @param Filenames		OUT		Receives the array of files.
+	 */
+	void GetFileList( TArray< FString >& Filenames ) const;
+
+	/**
+	* Get the list of Guids for all files described by this manifest
+	* @param DataGuids		OUT		Receives the array of Guids.
+	*/
+	void GetDataList( TArray< FGuid >& DataGuids ) const;
+
+	/**
+	 * Returns the manifest for a particular file in the app, invalid if non-existing
+	 * @param Filename	The filename.
+	 * @return	The file manifest, or invalid ptr
+	 */
+	const TSharedPtr< FBuildPatchFileManifest > GetFileManifest( const FString& Filename );
+
+	/**
+	 * Gets whether this manifest is made up of file data instead of chunk data
+	 * @return	True if the build is made from file data. False if the build is constructed from chunk data.
+	 */
+	const bool IsFileDataManifest() const;
+
+	/**
+	 * Gets the chunk hash for a given chunk
+	 * @param ChunkGuid		IN		The guid of the chunk to get hash for
+	 * @param OutHash		OUT		Receives the hash value if found
+	 * @return	true if we had the hash for this chunk
+	 */
+	const bool GetChunkHash( const FGuid& ChunkGuid, uint64& OutHash ) const;
+
+	/**
+	 * Gets the file hash for given file data
+	 * @param FileGuid		IN		The guid of the file data to get hash for
+	 * @param OutHash		OUT		Receives the hash value if found
+	 * @return	true if we had the hash for this file
+	 */
+	const bool GetFileDataHash( const FGuid& FileGuid, FSHAHash& OutHash ) const;
+
+	/**
+	 * Populates an array of chunks that should be producible from this local build, given the list of chunks needed. Also checks source files exist and match size.
+	 * @param InstallDirectory	IN		The directory of the build where chunks would be sourced from.
+	 * @param ChunksRequired	IN		A list of chunks that are needed.
+	 * @param ChunksAvailable	OUT		A list to receive the chunks that could be constructed locally.
+	 */
+	void EnumerateProducibleChunks( const FString& InstallDirectory, const TArray< FGuid >& ChunksRequired, TArray< FGuid >& ChunksAvailable ) const;
+
+	/**
+	 * Verifies a local directory structure against a given manifest.
+	 * NOTE: This function is blocking and will not return until finished. Don't run on main thread.
+	 * @param VerifyDirectory		IN		The directory to analyze.
+	 * @param OutDatedFiles			OUT		The array of files that do not match or are locally missing.
+	 * @param ProgressDelegate		IN		Delegate to receive progress updates in the form of a float range 0.0f to 1.0f
+	 * @param ShouldPauseDelegate	IN		Delegate that returns a bool, which if true will pause the process
+	 * @param TimeSpentPaused		OUT		The amount of time we were paused for, in seconds.
+	 * @return		true if no file errors occurred AND the verification was successful
+	 */
+	bool VerifyAgainstDirectory( const FString& VerifyDirectory, TArray < FString >& OutDatedFiles, FBuildPatchFloatDelegate ProgressDelegate, FBuildPatchBoolRetDelegate ShouldPauseDelegate, double& TimeSpentPaused );
+
+	/**
+	 * Gets a list of files that were installed with the Old Manifest, but no longer required by the New Manifest.
+	 * @param OldManifest		IN		The Build Manifest that is currently installed
+	 * @param NewManifest		IN		The Build Manifest that is being patched to
+	 * @param RemovableFiles	OUT		A list to receive the files that may be removed.
+	 */
+	static void GetRemovableFiles( FBuildPatchAppManifestRef OldManifest, FBuildPatchAppManifestRef NewManifest, TArray< FString >& RemovableFiles );
+
+	/**
+	 * Gets a list of files that have different hash values in the New Manifest, to those in the Old Manifest.
+	 * @param OldManifest		IN		The Build Manifest that is currently installed. Shared Ptr - Can be invalid.
+	 * @param NewManifest		IN		The Build Manifest that is being patched to. Shared Ref - Implicitly valid.
+	 * @param InstallDirectory	IN		The Build installation directory, so that it can be checked for missing files.
+	 * @param OutDatedFiles		OUT		The array of files that do not match or are new.
+	 */
+	static void GetOutdatedFiles( FBuildPatchAppManifestPtr OldManifest, FBuildPatchAppManifestRef NewManifest, const FString& InstallDirectory, TArray< FString >& OutDatedFiles );
+
+	/**
+	 * Check a single file to see if it will be effected by patching
+	 * @param OldManifest		The Build Manifest that is currently installed. Shared Ref - Implicitly valid.
+	 * @param NewManifest		The Build Manifest that is being patched to. Shared Ref - Implicitly valid.
+	 * @param Filename			The Build installation directory, so that it can be checked for missing files.
+	 */
+	static bool IsFileOutdated( FBuildPatchAppManifestRef OldManifest, FBuildPatchAppManifestRef NewManifest, const FString& Filename );
+
+	/**
+	 * Populates a map of FFileChunkParts referring to the chunks in the given array, that should be accessible locally.
+	 * @param ChunksRequired		IN		A list of chunks that are needed.
+	 * @param ChunkPartsAvailable	OUT		A map to receive the FFileChunkParts keyed by chunk guid.
+	 */
+	void EnumerateChunkPartInventory( const TArray< FGuid >& ChunksRequired, TMap< FGuid, TArray< FFileChunkPart > >& ChunkPartsAvailable ) const;
+
+private:
+
+	/**
+	 * Destroys any memory we have allocated and clears out ready for generation of a new manifest
+	 */
+	void DestroyData();
+
+private:
+
+	/**
+	 * Holds the version support for this manifest.
+	 */
+	EBuildPatchAppManifestVersion::Type ManifestFileVersion;
+
+	/**
+	 * Stores whether the manifest was produced from a file data build instead of chunk data
+	 */
+	bool bIsFileData;
+
+	/**
+	 * Holds the ID of this app.
+	 */
+	uint32 AppID;
+
+	/**
+	 * Holds the name of the app this Build is for
+	 */
+	FString AppNameString;
+
+	/**
+	 * Holds the version string of this Build
+	 */
+	FString BuildVersionString;
+
+	/**
+	 * Holds the app local exe path
+	 */
+	FString LaunchExeString;
+
+	/**
+	 * Holds the name of prerequisites installer
+	 */
+	FString PrereqNameString;
+
+	/**
+	 * Holds the exe path of the prerequisites installer
+	 */
+	FString PrereqPathString;
+
+	/**
+	 * Holds commandline arguments for the prerequisites installer
+	 */
+	FString PrereqArgsString;
+	
+	/**
+	 * Holds the app launch command
+	 */
+	FString LaunchCommandString;
+
+	/**
+	 * References filename by file data GUID
+	 */
+	TMap< FGuid, FString > FileNameReference;
+
+	/**
+	 * Holds the list of files and their generated file manifest
+	 */
+	TMap< FString, TSharedPtr< FBuildPatchFileManifest > > FileManifestList;
+
+	/**
+	 * Holds the list of chunks with their rolling chunk hash, does not apply to file data.
+	 */
+	TMap< FGuid, uint64 > ChunkHashList;
+
+	/**
+	 * Holds the filesize for each chunk, does not apply to file data.
+	 */
+	TMap< FGuid, int64 > ChunkFilesizeList;
+
+	/**
+	 * Holds which group the chunk (containing chunk or file data) is placed into for file system optimisations
+	 */
+	TMap< FGuid, uint8 > DataGroupList;
+
+	/**
+	 * Holds the list of custom fields
+	 */
+	TMap< FString, FString > CustomFields;
+};
diff --git a/Engine/Source/Runtime/Online/BuildPatchServices/Private/BuildPatchProgress.cpp b/Engine/Source/Runtime/Online/BuildPatchServices/Private/BuildPatchProgress.cpp
new file mode 100644
--- /dev/null
+++ b/Engine/Source/Runtime/Online/BuildPatchServices/Private/BuildPatchProgress.cpp
@@ -0,0 +1,314 @@
+// Copyright 1998-2014 Epic Games, Inc. All Rights Reserved.
+
+/*=============================================================================
+	BuildPatchProgress.cpp: Implements classes involved with tracking the patch
+	progress information.
+=============================================================================*/
+
+#include "BuildPatchServicesPrivatePCH.h"
+
+#define LOCTEXT_NAMESPACE "BuildPatchInstallProgress"
+
+/* EBuildPatchProgress implementation
+*****************************************************************************/
+const FString& EBuildPatchProgress::ToString(const EBuildPatchProgress::Type& ProgressValue)
+{
+	// Static const fixed FString values so that they are not constantly constructed
+	static const FString Initializing(TEXT("Initialising"));
+	static const FString Resuming(TEXT("Resuming"));
+	static const FString Downloading(TEXT("Downloading"));
+	static const FString Recycling(TEXT("Recycling"));
+	static const FString Installing(TEXT("Installing"));
+	static const FString MovingToInstall(TEXT("MovingToInstall"));
+	static const FString BuildVerification(TEXT("BuildVerification"));
+	static const FString CleanUp(TEXT("CleanUp"));
+	static const FString PrerequisitesInstall(TEXT("PrerequisitesInstall"));
+	static const FString Completed(TEXT("Completed"));
+	static const FString Error(TEXT("Error"));
+	static const FString Paused(TEXT("Paused"));
+	static const FString Default(TEXT("InvalidOrMax"));
+
+	switch (ProgressValue)
+	{
+		case EBuildPatchProgress::Initializing:
+			return Initializing;
+		case EBuildPatchProgress::Resuming:
+			return Resuming;
+		case EBuildPatchProgress::Downloading:
+			return Downloading;
+		case EBuildPatchProgress::Installing:
+			return Installing;
+		case EBuildPatchProgress::MovingToInstall:
+			return MovingToInstall;
+		case EBuildPatchProgress::BuildVerification:
+			return BuildVerification;
+		case EBuildPatchProgress::CleanUp:
+			return CleanUp;
+		case EBuildPatchProgress::PrerequisitesInstall:
+			return PrerequisitesInstall;
+		case EBuildPatchProgress::Completed:
+			return Completed;
+		case EBuildPatchProgress::Error:
+			return Error;
+		case EBuildPatchProgress::Paused:
+			return Paused;
+		default:
+			return Default;
+	}
+}
+
+const FText& EBuildPatchProgress::ToText(const EBuildPatchProgress::Type& ProgressType)
+{
+	// Static const fixed FText values so that they are not constantly constructed
+	static const FText Initializing = LOCTEXT("EBuildPatchProgress_Initialising", "Initializing");
+	static const FText Resuming = LOCTEXT("EBuildPatchProgress_Resuming", "Resuming");
+	static const FText Downloading = LOCTEXT("EBuildPatchProgress_Downloading", "Downloading");
+	static const FText Recycling = LOCTEXT("EBuildPatchProgress_Recycling", "Recycling Content");
+	static const FText Installing = LOCTEXT("EBuildPatchProgress_Installing", "Installing");
+	static const FText BuildVerification = LOCTEXT("EBuildPatchProgress_BuildVerification", "Verifying");
+	static const FText CleanUp = LOCTEXT("EBuildPatchProgress_CleanUp", "Cleaning up");
+	static const FText PrerequisitesInstall = LOCTEXT("EBuildPatchProgress_PrerequisitesInstall", "Installing Prerequisites");
+	static const FText Completed = LOCTEXT("EBuildPatchProgress_Complete", "Complete");
+	static const FText Error = LOCTEXT("EBuildPatchProgress_Error", "Error");
+	static const FText Paused = LOCTEXT("EBuildPatchProgress_Paused", "Paused");
+	static const FText Empty = FText::GetEmpty();
+
+	switch (ProgressType)
+	{
+		case EBuildPatchProgress::Initializing:
+			return Initializing;
+		case EBuildPatchProgress::Resuming:
+			return Resuming;
+		case EBuildPatchProgress::Downloading:
+			return Downloading;
+		case EBuildPatchProgress::Installing:
+		case EBuildPatchProgress::MovingToInstall: //User should see same as "Installing"
+			return Installing;
+		case EBuildPatchProgress::BuildVerification:
+			return BuildVerification;
+		case EBuildPatchProgress::CleanUp:
+			return CleanUp;
+		case EBuildPatchProgress::PrerequisitesInstall:
+			return PrerequisitesInstall;
+		case EBuildPatchProgress::Completed:
+			return Completed;
+		case EBuildPatchProgress::Error:
+			return Error;
+		case EBuildPatchProgress::Paused:
+			return Paused;
+		default:
+			return Empty;
+	}
+}
+
+/* FBuildPatchProgress implementation
+*****************************************************************************/
+FBuildPatchProgress::FBuildPatchProgress()
+{
+	Reset();
+}
+
+void FBuildPatchProgress::Reset()
+{
+	TotalWeight = 0.0f;
+	CurrentState = EBuildPatchProgress::Initializing;
+	CurrentProgress = 0.0f;
+	ErrorText = FText::GetEmpty();
+
+	// Initialize array data
+	for( uint32 idx = 0; idx < EBuildPatchProgress::NUM_PROGRESS_STATES; ++idx )
+	{
+		StateProgressValues[ idx ] = 0.0f;
+		StateProgressWeights[ idx ] = 1.0f;
+	}
+
+	// Update for initial values
+	UpdateCachedValues();
+	UpdateProgressInfo();
+}
+
+void FBuildPatchProgress::SetStateProgress( const EBuildPatchProgress::Type& State, const float& Value )
+{
+	check( !FMath::IsNaN( Value ) );
+	FScopeLock ScopeLock( &ThreadLock );
+	if( StateProgressValues[ State ] != Value )
+	{
+		StateProgressValues[ State ] = Value;
+		UpdateProgressInfo();
+	}
+}
+
+void FBuildPatchProgress::SetStateWeight( const EBuildPatchProgress::Type& State, const float& Value )
+{
+	check( !FMath::IsNaN( Value ) );
+	FScopeLock ScopeLock( &ThreadLock );
+	if( bCountsTowardsProgress[ State ] )
+	{
+		TotalWeight += Value - StateProgressWeights[ State ];
+		StateProgressWeights[ State ] = Value;
+	}
+}
+
+const FText& FBuildPatchProgress::GetStateText()
+{
+	FScopeLock ScopeLock( &ThreadLock );
+	if( CurrentState == EBuildPatchProgress::Error )
+	{
+		return ErrorText;
+	}
+	return EBuildPatchProgress::ToText( CurrentState );
+}
+
+const float FBuildPatchProgress::GetProgress()
+{
+	FScopeLock ScopeLock( &ThreadLock );
+	if( bHasProgressValue[ CurrentState ] )
+	{
+		return CurrentProgress;
+	}
+	// We return -1.0f to indicate a marquee style progress should be displayed
+	return -1.0f;
+}
+
+const float FBuildPatchProgress::GetStateProgress( const EBuildPatchProgress::Type& State )
+{
+	FScopeLock ScopeLock( &ThreadLock );
+	return StateProgressValues[ State ];
+}
+
+const float FBuildPatchProgress::GetStateWeight(const EBuildPatchProgress::Type& State)
+{
+	FScopeLock ScopeLock( &ThreadLock );
+	return StateProgressWeights[ State ];
+}
+
+const bool FBuildPatchProgress::TogglePauseState()
+{
+	FScopeLock ScopeLock( &ThreadLock );
+
+	// If in pause state revert to no state and recalculate progress
+	if( CurrentState == EBuildPatchProgress::Paused )
+	{
+		CurrentState = EBuildPatchProgress::NUM_PROGRESS_STATES;
+		UpdateProgressInfo();
+		return false;
+	}
+	else
+	{
+		// Else we just set paused
+		CurrentState = EBuildPatchProgress::Paused;
+		return true;
+	}
+}
+
+const bool FBuildPatchProgress::GetPauseState()
+{
+	FScopeLock ScopeLock( &ThreadLock );
+	return CurrentState == EBuildPatchProgress::Paused;
+}
+
+const double FBuildPatchProgress::WaitWhilePaused()
+{
+	// Pause if necessary
+	const double PrePauseTime = FPlatformTime::Seconds();
+	double PostPauseTime = PrePauseTime;
+	while ( GetPauseState() )
+	{
+		FPlatformProcess::Sleep( 0.1f );
+		PostPauseTime = FPlatformTime::Seconds();
+	}
+	// Return pause time
+	return PostPauseTime - PrePauseTime;
+}
+
+void FBuildPatchProgress::UpdateProgressInfo()
+{
+	FScopeLock ScopeLock( &ThreadLock );
+
+	// Override state to error if an error was set
+	if( FBuildPatchInstallError::HasFatalError() )
+	{
+		CurrentState = EBuildPatchProgress::Error;
+		ErrorText = FBuildPatchInstallError::GetErrorText();
+		return;
+	}
+
+	// If we are paused, we don't change the state or progress value
+	if( CurrentState == EBuildPatchProgress::Paused )
+	{
+		return;
+	}
+
+	// Reset values
+	CurrentState = EBuildPatchProgress::NUM_PROGRESS_STATES;
+	CurrentProgress = 0.0f;
+
+	// Loop through each progress value to find total progress and current state
+	for( uint32 idx = 0; idx < EBuildPatchProgress::NUM_PROGRESS_STATES; ++idx )
+	{
+		// Is this the current state
+		if( ( CurrentState == EBuildPatchProgress::NUM_PROGRESS_STATES ) && StateProgressValues[ idx ] < 1.0f )
+		{
+			CurrentState = static_cast< EBuildPatchProgress::Type >( idx );
+		}
+
+		// Do we count the progress
+		if( bCountsTowardsProgress[ idx ] )
+		{
+			const float StateProgress = StateProgressValues[ idx ] * ( StateProgressWeights[ idx ] / TotalWeight );
+			CurrentProgress += StateProgress;
+			check( !FMath::IsNaN( CurrentProgress ) );
+		}
+	}
+
+	// Ensure Sanity
+	CurrentProgress = FMath::Clamp( CurrentProgress, 0.0f, 1.0f );
+}
+
+void FBuildPatchProgress::UpdateCachedValues()
+{
+	TotalWeight = 0.0f;
+	for( uint32 idx = 0; idx < EBuildPatchProgress::NUM_PROGRESS_STATES; ++idx )
+	{
+		if( bCountsTowardsProgress[ idx ] )
+		{
+			TotalWeight += StateProgressWeights[ idx ];
+		}
+	}
+	// If you assert here you're going to cause a divide by 0.0f
+	check( TotalWeight != 0.0f );
+}
+
+/* FBuildPatchProgress static constants
+*****************************************************************************/
+const bool FBuildPatchProgress::bHasProgressValue[EBuildPatchProgress::NUM_PROGRESS_STATES] =
+{
+	false, // Initializing
+	true,  // Resuming
+	true,  // Downloading
+	true,  // Installing
+	true,  // MovingToInstall
+	true,  // BuildVerification
+	false, // CleanUp
+	false, // PrerequisitesInstall
+	false, // Completed
+	false, // Error
+	true   // Paused
+};
+
+const bool FBuildPatchProgress::bCountsTowardsProgress[EBuildPatchProgress::NUM_PROGRESS_STATES] =
+{
+	false, // Initializing
+	false, // Resuming
+	true,  // Downloading
+	true,  // Installing
+	true,  // MovingToInstall
+	true,  // BuildVerification
+	false, // CleanUp
+	false, // PrerequisitesInstall
+	false, // Completed
+	false, // Error
+	false  // Paused
+};
+
+#undef LOCTEXT_NAMESPACE
diff --git a/Engine/Source/Runtime/Online/BuildPatchServices/Private/BuildPatchProgress.h b/Engine/Source/Runtime/Online/BuildPatchServices/Private/BuildPatchProgress.h
new file mode 100644
--- /dev/null
+++ b/Engine/Source/Runtime/Online/BuildPatchServices/Private/BuildPatchProgress.h
@@ -0,0 +1,188 @@
+// Copyright 1998-2014 Epic Games, Inc. All Rights Reserved.
+
+/*=============================================================================
+	BuildPatchProgress.h: Declares classes involved with tracking the patch
+	progress information.
+=============================================================================*/
+
+#pragma once
+
+/**
+ * Namespace to declares the progress type enum
+ */
+namespace EBuildPatchProgress
+{
+	enum Type
+	{
+		// The patch process is initializing
+		Initializing = 0,
+
+		// The patch process is enumerating existing staged data
+		Resuming,
+
+		// The patch process is downloading patch data
+		Downloading,
+
+		// The patch process is installing files
+		Installing,
+
+		// The patch process is moving staged files to the install
+		MovingToInstall,
+
+		// The patch process is verifying the build
+		BuildVerification,
+
+		// The patch process is cleaning temp files
+		CleanUp,
+
+		// The patch process is installing prerequisites
+		PrerequisitesInstall,
+
+		// A state to catch the UI when progress is 100% but UI still being displayed
+		Completed,
+
+		// The patch process is in an error state
+		Error,
+
+		// The process has been set paused
+		Paused,
+
+		// Holds the number of states, for array sizes
+		NUM_PROGRESS_STATES,
+	};
+
+	/**
+	 * Returns the string representation of the EBuildPatchProgress value. Used for analytics and logging only.
+	 * @param ProgressValue - The value.
+	 * @return The enum's string value.
+	 */
+	const FString& ToString( const Type& ProgressValue );
+
+	/**
+	 * Returns the FText representation of the specified EBuildPatchProgress value. Used for displaying to the user.
+	 * @param ProgressValue - The error type value.
+	 * @return The display text associated with the progress step.
+	 */
+	const FText& ToText( const Type& ProgressValue );
+}
+
+/**
+ * A struct to hold patch progress tracking
+ */
+struct FBuildPatchProgress
+{
+private:
+
+	// Defines whether each state displays progress percent or is designed for a "please wait" or marquee style progress bar
+	// This is predefined and constant.
+	static const bool bHasProgressValue[ EBuildPatchProgress::NUM_PROGRESS_STATES ];
+
+	// Defines whether each state should count towards the overall progress
+	// This is predefined and constant.
+	static const bool bCountsTowardsProgress[ EBuildPatchProgress::NUM_PROGRESS_STATES ];
+
+	// Holds the current percentage complete for each state, this will decide the "current" state, being the first that is not complete.
+	// Range 0 to 1.
+	float StateProgressValues[ EBuildPatchProgress::NUM_PROGRESS_STATES ];
+
+	// Holds the weight that each stage has on overall progress. 
+	// Range 0 to 1.
+	float StateProgressWeights[ EBuildPatchProgress::NUM_PROGRESS_STATES ];
+
+	// Cached total weight value for progress calculation
+	float TotalWeight;
+
+	// The current state value for UI polling
+	EBuildPatchProgress::Type CurrentState;
+
+	// The current progress value for UI polling
+	float CurrentProgress;
+
+	// The error text that was set when in error state
+	FText ErrorText;
+
+	// Critical section to protect variable access
+	FCriticalSection ThreadLock;
+
+public:
+
+	/**
+	 * Default constructor
+	 */
+	FBuildPatchProgress();
+
+	/**
+	 * Resets internal variables to start over
+	 */
+	void Reset();
+
+	/**
+	 * Sets the progress value for a particular state
+	 * @param State		The state to set progress for
+	 * @param Value		The progress value
+	 */
+	void SetStateProgress( const EBuildPatchProgress::Type& State, const float& Value );
+
+	/**
+	 * Sets the progress weight for a particular state
+	 * @param State		The state to set weight for
+	 * @param Value		The weight value
+	 */
+	void SetStateWeight( const EBuildPatchProgress::Type& State, const float& Value );
+
+	/**
+	 * Gets the text for the current progress state
+	 * @return The display text for the current progress state
+	 */
+	const FText& GetStateText();
+
+	/**
+	 * Gets the current overall progress
+	 * @return The current progress value. Range 0 to 1. -1 indicates undetermined, i.e. show a marquee style bar.
+	 */
+	const float GetProgress();
+
+	/**
+	 * Gets the progress value for a particular state
+	 * @param State		The state to get progress for
+	 * @return The state progress value. Range 0 to 1.
+	 */
+	const float GetStateProgress( const EBuildPatchProgress::Type& State );
+
+	/**
+	 * Gets the weight value for a particular state
+	 * @param State		The state to get weight for
+	 * @return The state weight value.
+	 */
+	const float GetStateWeight( const EBuildPatchProgress::Type& State );
+
+	/**
+	 * Toggles the pause state
+	 * @return Whether the current state is now paused
+	 */
+	const bool TogglePauseState();
+
+	/**
+	 * Blocks calling thread while the progress is paused
+	 * @return How long we paused for, in seconds
+	 */
+	const double WaitWhilePaused();
+
+	/**
+	 * Gets the pause state
+	 * @return Whether the current state is paused
+	 */
+	const bool GetPauseState();
+
+private:
+
+	/**
+	 * Updates the current state and progress values
+	 */
+	void UpdateProgressInfo();
+
+	/**
+	 * Updates the cached values
+	 */
+	void UpdateCachedValues();
+};
diff --git a/Engine/Source/Runtime/Online/BuildPatchServices/Private/BuildPatchServicesModule.cpp b/Engine/Source/Runtime/Online/BuildPatchServices/Private/BuildPatchServicesModule.cpp
new file mode 100644
--- /dev/null
+++ b/Engine/Source/Runtime/Online/BuildPatchServices/Private/BuildPatchServicesModule.cpp
@@ -0,0 +1,252 @@
+// Copyright 1998-2014 Epic Games, Inc. All Rights Reserved.
+
+/*=============================================================================
+	BuildPatchServicesModule.cpp: Implements the FBuildPatchServicesModule class.
+=============================================================================*/
+
+#include "BuildPatchServicesPrivatePCH.h"
+
+IMPLEMENT_MODULE( FBuildPatchServicesModule, BuildPatchServices );
+
+/* FBuildPatchInstallationInfo implementation
+*****************************************************************************/
+
+void FBuildPatchInstallationInfo::RegisterAppInstallation(IBuildManifestRef AppManifest, const FString& AppInstallDirectory)
+{
+	FBuildPatchAppManifestRef InternalRef = StaticCastSharedRef<FBuildPatchAppManifest>(AppManifest);
+	AvailableInstallations.Add(AppInstallDirectory, InternalRef);
+}
+
+void FBuildPatchInstallationInfo::EnumerateProducibleChunks(const TArray< FGuid >& ChunksRequired, TArray< FGuid >& ChunksAvailable)
+{
+	for (auto AvailableInstallationsIt = AvailableInstallations.CreateConstIterator(); AvailableInstallationsIt; ++AvailableInstallationsIt)
+	{
+		const FBuildPatchAppManifest& AppManifest = *AvailableInstallationsIt.Value().Get();
+		const FString& AppInstallDir = AvailableInstallationsIt.Key();
+		TArray< FGuid > ChunksFromThisManifest;
+		AppManifest.EnumerateProducibleChunks(AppInstallDir, ChunksRequired, ChunksFromThisManifest);
+		for (auto ChunksFromThisManifestIt = ChunksFromThisManifest.CreateConstIterator(); ChunksFromThisManifestIt; ++ChunksFromThisManifestIt)
+		{
+			const FGuid& ChunkGuid = *ChunksFromThisManifestIt;
+			if (ChunksAvailable.Contains(ChunkGuid) == false)
+			{
+				ChunksAvailable.Add(ChunkGuid);
+				RecyclableChunks.Add(ChunkGuid, AvailableInstallationsIt.Value());
+			}
+		}
+	}
+}
+
+FBuildPatchAppManifestPtr FBuildPatchInstallationInfo::GetManifestContainingChunk(const FGuid& ChunkGuid)
+{
+	return RecyclableChunks.FindRef(ChunkGuid);
+}
+
+const FString FBuildPatchInstallationInfo::GetManifestInstallDir(FBuildPatchAppManifestPtr AppManifest)
+{
+	for (auto AvailableInstallationsIt = AvailableInstallations.CreateConstIterator(); AvailableInstallationsIt; ++AvailableInstallationsIt)
+	{
+		if (AppManifest == AvailableInstallationsIt.Value())
+		{
+			return AvailableInstallationsIt.Key();
+		}
+	}
+	return TEXT("");
+}
+
+/* FBuildPatchServicesModule implementation
+ *****************************************************************************/
+
+void FBuildPatchServicesModule::StartupModule()
+{
+	// We need to initialize the lookup for our hashing functions
+	FRollingHashConst::Init();
+
+	// Add our ticker
+	FTicker::GetCoreTicker().AddTicker( FTickerDelegate::CreateRaw( this, &FBuildPatchServicesModule::Tick ) );
+
+	// Test the rolling hash algorithm
+	check( CheckRollingHashAlgorithm() );
+}
+
+void FBuildPatchServicesModule::ShutdownModule()
+{
+	GWarn->Logf( TEXT( "BuildPatchServicesModule: Shutting Down" ) );
+	// We need to stop and wait for any threads to exit.
+	FBuildPatchInstallError::SetFatalError( EBuildPatchInstallError::ApplicationClosing );
+
+	// Release our ptr to analytics
+	FBuildPatchAnalytics::SetAnalyticsProvider( NULL );
+
+	// Reset all installer ptrs (this will cause us to wait for each thread to complete)
+	for(auto InstallerIt = BuildPatchInstallers.CreateIterator(); InstallerIt; ++InstallerIt)
+	{
+		// Make sure it is not paused, this function only un-pauses when in error state
+		(*InstallerIt)->TogglePauseInstall();
+		// And reset
+		(*InstallerIt).Reset();
+	}
+
+	// Remove our ticker
+	FTicker::GetCoreTicker().RemoveTicker( FTickerDelegate::CreateRaw( this, &FBuildPatchServicesModule::Tick ) );
+
+	FBuildPatchHTTP::OnShutdown();
+}
+
+IBuildManifestPtr FBuildPatchServicesModule::LoadManifestFromFile( const FString& Filename )
+{
+	FBuildPatchAppManifestRef Manifest = MakeShareable( new FBuildPatchAppManifest() );
+	if( Manifest->LoadFromFile( Filename ) )
+	{
+		return Manifest;
+	}
+	else
+	{
+		return NULL;
+	}
+}
+
+IBuildManifestPtr FBuildPatchServicesModule::MakeManifestFromJSON( const FString& ManifestJSON )
+{
+	FBuildPatchAppManifestRef Manifest = MakeShareable( new FBuildPatchAppManifest() );
+	if( Manifest->DeserializeFromJSON( ManifestJSON ) )
+	{
+		return Manifest;
+	}
+	return NULL;
+}
+
+bool FBuildPatchServicesModule::SaveManifestToFile( const FString& Filename, IBuildManifestRef Manifest )
+{
+	return StaticCastSharedRef< FBuildPatchAppManifest >( Manifest )->SaveToFile( Filename );
+}
+
+IBuildInstallerPtr FBuildPatchServicesModule::StartBuildInstall( IBuildManifestPtr CurrentManifest, IBuildManifestPtr InstallManifest, const FString& InstallDirectory, FBuildPatchBoolManifestDelegate OnCompleteDelegate )
+{
+	// Using a local bool for this check will improve the assert message that gets displayed
+	const bool bIsCalledFromMainThread = IsInGameThread();
+	check( bIsCalledFromMainThread );
+	// Cast manifest parameters
+	FBuildPatchAppManifestPtr CurrentManifestInternal = StaticCastSharedPtr< FBuildPatchAppManifest >( CurrentManifest );
+	FBuildPatchAppManifestPtr InstallManifestInternal = StaticCastSharedPtr< FBuildPatchAppManifest >( InstallManifest );
+	if( !InstallManifestInternal.IsValid() )
+	{
+		// We must have an install manifest to continue
+		return NULL;
+	}
+	// Make directory
+	IFileManager::Get().MakeDirectory( *InstallDirectory, true );
+	if( !IFileManager::Get().DirectoryExists( *InstallDirectory ) )
+	{
+		return NULL;
+	}
+	// Run the install thread
+	BuildPatchInstallers.Add( MakeShareable( new FBuildPatchInstaller( OnCompleteDelegate, CurrentManifestInternal, InstallManifestInternal.ToSharedRef(), InstallDirectory, GetStagingDirectory(), InstallationInfo ) ) );
+	return BuildPatchInstallers.Top();
+}
+
+bool FBuildPatchServicesModule::Tick( float Delta )
+{
+	// Using a local bool for this check will improve the assert message that gets displayed
+	// This one is unlikely to assert unless the FTicker's core tick is not ticked on the main thread for some reason
+	const bool bIsCalledFromMainThread = IsInGameThread();
+	check( bIsCalledFromMainThread );
+
+	// Call complete delegate on each finished installer
+	for(auto InstallerIt = BuildPatchInstallers.CreateIterator(); InstallerIt; ++InstallerIt)
+	{
+		if( (*InstallerIt).IsValid() && (*InstallerIt)->IsComplete() )
+		{
+			(*InstallerIt)->ExecuteCompleteDelegate();
+			(*InstallerIt).Reset();
+		}
+	}
+
+	// Remove completed (invalids) from the list
+	for(int32 BuildPatchInstallersIdx = 0; BuildPatchInstallersIdx < BuildPatchInstallers.Num(); ++BuildPatchInstallersIdx )
+	{
+		const FBuildPatchInstallerPtr* Installer = &BuildPatchInstallers[ BuildPatchInstallersIdx ];
+		if( !Installer->IsValid() )
+		{
+			BuildPatchInstallers.RemoveAt( BuildPatchInstallersIdx-- );
+		}
+	}
+
+	// More ticks
+	return true;
+}
+
+#if WITH_BUILDPATCHGENERATION
+bool FBuildPatchServicesModule::GenerateChunksManifestFromDirectory( const FBuildPatchSettings& Settings )
+{
+	return FBuildDataGenerator::GenerateChunksManifestFromDirectory( Settings );
+}
+
+bool FBuildPatchServicesModule::GenerateFilesManifestFromDirectory( const FBuildPatchSettings& Settings )
+{
+	return FBuildDataGenerator::GenerateFilesManifestFromDirectory( Settings );
+}
+
+bool FBuildPatchServicesModule::CompactifyCloudDirectory(const bool bPreview)
+{
+	return FBuildDataCompactifier::CompactifyCloudDirectory(bPreview);
+}
+
+#endif //WITH_BUILDPATCHGENERATION
+
+void FBuildPatchServicesModule::SetStagingDirectory( const FString& StagingDir )
+{
+	StagingDirectory = StagingDir;
+}
+
+void FBuildPatchServicesModule::SetCloudDirectory( const FString& CloudDir )
+{
+	CloudDirectory = CloudDir;
+}
+
+void FBuildPatchServicesModule::SetBackupDirectory( const FString& BackupDir )
+{
+	BackupDirectory = BackupDir;
+}
+
+void FBuildPatchServicesModule::SetAnalyticsProvider( TSharedPtr< IAnalyticsProvider > AnalyticsProvider )
+{
+	FBuildPatchAnalytics::SetAnalyticsProvider( AnalyticsProvider );
+}
+
+void FBuildPatchServicesModule::RegisterAppInstallation(IBuildManifestRef AppManifest, const FString AppInstallDirectory)
+{
+	InstallationInfo.RegisterAppInstallation(AppManifest, AppInstallDirectory);
+}
+
+const FString& FBuildPatchServicesModule::GetStagingDirectory()
+{
+	// Default staging directory
+	if( StagingDirectory.IsEmpty() )
+	{
+		StagingDirectory = FPaths::GameDir() + TEXT( "BuildStaging/" );
+	}
+	return StagingDirectory;
+}
+
+const FString& FBuildPatchServicesModule::GetCloudDirectory()
+{
+	// Default cloud directory
+	if( CloudDirectory.IsEmpty() )
+	{
+		CloudDirectory = FPaths::CloudDir();
+	}
+	return CloudDirectory;
+}
+
+const FString& FBuildPatchServicesModule::GetBackupDirectory()
+{
+	// Default backup directory stays empty which simply doesn't backup
+	return BackupDirectory;
+}
+
+/* Static variables
+ *****************************************************************************/
+FString FBuildPatchServicesModule::CloudDirectory;
+FString FBuildPatchServicesModule::StagingDirectory;
+FString FBuildPatchServicesModule::BackupDirectory;
diff --git a/Engine/Source/Runtime/Online/BuildPatchServices/Private/BuildPatchServicesModule.h b/Engine/Source/Runtime/Online/BuildPatchServices/Private/BuildPatchServicesModule.h
new file mode 100644
--- /dev/null
+++ b/Engine/Source/Runtime/Online/BuildPatchServices/Private/BuildPatchServicesModule.h
@@ -0,0 +1,134 @@
+// Copyright 1998-2014 Epic Games, Inc. All Rights Reserved.
+
+/*=============================================================================
+	BuildPatchServicesModule.h: Declares the FBuildPatchServicesModule class.
+=============================================================================*/
+
+#pragma once
+
+/**
+ * Constant values and typedefs
+ */
+enum
+{
+	// Sizes
+	FileBufferSize		= 1024*1024*4,		// When reading from files, how much to buffer
+	StreamBufferSize	= FileBufferSize*4,	// When reading from build data stream, how much to buffer.
+};
+
+/**
+ * Implements a class that stores information about installed builds that are registered with the module
+ */
+class FBuildPatchInstallationInfo
+{
+private:
+	// Stores a list of installations available on this machine
+	TMap<FString, FBuildPatchAppManifestPtr> AvailableInstallations;
+
+	// Stores the information for which chunks can be gathered from which installations
+	TMap<FGuid, FBuildPatchAppManifestPtr> RecyclableChunks;
+
+public:
+	/**
+	 * Registers an installation on this machine. This information is used to gather a list of install locations that can be used as chunk sources.
+	 * @param AppManifest			Ref to the manifest for this installation
+	 * @param AppInstallDirectory	The install location
+	 */
+	void RegisterAppInstallation(IBuildManifestRef AppManifest, const FString& AppInstallDirectory);
+
+	/**
+	 * Populates an array of chunks that should be producible from the list of installations on the machine.
+	 * @param ChunksRequired	IN		A list of chunks that are needed.
+	 * @param ChunksAvailable	OUT		A list to receive the chunks that could be constructed locally.
+	 */
+	void EnumerateProducibleChunks(const TArray< FGuid >& ChunksRequired, TArray< FGuid >& ChunksAvailable);
+
+	/**
+	 * Get the manifest for the installation that a given chunk can be recycled from.
+	 * @param ChunkGuid		The GUID for the chunk required
+	 * @return Shared pointer to the installed manifest that contains this chunk
+	 */
+	FBuildPatchAppManifestPtr GetManifestContainingChunk(const FGuid& ChunkGuid);
+
+	/**
+	 * Get the installation directory for a given manifest
+	 * @param AppManifest	The manifest that has been registered with this system
+	 * @return The installation directory
+	 */
+	const FString GetManifestInstallDir(FBuildPatchAppManifestPtr AppManifest);
+};
+
+/**
+ * Implements the BuildPatchServicesModule.
+ */
+class FBuildPatchServicesModule
+	: public IBuildPatchServicesModule
+{
+public:
+
+	// Begin IModuleInterface interface
+	virtual void StartupModule() override;
+	virtual void ShutdownModule() override;
+	// End IModuleInterface interface
+
+	// Begin IBuildPatchServicesModule interface
+	virtual IBuildManifestPtr LoadManifestFromFile( const FString& Filename ) override;
+	virtual IBuildManifestPtr MakeManifestFromJSON( const FString& ManifestJSON ) override;
+	virtual bool SaveManifestToFile( const FString& Filename, IBuildManifestRef Manifest ) override;
+	virtual IBuildInstallerPtr StartBuildInstall( IBuildManifestPtr CurrentManifest, IBuildManifestPtr InstallManifest, const FString& InstallDirectory, FBuildPatchBoolManifestDelegate OnCompleteDelegate ) override;
+	virtual void SetStagingDirectory( const FString& StagingDir ) override;
+	virtual void SetCloudDirectory( const FString& CloudDir ) override;
+	virtual void SetBackupDirectory( const FString& BackupDir ) override;
+	virtual void SetAnalyticsProvider( TSharedPtr< IAnalyticsProvider > AnalyticsProvider ) override;
+	virtual void RegisterAppInstallation(IBuildManifestRef AppManifest, const FString AppInstallDirectory) override;
+
+#if WITH_BUILDPATCHGENERATION
+	virtual bool GenerateChunksManifestFromDirectory( const FBuildPatchSettings& Settings ) override;
+	virtual bool GenerateFilesManifestFromDirectory( const FBuildPatchSettings& Settings ) override;
+	virtual bool CompactifyCloudDirectory(const bool bPreview) override;
+#endif
+
+	// End IBuildPatchServicesModule interface
+
+	/**
+	 * Gets the directory used for staging intermediate files.
+	 * @return	The staging directory
+	 */
+	static const FString& GetStagingDirectory();
+
+	/**
+	 * Gets the cloud directory where chunks and manifests will be pulled from.
+	 * @return	The cloud directory
+	 */
+	static const FString& GetCloudDirectory();
+
+	/**
+	 * Gets the backup directory for saving files clobbered by repair/patch.
+	 * @return	The backup directory
+	 */
+	static const FString& GetBackupDirectory();
+
+private:
+	/**
+	 * Tick function to monitor installers for completion, so that we can call delegates on the main thread
+	 * @param		Delta	Time since last tick
+	 * @return	Whether to continue ticking
+	 */
+	bool Tick( float Delta );
+
+private:
+	// Holds the cloud directory where chunks should belong
+	static FString CloudDirectory;
+
+	// Holds the staging directory where we can perform any temporary work
+	static FString StagingDirectory;
+
+	// Holds the backup directory where we can move files that will be clobbered by repair or patch
+	static FString BackupDirectory;
+
+	// Save the ptr to the build installer thread that we create
+	TArray< FBuildPatchInstallerPtr > BuildPatchInstallers;
+
+	// The system that holds available installations used for recycling install data
+	FBuildPatchInstallationInfo InstallationInfo;
+};
diff --git a/Engine/Source/Runtime/Online/BuildPatchServices/Private/BuildPatchServicesPrivatePCH.h b/Engine/Source/Runtime/Online/BuildPatchServices/Private/BuildPatchServicesPrivatePCH.h
new file mode 100644
--- /dev/null
+++ b/Engine/Source/Runtime/Online/BuildPatchServices/Private/BuildPatchServicesPrivatePCH.h
@@ -0,0 +1,52 @@
+// Copyright 1998-2014 Epic Games, Inc. All Rights Reserved.
+
+/*=============================================================================
+	BuildPatchServicesPrivatePCH.h: Pre-compiled header file for the 
+	BuildPatchServices module.
+=============================================================================*/
+
+#pragma once
+
+// Still saving out old chunk data filenames
+// WARNING!! Disabling this will cause older clients to be incapable of patching via chunked patches. (e.g. game-patch for Community Portal project)
+//           Only remove if the latest versions pushed of clients support new filename chunk data.
+#define SAVE_OLD_CHUNKDATA_FILENAMES 0
+
+// Still saving out old data file format
+// WARNING!! Disabling this will cause older clients to be incapable of patching via -nochunks patches. (e.g. self-patch for Community Portal project)
+//           Only remove if absolutely sure there could be no clients in the wild that cannot fetch new filename file data.
+#define SAVE_OLD_FILEDATA_FILENAMES 0
+
+#include "BuildPatchServices.h"
+
+/* Dependencies
+*****************************************************************************/
+#include "Core.h"
+#include "SecureHash.h"
+#include "Json.h"
+#include "Http.h"
+#include "DefaultValueHelper.h"
+#include "IAnalyticsProvider.h"
+
+/* Public includes
+*****************************************************************************/
+#include "IBuildManifest.h"
+#include "IBuildPatchServicesModule.h"
+
+/* Private includes
+*****************************************************************************/
+#include "BuildPatchError.h"
+#include "BuildPatchProgress.h"
+#include "BuildPatchHTTP.h"
+#include "BuildPatchChunk.h"
+#include "BuildPatchManifest.h"
+#include "BuildPatchChunkCache.h"
+#include "BuildPatchCompactifier.h"
+#include "BuildPatchFileConstructor.h"
+#include "BuildPatchDownloader.h"
+#include "BuildPatchInstaller.h"
+#include "BuildPatchServicesModule.h"
+#include "BuildPatchHash.h"
+#include "BuildPatchGeneration.h"
+#include "BuildPatchUtil.h"
+#include "BuildPatchAnalytics.h"
diff --git a/Engine/Source/Runtime/Online/BuildPatchServices/Private/BuildPatchUtil.cpp b/Engine/Source/Runtime/Online/BuildPatchServices/Private/BuildPatchUtil.cpp
new file mode 100644
--- /dev/null
+++ b/Engine/Source/Runtime/Online/BuildPatchServices/Private/BuildPatchUtil.cpp
@@ -0,0 +1,302 @@
+// Copyright 1998-2014 Epic Games, Inc. All Rights Reserved.
+
+/*=============================================================================
+	BuildPatchUtil.cpp: Implements miscellaneous utility functions.
+=============================================================================*/
+
+#include "BuildPatchServicesPrivatePCH.h"
+
+/* FBuildPatchUtils implementation
+*****************************************************************************/
+const FString FBuildPatchUtils::GetChunkNewFilename( const EBuildPatchAppManifestVersion::Type ManifestVersion, const FString& RootDirectory, const FGuid& ChunkGUID, const uint64& ChunkHash )
+{
+	check( ChunkGUID.IsValid() );
+	return FPaths::Combine( *RootDirectory, *FString::Printf( TEXT("%s/%02d/%016llX_%s.chunk"), *EBuildPatchAppManifestVersion::GetChunkSubdir( ManifestVersion ),FCrc::MemCrc32( &ChunkGUID, sizeof( FGuid ) ) % 100, ChunkHash, *ChunkGUID.ToString() ) );
+}
+
+const FString FBuildPatchUtils::GetFileNewFilename( const EBuildPatchAppManifestVersion::Type ManifestVersion, const FString& RootDirectory, const FGuid& FileGUID, const FSHAHash& FileHash )
+{
+	check( FileGUID.IsValid() );
+	return FPaths::Combine( *RootDirectory, *FString::Printf( TEXT("%s/%02d/%s_%s.file"), *EBuildPatchAppManifestVersion::GetFileSubdir( ManifestVersion ), FCrc::MemCrc32( &FileGUID, sizeof( FGuid ) ) % 100, *FileHash.ToString(), *FileGUID.ToString() ) );
+}
+
+void FBuildPatchUtils::GetChunkDetailFromNewFilename( const FString& ChunkNewFilename, FGuid& ChunkGUID, uint64& ChunkHash )
+{
+	const FString ChunkFilename = FPaths::GetBaseFilename( ChunkNewFilename );
+	FString GuidString;
+	FString HashString;
+	ChunkFilename.Split( TEXT( "_" ), &HashString, &GuidString );
+	// Check that the strings are exactly as we expect otherwise this is not being used properly
+	check( GuidString.Len() == 32 );
+	check( HashString.Len() == 16 );
+	ChunkHash = FCString::Strtoui64( *HashString, NULL, 16 );
+	FGuid::Parse( GuidString, ChunkGUID );
+}
+
+void FBuildPatchUtils::GetFileDetailFromNewFilename( const FString& FileNewFilename, FGuid& FileGUID, FSHAHash& FileHash )
+{
+	const FString FileFilename = FPaths::GetBaseFilename( FileNewFilename );
+	FString GuidString;
+	FString HashString;
+	FileFilename.Split( TEXT( "_" ), &HashString, &GuidString );
+	// Check that the strings are exactly as we expect otherwise this is not being used properly
+	check( GuidString.Len() == 32 );
+	check( HashString.Len() == 40 );
+	HexToBytes( HashString, FileHash.Hash );
+	FGuid::Parse( GuidString, FileGUID );
+}
+
+const FString FBuildPatchUtils::GetChunkOldFilename( const FString& RootDirectory, const FGuid& ChunkGUID )
+{
+	check( ChunkGUID.IsValid() );
+	return FPaths::Combine( *RootDirectory, *FString::Printf( TEXT("Chunks/%02d/%s.chunk"), FCrc::MemCrc_DEPRECATED( &ChunkGUID, sizeof( FGuid ) ) % 100, *ChunkGUID.ToString() ) );
+}
+
+const FString FBuildPatchUtils::GetFileOldFilename( const FString& RootDirectory, const FGuid& FileGUID )
+{
+	check( FileGUID.IsValid() );
+	return FPaths::Combine( *RootDirectory, *FString::Printf( TEXT("Files/%02d/%s.file"), FCrc::MemCrc_DEPRECATED( &FileGUID, sizeof( FGuid ) ) % 100, *FileGUID.ToString() ) );
+}
+
+const FString FBuildPatchUtils::GetDataTypeOldFilename( const FBuildPatchData::Type DataType, const FString& RootDirectory, const FGuid& Guid )
+{
+	check( Guid.IsValid() );
+
+	switch ( DataType )
+	{
+	case FBuildPatchData::ChunkData:
+		return GetChunkOldFilename( RootDirectory, Guid );
+	case FBuildPatchData::FileData:
+		return GetFileOldFilename( RootDirectory, Guid );
+	}
+
+	// Error, didn't case type
+	check( false );
+	return TEXT( "" );
+}
+
+bool FBuildPatchUtils::GetGUIDFromFilename( const FString& DataFilename, FGuid& DataGUID )
+{
+	const FString DataBaseFilename = FPaths::GetBaseFilename( DataFilename );
+	FString GuidString;
+	if( DataBaseFilename.Contains( TEXT( "_" ) ) )
+	{
+		DataBaseFilename.Split( TEXT( "_" ), NULL, &GuidString );
+	}
+	else
+	{
+		GuidString = DataBaseFilename;
+	}
+	check( GuidString.Len() == 32 );
+	return FGuid::Parse( GuidString, DataGUID );
+}
+
+uint8 FBuildPatchUtils::VerifyFile(const FString& FileToVerify, const FSHAHash& Hash1, const FSHAHash& Hash2)
+{
+	FBuildPatchFloatDelegate NoProgressDelegate;
+	FBuildPatchBoolRetDelegate NoPauseDelegate;
+	double NoDouble;
+	return VerifyFile(FileToVerify, Hash1, Hash2, NoProgressDelegate, NoPauseDelegate, NoDouble);
+}
+
+uint8 FBuildPatchUtils::VerifyFile(const FString& FileToVerify, const FSHAHash& Hash1, const FSHAHash& Hash2, FBuildPatchFloatDelegate ProgressDelegate, FBuildPatchBoolRetDelegate ShouldPauseDelegate, double& TimeSpentPaused)
+{
+	uint8 ReturnValue = 0;
+	FArchive* FileReader = IFileManager::Get().CreateFileReader(*FileToVerify);
+	ProgressDelegate.ExecuteIfBound(0.0f);
+	if (FileReader != NULL)
+	{
+		FSHA1 HashState;
+		FSHAHash HashValue;
+		const int64 FileSize = FileReader->TotalSize();
+		uint8* FileReadBuffer = new uint8[FileBufferSize];
+		while (!FileReader->AtEnd() && !FBuildPatchInstallError::HasFatalError())
+		{
+			// Pause if necessary
+			const double PrePauseTime = FPlatformTime::Seconds();
+			double PostPauseTime = PrePauseTime;
+			bool bShouldPause = ShouldPauseDelegate.IsBound() && ShouldPauseDelegate.Execute();
+			while (bShouldPause && !FBuildPatchInstallError::HasFatalError())
+			{
+				FPlatformProcess::Sleep(0.1f);
+				bShouldPause = ShouldPauseDelegate.Execute();
+				PostPauseTime = FPlatformTime::Seconds();
+			}
+			// Count up pause time
+			TimeSpentPaused += PostPauseTime - PrePauseTime;
+			// Read file and update hash state
+			const int64 SizeLeft = FileSize - FileReader->Tell();
+			const uint32 ReadLen = FMath::Min< int64 >(FileBufferSize, SizeLeft);
+			FileReader->Serialize(FileReadBuffer, ReadLen);
+			HashState.Update(FileReadBuffer, ReadLen);
+			const double FileSizeTemp = FileSize;
+			const float Progress = 1.0f - ((SizeLeft - ReadLen) / FileSizeTemp);
+			ProgressDelegate.ExecuteIfBound(Progress);
+		}
+		delete[] FileReadBuffer;
+		HashState.Final();
+		HashState.GetHash(HashValue.Hash);
+		ReturnValue = (HashValue == Hash1) ? 1 : (HashValue == Hash2) ? 2 : 0;
+		if (ReturnValue == 0)
+		{
+			GLog->Logf(TEXT("BuildDataGenerator: Verify failed on %s"), *FPaths::GetCleanFilename(FileToVerify));
+		}
+		FileReader->Close();
+		delete FileReader;
+	}
+	else
+	{
+		GLog->Logf(TEXT("BuildDataGenerator: ERROR VerifyFile cannot open %s"), *FileToVerify);
+	}
+	ProgressDelegate.ExecuteIfBound(1.0f);
+	return ReturnValue;
+}
+
+bool FBuildPatchUtils::VerifyChunkFile( const TArray< uint8 >& ChunkFileArray, bool bQuickCheck )
+{
+	FMemoryReader ChunkArrayReader( ChunkFileArray );
+	return VerifyChunkFile( ChunkArrayReader, bQuickCheck );
+}
+
+bool FBuildPatchUtils::VerifyChunkFile( FArchive& ChunkFileData, bool bQuickCheck )
+{
+	const int64 FileSize = ChunkFileData.TotalSize();
+	bool bSuccess = ChunkFileData.IsLoading();
+	if ( !bSuccess )
+	{
+		GLog->Logf( TEXT( "BuildPatchServices: ERROR: VerifyChunkFile expected readonly archive" ) );
+	}
+	else
+	{
+		// Read the header
+		FChunkHeader Header;
+		ChunkFileData << Header;
+		// Check header magic
+		if ( !Header.IsValidMagic() )
+		{
+			bSuccess = false;
+			GLog->Logf( TEXT( "BuildPatchServices: ERROR: VerifyChunkFile corrupt header" ) );
+		}
+		// Check Header and data size
+		if ( bSuccess && ( Header.HeaderSize + Header.DataSize ) != FileSize )
+		{
+			bSuccess = false;
+			GLog->Logf( TEXT( "BuildPatchServices: ERROR: VerifyChunkFile header info does not match file size" ) );
+		}
+		if( bSuccess && !bQuickCheck )
+		{
+			// Hashes for checking data
+			FSHA1 SHAHasher;
+			FSHAHash SHAHash;
+			// Load the data to check
+			uint8* FileReadBuffer = new uint8[ FileBufferSize ];
+			int64 DataOffset = 0;
+			switch ( Header.StoredAs )
+			{
+			case FChunkHeader::STORED_RAW:
+				while( !ChunkFileData.AtEnd() )
+				{
+					const int64 SizeLeft = FileSize - ChunkFileData.Tell();
+					const uint32 ReadLen = FMath::Min< int64 >( FileBufferSize, SizeLeft );
+					ChunkFileData.Serialize( FileReadBuffer, ReadLen );
+					switch ( Header.HashType )
+					{
+					case FChunkHeader::HASH_ROLLING:
+						bSuccess = bSuccess && ( ReadLen == FBuildPatchData::ChunkDataSize );
+						if ( !bSuccess )
+						{
+							GLog->Logf( TEXT( "BuildPatchServices: ERROR: VerifyChunkFile invalid chunk size %d" ), ReadLen );
+						}
+						bSuccess = bSuccess && Header.RollingHash == FRollingHash< FBuildPatchData::ChunkDataSize >::GetHashForDataSet( FileReadBuffer );
+						if ( !bSuccess )
+						{
+							GLog->Logf( TEXT( "BuildPatchServices: ERROR: VerifyChunkFile chunk hashcheck failed" ) );
+						}
+						break;
+					case  FChunkHeader::HASH_SHA1:
+						SHAHasher.Update( FileReadBuffer, ReadLen );
+						break;
+					default:
+						check( false ); // @TODO LSwift: Implement other storage methods!
+						bSuccess = false;
+						break;
+					}
+					DataOffset += ReadLen;
+				}
+				if( bSuccess )
+				{
+					switch ( Header.HashType )
+					{
+					case FChunkHeader::HASH_ROLLING:
+						// Already checked
+						break;
+					case  FChunkHeader::HASH_SHA1:
+						SHAHasher.Final();
+						SHAHasher.GetHash( SHAHash.Hash );
+						bSuccess = SHAHash == Header.SHAHash;
+						if ( !bSuccess )
+						{
+							GLog->Logf( TEXT( "BuildPatchServices: ERROR: VerifyChunkFile file hashcheck failed" ) );
+						}
+					}
+				}
+				break;
+			default:
+				GLog->Logf( TEXT( "BuildPatchServices: ERROR: VerifyChunkFile failed, unknown storage type" ) );
+				bSuccess = false;
+				break;
+			}
+			delete[] FileReadBuffer;
+		}
+	}
+
+	return bSuccess;
+}
+
+bool FBuildPatchUtils::UncompressChunkFile( TArray< uint8 >& ChunkFileArray )
+{
+	FMemoryReader ChunkArrayReader( ChunkFileArray );
+	// Read the header
+	FChunkHeader Header;
+	ChunkArrayReader << Header;
+	// Check header
+	const bool bValidHeader = Header.IsValidMagic();
+	const bool bSupportedFormat = Header.HashType == FChunkHeader::HASH_ROLLING && !( Header.StoredAs & FChunkHeader::STORED_ENCRYPTED );
+	if( bValidHeader && bSupportedFormat )
+	{
+		bool bSuccess = true;
+		// Uncompress if we need to
+		if( Header.StoredAs == FChunkHeader::STORED_COMPRESSED )
+		{
+			// Load the compressed chunk data
+			TArray< uint8 > CompressedData;
+			TArray< uint8 > UncompressedData;
+			CompressedData.Empty( Header.DataSize );
+			CompressedData.AddUninitialized( Header.DataSize );
+			UncompressedData.Empty( FBuildPatchData::ChunkDataSize );
+			UncompressedData.AddUninitialized( FBuildPatchData::ChunkDataSize );
+			ChunkArrayReader.Serialize( CompressedData.GetData(), Header.DataSize );
+			ChunkArrayReader.Close();
+			// Uncompress
+			bSuccess = FCompression::UncompressMemory(
+				static_cast< ECompressionFlags >( COMPRESS_ZLIB | COMPRESS_BiasMemory ),
+				UncompressedData.GetData(),
+				UncompressedData.Num(),
+				CompressedData.GetData(),
+				CompressedData.Num() );
+			// If successful, write back over the original array
+			if( bSuccess )
+			{
+				ChunkFileArray.Empty();
+				FMemoryWriter ChunkArrayWriter( ChunkFileArray );
+				Header.StoredAs = FChunkHeader::STORED_RAW;
+				Header.DataSize = FBuildPatchData::ChunkDataSize;
+				ChunkArrayWriter << Header;
+				ChunkArrayWriter.Serialize( UncompressedData.GetData(), UncompressedData.Num() );
+				ChunkArrayWriter.Close();
+			}
+		}
+		return bSuccess;
+	}
+	return false;
+}
