Moving Build Patch Services from NoRedist to Runtime

diff --git a/Engine/Source/Runtime/Online/BuildPatchServices/Private/BuildPatchUtil.h b/Engine/Source/Runtime/Online/BuildPatchServices/Private/BuildPatchUtil.h
new file mode 100644
--- /dev/null
+++ b/Engine/Source/Runtime/Online/BuildPatchServices/Private/BuildPatchUtil.h
@@ -0,0 +1,130 @@
+// Copyright 1998-2014 Epic Games, Inc. All Rights Reserved.
+
+/*=============================================================================
+	BuildPatchUtil.h: Declares miscellaneous utility functions.
+=============================================================================*/
+
+#pragma once
+
+/**
+ * Some utility functions
+ */
+struct FBuildPatchUtils
+{
+	/**
+	 * Gets the filename for a chunk generated from it's GUID and Hash, which is the new format.
+	 * @param ManifestVersion	The manifest version that references this chunk
+	 * @param RootDirectory		The root directory
+	 * @param ChunkGuid			The chunk Guid
+	 * @param ChunkHash			The chunk rolling hash value
+	 * @return	the chunk path
+	 */
+	static const FString GetChunkNewFilename( const EBuildPatchAppManifestVersion::Type ManifestVersion, const FString& RootDirectory, const FGuid& ChunkGUID, const uint64& ChunkHash );
+
+	/**
+	 * Gets the filename for a file chunk generated from it's GUID and Hash, which is the new format.
+	 * @param ManifestVersion	The manifest version that references this file
+	 * @param RootDirectory		The root directory
+	 * @param FileGUID			The file chunk Guid
+	 * @param FileHash			The file hash value
+	 * @return	The file chunk path
+	 */
+	static const FString GetFileNewFilename( const EBuildPatchAppManifestVersion::Type ManifestVersion, const FString& RootDirectory, const FGuid& FileGUID, const FSHAHash& FileHash );
+
+	/**
+	 * Gets the chunk GUID and Hash, from the filename, which is the new format.
+	 * @param ChunkNewFilename	IN	The chunk filename that would have been generated from GetChunkNewFilename
+	 * @param ChunkGuid			OUT	The chunk Guid
+	 * @param ChunkHash			OUT	The chunk rolling hash value
+	 */
+	static void GetChunkDetailFromNewFilename( const FString& ChunkNewFilename, FGuid& ChunkGUID, uint64& ChunkHash );
+
+	/**
+	 * Gets the file chunk GUID and file hash, from the filename, which is the new format.
+	 * @param FileNewFilename	IN	The file chunk filename that would have been generated from GetChunkNewFilename
+	 * @param FileGUID			OUT	The file chunk Guid
+	 * @param FileHash			OUT	The file chunk rolling hash value
+	 */
+	static void GetFileDetailFromNewFilename( const FString& FileNewFilename, FGuid& FileGUID, FSHAHash& FileHash );
+
+	/**
+	 * Gets the filename for a chunk generated from it's GUID
+	 * @param RootDirectory		The root directory
+	 * @param ChunkGuid			The chunk Guid
+	 * @return	the chunk path
+	 */
+	static const FString GetChunkOldFilename( const FString& RootDirectory, const FGuid& ChunkGUID );
+
+	/**
+	 * Gets the filename for a file data part generated from it's GUID
+	 * @param RootDirectory		The root directory
+	 * @param FileGUID			The file part Guid
+	 * @return	The file data path
+	 */
+	static const FString GetFileOldFilename( const FString& RootDirectory, const FGuid& FileGUID );
+
+	/**
+	 * Gets the filename for a specific data part type from it's GUID
+	 * @param DataType			The type of data
+	 * @param RootDirectory		The root directory
+	 * @param DataGUID			The data Guid
+	 * @return	the data part path
+	 */
+	static const FString GetDataTypeOldFilename( const FBuildPatchData::Type DataType, const FString& RootDirectory, const FGuid& Guid );
+
+	/**
+	 * Gets the GUID for a data file according to it's filename (new or old)
+	 * @param DataFilename		IN		The data filename, or URL
+	 * @param DataGUID			OUT		Receives the GUID of the data
+	 * @return  True if successful, false otherwise
+	 */
+	static bool GetGUIDFromFilename( const FString& DataFilename, FGuid& DataGUID );
+
+	/**
+	 * Helper function to uncompress chunk file data. Can be called without knowing if needed and process will be just skipped.
+	 * @param ChunkFileArray	IN OUT		The data array, should contain full chunk - header plus data. Will be overwritten with uncompressed chunk version.
+	 * @return		true if no errors occurred and the data is not corrupted
+	 */
+	static bool UncompressChunkFile( TArray< uint8 >& ChunkFileArray );
+
+	/**
+	 * Checks a file against SHA1 hashes. The function takes two so that it can return no match, match with Hash1, or match with Hash2, that way we can check the file for being the same as an old manifest or new manifest
+	 * NOTE: This function is blocking and will not return until finished. Don't run on main thread.
+	 * @param FileToVerify			IN		The file to analyze.
+	 * @param Hash1					IN		A Hash to match against the file
+	 * @param Hash2					IN		A second Hash to match against the file
+	 * @param ProgressDelegate		IN		Delegate to receive progress updates in the form of a float range 0.0f to 1.0f
+	 * @param ShouldPauseDelegate	IN		Delegate that returns a bool, which if true will pause the process
+	 * @param TimeSpentPaused		OUT		The amount of time we were paused for, in seconds.
+	 * @return		0 if no match, 1 for match with Hash1, and 2 for match with Hash2
+	 */
+	static uint8 VerifyFile(const FString& FileToVerify, const FSHAHash& Hash1, const FSHAHash& Hash2, FBuildPatchFloatDelegate ProgressDelegate, FBuildPatchBoolRetDelegate ShouldPauseDelegate, double& TimeSpentPaused);
+
+	/**
+	 * Checks a file against SHA1 hashes. The function takes two so that it can return no match, match with Hash1, or match with Hash2, that way we can check the file for being the same as an old manifest or new manifest
+	 * NOTE: This function is blocking and will not return until finished. Don't run on main thread. This allows the above function to be easily called without delegates
+	 * @param FileToVerify			IN		The file to analyze.
+	 * @param Hash1					IN		A Hash to match against the file
+	 * @param Hash2					IN		A second Hash to match against the file
+	 * @return		0 if no match, 1 for match with Hash1, and 2 for match with Hash2
+	 */
+	static uint8 VerifyFile(const FString& FileToVerify, const FSHAHash& Hash1, const FSHAHash& Hash2);
+
+	/**
+	 * Helper function to verify the integrity of a chunk file that is in memory
+	 * @param ChunkFileArray	The data array.
+	 * @param bQuickCheck		If true, will only check header and file size, not data hash.
+	 * @return		true if no errors occurred and the data is not corrupted
+	 */
+	static bool VerifyChunkFile( const TArray< uint8 >& ChunkFileArray, bool bQuickCheck = false );
+
+private:
+
+	/**
+	 * Internal function to verify the integrity of a chunk file called by the various public versions of the function
+	 * @param ChunkFileData		The archive to test data on
+	 * @param bQuickCheck		If true, will only check header and file size, not data hash.
+	 * @return		true if no errors occurred and the data is not corrupted
+	 */
+	static bool VerifyChunkFile( FArchive& ChunkFileData, bool bQuickCheck = false );
+};
diff --git a/Engine/Source/Runtime/Online/BuildPatchServices/Public/BuildPatchServices.h b/Engine/Source/Runtime/Online/BuildPatchServices/Public/BuildPatchServices.h
new file mode 100644
--- /dev/null
+++ b/Engine/Source/Runtime/Online/BuildPatchServices/Public/BuildPatchServices.h
@@ -0,0 +1,25 @@
+// Copyright 1998-2014 Epic Games, Inc. All Rights Reserved.
+
+/*=============================================================================
+	BuildPatchServices.h: BuildPatchServices module public header file.
+=============================================================================*/
+
+#ifndef __BuildPatchServices_h__
+#define __BuildPatchServices_h__
+
+#pragma once
+
+/* Dependencies
+*****************************************************************************/
+#include "Core.h"
+#include "ModuleManager.h"
+
+/* Private includes
+*****************************************************************************/
+#include "BuildPatchSettings.h"
+#include "IBuildInstaller.h"
+#include "IBuildManifest.h"
+#include "IBuildPatchServicesModule.h"
+
+
+#endif // __BuildPatchServices_h__
diff --git a/Engine/Source/Runtime/Online/BuildPatchServices/Public/BuildPatchSettings.h b/Engine/Source/Runtime/Online/BuildPatchServices/Public/BuildPatchSettings.h
new file mode 100644
--- /dev/null
+++ b/Engine/Source/Runtime/Online/BuildPatchServices/Public/BuildPatchSettings.h
@@ -0,0 +1,31 @@
+// Copyright 1998-2014 Epic Games, Inc. All Rights Reserved.
+#pragma once
+
+/**
+* A helper for passing patch settings around
+*/
+struct FBuildPatchSettings
+{
+	// The directory to analyze
+	FString RootDirectory;
+	// The ID of the app of this build
+	uint32 InAppID;
+	// The name of the app of this build
+	FString AppName;
+	// The version string for this build
+	FString BuildVersion;
+	// The local exe path that would launch this build
+	FString LaunchExe;
+	// The command line that would launch this build
+	FString LaunchCommand;
+	// The path to a file containing a \r\n separated list of RootDirectory relative files to ignore.
+	FString IgnoreListFile;
+	// The display name of the prerequisites installer
+	FString PrereqName;
+	// The path to the prerequisites installer
+	FString PrereqPath;
+	// The command line arguments for the prerequisites installer
+	FString PrereqArgs;
+	// Map of custom fields to add to the manifest
+	TMap<FString, FVariant> CustomFields;
+};
diff --git a/Engine/Source/Runtime/Online/BuildPatchServices/Public/Interfaces/IBuildInstaller.h b/Engine/Source/Runtime/Online/BuildPatchServices/Public/Interfaces/IBuildInstaller.h
new file mode 100644
--- /dev/null
+++ b/Engine/Source/Runtime/Online/BuildPatchServices/Public/Interfaces/IBuildInstaller.h
@@ -0,0 +1,175 @@
+// Copyright 1998-2014 Epic Games, Inc. All Rights Reserved.
+
+/*=============================================================================
+	IBuildInstaller.h: Declares the IBuildInstaller interface.
+=============================================================================*/
+
+#pragma once
+
+typedef TSharedPtr< class IBuildInstaller, ESPMode::ThreadSafe > IBuildInstallerPtr;
+typedef TSharedRef< class IBuildInstaller, ESPMode::ThreadSafe > IBuildInstallerRef;
+
+/**
+ * A struct to hold stats for the build process.
+ */
+struct FBuildInstallStats
+{
+	// The name of the app being installed
+	FString AppName;
+	// The version string currently installed, or "NONE"
+	FString AppInstalledVersion;
+	// The version string patching to
+	FString AppPatchVersion;
+	// The cloud directory used for this install
+	FString CloudDirectory;
+	// The total number of files in the build
+	uint32 NumFilesInBuild;
+	// The total number of files outdated
+	uint32 NumFilesOutdated;
+	// The total number of files in the previous build that can be deleted
+	uint32 NumFilesToRemove;
+	// The total number of chunks making up those files
+	uint32 NumChunksRequired;
+	// The number of required chunks queued for download
+	uint32 ChunksQueuedForDownload;
+	// The number of chunks locally available in the build
+	uint32 ChunksLocallyAvailable;
+	// The total number of chunks that were downloaded
+	uint32 NumChunksDownloaded;
+	// The number of chunks successfully recycled
+	uint32 NumChunksRecycled;
+	// The number of chunks that had to be booted from the cache
+	uint32 NumChunksCacheBooted;
+	// The number of chunks that had to be loaded from the drive cache
+	uint32 NumDriveCacheChunkLoads;
+	// The number of chunks that failed to be recycled from existing build
+	uint32 NumRecycleFailures;
+	// The number of chunks that failed to load from the drive cache
+	uint32 NumDriveCacheLoadFailures;
+	// The total number of bytes downloaded
+	int64 TotalDownloadedData;
+	// The average chunk download speed
+	double AverageDownloadSpeed;
+	// The theoretical download time (data/speed)
+	float TheoreticalDownloadTime;
+	// The total time that the install verification took to complete
+	float VerifyTime;
+	// The total time that the install cleanup took to complete
+	float CleanUpTime;
+	// The total time that the install process took to complete
+	float ProcessExecuteTime;
+	// The amount of time that was spent paused
+	float ProcessPausedTime;
+	// Whether the process was successful
+	bool ProcessSuccess;
+	// The reason for failure, or "SUCCESS"
+	FString FailureReason;
+	// The localised, more generic failure reason
+	FText FailureReasonText;
+};
+
+/**
+ * Interface to a Build Manifest.
+ */
+class IBuildInstaller
+{
+public:
+	/**
+	 * Virtual destructor.
+	 */
+	virtual ~IBuildInstaller() { }
+
+	/**
+	 * Get whether the install has complete
+	 * @return	true if the thread completed
+	 */
+	virtual bool IsComplete() = 0;
+
+	/**
+	 * Get whether the install was canceled. Only valid if complete.
+	 * @return	true if installation was canceled
+	 */
+	virtual bool IsCanceled() = 0;
+
+	/**
+	 * Get whether the install is currently paused.
+	 * @return	true if installation is paused
+	 */
+	virtual bool IsPaused() = 0;
+
+	/**
+	 * Get whether the install failed. Only valid if complete.
+	 * @return	true if installation was a failure
+	 */
+	virtual bool HasError() = 0;
+
+	/**
+	 * This is deprecated and shouldn't be used anymore [6/4/2014 justin.sargent]
+	 *
+	 * Get the percentage complete text for the current process
+	 * @return	percentage complete text progress
+	 */
+	virtual FText GetPercentageText() = 0;
+
+	/**
+	 * This is deprecated and shouldn't be used anymore [6/4/2014 justin.sargent]
+	 *
+	 * Get the download speed text for the current process
+	 * @return	download speed text progress
+	 */
+	virtual FText GetDownloadSpeedText()  = 0;
+
+	/**
+	 * Get the download speed for the current process
+	 * @return	download speed progress
+	 */
+	virtual double GetDownloadSpeed() const  = 0;
+
+	/**
+	 * Get the initial download size
+	 * @return	the initial download size
+	 */
+	virtual int64 GetInitialDownloadSize() const  = 0;
+
+	/**
+	* Get the total currently downloaded
+	* @return	the total currently downloaded
+	 */
+	virtual int64 GetTotalDownloaded() const = 0;
+
+	/**
+	 * Get the text for status of the install process
+	 * @return	status of the install process text
+	 */
+	virtual FText GetStatusText() = 0;
+
+	/**
+	 * Get the update progress
+	 * @return	A float describing progress: Between 0 and 1 for known progress, or less than 0 for unknown progress.
+	 */
+	virtual float GetUpdateProgress() = 0;
+
+	/**
+	 * Get the build stats for the process. This should only be called after the install has completed
+	 * @return	A struct containing information about the build
+	 */
+	virtual FBuildInstallStats GetBuildStatistics() = 0;
+
+	/**
+	 * Get the display text for the error that occurred. Only valid to call after completion
+	 * @return	display error text
+	 */
+	virtual FText GetErrorText() = 0;
+
+	/**
+	 * Cancel the current install
+	 */
+	virtual void CancelInstall() = 0;
+
+	/**
+	 * Toggle the install paused state
+	 * @return true if the installer is now paused
+	 */
+	virtual bool TogglePauseInstall() = 0;
+};
+
diff --git a/Engine/Source/Runtime/Online/BuildPatchServices/Public/Interfaces/IBuildPatchServicesModule.h b/Engine/Source/Runtime/Online/BuildPatchServices/Public/Interfaces/IBuildPatchServicesModule.h
new file mode 100644
--- /dev/null
+++ b/Engine/Source/Runtime/Online/BuildPatchServices/Public/Interfaces/IBuildPatchServicesModule.h
@@ -0,0 +1,119 @@
+// Copyright 1998-2014 Epic Games, Inc. All Rights Reserved.
+
+/*=============================================================================
+	IBuildPatchServicesModule.h: Declares the IBuildPatchServicesModule interface.
+=============================================================================*/
+
+#pragma once
+
+class IAnalyticsProvider;
+
+/**
+ * Delegates that will be accepted and fired of by the implementation
+ */
+DECLARE_DELEGATE( FBuildPatchDelegate );
+DECLARE_DELEGATE_OneParam( FBuildPatchManifestDelegate, IBuildManifestRef );
+DECLARE_DELEGATE_TwoParams( FBuildPatchBoolManifestDelegate, bool, IBuildManifestRef );
+DECLARE_DELEGATE_OneParam( FBuildPatchFloatDelegate, float );
+
+/**
+ * Interface for the services manager.
+ */
+class IBuildPatchServicesModule
+	: public IModuleInterface
+{
+public:
+	/**
+	 * Virtual destructor.
+	 */
+	virtual ~IBuildPatchServicesModule( ) { }
+	
+	/**
+	 * Loads a Build Manifest from file and returns the interface
+	 * @param Filename		The file to load from
+	 * @return		a shared pointer to the manifest, which will be invalid if load failed.
+	 */
+	virtual IBuildManifestPtr LoadManifestFromFile( const FString& Filename ) = 0;
+	
+	/**
+	 * Constructs a Build Manifest from a JSON string
+	 * @param ManifestJSON		The JSON received from a web api
+	 * @return		a shared pointer to the manifest, which will be invalid if creation failed.
+	 */
+	virtual IBuildManifestPtr MakeManifestFromJSON( const FString& ManifestJSON ) = 0;
+	
+	/**
+	 * Saves a Build Manifest to file
+	 * @param Filename		The file to save to
+	 * @param Manifest		The manifest to save out
+	 * @return		If the save was successful.
+	 */
+	virtual bool SaveManifestToFile( const FString& Filename, IBuildManifestRef Manifest ) = 0;
+
+	/**
+	 * Starts an installer thread for the provided manifests
+	 * @param	CurrentManifest			The manifest that the current install was generated from (if applicable)
+	 * @param	InstallManifest			The manifest to be installed
+	 * @param	InstallDirectory		The directory to install the App to
+	 * @param	OnCompleteDelegate		The delegate to call on completion
+	 * @return		An interface to the created installer. Will be an invalid ptr if error.
+	 */
+	virtual IBuildInstallerPtr StartBuildInstall( IBuildManifestPtr CurrentManifest, IBuildManifestPtr InstallManifest, const FString& InstallDirectory, FBuildPatchBoolManifestDelegate OnCompleteDelegate ) = 0;
+
+	/**
+	 * Sets the directory used for staging intermediate files.
+	 * @param StagingDir	The staging directory
+	 */
+	virtual void SetStagingDirectory( const FString& StagingDir ) = 0;
+
+	/**
+	 * Sets the cloud directory where chunks and manifests will be pulled from and saved to.
+	 * @param CloudDir		The cloud directory
+	 */
+	virtual void SetCloudDirectory( const FString& CloudDir ) = 0;
+
+	/**
+	 * Sets the backup directory where files that are being clobbered by repair/patch will be placed.
+	 * @param BackupDir		The backup directory
+	 */
+	virtual void SetBackupDirectory( const FString& BackupDir ) = 0;
+
+	/**
+	 * Sets the Analytics provider that will be used to register errors with patch/build installs
+	 * @param AnalyticsProvider		Shared ptr to an analytics interface to use. If NULL analytics will be disabled.
+	 */
+	virtual void SetAnalyticsProvider( TSharedPtr< IAnalyticsProvider > AnalyticsProvider ) = 0;
+
+	/**
+	 * Registers an installation on this machine. This information is used to gather a list of install locations that can be used as chunk sources.
+	 * @param AppManifest			Ref to the manifest for this installation
+	 * @param AppInstallDirectory	The install location
+	 */
+	virtual void RegisterAppInstallation(IBuildManifestRef AppManifest, const FString AppInstallDirectory) = 0;
+
+#if WITH_BUILDPATCHGENERATION
+	/**
+	 * Processes a Build Image to determine new chunks and produce a chunk based manifest, all saved to the cloud.
+	 * NOTE: This function is blocking and will not return until finished. Don't run on main thread.
+	 * @param Settings				Specifies the settings for the operation.  See FBuildPatchSettings documentation.
+	 * @return		true if no file errors occurred
+	 */
+	virtual bool GenerateChunksManifestFromDirectory( const FBuildPatchSettings& Settings ) = 0;
+
+	/**
+	 * Processes a Build Image to determine new raw files and produce a raw file based manifest, all saved to the cloud.
+	 * NOTE: This function is blocking and will not return until finished. Don't run on main thread.
+	 * @param Settings				Specifies the settings for the operation.  See FBuildPatchSettings documentation.
+	 * @return		true if no file errors occurred
+	 */
+	virtual bool GenerateFilesManifestFromDirectory( const FBuildPatchSettings& Settings ) = 0;
+
+	/**
+	 * Processes a Cloud Directory to identify and delete any orphaned chunks or files.
+	 * NOTE: THIS function is blocking and will not return until finished. Don't run on main thread.
+	 * @param bPreview      If true, then no actual work will be done, but all operations which would be carried out will be logged.
+	 * @return              true if no file errors occurred
+	 */
+	virtual bool CompactifyCloudDirectory( const bool bPreview ) = 0;
+#endif // WITH_BUILDPATCHGENERATION
+};
