CIS fix. We have more platforms than Win and Mac ;-)

diff --git a/Engine/Source/Runtime/Online/OnlineSubsystemSteam/Private/OnlineSubsystemModuleSteam.cpp b/Engine/Source/Runtime/Online/OnlineSubsystemSteam/Private/OnlineSubsystemModuleSteam.cpp
--- a/Engine/Source/Runtime/Online/OnlineSubsystemSteam/Private/OnlineSubsystemModuleSteam.cpp
+++ b/Engine/Source/Runtime/Online/OnlineSubsystemSteam/Private/OnlineSubsystemModuleSteam.cpp
@@ -91,7 +91,7 @@ static FString GetSteamModulePath()
 		return FPaths::EngineDir() / STEAM_SDK_ROOT_PATH / STEAM_SDK_VER_PATH / TEXT("Win32/");
 	#endif	//PLATFORM_64BITS
 
-#elif PLATFORM_MAC
+#else
 
 	return FString();
 
