Add overloads for some OSS identity interface functions that take FUniqueNetId instead of LocalUserNum

diff --git a/Engine/Source/Runtime/Online/Android/OnlineSubsystemGooglePlay/Private/OnlineIdentityInterfaceGooglePlay.h b/Engine/Source/Runtime/Online/Android/OnlineSubsystemGooglePlay/Private/OnlineIdentityInterfaceGooglePlay.h
--- a/Engine/Source/Runtime/Online/Android/OnlineSubsystemGooglePlay/Private/OnlineIdentityInterfaceGooglePlay.h
+++ b/Engine/Source/Runtime/Online/Android/OnlineSubsystemGooglePlay/Private/OnlineIdentityInterfaceGooglePlay.h
@@ -61,7 +61,9 @@ public:
 	virtual TSharedPtr<FUniqueNetId> CreateUniquePlayerId(uint8* Bytes, int32 Size) override;
 	virtual TSharedPtr<FUniqueNetId> CreateUniquePlayerId(const FString& Str) override;
 	virtual ELoginStatus::Type GetLoginStatus(int32 LocalUserNum) const override;
+	virtual ELoginStatus::Type GetLoginStatus(const FUniqueNetId& UserId) const override;
 	virtual FString GetPlayerNickname(int32 LocalUserNum) const override;
+	virtual FString GetPlayerNickname(const FUniqueNetId& UserId) const override;
 	virtual FString GetAuthToken(int32 LocalUserNum) const override;
 	// End IOnlineIdentity interface
 
diff --git a/Engine/Source/Runtime/Online/Android/OnlineSubsystemGooglePlay/Private/OnlineIdentityInterfaceGoolgePlay.cpp b/Engine/Source/Runtime/Online/Android/OnlineSubsystemGooglePlay/Private/OnlineIdentityInterfaceGoolgePlay.cpp
--- a/Engine/Source/Runtime/Online/Android/OnlineSubsystemGooglePlay/Private/OnlineIdentityInterfaceGoolgePlay.cpp
+++ b/Engine/Source/Runtime/Online/Android/OnlineSubsystemGooglePlay/Private/OnlineIdentityInterfaceGoolgePlay.cpp
@@ -118,6 +118,24 @@ ELoginStatus::Type FOnlineIdentityGooglePlay::GetLoginStatus(int32 LocalUserNum)
 }
 
 
+ELoginStatus::Type FOnlineIdentityGooglePlay::GetLoginStatus(const FUniqueNetId& UserId) const
+{
+	ELoginStatus::Type LoginStatus = ELoginStatus::NotLoggedIn;
+
+	if(bLoggedIn)
+	{
+		UE_LOG(LogOnline, Display, TEXT("FOnlineIdentityAndroid::GetLoginStatus - Logged in"));
+		LoginStatus = ELoginStatus::LoggedIn;
+	}
+	else
+	{
+		UE_LOG(LogOnline, Display, TEXT("FOnlineIdentityAndroid::GetLoginStatus - Not logged in"));
+	}
+
+	return LoginStatus;
+}
+
+
 TSharedPtr<FUniqueNetId> FOnlineIdentityGooglePlay::GetUniquePlayerId(int32 LocalUserNum) const
 {
 	TSharedPtr<FUniqueNetId> NewID = MakeShareable(new FUniqueNetIdString(""));
@@ -152,6 +170,11 @@ FString FOnlineIdentityGooglePlay::GetPlayerNickname(int32 LocalUserNum) const
 	return PlayerAlias;
 }
 
+FString FOnlineIdentityGooglePlay::GetPlayerNickname(const FUniqueNetId& UserId) const
+{
+	UE_LOG(LogOnline, Display, TEXT("FOnlineIdentityAndroid::GetPlayerNickname"));
+	return PlayerAlias;
+}
 
 FString FOnlineIdentityGooglePlay::GetAuthToken(int32 LocalUserNum) const
 {
diff --git a/Engine/Source/Runtime/Online/IOS/OnlineSubsystemIOS/Private/OnlineIdentityInterfaceIOS.cpp b/Engine/Source/Runtime/Online/IOS/OnlineSubsystemIOS/Private/OnlineIdentityInterfaceIOS.cpp
--- a/Engine/Source/Runtime/Online/IOS/OnlineSubsystemIOS/Private/OnlineIdentityInterfaceIOS.cpp
+++ b/Engine/Source/Runtime/Online/IOS/OnlineSubsystemIOS/Private/OnlineIdentityInterfaceIOS.cpp
@@ -125,6 +125,18 @@ ELoginStatus::Type FOnlineIdentityIOS::GetLoginStatus(int32 LocalUserNum) const
 	return LoginStatus;
 }
 
+ELoginStatus::Type FOnlineIdentityIOS::GetLoginStatus(const FUniqueNetId& UserId) const 
+{
+	ELoginStatus::Type LoginStatus = ELoginStatus::NotLoggedIn;
+
+	if(GetLocalGameCenterUser() != NULL && GetLocalGameCenterUser().isAuthenticated == YES)
+	{
+		LoginStatus = ELoginStatus::LoggedIn;
+	}
+
+	return LoginStatus;
+}
+
 TSharedPtr<FUniqueNetId> FOnlineIdentityIOS::GetUniquePlayerId(int32 LocalUserNum) const
 {
 	return UniqueNetId;
@@ -165,6 +177,21 @@ FString FOnlineIdentityIOS::GetPlayerNickname(int32 LocalUserNum) const
 	return FString();
 }
 
+FString FOnlineIdentityIOS::GetPlayerNickname(const FUniqueNetId& UserId) const 
+{
+	if (GetLocalGameCenterUser() != NULL)
+	{
+		NSString* PersonaName = [GetLocalGameCenterUser() alias];
+		
+        if (PersonaName != nil)
+        {
+            return FString(PersonaName);
+        }
+	}
+
+	return FString();
+}
+
 FString FOnlineIdentityIOS::GetAuthToken(int32 LocalUserNum) const
 {
 	FString ResultToken;
diff --git a/Engine/Source/Runtime/Online/IOS/OnlineSubsystemIOS/Private/OnlineIdentityInterfaceIOS.h b/Engine/Source/Runtime/Online/IOS/OnlineSubsystemIOS/Private/OnlineIdentityInterfaceIOS.h
--- a/Engine/Source/Runtime/Online/IOS/OnlineSubsystemIOS/Private/OnlineIdentityInterfaceIOS.h
+++ b/Engine/Source/Runtime/Online/IOS/OnlineSubsystemIOS/Private/OnlineIdentityInterfaceIOS.h
@@ -37,7 +37,9 @@ public:
 	virtual TSharedPtr<FUniqueNetId> CreateUniquePlayerId(uint8* Bytes, int32 Size) override;
 	virtual TSharedPtr<FUniqueNetId> CreateUniquePlayerId(const FString& Str) override;
 	virtual ELoginStatus::Type GetLoginStatus(int32 LocalUserNum) const override;
+	virtual ELoginStatus::Type GetLoginStatus(const FUniqueNetId& UserId) const override;
 	virtual FString GetPlayerNickname(int32 LocalUserNum) const override;
+	virtual FString GetPlayerNickname(const FUniqueNetId& UserId) const override;
 	virtual FString GetAuthToken(int32 LocalUserNum) const override;
 	// End IOnlineIdentity interface
 
diff --git a/Engine/Source/Runtime/Online/OnlineSubsystem/Public/Interfaces/OnlineIdentityInterface.h b/Engine/Source/Runtime/Online/OnlineSubsystem/Public/Interfaces/OnlineIdentityInterface.h
--- a/Engine/Source/Runtime/Online/OnlineSubsystem/Public/Interfaces/OnlineIdentityInterface.h
+++ b/Engine/Source/Runtime/Online/OnlineSubsystem/Public/Interfaces/OnlineIdentityInterface.h
@@ -247,6 +247,15 @@ public:
 	 */
 	virtual ELoginStatus::Type GetLoginStatus(int32 LocalUserNum) const = 0;
 
+	/**
+	 * Fetches the login status for a given player
+	 *
+	 * @param UserId the unique net id of the associated user
+	 *
+	 * @return the enum value of their status
+	 */
+	virtual ELoginStatus::Type GetLoginStatus(const FUniqueNetId& UserId) const = 0;
+
 	/**
 	 * Reads the player's nick name from the online service
 	 *
@@ -257,6 +266,16 @@ public:
 	//@todo - move to user interface
 	virtual FString GetPlayerNickname(int32 LocalUserNum) const = 0;
 
+	/**
+	 * Reads the player's nick name from the online service
+	 *
+	 * @param UserId the unique net of the associated user
+	 *
+	 * @return a string containing the players nick name
+	 */
+	//@todo - move to user interface
+	virtual FString GetPlayerNickname(const FUniqueNetId& UserId) const = 0;
+
 	/**
 	 * Gets a user's platform specific authentication token to verify their identity
 	 *
diff --git a/Engine/Source/Runtime/Online/OnlineSubsystemAmazon/Private/OnlineIdentityAmazon.cpp b/Engine/Source/Runtime/Online/OnlineSubsystemAmazon/Private/OnlineIdentityAmazon.cpp
--- a/Engine/Source/Runtime/Online/OnlineSubsystemAmazon/Private/OnlineIdentityAmazon.cpp
+++ b/Engine/Source/Runtime/Online/OnlineSubsystemAmazon/Private/OnlineIdentityAmazon.cpp
@@ -293,13 +293,19 @@ ELoginStatus::Type FOnlineIdentityAmazon::GetLoginStatus(int32 LocalUserNum) con
 	TSharedPtr<FUniqueNetId> UserId = GetUniquePlayerId(LocalUserNum);
 	if (UserId.IsValid())
 	{
-		TSharedPtr<FUserOnlineAccount> UserAccount = GetUserAccount(*UserId);
-		if (UserAccount.IsValid() &&
-			UserAccount->GetUserId()->IsValid() &&
-			!UserAccount->GetAccessToken().IsEmpty())
-		{
-			return ELoginStatus::LoggedIn;
-		}
+		return GetLoginStatus(*UserId);
+	}
+	return ELoginStatus::NotLoggedIn;
+}
+
+ELoginStatus::Type FOnlineIdentityAmazon::GetLoginStatus(const FUniqueNetId& UserId) const 
+{
+	TSharedPtr<FUserOnlineAccount> UserAccount = GetUserAccount(UserId);
+	if (UserAccount.IsValid() &&
+		UserAccount->GetUserId()->IsValid() &&
+		!UserAccount->GetAccessToken().IsEmpty())
+	{
+		return ELoginStatus::LoggedIn;
 	}
 	return ELoginStatus::NotLoggedIn;
 }
@@ -310,6 +316,12 @@ FString FOnlineIdentityAmazon::GetPlayerNickname(int32 LocalUserNum) const
 	return TEXT("AmazonUser");
 }
 
+FString FOnlineIdentityAmazon::GetPlayerNickname(const FUniqueNetId& UserId) const
+{
+	//@todo - not implemented
+	return TEXT("AmazonUser");
+}
+
 FString FOnlineIdentityAmazon::GetAuthToken(int32 LocalUserNum) const
 {
 	TSharedPtr<FUniqueNetId> UserId = GetUniquePlayerId(LocalUserNum);
diff --git a/Engine/Source/Runtime/Online/OnlineSubsystemAmazon/Private/OnlineIdentityAmazon.h b/Engine/Source/Runtime/Online/OnlineSubsystemAmazon/Private/OnlineIdentityAmazon.h
--- a/Engine/Source/Runtime/Online/OnlineSubsystemAmazon/Private/OnlineIdentityAmazon.h
+++ b/Engine/Source/Runtime/Online/OnlineSubsystemAmazon/Private/OnlineIdentityAmazon.h
@@ -107,7 +107,9 @@ public:
 	virtual TSharedPtr<FUniqueNetId> CreateUniquePlayerId(uint8* Bytes, int32 Size) override;
 	virtual TSharedPtr<FUniqueNetId> CreateUniquePlayerId(const FString& Str) override;
 	virtual ELoginStatus::Type GetLoginStatus(int32 LocalUserNum) const override;
+	virtual ELoginStatus::Type GetLoginStatus(const FUniqueNetId& UserId) const override;
 	virtual FString GetPlayerNickname(int32 LocalUserNum) const override;
+	virtual FString GetPlayerNickname(const FUniqueNetId& UserId) const override;
 	virtual FString GetAuthToken(int32 LocalUserNum) const override;
 
 	// FOnlineIdentityAmazon
diff --git a/Engine/Source/Runtime/Online/OnlineSubsystemFacebook/Private/IOS/OnlineIdentityFacebook.cpp b/Engine/Source/Runtime/Online/OnlineSubsystemFacebook/Private/IOS/OnlineIdentityFacebook.cpp
--- a/Engine/Source/Runtime/Online/OnlineSubsystemFacebook/Private/IOS/OnlineIdentityFacebook.cpp
+++ b/Engine/Source/Runtime/Online/OnlineSubsystemFacebook/Private/IOS/OnlineIdentityFacebook.cpp
@@ -191,12 +191,22 @@ ELoginStatus::Type FOnlineIdentityFacebook::GetLoginStatus(int32 LocalUserNum) c
 	return LoginStatus;
 }
 
+ELoginStatus::Type FOnlineIdentityFacebook::GetLoginStatus(const FUniqueNetId& UserId) const 
+{
+	return LoginStatus;
+}
+
 
 FString FOnlineIdentityFacebook::GetPlayerNickname(int32 LocalUserNum) const
 {
 	return TEXT("FacebookUser");
 }
 
+FString FOnlineIdentityFacebook::GetPlayerNickname(const FUniqueNetId& UserId) const 
+{
+	return TEXT("FacebookUser");
+}
+
 
 FString FOnlineIdentityFacebook::GetAuthToken(int32 LocalUserNum) const
 {
diff --git a/Engine/Source/Runtime/Online/OnlineSubsystemFacebook/Private/IOS/OnlineIdentityFacebook.h b/Engine/Source/Runtime/Online/OnlineSubsystemFacebook/Private/IOS/OnlineIdentityFacebook.h
--- a/Engine/Source/Runtime/Online/OnlineSubsystemFacebook/Private/IOS/OnlineIdentityFacebook.h
+++ b/Engine/Source/Runtime/Online/OnlineSubsystemFacebook/Private/IOS/OnlineIdentityFacebook.h
@@ -82,7 +82,9 @@ public:
 	virtual TSharedPtr<FUniqueNetId> CreateUniquePlayerId(uint8* Bytes, int32 Size) override;
 	virtual TSharedPtr<FUniqueNetId> CreateUniquePlayerId(const FString& Str) override;
 	virtual ELoginStatus::Type GetLoginStatus(int32 LocalUserNum) const override;
+	virtual ELoginStatus::Type GetLoginStatus(const FUniqueNetId& UserId) const override;
 	virtual FString GetPlayerNickname(int32 LocalUserNum) const override;
+	virtual FString GetPlayerNickname(const FUniqueNetId& UserId) const override;
 	virtual FString GetAuthToken(int32 LocalUserNum) const override;
 	// End IOnlineIdentity interface
 
diff --git a/Engine/Source/Runtime/Online/OnlineSubsystemFacebook/Private/Windows/OnlineIdentityFacebook.cpp b/Engine/Source/Runtime/Online/OnlineSubsystemFacebook/Private/Windows/OnlineIdentityFacebook.cpp
--- a/Engine/Source/Runtime/Online/OnlineSubsystemFacebook/Private/Windows/OnlineIdentityFacebook.cpp
+++ b/Engine/Source/Runtime/Online/OnlineSubsystemFacebook/Private/Windows/OnlineIdentityFacebook.cpp
@@ -313,13 +313,19 @@ ELoginStatus::Type FOnlineIdentityFacebook::GetLoginStatus(int32 LocalUserNum) c
 	TSharedPtr<FUniqueNetId> UserId = GetUniquePlayerId(LocalUserNum);
 	if (UserId.IsValid())
 	{
-		TSharedPtr<FUserOnlineAccount> UserAccount = GetUserAccount(*UserId);
-		if (UserAccount.IsValid() &&
-			UserAccount->GetUserId()->IsValid() &&
-			!UserAccount->GetAccessToken().IsEmpty())
-		{
-			return ELoginStatus::LoggedIn;
-		}
+		return GetLoginStatus(*UserId);
+	}
+	return ELoginStatus::NotLoggedIn;
+}
+
+ELoginStatus::Type FOnlineIdentityFacebook::GetLoginStatus(const FUniqueNetId& UserId) const
+{
+	TSharedPtr<FUserOnlineAccount> UserAccount = GetUserAccount(UserId);
+	if (UserAccount.IsValid() &&
+		UserAccount->GetUserId()->IsValid() &&
+		!UserAccount->GetAccessToken().IsEmpty())
+	{
+		return ELoginStatus::LoggedIn;
 	}
 	return ELoginStatus::NotLoggedIn;
 }
@@ -329,12 +335,18 @@ FString FOnlineIdentityFacebook::GetPlayerNickname(int32 LocalUserNum) const
 	TSharedPtr<FUniqueNetId> UserId = GetUniquePlayerId(LocalUserNum);
 	if (UserId.IsValid())
 	{
-		const TSharedRef<FUserOnlineAccountFacebook>* FoundUserAccount = UserAccounts.Find(UserId->ToString());
-		if (FoundUserAccount != NULL)
-		{
-			const TSharedRef<FUserOnlineAccountFacebook>& UserAccount = *FoundUserAccount;
-			return UserAccount->RealName;
-		}
+		return  GetPlayerNickname(*UserId);
+	}
+	return TEXT("");
+}
+
+FString FOnlineIdentityFacebook::GetPlayerNickname(const FUniqueNetId& UserId) const
+{
+	const TSharedRef<FUserOnlineAccountFacebook>* FoundUserAccount = UserAccounts.Find(UserId.ToString());
+	if (FoundUserAccount != NULL)
+	{
+		const TSharedRef<FUserOnlineAccountFacebook>& UserAccount = *FoundUserAccount;
+		return UserAccount->RealName;
 	}
 	return TEXT("");
 }
diff --git a/Engine/Source/Runtime/Online/OnlineSubsystemFacebook/Private/Windows/OnlineIdentityFacebook.h b/Engine/Source/Runtime/Online/OnlineSubsystemFacebook/Private/Windows/OnlineIdentityFacebook.h
--- a/Engine/Source/Runtime/Online/OnlineSubsystemFacebook/Private/Windows/OnlineIdentityFacebook.h
+++ b/Engine/Source/Runtime/Online/OnlineSubsystemFacebook/Private/Windows/OnlineIdentityFacebook.h
@@ -105,7 +105,9 @@ public:
 	virtual TSharedPtr<FUniqueNetId> CreateUniquePlayerId(uint8* Bytes, int32 Size) override;
 	virtual TSharedPtr<FUniqueNetId> CreateUniquePlayerId(const FString& Str) override;
 	virtual ELoginStatus::Type GetLoginStatus(int32 LocalUserNum) const override;
+	virtual ELoginStatus::Type GetLoginStatus(const FUniqueNetId& UserId) const override;
 	virtual FString GetPlayerNickname(int32 LocalUserNum) const override;
+	virtual FString GetPlayerNickname(const FUniqueNetId& UserId) const override;
 	virtual FString GetAuthToken(int32 LocalUserNum) const override;
 
 	// FOnlineIdentityFacebook
diff --git a/Engine/Source/Runtime/Online/OnlineSubsystemNull/Private/OnlineIdentityNull.cpp b/Engine/Source/Runtime/Online/OnlineSubsystemNull/Private/OnlineIdentityNull.cpp
--- a/Engine/Source/Runtime/Online/OnlineSubsystemNull/Private/OnlineIdentityNull.cpp
+++ b/Engine/Source/Runtime/Online/OnlineSubsystemNull/Private/OnlineIdentityNull.cpp
@@ -205,12 +205,18 @@ ELoginStatus::Type FOnlineIdentityNull::GetLoginStatus(int32 LocalUserNum) const
 	TSharedPtr<FUniqueNetId> UserId = GetUniquePlayerId(LocalUserNum);
 	if (UserId.IsValid())
 	{
-		TSharedPtr<FUserOnlineAccount> UserAccount = GetUserAccount(*UserId);
-		if (UserAccount.IsValid() &&
-			UserAccount->GetUserId()->IsValid())
-		{
-			return ELoginStatus::LoggedIn;
-		}
+		return GetLoginStatus(*UserId);
+	}
+	return ELoginStatus::NotLoggedIn;
+}
+
+ELoginStatus::Type FOnlineIdentityNull::GetLoginStatus(const FUniqueNetId& UserId) const 
+{
+	TSharedPtr<FUserOnlineAccount> UserAccount = GetUserAccount(UserId);
+	if (UserAccount.IsValid() &&
+		UserAccount->GetUserId()->IsValid())
+	{
+		return ELoginStatus::LoggedIn;
 	}
 	return ELoginStatus::NotLoggedIn;
 }
@@ -226,6 +232,11 @@ FString FOnlineIdentityNull::GetPlayerNickname(int32 LocalUserNum) const
 	return TEXT("NullUser");
 }
 
+FString FOnlineIdentityNull::GetPlayerNickname(const FUniqueNetId& UserId) const
+{
+	return UserId.ToString();
+}
+
 FString FOnlineIdentityNull::GetAuthToken(int32 LocalUserNum) const
 {
 	TSharedPtr<FUniqueNetId> UserId = GetUniquePlayerId(LocalUserNum);
diff --git a/Engine/Source/Runtime/Online/OnlineSubsystemNull/Private/OnlineIdentityNull.h b/Engine/Source/Runtime/Online/OnlineSubsystemNull/Private/OnlineIdentityNull.h
--- a/Engine/Source/Runtime/Online/OnlineSubsystemNull/Private/OnlineIdentityNull.h
+++ b/Engine/Source/Runtime/Online/OnlineSubsystemNull/Private/OnlineIdentityNull.h
@@ -62,7 +62,9 @@ public:
 	virtual TSharedPtr<FUniqueNetId> CreateUniquePlayerId(uint8* Bytes, int32 Size) override;
 	virtual TSharedPtr<FUniqueNetId> CreateUniquePlayerId(const FString& Str) override;
 	virtual ELoginStatus::Type GetLoginStatus(int32 LocalUserNum) const override;
+	virtual ELoginStatus::Type GetLoginStatus(const FUniqueNetId& UserId) const override;
 	virtual FString GetPlayerNickname(int32 LocalUserNum) const override;
+	virtual FString GetPlayerNickname(const FUniqueNetId& UserId) const override;
 	virtual FString GetAuthToken(int32 LocalUserNum) const override;
 
 	// FOnlineIdentityNull
diff --git a/Engine/Source/Runtime/Online/OnlineSubsystemSteam/Private/OnlineIdentityInterfaceSteam.cpp b/Engine/Source/Runtime/Online/OnlineSubsystemSteam/Private/OnlineIdentityInterfaceSteam.cpp
--- a/Engine/Source/Runtime/Online/OnlineSubsystemSteam/Private/OnlineIdentityInterfaceSteam.cpp
+++ b/Engine/Source/Runtime/Online/OnlineSubsystemSteam/Private/OnlineIdentityInterfaceSteam.cpp
@@ -101,6 +101,11 @@ ELoginStatus::Type FOnlineIdentitySteam::GetLoginStatus(int32 LocalUserNum) cons
 	return ELoginStatus::NotLoggedIn; 
 }
 
+ELoginStatus::Type FOnlineIdentitySteam::GetLoginStatus(const FUniqueNetId& UserId) const 
+{
+	return GetLoginStatus(0);
+}
+
 TSharedPtr<FUniqueNetId> FOnlineIdentitySteam::GetUniquePlayerId(int32 LocalUserNum) const
 {
 	if (LocalUserNum < MAX_LOCAL_PLAYERS &&
@@ -149,6 +154,16 @@ FString FOnlineIdentitySteam::GetPlayerNickname(int32 LocalUserNum) const
 	return FString(TEXT(""));
 }
 
+FString FOnlineIdentitySteam::GetPlayerNickname(const FUniqueNetId& UserId) const
+{
+	if (SteamFriendsPtr != NULL)
+	{
+		const char* PersonaName = SteamFriendsPtr->GetPersonaName();
+		return FString(UTF8_TO_TCHAR(PersonaName));
+	}
+	return FString(TEXT(""));
+}
+
 /**
  * Gets a user's platform specific authentication token to verify their identity
  *
diff --git a/Engine/Source/Runtime/Online/OnlineSubsystemSteam/Private/OnlineIdentityInterfaceSteam.h b/Engine/Source/Runtime/Online/OnlineSubsystemSteam/Private/OnlineIdentityInterfaceSteam.h
--- a/Engine/Source/Runtime/Online/OnlineSubsystemSteam/Private/OnlineIdentityInterfaceSteam.h
+++ b/Engine/Source/Runtime/Online/OnlineSubsystemSteam/Private/OnlineIdentityInterfaceSteam.h
@@ -33,7 +33,9 @@ public:
 	virtual TSharedPtr<FUniqueNetId> CreateUniquePlayerId(uint8* Bytes, int32 Size) override;
 	virtual TSharedPtr<FUniqueNetId> CreateUniquePlayerId(const FString& Str) override;
 	virtual ELoginStatus::Type GetLoginStatus(int32 LocalUserNum) const override;
+	virtual ELoginStatus::Type GetLoginStatus(const FUniqueNetId& UserId) const override;
 	virtual FString GetPlayerNickname(int32 LocalUserNum) const override;
+	virtual FString GetPlayerNickname(const FUniqueNetId& UserId) const override;
 	virtual FString GetAuthToken(int32 LocalUserNum) const override;
 };
 
