Add overloads for OSS session functions that take an FUniqueNetId instead of an integer LocalPlayerNum.
Update ShooterGame to use these overloads.

diff --git a/Engine/Source/Runtime/Online/IOS/OnlineSubsystemIOS/Private/OnlineSessionInterfaceIOS.cpp b/Engine/Source/Runtime/Online/IOS/OnlineSubsystemIOS/Private/OnlineSessionInterfaceIOS.cpp
--- a/Engine/Source/Runtime/Online/IOS/OnlineSubsystemIOS/Private/OnlineSessionInterfaceIOS.cpp
+++ b/Engine/Source/Runtime/Online/IOS/OnlineSubsystemIOS/Private/OnlineSessionInterfaceIOS.cpp
@@ -231,6 +231,13 @@ bool FOnlineSessionIOS::CreateSession(int32 HostingPlayerNum, FName SessionName,
 }
 
 
+bool FOnlineSessionIOS::CreateSession(const FUniqueNetId& HostingPlayerId, FName SessionName, const FOnlineSessionSettings& NewSessionSettings)
+{
+	// HostingPlayerNum is unused, can pass in anything
+	return CreateSession(0, SessionName, NewSessionSettings);
+}
+
+
 bool FOnlineSessionIOS::StartSession(FName SessionName)
 {
 	bool bSuccessfullyStartedSession = false;
@@ -319,6 +326,15 @@ bool FOnlineSessionIOS::StartMatchmaking(int32 SearchingPlayerNum, FName Session
 	return false;
 }
 
+
+bool FOnlineSessionIOS::StartMatchmaking(const FUniqueNetId& SearchingPlayerId, FName SessionName, const FOnlineSessionSettings& NewSessionSettings, TSharedRef<FOnlineSessionSearch>& SearchSettings)
+{
+	UE_LOG(LogOnline, Warning, TEXT("Matchmaking is not supported on this platform."));
+	TriggerOnMatchmakingCompleteDelegates(SessionName, false);
+	return false;
+}
+
+
 bool FOnlineSessionIOS::CancelMatchmaking(int32 SearchingPlayerNum, FName SessionName)
 {
 	UE_LOG(LogOnline, Warning, TEXT("Matchmaking is not supported on this platform."));
@@ -326,6 +342,14 @@ bool FOnlineSessionIOS::CancelMatchmaking(int32 SearchingPlayerNum, FName Sessio
 	return false;
 }
 
+
+bool FOnlineSessionIOS::CancelMatchmaking(const FUniqueNetId& SearchingPlayerId, FName SessionName, const FOnlineSessionSettings& NewSessionSettings, TSharedRef<FOnlineSessionSearch>& SearchSettings)
+{
+	UE_LOG(LogOnline, Warning, TEXT("Matchmaking is not supported on this platform."));
+	TriggerOnCancelMatchmakingCompleteDelegates(SessionName, false);
+	return false;
+}
+
 bool FOnlineSessionIOS::FindSessions(int32 SearchingPlayerNum, const TSharedRef<FOnlineSessionSearch>& SearchSettings)
 {
 	bool bSuccessfullyFoundSessions = false;
@@ -347,6 +371,13 @@ bool FOnlineSessionIOS::FindSessions(int32 SearchingPlayerNum, const TSharedRef<
 }
 
 
+bool FOnlineSessionIOS::FindSessions(const FUniqueNetId& SearchingPlayerId, const TSharedRef<FOnlineSessionSearch>& SearchSettings)
+{
+	// SearchingPlayerNum is unused, can pass in anything
+	return FindSessions(0, SearchSettings);
+}
+
+
 bool FOnlineSessionIOS::CancelFindSessions()
 {
 	bool bSuccessfullyCancelledSession = false;
@@ -389,6 +420,13 @@ bool FOnlineSessionIOS::JoinSession(int32 PlayerNum, FName SessionName, const FO
 }
 
 
+bool FOnlineSessionIOS::JoinSession(const FUniqueNetId& PlayerId, FName SessionName, const FOnlineSessionSearchResult& DesiredSession)
+{
+	// PlayerNum is unused, can pass in anything
+	return JoinSession(0, SessionName, DesiredSession);
+}
+
+
 bool FOnlineSessionIOS::FindFriendSession(int32 LocalUserNum, const FUniqueNetId& Friend)
 {
 	bool bSuccessfullyJointFriendSession = false;
@@ -402,6 +440,12 @@ bool FOnlineSessionIOS::FindFriendSession(int32 LocalUserNum, const FUniqueNetId
 }
 
 
+bool FOnlineSessionIOS::FindFriendSession(const FUniqueNetId& LocalUserId, const FUniqueNetId& Friend)
+{
+	return FindFriendSession(0, Friend);
+}
+
+
 bool FOnlineSessionIOS::SendSessionInviteToFriend(int32 LocalUserNum, FName SessionName, const FUniqueNetId& Friend)
 {
 	bool bSuccessfullySentSessionInviteToFriend = false;
@@ -412,6 +456,14 @@ bool FOnlineSessionIOS::SendSessionInviteToFriend(int32 LocalUserNum, FName Sess
 }
 
 
+bool FOnlineSessionIOS::SendSessionInviteToFriend(const FUniqueNetId& LocalUserId, FName SessionName, const FUniqueNetId& Friend)
+{
+	UE_LOG(LogOnline, Display, TEXT("FOnlineSessionIOS::SendSessionInviteToFriend - not implemented"));
+	
+	return false;
+}
+
+
 bool FOnlineSessionIOS::SendSessionInviteToFriends(int32 LocalUserNum, FName SessionName, const TArray< TSharedRef<FUniqueNetId> >& Friends)
 {
 	bool bSuccessfullySentSessionInviteToFriends = false;
@@ -422,6 +474,14 @@ bool FOnlineSessionIOS::SendSessionInviteToFriends(int32 LocalUserNum, FName Ses
 }
 
 
+bool FOnlineSessionIOS::SendSessionInviteToFriends(const FUniqueNetId& LocalUserId, FName SessionName, const TArray< TSharedRef<FUniqueNetId> >& Friends)
+{
+	UE_LOG(LogOnline, Display, TEXT("FOnlineSessionIOS::SendSessionInviteToFriends - not implemented"));
+	
+	return false;
+}
+
+
 bool FOnlineSessionIOS::GetResolvedConnectString(FName SessionName, FString& ConnectInfo)
 {
 	bool bSuccessfullyGotResolvedConnectString = false;
diff --git a/Engine/Source/Runtime/Online/IOS/OnlineSubsystemIOS/Private/OnlineSessionInterfaceIOS.h b/Engine/Source/Runtime/Online/IOS/OnlineSubsystemIOS/Private/OnlineSessionInterfaceIOS.h
--- a/Engine/Source/Runtime/Online/IOS/OnlineSubsystemIOS/Private/OnlineSessionInterfaceIOS.h
+++ b/Engine/Source/Runtime/Online/IOS/OnlineSubsystemIOS/Private/OnlineSessionInterfaceIOS.h
@@ -80,6 +80,8 @@ public:
 
 	virtual bool CreateSession(int32 HostingPlayerNum, FName SessionName, const FOnlineSessionSettings& NewSessionSettings) override;
 
+	virtual bool CreateSession(const FUniqueNetId& HostingPlayerId, FName SessionName, const FOnlineSessionSettings& NewSessionSettings) override;
+
 	virtual bool StartSession(FName SessionName) override;
 
 	virtual bool UpdateSession(FName SessionName, FOnlineSessionSettings& UpdatedSessionSettings, bool bShouldRefreshOnlineData = false) override;
@@ -92,22 +94,36 @@ public:
 
 	virtual bool StartMatchmaking(int32 SearchingPlayerNum, FName SessionName, const FOnlineSessionSettings& NewSessionSettings, TSharedRef<FOnlineSessionSearch>& SearchSettings) override;
 
+	virtual bool StartMatchmaking(const FUniqueNetId& SearchingPlayerId, FName SessionName, const FOnlineSessionSettings& NewSessionSettings, TSharedRef<FOnlineSessionSearch>& SearchSettings) override;
+
 	virtual bool CancelMatchmaking(int32 SearchingPlayerNum, FName SessionName) override;
 
+	virtual bool CancelMatchmaking(const FUniqueNetId& SearchingPlayerId, FName SessionName) override;
+
 	virtual bool FindSessions(int32 SearchingPlayerNum, const TSharedRef<FOnlineSessionSearch>& SearchSettings) override;
 
+	virtual bool FindSessions(const FUniqueNetId& SearchingPlayerId, const TSharedRef<FOnlineSessionSearch>& SearchSettings) override;
+
 	virtual bool CancelFindSessions() override;
 
 	virtual bool PingSearchResults(const FOnlineSessionSearchResult& SearchResult) override;
 
 	virtual bool JoinSession(int32 PlayerNum, FName SessionName, const FOnlineSessionSearchResult& DesiredSession) override;
 
+	virtual bool JoinSession(const FUniqueNetId& PlayerId, FName SessionName, const FOnlineSessionSearchResult& DesiredSession) override;
+
 	virtual bool FindFriendSession(int32 LocalUserNum, const FUniqueNetId& Friend) override;
 
+	virtual bool FindFriendSession(const FUniqueNetId& LocalUserId, const FUniqueNetId& Friend) override;
+
 	virtual bool SendSessionInviteToFriend(int32 LocalUserNum, FName SessionName, const FUniqueNetId& Friend) override;
 
+	virtual bool SendSessionInviteToFriend(const FUniqueNetId& LocalUserId, FName SessionName, const FUniqueNetId& Friend) override;
+
 	virtual bool SendSessionInviteToFriends(int32 LocalUserNum, FName SessionName, const TArray< TSharedRef<FUniqueNetId> >& Friends) override;
 
+	virtual bool SendSessionInviteToFriends(const FUniqueNetId& LocalUserId, FName SessionName, const TArray< TSharedRef<FUniqueNetId> >& Friends) override;
+
 	virtual bool GetResolvedConnectString(FName SessionName, FString& ConnectInfo) override;
 
 	virtual bool GetResolvedConnectString(const class FOnlineSessionSearchResult& SearchResult, FName PortType, FString& ConnectInfo) override;
diff --git a/Engine/Source/Runtime/Online/OnlineSubsystem/Public/Interfaces/OnlineSessionInterface.h b/Engine/Source/Runtime/Online/OnlineSubsystem/Public/Interfaces/OnlineSessionInterface.h
--- a/Engine/Source/Runtime/Online/OnlineSubsystem/Public/Interfaces/OnlineSessionInterface.h
+++ b/Engine/Source/Runtime/Online/OnlineSubsystem/Public/Interfaces/OnlineSessionInterface.h
@@ -228,6 +228,19 @@ public:
 	 */
 	virtual bool CreateSession(int32 HostingPlayerNum, FName SessionName, const FOnlineSessionSettings& NewSessionSettings) = 0;
 
+	/**
+	 * Creates an online session based upon the settings object specified.
+	 * NOTE: online session registration is an async process and does not complete
+	 * until the OnCreateSessionComplete delegate is called.
+	 *
+	 * @param HostingPlayerId the index of the player hosting the session
+	 * @param SessionName the name to use for this session so that multiple sessions can exist at the same time
+	 * @param NewSessionSettings the settings to use for the new session
+	 *
+	 * @return true if successful creating the session, false otherwise
+	 */
+	virtual bool CreateSession(const FUniqueNetId& HostingPlayerId, FName SessionName, const FOnlineSessionSettings& NewSessionSettings) = 0;
+
 	/**
 	* Delegate fired when a session create request has completed
 	*
@@ -328,6 +341,19 @@ public:
 	 */
 	virtual bool StartMatchmaking(int32 SearchingPlayerNum, FName SessionName, const FOnlineSessionSettings& NewSessionSettings, TSharedRef<FOnlineSessionSearch>& SearchSettings) = 0;
 
+	/**
+	 * Begins cloud based matchmaking for a session
+	 *
+	 * @param SearchingPlayerId the id of the player searching for a match
+	 * @param SessionName the name of the session to use, usually will be GAME_SESSION_NAME
+	 * @param NewSessionSettings the desired settings to match against or create with when forming new sessions
+	 * @param SearchSettings the desired settings that the matched session will have
+	 *
+	 * @return true if successful searching for sessions, false otherwise
+	 */
+	virtual bool StartMatchmaking(const FUniqueNetId& SearchingPlayerId, FName SessionName, const FOnlineSessionSettings& NewSessionSettings, TSharedRef<FOnlineSessionSearch>& SearchSettings) = 0;
+
+
 	/**
 	 * Delegate fired when the cloud matchmaking has completed
 	 *
@@ -344,6 +370,14 @@ public:
 	 */ 
 	virtual bool CancelMatchmaking(int32 SearchingPlayerNum, FName SessionName) = 0;
 
+	/**
+	 * Cancel a Matchmaking request for a given session name
+	 *
+	 * @param SearchingPlayerId the id of the player canceling the search
+	 * @param SessionName the name of the session that was passed to StartMatchmaking (or CreateSession)
+	 */ 
+	virtual bool CancelMatchmaking(const FUniqueNetId& SearchingPlayerId, FName SessionName) = 0;
+
 	/**
 	 * Delegate fired when the cloud matchmaking has been canceled
 	 *
@@ -362,6 +396,16 @@ public:
 	 */
 	virtual bool FindSessions(int32 SearchingPlayerNum, const TSharedRef<FOnlineSessionSearch>& SearchSettings) = 0;
 
+	/**
+	 * Searches for sessions matching the settings specified
+	 *
+	 * @param SearchingPlayerId the id of the player searching for a match
+	 * @param SearchSettings the desired settings that the returned sessions will have
+	 *
+	 * @return true if successful searching for sessions, false otherwise
+	 */
+	virtual bool FindSessions(const FUniqueNetId& SearchingPlayerId, const TSharedRef<FOnlineSessionSearch>& SearchSettings) = 0;
+
 	/**
 	 * Delegate fired when the search for an online session has completed
 	 *
@@ -411,6 +455,18 @@ public:
 	 */
 	virtual bool JoinSession(int32 LocalUserNum, FName SessionName, const FOnlineSessionSearchResult& DesiredSession) = 0;
 
+	/**
+	 * Joins the session specified
+	 *
+	 * @param LocalUserId the id of the player searching for a match
+	 * @param SessionName the name of the session to join
+	 * @param DesiredSession the desired session to join
+	 *
+	 * @return true if the call completed successfully, false otherwise
+	 */
+	virtual bool JoinSession(const FUniqueNetId& LocalUserId, FName SessionName, const FOnlineSessionSearchResult& DesiredSession) = 0;
+
+
 	/**
 	 * Delegate fired when the joining process for an online session has completed
 	 *
@@ -429,6 +485,16 @@ public:
 	 */
 	virtual bool FindFriendSession(int32 LocalUserNum, const FUniqueNetId& Friend) = 0;
 
+	/**
+	 * Allows the local player to follow a friend into a session
+	 *
+	 * @param LocalUserId the local player wanting to join
+	 * @param Friend the player that is being followed
+	 *
+	 * @return true if the async call worked, false otherwise
+	 */
+	virtual bool FindFriendSession(const FUniqueNetId& LocalUserId, const FUniqueNetId& Friend) = 0;
+
 	/**
 	 * Delegate fired once the find friend task has completed
 	 * Session has not been joined at this point, and still requires a call to JoinSession()
@@ -450,6 +516,17 @@ public:
 	 */
 	virtual bool SendSessionInviteToFriend(int32 LocalUserNum, FName SessionName, const FUniqueNetId& Friend) = 0;
 
+	/**
+	 * Sends an invitation to play in the player's current session
+	 *
+	 * @param LocalUserId the user that is sending the invite
+	 * @param SessionName session to invite them to
+	 * @param Friend the player to send the invite to
+	 *
+	 * @return true if successful, false otherwise
+	 */
+	virtual bool SendSessionInviteToFriend(const FUniqueNetId& LocalUserId, FName SessionName, const FUniqueNetId& Friend) = 0;
+
 	/**
 	 * Sends invitations to play in the player's current session
 	 *
@@ -461,6 +538,17 @@ public:
 	 */
 	virtual bool SendSessionInviteToFriends(int32 LocalUserNum, FName SessionName, const TArray< TSharedRef<FUniqueNetId> >& Friends) = 0;
 
+	/**
+	 * Sends invitations to play in the player's current session
+	 *
+	 * @param LocalUserId the user that is sending the invite
+	 * @param SessionName session to invite them to
+	 * @param Friends the player to send the invite to
+	 *
+	 * @return true if successful, false otherwise
+	 */
+	virtual bool SendSessionInviteToFriends(const FUniqueNetId& LocalUserId, FName SessionName, const TArray< TSharedRef<FUniqueNetId> >& Friends) = 0;
+
 	/**
 	 * Called when a user accepts a session invitation. Allows the game code a chance
 	 * to clean up any existing state before accepting the invite. The invite must be
diff --git a/Engine/Source/Runtime/Online/OnlineSubsystemNull/Private/OnlineSessionInterfaceNull.cpp b/Engine/Source/Runtime/Online/OnlineSubsystemNull/Private/OnlineSessionInterfaceNull.cpp
--- a/Engine/Source/Runtime/Online/OnlineSubsystemNull/Private/OnlineSessionInterfaceNull.cpp
+++ b/Engine/Source/Runtime/Online/OnlineSubsystemNull/Private/OnlineSessionInterfaceNull.cpp
@@ -235,6 +235,12 @@ bool FOnlineSessionNull::CreateSession(int32 HostingPlayerNum, FName SessionName
 	return Result == ERROR_IO_PENDING || Result == ERROR_SUCCESS;
 }
 
+bool FOnlineSessionNull::CreateSession(const FUniqueNetId& HostingPlayerId, FName SessionName, const FOnlineSessionSettings& NewSessionSettings)
+{
+	// todo: use proper	HostingPlayerId
+	return CreateSession(0, SessionName, NewSessionSettings);
+}
+
 bool FOnlineSessionNull::NeedsToAdvertise()
 {
 	FScopeLock ScopeLock(&SessionLock);
@@ -442,6 +448,13 @@ bool FOnlineSessionNull::StartMatchmaking(int32 SearchingPlayerNum, FName Sessio
 	return false;
 }
 
+bool FOnlineSessionNull::StartMatchmaking(const FUniqueNetId& SearchingPlayerId, FName SessionName, const FOnlineSessionSettings& NewSessionSettings, TSharedRef<FOnlineSessionSearch>& SearchSettings)
+{
+	UE_LOG(LogOnline, Warning, TEXT("Matchmaking is not supported on this platform."));
+	TriggerOnMatchmakingCompleteDelegates(SessionName, false);
+	return false;
+}
+
 bool FOnlineSessionNull::CancelMatchmaking(int32 SearchingPlayerNum, FName SessionName)
 {
 	UE_LOG(LogOnline, Warning, TEXT("Matchmaking is not supported on this platform."));
@@ -449,6 +462,13 @@ bool FOnlineSessionNull::CancelMatchmaking(int32 SearchingPlayerNum, FName Sessi
 	return false;
 }
 
+bool FOnlineSessionNull::CancelMatchmaking(const FUniqueNetId& SearchingPlayerId, FName SessionName)
+{
+	UE_LOG(LogOnline, Warning, TEXT("Matchmaking is not supported on this platform."));
+	TriggerOnCancelMatchmakingCompleteDelegates(SessionName, false);
+	return false;
+}
+
 bool FOnlineSessionNull::FindSessions(int32 SearchingPlayerNum, const TSharedRef<FOnlineSessionSearch>& SearchSettings)
 {
 	uint32 Return = E_FAIL;
@@ -482,6 +502,12 @@ bool FOnlineSessionNull::FindSessions(int32 SearchingPlayerNum, const TSharedRef
 	return Return == ERROR_SUCCESS || Return == ERROR_IO_PENDING;
 }
 
+bool FOnlineSessionNull::FindSessions(const FUniqueNetId& SearchingPlayerId, const TSharedRef<FOnlineSessionSearch>& SearchSettings)
+{
+	// This function doesn't use the SearchingPlayerNum parameter, so passing in anything is fine.
+	return FindSessions(0, SearchSettings);
+}
+
 uint32 FOnlineSessionNull::FindLANSession()
 {
 	uint32 Return = ERROR_IO_PENDING;
@@ -586,6 +612,12 @@ bool FOnlineSessionNull::JoinSession(int32 PlayerNum, FName SessionName, const F
 	return Return == ERROR_SUCCESS || Return == ERROR_IO_PENDING;
 }
 
+bool FOnlineSessionNull::JoinSession(const FUniqueNetId& PlayerId, FName SessionName, const FOnlineSessionSearchResult& DesiredSession)
+{
+	// Assuming player 0 should be OK here
+	return JoinSession(0, SessionName, DesiredSession);
+}
+
 bool FOnlineSessionNull::FindFriendSession(int32 LocalUserNum, const FUniqueNetId& Friend)
 {
 	// this function has to exist due to interface definition, but it does not have a meaningful implementation in Null subsystem
@@ -594,18 +626,38 @@ bool FOnlineSessionNull::FindFriendSession(int32 LocalUserNum, const FUniqueNetI
 	return false;
 };
 
+bool FOnlineSessionNull::FindFriendSession(const FUniqueNetId& LocalUserId, const FUniqueNetId& Friend)
+{
+	// this function has to exist due to interface definition, but it does not have a meaningful implementation in Null subsystem
+	FOnlineSessionSearchResult EmptySearchResult;
+	TriggerOnFindFriendSessionCompleteDelegates(0, false, EmptySearchResult);
+	return false;
+}
+
 bool FOnlineSessionNull::SendSessionInviteToFriend(int32 LocalUserNum, FName SessionName, const FUniqueNetId& Friend)
 {
 	// this function has to exist due to interface definition, but it does not have a meaningful implementation in Null subsystem
 	return false;
 };
 
+bool FOnlineSessionNull::SendSessionInviteToFriend(const FUniqueNetId& LocalUserId, FName SessionName, const FUniqueNetId& Friend)
+{
+	// this function has to exist due to interface definition, but it does not have a meaningful implementation in Null subsystem
+	return false;
+}
+
 bool FOnlineSessionNull::SendSessionInviteToFriends(int32 LocalUserNum, FName SessionName, const TArray< TSharedRef<FUniqueNetId> >& Friends)
 {
 	// this function has to exist due to interface definition, but it does not have a meaningful implementation in Null subsystem
 	return false;
 };
 
+bool FOnlineSessionNull::SendSessionInviteToFriends(const FUniqueNetId& LocalUserId, FName SessionName, const TArray< TSharedRef<FUniqueNetId> >& Friends)
+{
+	// this function has to exist due to interface definition, but it does not have a meaningful implementation in Null subsystem
+	return false;
+}
+
 uint32 FOnlineSessionNull::JoinLANSession(int32 PlayerNum, FNamedOnlineSession* Session, const FOnlineSession* SearchSession)
 {
 	uint32 Result = E_FAIL;
diff --git a/Engine/Source/Runtime/Online/OnlineSubsystemNull/Private/OnlineSessionInterfaceNull.h b/Engine/Source/Runtime/Online/OnlineSubsystemNull/Private/OnlineSessionInterfaceNull.h
--- a/Engine/Source/Runtime/Online/OnlineSubsystemNull/Private/OnlineSessionInterfaceNull.h
+++ b/Engine/Source/Runtime/Online/OnlineSubsystemNull/Private/OnlineSessionInterfaceNull.h
@@ -264,20 +264,28 @@ public:
 
 	// IOnlineSession
 	virtual bool CreateSession(int32 HostingPlayerNum, FName SessionName, const FOnlineSessionSettings& NewSessionSettings) override;
+	virtual bool CreateSession(const FUniqueNetId& HostingPlayerId, FName SessionName, const FOnlineSessionSettings& NewSessionSettings) override;
 	virtual bool StartSession(FName SessionName) override;
 	virtual bool UpdateSession(FName SessionName, FOnlineSessionSettings& UpdatedSessionSettings, bool bShouldRefreshOnlineData = true) override;
 	virtual bool EndSession(FName SessionName) override;
 	virtual bool DestroySession(FName SessionName) override;
 	virtual bool IsPlayerInSession(FName SessionName, const FUniqueNetId& UniqueId) override;
 	virtual bool StartMatchmaking(int32 SearchingPlayerNum, FName SessionName, const FOnlineSessionSettings& NewSessionSettings, TSharedRef<FOnlineSessionSearch>& SearchSettings) override;
+	virtual bool StartMatchmaking(const FUniqueNetId& SearchingPlayerId, FName SessionName, const FOnlineSessionSettings& NewSessionSettings, TSharedRef<FOnlineSessionSearch>& SearchSettings) override;
 	virtual bool CancelMatchmaking(int32 SearchingPlayerNum, FName SessionName) override;
+	virtual bool CancelMatchmaking(const FUniqueNetId& SearchingPlayerId, FName SessionName) override;
 	virtual bool FindSessions(int32 SearchingPlayerNum, const TSharedRef<FOnlineSessionSearch>& SearchSettings) override;
+	virtual bool FindSessions(const FUniqueNetId& SearchingPlayerId, const TSharedRef<FOnlineSessionSearch>& SearchSettings) override;
 	virtual bool CancelFindSessions() override;
 	virtual bool PingSearchResults(const FOnlineSessionSearchResult& SearchResult) override;
 	virtual bool JoinSession(int32 PlayerNum, FName SessionName, const FOnlineSessionSearchResult& DesiredSession) override;
+	virtual bool JoinSession(const FUniqueNetId& PlayerId, FName SessionName, const FOnlineSessionSearchResult& DesiredSession) override;
 	virtual bool FindFriendSession(int32 LocalUserNum, const FUniqueNetId& Friend) override;
+	virtual bool FindFriendSession(const FUniqueNetId& LocalUserId, const FUniqueNetId& Friend) override;
 	virtual bool SendSessionInviteToFriend(int32 LocalUserNum, FName SessionName, const FUniqueNetId& Friend) override;
+	virtual bool SendSessionInviteToFriend(const FUniqueNetId& LocalUserId, FName SessionName, const FUniqueNetId& Friend) override;
 	virtual bool SendSessionInviteToFriends(int32 LocalUserNum, FName SessionName, const TArray< TSharedRef<FUniqueNetId> >& Friends) override;
+	virtual bool SendSessionInviteToFriends(const FUniqueNetId& LocalUserId, FName SessionName, const TArray< TSharedRef<FUniqueNetId> >& Friends) override;
 	virtual bool GetResolvedConnectString(FName SessionName, FString& ConnectInfo) override;
 	virtual bool GetResolvedConnectString(const class FOnlineSessionSearchResult& SearchResult, FName PortType, FString& ConnectInfo) override;
 	virtual FOnlineSessionSettings* GetSessionSettings(FName SessionName) override;
diff --git a/Engine/Source/Runtime/Online/OnlineSubsystemSteam/Private/OnlineSessionInterfaceSteam.cpp b/Engine/Source/Runtime/Online/OnlineSubsystemSteam/Private/OnlineSessionInterfaceSteam.cpp
--- a/Engine/Source/Runtime/Online/OnlineSubsystemSteam/Private/OnlineSessionInterfaceSteam.cpp
+++ b/Engine/Source/Runtime/Online/OnlineSubsystemSteam/Private/OnlineSessionInterfaceSteam.cpp
@@ -270,6 +270,12 @@ bool FOnlineSessionSteam::CreateSession(int32 HostingPlayerNum, FName SessionNam
 	return Result == ERROR_IO_PENDING || Result == ERROR_SUCCESS;
 }
 
+bool FOnlineSessionSteam::CreateSession(const FUniqueNetId& HostingPlayerId, FName SessionName, const FOnlineSessionSettings& NewSessionSettings)
+{
+	// @todo: use proper HostingPlayerId
+	return CreateSession(0, SessionName, NewSessionSettings);
+}
+
 uint32 FOnlineSessionSteam::CreateLobbySession(int32 HostingPlayerNum, FNamedOnlineSession* Session)
 {
 	uint32 Result = E_FAIL; 
@@ -641,6 +647,13 @@ bool FOnlineSessionSteam::StartMatchmaking(int32 SearchingPlayerNum, FName Sessi
 	return false;
 }
 
+bool FOnlineSessionSteam::StartMatchmaking(const FUniqueNetId& SearchingPlayerId, FName SessionName, const FOnlineSessionSettings& NewSessionSettings, TSharedRef<FOnlineSessionSearch>& SearchSettings)
+{
+	UE_LOG(LogOnline, Warning, TEXT("Matchmaking is not supported on this platform."));
+	TriggerOnMatchmakingCompleteDelegates(SessionName, false);
+	return false;
+}
+
 bool FOnlineSessionSteam::CancelMatchmaking(int32 SearchingPlayerNum, FName SessionName)
 {
 	UE_LOG(LogOnline, Warning, TEXT("Matchmaking is not supported on this platform."));
@@ -648,6 +661,13 @@ bool FOnlineSessionSteam::CancelMatchmaking(int32 SearchingPlayerNum, FName Sess
 	return false;
 }
 
+bool FOnlineSessionSteam::CancelMatchmaking(const FUniqueNetId& SearchingPlayerId, FName SessionName)
+{
+	UE_LOG(LogOnline, Warning, TEXT("Matchmaking is not supported on this platform."));
+	TriggerOnCancelMatchmakingCompleteDelegates(SessionName, false);
+	return false;
+}
+
 bool FOnlineSessionSteam::FindSessions(int32 SearchingPlayerNum, const TSharedRef<FOnlineSessionSearch>& SearchSettings)
 {
 	uint32 Return = E_FAIL;
@@ -685,6 +705,12 @@ bool FOnlineSessionSteam::FindSessions(int32 SearchingPlayerNum, const TSharedRe
 	return Return == ERROR_SUCCESS || Return == ERROR_IO_PENDING;
 }
 
+bool FOnlineSessionSteam::FindSessions(const FUniqueNetId& SearchingPlayerId, const TSharedRef<FOnlineSessionSearch>& SearchSettings)
+{
+	// @todo: use proper SearchingPlayerId
+	return FindSessions(0, SearchSettings);
+}
+
 uint32 FOnlineSessionSteam::FindInternetSession(const TSharedRef<FOnlineSessionSearch>& SearchSettings)
 {
 	bool PresenceSearch = false;
@@ -841,6 +867,12 @@ bool FOnlineSessionSteam::JoinSession(int32 PlayerNum, FName SessionName, const
 	return Return == ERROR_SUCCESS || Return == ERROR_IO_PENDING;
 }
 
+bool FOnlineSessionSteam::JoinSession(const FUniqueNetId& PlayerId, FName SessionName, const FOnlineSessionSearchResult& DesiredSession)
+{
+	// @todo: use proper PlayerId
+	return JoinSession(0, SessionName, DesiredSession);
+}
+
 uint32 FOnlineSessionSteam::JoinLobbySession(int32 PlayerNum, FNamedOnlineSession* Session, const FOnlineSession* SearchSession)
 {
 	uint32 Result = E_FAIL;
@@ -959,6 +991,12 @@ bool FOnlineSessionSteam::FindFriendSession(int32 LocalUserNum, const FUniqueNet
 	return bSuccess;
 }
 
+bool FOnlineSessionSteam::FindFriendSession(const FUniqueNetId& LocalUserId, const FUniqueNetId& Friend)
+{
+	// @todo: use proper LocalUserId
+	return FindFriendSession(0, Friend);
+}
+
 bool FOnlineSessionSteam::PingSearchResults(const FOnlineSessionSearchResult& SearchResult)
 {
 	return false;
@@ -1009,6 +1047,12 @@ bool FOnlineSessionSteam::SendSessionInviteToFriend(int32 LocalUserNum, FName Se
 	return SendSessionInviteToFriends(LocalUserNum, SessionName, Friends);
 }
 
+bool FOnlineSessionSteam::SendSessionInviteToFriend(const FUniqueNetId& LocalUserId, FName SessionName, const FUniqueNetId& Friend)
+{
+	// @todo: use proper LocalUserId
+	return SendSessionInviteToFriend(0, SessionName, Friend);
+}
+
 bool FOnlineSessionSteam::SendSessionInviteToFriends(int32 LocalUserNum, FName SessionName, const TArray< TSharedRef<FUniqueNetId> >& Friends)
 {
 	bool bSuccess = false;
@@ -1069,6 +1113,12 @@ bool FOnlineSessionSteam::SendSessionInviteToFriends(int32 LocalUserNum, FName S
 	return bSuccess;
 }
 
+bool FOnlineSessionSteam::SendSessionInviteToFriends(const FUniqueNetId& LocalUserId, FName SessionName, const TArray< TSharedRef<FUniqueNetId> >& Friends)
+{
+	// @todo: use proper LocalUserId
+	return SendSessionInviteToFriends(0, SessionName, Friends);
+}
+
 FString FOnlineSessionSteam::GetSteamConnectionString(FName SessionName)
 {
 	FString ConnectionString;
diff --git a/Engine/Source/Runtime/Online/OnlineSubsystemSteam/Private/OnlineSessionInterfaceSteam.h b/Engine/Source/Runtime/Online/OnlineSubsystemSteam/Private/OnlineSessionInterfaceSteam.h
--- a/Engine/Source/Runtime/Online/OnlineSubsystemSteam/Private/OnlineSessionInterfaceSteam.h
+++ b/Engine/Source/Runtime/Online/OnlineSubsystemSteam/Private/OnlineSessionInterfaceSteam.h
@@ -490,20 +490,28 @@ public:
 
 	// IOnlineSession
 	virtual bool CreateSession(int32 HostingPlayerNum, FName SessionName, const FOnlineSessionSettings& NewSessionSettings) override;
+	virtual bool CreateSession(const FUniqueNetId& HostingPlayerId, FName SessionName, const FOnlineSessionSettings& NewSessionSettings) override;
 	virtual bool StartSession(FName SessionName) override;
 	virtual bool UpdateSession(FName SessionName, FOnlineSessionSettings& UpdatedSessionSettings, bool bShouldRefreshOnlineData = true) override;
 	virtual bool EndSession(FName SessionName) override;
 	virtual bool DestroySession(FName SessionName) override;
 	virtual bool IsPlayerInSession(FName SessionName, const FUniqueNetId& UniqueId) override;
 	virtual bool StartMatchmaking(int32 SearchingPlayerNum, FName SessionName, const FOnlineSessionSettings& NewSessionSettings, TSharedRef<FOnlineSessionSearch>& SearchSettings) override;
+	virtual bool StartMatchmaking(const FUniqueNetId& SearchingPlayerId, FName SessionName, const FOnlineSessionSettings& NewSessionSettings, TSharedRef<FOnlineSessionSearch>& SearchSettings) override;
 	virtual bool CancelMatchmaking(int32 SearchingPlayerNum, FName SessionName) override;
+	virtual bool CancelMatchmaking(const FUniqueNetId& SearchingPlayerId, FName SessionName) override;
 	virtual bool FindSessions(int32 SearchingPlayerNum, const TSharedRef<FOnlineSessionSearch>& SearchSettings) override;
+	virtual bool FindSessions(const FUniqueNetId& SearchingPlayerId, const TSharedRef<FOnlineSessionSearch>& SearchSettings) override;
 	virtual bool CancelFindSessions() override;
 	virtual bool PingSearchResults(const FOnlineSessionSearchResult& SearchResult) override;
 	virtual bool JoinSession(int32 PlayerNum, FName SessionName, const FOnlineSessionSearchResult& DesiredSession) override;
+	virtual bool JoinSession(const FUniqueNetId& PlayerId, FName SessionName, const FOnlineSessionSearchResult& DesiredSession) override;
 	virtual bool FindFriendSession(int32 LocalUserNum, const FUniqueNetId& Friend) override;
+	virtual bool FindFriendSession(const FUniqueNetId& LocalUserId, const FUniqueNetId& Friend) override;
 	virtual bool SendSessionInviteToFriend(int32 LocalUserNum, FName SessionName, const FUniqueNetId& Friend) override;
+	virtual bool SendSessionInviteToFriend(const FUniqueNetId& LocalUserId, FName SessionName, const FUniqueNetId& Friend) override;
 	virtual bool SendSessionInviteToFriends(int32 LocalUserNum, FName SessionName, const TArray< TSharedRef<FUniqueNetId> >& Friends) override;
+	virtual bool SendSessionInviteToFriends(const FUniqueNetId& LocalUserId, FName SessionName, const TArray< TSharedRef<FUniqueNetId> >& Friends) override;
 	virtual bool GetResolvedConnectString(FName SessionName, FString& ConnectInfo) override;
 	virtual bool GetResolvedConnectString(const class FOnlineSessionSearchResult& SearchResult, FName PortType, FString& ConnectInfo) override;
 	virtual FOnlineSessionSettings* GetSessionSettings(FName SessionName) override;
