Fix copy/paste error

diff --git a/Engine/Source/Runtime/Online/IOS/OnlineSubsystemIOS/Private/OnlineSessionInterfaceIOS.cpp b/Engine/Source/Runtime/Online/IOS/OnlineSubsystemIOS/Private/OnlineSessionInterfaceIOS.cpp
--- a/Engine/Source/Runtime/Online/IOS/OnlineSubsystemIOS/Private/OnlineSessionInterfaceIOS.cpp
+++ b/Engine/Source/Runtime/Online/IOS/OnlineSubsystemIOS/Private/OnlineSessionInterfaceIOS.cpp
@@ -343,7 +343,7 @@ bool FOnlineSessionIOS::CancelMatchmaking(int32 SearchingPlayerNum, FName Sessio
 }
 
 
-bool FOnlineSessionIOS::CancelMatchmaking(const FUniqueNetId& SearchingPlayerId, FName SessionName, const FOnlineSessionSettings& NewSessionSettings, TSharedRef<FOnlineSessionSearch>& SearchSettings)
+bool FOnlineSessionIOS::CancelMatchmaking(const FUniqueNetId& SearchingPlayerId, FName SessionName)
 {
 	UE_LOG(LogOnline, Warning, TEXT("Matchmaking is not supported on this platform."));
 	TriggerOnCancelMatchmakingCompleteDelegates(SessionName, false);
