BuildPatchServices module ensures it does not complete shutdown until all module threads have exited

diff --git a/Engine/Source/Runtime/Online/BuildPatchServices/Private/BuildPatchInstaller.cpp b/Engine/Source/Runtime/Online/BuildPatchServices/Private/BuildPatchInstaller.cpp
--- a/Engine/Source/Runtime/Online/BuildPatchServices/Private/BuildPatchInstaller.cpp
+++ b/Engine/Source/Runtime/Online/BuildPatchServices/Private/BuildPatchInstaller.cpp
@@ -958,6 +958,12 @@ void FBuildPatchInstaller::ExecuteCompleteDelegate()
 	OnCompleteDelegate.Execute( bSuccess, NewBuildManifest );
 }
 
-
+void FBuildPatchInstaller::WaitForThread() const
+{
+	if (Thread != NULL)
+	{
+		Thread->WaitForCompletion();
+	}
+}
 
 #undef LOCTEXT_NAMESPACE
diff --git a/Engine/Source/Runtime/Online/BuildPatchServices/Private/BuildPatchInstaller.h b/Engine/Source/Runtime/Online/BuildPatchServices/Private/BuildPatchInstaller.h
--- a/Engine/Source/Runtime/Online/BuildPatchServices/Private/BuildPatchInstaller.h
+++ b/Engine/Source/Runtime/Online/BuildPatchServices/Private/BuildPatchInstaller.h
@@ -144,6 +144,11 @@ public:
 	 */
 	void ExecuteCompleteDelegate();
 
+	/**
+	 * Only returns once the thread has finished running
+	 */
+	void WaitForThread() const;
+
 private:
 
 	/**
diff --git a/Engine/Source/Runtime/Online/BuildPatchServices/Private/BuildPatchManifest.cpp b/Engine/Source/Runtime/Online/BuildPatchServices/Private/BuildPatchManifest.cpp
--- a/Engine/Source/Runtime/Online/BuildPatchServices/Private/BuildPatchManifest.cpp
+++ b/Engine/Source/Runtime/Online/BuildPatchServices/Private/BuildPatchManifest.cpp
@@ -859,11 +859,11 @@ void FBuildPatchAppManifest::EnumerateChunkPartInventory( const TArray< FGuid >&
 {
 	ChunkPartsAvailable.Empty();
 	// For each file in the manifest, check what chunks it is made out of, and grab details for the ones in ChunksRequired
-	for( auto FileManifestIt = FileManifestList.CreateConstIterator(); FileManifestIt; ++FileManifestIt )
+	for (auto FileManifestIt = FileManifestList.CreateConstIterator(); FileManifestIt && !FBuildPatchInstallError::HasFatalError(); ++FileManifestIt)
 	{
 		const TSharedPtr< FBuildPatchFileManifest > FileManifest = FileManifestIt.Value();
 		uint64 FileOffset = 0;
-		for( auto ChunkPartIt = FileManifest->FileChunkParts.CreateConstIterator(); ChunkPartIt; ++ChunkPartIt )
+		for (auto ChunkPartIt = FileManifest->FileChunkParts.CreateConstIterator(); ChunkPartIt && !FBuildPatchInstallError::HasFatalError(); ++ChunkPartIt)
 		{
 			const FChunkPart& ChunkPart = *ChunkPartIt;
 			if( ChunksRequired.Contains( ChunkPart.Guid ) )
diff --git a/Engine/Source/Runtime/Online/BuildPatchServices/Private/BuildPatchServicesModule.cpp b/Engine/Source/Runtime/Online/BuildPatchServices/Private/BuildPatchServicesModule.cpp
--- a/Engine/Source/Runtime/Online/BuildPatchServices/Private/BuildPatchServicesModule.cpp
+++ b/Engine/Source/Runtime/Online/BuildPatchServices/Private/BuildPatchServicesModule.cpp
@@ -79,13 +79,15 @@ void FBuildPatchServicesModule::ShutdownModule()
 	FBuildPatchAnalytics::SetAnalyticsProvider( NULL );
 
 	// Reset all installer ptrs (this will cause us to wait for each thread to complete)
-	for(auto InstallerIt = BuildPatchInstallers.CreateIterator(); InstallerIt; ++InstallerIt)
+	for (auto& BuildPatchInstaller : BuildPatchInstallers)
 	{
 		// Make sure it is not paused, this function only un-pauses when in error state
-		(*InstallerIt)->TogglePauseInstall();
-		// And reset
-		(*InstallerIt).Reset();
+		BuildPatchInstaller->TogglePauseInstall();
+		// We still have to manually wait for the thread as another system could hold a shared ptr
+		// thus we would not be calling the destructor here
+		BuildPatchInstaller->WaitForThread();
 	}
+	BuildPatchInstallers.Empty();
 
 	// Remove our ticker
 	FTicker::GetCoreTicker().RemoveTicker( FTickerDelegate::CreateRaw( this, &FBuildPatchServicesModule::Tick ) );
