#uelauncher
- refactor of analytics events to better match new recommended guidelines of "Product.Category.Event"
OPP-538

Old to new mapping:
CPBootedFromLaunch		- Launcher.Error.Booted
CPWARNING.MULTIPLE.INSTANCE	- Launcher.Error.MultipleInstance
CPERROR.SELFPATCH.FAILED		- Launcher.SelfPatch.Failed
CPERROR.SELFPATCH.DELETEFILEFAIL	- Launcher.SelfPatch.Error, Op=DeleteFileFail
CPERROR.SELFPATCH.MOVEFILEFAIL	- Launcher.SelfPatch.Error, Op=MoveFileFail
CPERROR.SELFPATCH.RENAMEFILEFAIL	- Launcher.SelfPatch.Error, Op=RenameFileFail
CPLaunchApp			- Launcher.Install.Launch
CPBUILDSTART			- Launcher.Install.Start
CPBUILDSTATS.4		- Launcher.Install.Stats
Launcher.UI			- Launcher.Usage.ViewTutorial
			- Launcher.Usage.DismissTutorial
			- Launcher.Usage.QuickLaunch
			- Launcher.Usage.UserProject
			- Launcher.Usage.EngineSlot
LAUNCHER.NEWSCLICKED		- Launcher.Usage.News
CPLOGIN 			- Launcher.User.Login
BUILDPATCHERROR.CHUNK.CACHE	- Patcher.Error.Cache
BUILDPATCHERROR.CHUNK.DOWNLOAD	- Patcher.Error.Download
BUILDPATCHERROR.CONSTRUCTION	- Patcher.Error.Construction
BUILDPATCHERROR.PREREQUISITES	- Patcher.Error.Prerequisites

diff --git a/Engine/Source/Runtime/Online/BuildPatchServices/Private/BuildPatchAnalytics.cpp b/Engine/Source/Runtime/Online/BuildPatchServices/Private/BuildPatchAnalytics.cpp
--- a/Engine/Source/Runtime/Online/BuildPatchServices/Private/BuildPatchAnalytics.cpp
+++ b/Engine/Source/Runtime/Online/BuildPatchServices/Private/BuildPatchAnalytics.cpp
@@ -37,7 +37,7 @@ void FBuildPatchAnalytics::RecordChunkDownloadError( const FString& ChunkUrl, co
 		Attributes.Add( FAnalyticsEventAttribute( TEXT( "ChunkURL" ), ChunkUrl ) );
 		Attributes.Add( FAnalyticsEventAttribute( TEXT( "ResponseCode" ), ResponseCode ) );
 		Attributes.Add( FAnalyticsEventAttribute( TEXT( "ErrorString" ), ErrorString ) );
-		FBuildPatchHTTP::QueueAnalyticsEvent( TEXT( "BUILDPATCHERROR.CHUNK.DOWNLOAD" ) , Attributes );
+		FBuildPatchHTTP::QueueAnalyticsEvent( TEXT( "Patcher.Error.Download" ) , Attributes );
 	}
 }
 
@@ -51,7 +51,7 @@ void FBuildPatchAnalytics::RecordChunkCacheError( const FGuid& ChunkGuid, const
 		Attributes.Add( FAnalyticsEventAttribute( TEXT( "LastError" ), LastError ) );
 		Attributes.Add( FAnalyticsEventAttribute( TEXT( "SystemName" ), SystemName ) );
 		Attributes.Add( FAnalyticsEventAttribute( TEXT( "ErrorString" ), ErrorString ) );
-		FBuildPatchHTTP::QueueAnalyticsEvent( TEXT( "BUILDPATCHERROR.CHUNK.CACHE" ) , Attributes );
+		FBuildPatchHTTP::QueueAnalyticsEvent( TEXT( "Patcher.Error.Cache" ) , Attributes );
 	}
 }
 
@@ -63,7 +63,7 @@ void FBuildPatchAnalytics::RecordConstructionError( const FString& Filename, con
 		Attributes.Add( FAnalyticsEventAttribute( TEXT( "Filename" ), Filename ) );
 		Attributes.Add( FAnalyticsEventAttribute( TEXT( "LastError" ), LastError ) );
 		Attributes.Add( FAnalyticsEventAttribute( TEXT( "ErrorString" ), ErrorString ) );
-		FBuildPatchHTTP::QueueAnalyticsEvent( TEXT( "BUILDPATCHERROR.CONSTRUCTION" ) , Attributes );
+		FBuildPatchHTTP::QueueAnalyticsEvent( TEXT( "Patcher.Error.Construction" ) , Attributes );
 	}
 }
 
@@ -74,7 +74,7 @@ void FBuildPatchAnalytics::RecordPrereqInstallnError( const FString& Filename, c
 	Attributes.Add( FAnalyticsEventAttribute( TEXT( "CommandLine" ), CommandLine ) );
 	Attributes.Add( FAnalyticsEventAttribute( TEXT( "ErrorCode" ), ErrorCode ) );
 	Attributes.Add( FAnalyticsEventAttribute( TEXT( "ErrorString" ), ErrorString ) );
-	FBuildPatchHTTP::QueueAnalyticsEvent( TEXT( "BUILDPATCHERROR.PREREQUISITES" ) , Attributes );
+	FBuildPatchHTTP::QueueAnalyticsEvent( TEXT( "Patcher.Error.Prerequisites" ) , Attributes );
 
 }
 
