Added the OnlineFriendsInterface.h to Online.h so that UE4 code projects can access online social features

diff --git a/Engine/Source/Runtime/Online/OnlineSubsystem/Public/Online.h b/Engine/Source/Runtime/Online/OnlineSubsystem/Public/Online.h
--- a/Engine/Source/Runtime/Online/OnlineSubsystem/Public/Online.h
+++ b/Engine/Source/Runtime/Online/OnlineSubsystem/Public/Online.h
@@ -7,6 +7,7 @@
 #include "ModuleManager.h"
 
 #include "OnlineIdentityInterface.h"
+#include "OnlineFriendsInterface.h"
 #include "OnlineEventsInterface.h"
 #include "OnlineSessionInterface.h"
 #include "OnlineExternalUIInterface.h"
