#uelauncher
- fix for sending loc text for build stat analytics events

OPP-354 BuildPatchServices: Analytics event contains localised UI string (FailureReasonText) which contains UTF8 characters that cannot be read by the SQL server

diff --git a/Engine/Source/Runtime/Online/BuildPatchServices/Private/BuildPatchInstaller.cpp b/Engine/Source/Runtime/Online/BuildPatchServices/Private/BuildPatchInstaller.cpp
--- a/Engine/Source/Runtime/Online/BuildPatchServices/Private/BuildPatchInstaller.cpp
+++ b/Engine/Source/Runtime/Online/BuildPatchServices/Private/BuildPatchInstaller.cpp
@@ -199,7 +199,7 @@ uint32 FBuildPatchInstaller::Run()
 		GLog->Logf(TEXT("BuildPatchServices: Build Stat: ProcessPausedTime: %.1f sec"), BuildStats.ProcessPausedTime);
 		GLog->Logf(TEXT("BuildPatchServices: Build Stat: ProcessSuccess: %s"), BuildStats.ProcessSuccess ? TEXT("TRUE") : TEXT("FALSE"));
 		GLog->Logf(TEXT("BuildPatchServices: Build Stat: FailureReason: %s"), *BuildStats.FailureReason);
-		GLog->Logf(TEXT("BuildPatchServices: Build Stat: FailureReasonText: %s"), *BuildStats.FailureReasonText.ToString());
+		GLog->Logf(TEXT("BuildPatchServices: Build Stat: FailureReasonText: %s"), *BuildStats.FailureReasonText.BuildSourceString());
 	}
 
 	// Mark that we are done
