FHttpRequestWinInet - Fix the SetContentAsString method to use the converted string length instead of the source string length, to avoid truncating strings which actually use wider characters.

diff --git a/Engine/Source/Runtime/Online/HTTP/Private/Windows/HttpWinInet.cpp b/Engine/Source/Runtime/Online/HTTP/Private/Windows/HttpWinInet.cpp
--- a/Engine/Source/Runtime/Online/HTTP/Private/Windows/HttpWinInet.cpp
+++ b/Engine/Source/Runtime/Online/HTTP/Private/Windows/HttpWinInet.cpp
@@ -417,7 +417,7 @@ void FHttpRequestWinInet::SetContent(const TArray<uint8>& ContentPayload)
 void FHttpRequestWinInet::SetContentAsString(const FString& ContentString)
 {
 	FTCHARToUTF8 Converter(*ContentString);
-	RequestPayload.SetNumUninitialized(ContentString.Len());
+	RequestPayload.SetNumUninitialized(Converter.Length());
 	FMemory::Memcpy(RequestPayload.GetTypedData(), (const uint8*)Converter.Get(), RequestPayload.Num());
 }
 
