Fixed mentioning older Steamworks SDK in logs.

diff --git a/Engine/Source/Runtime/Online/OnlineSubsystemSteam/Private/OnlineSubsystemSteamPrivatePCH.h b/Engine/Source/Runtime/Online/OnlineSubsystemSteam/Private/OnlineSubsystemSteamPrivatePCH.h
--- a/Engine/Source/Runtime/Online/OnlineSubsystemSteam/Private/OnlineSubsystemSteamPrivatePCH.h
+++ b/Engine/Source/Runtime/Online/OnlineSubsystemSteam/Private/OnlineSubsystemSteamPrivatePCH.h
@@ -16,7 +16,7 @@
 #define INVALID_INDEX -1
 
 /** Compile Steam SDK version in use */
-#define STEAM_SDK_VER TEXT("1.29a")
+#define STEAM_SDK_VER TEXT("1.30")
 /** Path of the current Steam SDK version in use */
 #define STEAM_SDK_VER_PATH TEXT("Steamv130")
 /** Root location of Steam SDK */
