[GitHub] 385 : Crash fix related to Steam voice implementation
https://github.com/EpicGames/UnrealEngine/pull/385

diff --git a/Engine/Source/Runtime/Online/OnlineSubsystemSteam/Private/VoiceInterfaceSteam.cpp b/Engine/Source/Runtime/Online/OnlineSubsystemSteam/Private/VoiceInterfaceSteam.cpp
--- a/Engine/Source/Runtime/Online/OnlineSubsystemSteam/Private/VoiceInterfaceSteam.cpp
+++ b/Engine/Source/Runtime/Online/OnlineSubsystemSteam/Private/VoiceInterfaceSteam.cpp
@@ -127,7 +127,7 @@ void FOnlineVoiceSteam::StartNetworkedVoice(uint8 LocalUserNum)
 void FOnlineVoiceSteam::StopNetworkedVoice(uint8 LocalUserNum)
 {
 	// Validate the range of the entry
-	if (LocalUserNum >= 0 && LocalUserNum < MAX_LOCAL_PLAYERS)
+	if (LocalTalkers.Num() > LocalUserNum && LocalUserNum >= 0 && LocalUserNum < MAX_LOCAL_PLAYERS)
 	{
 		LocalTalkers[LocalUserNum].bHasNetworkedVoice = false;
 		UE_LOG(LogVoice, Log, TEXT("Stopping networked voice for user: %d"), LocalUserNum);
