Merge UE4-ShooterGame -> UE4

diff --git a/Engine/Source/Programs/AutomationTool/CommandUtils.cs b/Engine/Source/Programs/AutomationTool/CommandUtils.cs
--- a/Engine/Source/Programs/AutomationTool/CommandUtils.cs
+++ b/Engine/Source/Programs/AutomationTool/CommandUtils.cs
@@ -1310,7 +1310,7 @@ namespace AutomationTool
 			}
 			foreach (var SourceSubDirectory in Directory.GetDirectories(Source))
 			{
-				string DestPath = Dest + PathSeparator.Default.ToString() + GetDirectoryName(SourceSubDirectory);
+				string DestPath = Dest + GetPathSeparatorChar(PathSeparator.Default) + GetLastDirectoryName(SourceSubDirectory);
 				if (!CopyDirectory_NoExceptions(SourceSubDirectory, DestPath, bQuiet))
 				{
 					return false;
@@ -1318,7 +1318,7 @@ namespace AutomationTool
 			}
 			foreach (var SourceFile in Directory.GetFiles(Source))
 			{
-				int FilenameStart = SourceFile.LastIndexOf(PathSeparator.Default.ToString());
+				int FilenameStart = SourceFile.LastIndexOf(GetPathSeparatorChar(PathSeparator.Default));
 				string DestPath = Dest + SourceFile.Substring(FilenameStart);
 				if (!CopyFile_NoExceptions(SourceFile, DestPath, bQuiet))
 				{
diff --git a/Engine/Source/Programs/UnrealBuildTool/System/Unity.cs b/Engine/Source/Programs/UnrealBuildTool/System/Unity.cs
--- a/Engine/Source/Programs/UnrealBuildTool/System/Unity.cs
+++ b/Engine/Source/Programs/UnrealBuildTool/System/Unity.cs
@@ -154,7 +154,9 @@ namespace UnrealBuildTool
 				AllUnityFiles = CPPUnityFileBuilder.GetUnityFiles();
 			}
 
-			string PCHHeaderNameInCode = CPPFiles[0].PCHHeaderNameInCode;
+			//THIS CHANGE WAS MADE TO FIX THE BUILD IN OUR BRANCH
+			//DO NOT MERGE THIS BACK TO MAIN
+			string PCHHeaderNameInCode = CPPFiles.Count > 0 ? CPPFiles[0].PCHHeaderNameInCode : "";
 			if( CompileEnvironment.Config.PrecompiledHeaderIncludeFilename != null )
 			{
 				PCHHeaderNameInCode = ToolChain.ConvertPath( CompileEnvironment.Config.PrecompiledHeaderIncludeFilename );
diff --git a/Engine/Source/Runtime/Core/Private/Misc/CallbackDevice.cpp b/Engine/Source/Runtime/Core/Private/Misc/CallbackDevice.cpp
--- a/Engine/Source/Runtime/Core/Private/Misc/CallbackDevice.cpp
+++ b/Engine/Source/Runtime/Core/Private/Misc/CallbackDevice.cpp
@@ -52,6 +52,7 @@ FSimpleMulticastDelegate FCoreDelegates::PostLoadMap;
 FSimpleMulticastDelegate FCoreDelegates::ColorPickerChanged;
 FCoreDelegates::FOnModalMessageBox FCoreDelegates::ModalErrorMessage;
 FCoreDelegates::FOnLoadObjectsOnTop FCoreDelegates::ShouldLoadOnTop;
+FCoreDelegates::FOnInviteAccepted FCoreDelegates::OnInviteAccepted;
 FCoreDelegates::FStringAssetReferenceLoaded FCoreDelegates::StringAssetReferenceLoaded;
 FCoreDelegates::FStringAssetReferenceSaving FCoreDelegates::StringAssetReferenceSaving;
 FCoreDelegates::FPackageCreatedForLoad FCoreDelegates::PackageCreatedForLoad;
@@ -68,4 +69,6 @@ FCoreDelegates::FApplicationLifetimeDelegate FCoreDelegates::ApplicationWillTerm
 FCoreDelegates::FStatCheckEnabled FCoreDelegates::StatCheckEnabled;
 FCoreDelegates::FStatEnabled FCoreDelegates::StatEnabled;
 FCoreDelegates::FStatDisabled FCoreDelegates::StatDisabled;
-FCoreDelegates::FStatDisableAll FCoreDelegates::StatDisableAll;
\ No newline at end of file
+FCoreDelegates::FStatDisableAll FCoreDelegates::StatDisableAll;
+
+FCoreDelegates::FApplicationLicenseChange FCoreDelegates::ApplicationLicenseChange;
diff --git a/Engine/Source/Runtime/Core/Public/GenericPlatform/GenericApplication.h b/Engine/Source/Runtime/Core/Public/GenericPlatform/GenericApplication.h
--- a/Engine/Source/Runtime/Core/Public/GenericPlatform/GenericApplication.h
+++ b/Engine/Source/Runtime/Core/Public/GenericPlatform/GenericApplication.h
@@ -384,6 +384,8 @@ public:
 
 	virtual void ShowSystemHelp() {}
 
+	virtual bool ApplicationLicenseValid(FPlatformUserId PlatformUser = PLATFORMUSERID_NONE) { return true; }
+
 public:
 
 	const TSharedPtr< ICursor > Cursor;
diff --git a/Engine/Source/Runtime/Core/Public/Misc/CallbackDevice.h b/Engine/Source/Runtime/Core/Public/Misc/CallbackDevice.h
--- a/Engine/Source/Runtime/Core/Public/Misc/CallbackDevice.h
+++ b/Engine/Source/Runtime/Core/Public/Misc/CallbackDevice.h
@@ -37,6 +37,8 @@ struct FTestHotFixPayload
 	bool Result;
 };
 
+//forward declares
+class FUniqueNetIdString;
 
 class CORE_API FCoreDelegates
 {
@@ -94,6 +96,9 @@ public:
 	// Callback for handling safe frame area size changes
 	DECLARE_MULTICAST_DELEGATE(FOnSafeFrameChangedEvent);
 
+	// Callback for handling accepting invitations - generally for engine code
+	DECLARE_MULTICAST_DELEGATE_TwoParams(FOnInviteAccepted, const FString&, const FUniqueNetIdString&);
+
 	// Callback for handling the Controller connection / disconnection
 	// first param is true for a connection, false for a disconnection.
 	// second param is UserID, third is UserIndex / ControllerId.
@@ -104,7 +109,7 @@ public:
 
 	// Callback when a user logs in/out of the platform.
 	static FOnUserLoginChangedEvent OnUserLoginChangedEvent;
-	
+
 	// Callback when controllers disconnected / reconnected
 	static FOnUserControllerConnectionChange OnControllerConnectionChange;
 
@@ -190,6 +195,9 @@ public:
 	/** Queries whether an object should be loaded on top ( replace ) an already existing one */
 	static FOnLoadObjectsOnTop ShouldLoadOnTop;
 
+	/** Called when the user accepts an invitation to the current game */
+	static FOnInviteAccepted OnInviteAccepted;
+
 	/** called when loading a string asset reference */
 	DECLARE_DELEGATE_OneParam(FStringAssetReferenceLoaded, FString const& /*LoadedAssetLongPathname*/);
 	static FStringAssetReferenceLoaded StringAssetReferenceLoaded;
@@ -252,6 +260,11 @@ public:
 	DECLARE_MULTICAST_DELEGATE_OneParam(FStatDisableAll, const bool);
 	static FStatDisableAll StatDisableAll;
 
+	// Called when an application is notified that the application license info has been updated.
+	// The new license data should be polled and steps taken based on the results (i.e. halt application if license is no longer valid).
+	DECLARE_MULTICAST_DELEGATE(FApplicationLicenseChange);
+	static FApplicationLicenseChange ApplicationLicenseChange;
+
 private:
 	// Callbacks for hotfixes
 	static TArray<FHotFixDelegate> HotFixDelegates;
diff --git a/Engine/Source/Runtime/Core/Public/Misc/CoreMiscDefines.h b/Engine/Source/Runtime/Core/Public/Misc/CoreMiscDefines.h
--- a/Engine/Source/Runtime/Core/Public/Misc/CoreMiscDefines.h
+++ b/Engine/Source/Runtime/Core/Public/Misc/CoreMiscDefines.h
@@ -100,6 +100,9 @@ enum EForceInit
 };
 enum ENoInit {NoInit};
 
+// Handle type to stably track users on a specific platform
+typedef int32 FPlatformUserId;
+const FPlatformUserId PLATFORMUSERID_NONE = INDEX_NONE;
 #endif // RC_INVOKED
 
 // helpers to turn an preprocessor token into a real string (see UBT_COMPILED_PLATFORM)
diff --git a/Engine/Source/Runtime/Engine/Classes/Engine/GameInstance.h b/Engine/Source/Runtime/Engine/Classes/Engine/GameInstance.h
--- a/Engine/Source/Runtime/Engine/Classes/Engine/GameInstance.h
+++ b/Engine/Source/Runtime/Engine/Classes/Engine/GameInstance.h
@@ -25,6 +25,8 @@ namespace GameInstanceState
 	extern ENGINE_API const FName Playing;			// We are playing the game
 }
 
+class FOnlineSessionSearchResult;
+
 
 UCLASS(config=Game, transient, BlueprintType, Blueprintable)
 class ENGINE_API UGameInstance : public UObject, public FExec
@@ -69,6 +71,8 @@ public:
 
 	/** Starts the GameInstance state machine running */
 	virtual void StartGameInstance();
+	virtual bool JoinSession(ULocalPlayer* LocalPlayer, int32 SessionIndexInSearchResults) { return false; }
+	virtual bool JoinSession(ULocalPlayer* LocalPlayer, const FOnlineSessionSearchResult & SearchResult) { return false; }
 
 	/** Local player access */
 
diff --git a/Engine/Source/Runtime/Engine/Private/GameMode.cpp b/Engine/Source/Runtime/Engine/Private/GameMode.cpp
--- a/Engine/Source/Runtime/Engine/Private/GameMode.cpp
+++ b/Engine/Source/Runtime/Engine/Private/GameMode.cpp
@@ -1582,7 +1582,7 @@ bool AGameMode::FindInactivePlayer(APlayerController* PC)
 			InactivePlayerArray.RemoveAt(i,1);
 			i--;
 		}
-		else if ( (bIsConsole /** @TODO ONLINE - Add back uniqueid compare */) ||
+		else if ( (bIsConsole && ( CurrentPlayerState->UniqueId == PC->PlayerState->UniqueId ) ) ||
 			(!bIsConsole && (FCString::Stricmp(*CurrentPlayerState->SavedNetworkAddress, *NewNetworkAddress) == 0) && (FCString::Stricmp(*CurrentPlayerState->PlayerName, *NewName) == 0)) )
 		{
 			// found it!
diff --git a/Engine/Source/Runtime/Engine/Private/GameSession.cpp b/Engine/Source/Runtime/Engine/Private/GameSession.cpp
--- a/Engine/Source/Runtime/Engine/Private/GameSession.cpp
+++ b/Engine/Source/Runtime/Engine/Private/GameSession.cpp
@@ -11,6 +11,8 @@
 
 DEFINE_LOG_CATEGORY_STATIC(LogGameSession, Log, All);
 
+static TAutoConsoleVariable<int32> CVarMaxPlayersOverride( TEXT( "net.MaxPlayersOverride" ), 0, TEXT( "If greater than 0, will override the standard max players count. Useful for testing full servers." ) );
+
 /** 
  * Returns the player controller associated with this net id
  * @param PlayerNetId the id to search for
@@ -195,7 +197,9 @@ bool AGameSession::AtCapacity(bool bSpectator)
 	}
 	else
 	{
-		return ( (MaxPlayers>0) && (GetWorld()->GetAuthGameMode()->GetNumPlayers() >= MaxPlayers) );
+		const int32 MaxPlayersToUse = CVarMaxPlayersOverride.GetValueOnGameThread() > 0 ? CVarMaxPlayersOverride.GetValueOnGameThread() : MaxPlayers;
+
+		return ( (MaxPlayersToUse>0) && (GetWorld()->GetAuthGameMode()->GetNumPlayers() >= MaxPlayersToUse) );
 	}
 }
 
diff --git a/Engine/Source/Runtime/Engine/Private/PlayerController.cpp b/Engine/Source/Runtime/Engine/Private/PlayerController.cpp
--- a/Engine/Source/Runtime/Engine/Private/PlayerController.cpp
+++ b/Engine/Source/Runtime/Engine/Private/PlayerController.cpp
@@ -1590,7 +1590,7 @@ void APlayerController::ClientReturnToMainMenu_Implementation(const FString& Ret
 bool APlayerController::SetPause( bool bPause, FCanUnpause CanUnpauseDelegate)
 {
 	bool bResult = false;
-	if (GetNetMode() != NM_Client)
+	if (GetNetMode() == NM_Standalone)
 	{
 		AGameMode* const GameMode = GetWorld()->GetAuthGameMode();
 		if (bPause)
diff --git a/Engine/Source/Runtime/Engine/Private/PlayerState.cpp b/Engine/Source/Runtime/Engine/Private/PlayerState.cpp
--- a/Engine/Source/Runtime/Engine/Private/PlayerState.cpp
+++ b/Engine/Source/Runtime/Engine/Private/PlayerState.cpp
@@ -77,6 +77,7 @@ void APlayerState::OverrideWith(APlayerState* PlayerState)
 {
 	bIsSpectator = PlayerState->bIsSpectator;
 	bOnlySpectator = PlayerState->bOnlySpectator;
+	PlayerName = PlayerState->PlayerName;
 	SetUniqueId(PlayerState->UniqueId.GetUniqueNetId());
 }
 
diff --git a/Engine/Source/Runtime/Launch/Private/LaunchEngineLoop.cpp b/Engine/Source/Runtime/Launch/Private/LaunchEngineLoop.cpp
--- a/Engine/Source/Runtime/Launch/Private/LaunchEngineLoop.cpp
+++ b/Engine/Source/Runtime/Launch/Private/LaunchEngineLoop.cpp
@@ -1274,7 +1274,7 @@ int32 FEngineLoop::PreInit( const TCHAR* CmdLine )
 		// If Slate is being used, initialize the renderer after RHIInit 
 		FSlateApplication& CurrentSlateApp = FSlateApplication::Get();
 		CurrentSlateApp.InitializeRenderer( SlateRenderer );
-		
+
 		GetMoviePlayer()->SetSlateRenderer(SlateRenderer);
 	}
 #endif
@@ -1306,7 +1306,6 @@ int32 FEngineLoop::PreInit( const TCHAR* CmdLine )
         // do any post appInit processing, before the render thread is started.
         FPlatformMisc::PlatformPostInit(true);
     }
-
 	if (GUseThreadedRendering)
 	{
 #if PLATFORM_SUPPORTS_RHI_THREAD
@@ -2242,6 +2241,9 @@ void FEngineLoop::OnResuming(_In_ Platform::Object^ Sender, _In_ Platform::Objec
 {
 	// Make the call down to the RHI to Resume the GPU state
 	RHIResumeRendering();
+
+	// Notify application of resume
+	FCoreDelegates::ApplicationHasEnteredForegroundDelegate.Broadcast();
 }
 
 void FEngineLoop::OnSuspending(_In_ Platform::Object^ Sender, _In_ Windows::ApplicationModel::SuspendingEventArgs^ Args)
@@ -2249,7 +2251,8 @@ void FEngineLoop::OnSuspending(_In_ Platform::Object^ Sender, _In_ Windows::Appl
 	// Get the Suspending Event
 	Windows::ApplicationModel::SuspendingDeferral^ SuspendingEvent = Args->SuspendingOperation->GetDeferral();
 	
-	// @TODO ASync Call to Save the Game - This will run in Parallel to the Flush
+	// Notify application of suspend. Application should kick off an async save at this point.
+	FCoreDelegates::ApplicationWillEnterBackgroundDelegate.Broadcast();
 
 	// Flush the RenderingThread
 	FlushRenderingCommands();
@@ -2257,6 +2260,7 @@ void FEngineLoop::OnSuspending(_In_ Platform::Object^ Sender, _In_ Windows::Appl
 	// Make the call down to the RHI to Suspend the GPU state
 	RHISuspendRendering();
 
+	// @TODO Wait for async save to complete
 	// Flush the log so it's all written to disk
 	GLog->FlushThreadedLogs();
 	GLog->Flush();
diff --git a/Engine/Source/Runtime/Online/Android/OnlineSubsystemGooglePlay/Private/OnlineIdentityInterfaceGooglePlay.h b/Engine/Source/Runtime/Online/Android/OnlineSubsystemGooglePlay/Private/OnlineIdentityInterfaceGooglePlay.h
--- a/Engine/Source/Runtime/Online/Android/OnlineSubsystemGooglePlay/Private/OnlineIdentityInterfaceGooglePlay.h
+++ b/Engine/Source/Runtime/Online/Android/OnlineSubsystemGooglePlay/Private/OnlineIdentityInterfaceGooglePlay.h
@@ -65,6 +65,8 @@ public:
 	virtual FString GetPlayerNickname(int32 LocalUserNum) const override;
 	virtual FString GetPlayerNickname(const FUniqueNetId& UserId) const override;
 	virtual FString GetAuthToken(int32 LocalUserNum) const override;
+	virtual void GetUserPrivilege(const FUniqueNetId& UserId, EUserPrivileges::Type Privilege, const FOnGetUserPrivilegeCompleteDelegate& Delegate) override;
+	virtual FPlatformUserId GetPlatformUserIdFromUniqueNetId(const FUniqueNetId& UniqueNetId) override;
 	// End IOnlineIdentity interface
 
 public:
diff --git a/Engine/Source/Runtime/Online/Android/OnlineSubsystemGooglePlay/Private/OnlineIdentityInterfaceGoolgePlay.cpp b/Engine/Source/Runtime/Online/Android/OnlineSubsystemGooglePlay/Private/OnlineIdentityInterfaceGoolgePlay.cpp
--- a/Engine/Source/Runtime/Online/Android/OnlineSubsystemGooglePlay/Private/OnlineIdentityInterfaceGoolgePlay.cpp
+++ b/Engine/Source/Runtime/Online/Android/OnlineSubsystemGooglePlay/Private/OnlineIdentityInterfaceGoolgePlay.cpp
@@ -199,3 +199,22 @@ void FOnlineIdentityGooglePlay::OnLoginCompleted(const int playerID, const int e
 
 	PendingConnectRequest.IsConnectionPending = false;
 }
+
+void FOnlineIdentityGooglePlay::GetUserPrivilege(const FUniqueNetId& UserId, EUserPrivileges::Type Privilege, const FOnGetUserPrivilegeCompleteDelegate& Delegate)
+{
+	Delegate.ExecuteIfBound(UserId, Privilege, (uint32)EPrivilegeResults::NoFailures);
+}
+
+FPlatformUserId FOnlineIdentityGooglePlay::GetPlatformUserIdFromUniqueNetId(const FUniqueNetId& UniqueNetId)
+{
+	for (int i = 0; i < MAX_LOCAL_PLAYERS; ++i)
+	{
+		auto CurrentUniqueId = GetUniquePlayerId(i);
+		if (CurrentUniqueId.IsValid() && (*CurrentUniqueId == UniqueNetId))
+		{
+			return i;
+		}
+	}
+
+	return PLATFORMUSERID_NONE;
+}
diff --git a/Engine/Source/Runtime/Online/Android/OnlineSubsystemGooglePlay/Public/OnlineSubsystemGooglePlay.h b/Engine/Source/Runtime/Online/Android/OnlineSubsystemGooglePlay/Public/OnlineSubsystemGooglePlay.h
--- a/Engine/Source/Runtime/Online/Android/OnlineSubsystemGooglePlay/Public/OnlineSubsystemGooglePlay.h
+++ b/Engine/Source/Runtime/Online/Android/OnlineSubsystemGooglePlay/Public/OnlineSubsystemGooglePlay.h
@@ -13,7 +13,7 @@
  * OnlineSubsystemGooglePlay - Implementation of the online subsystem for Google Play services
  */
 class ONLINESUBSYSTEMGOOGLEPLAY_API FOnlineSubsystemGooglePlay : 
-	public IOnlineSubsystem,
+	public FOnlineSubsystemImpl,
 	public FTickerObjectBase
 {
 public:
diff --git a/Engine/Source/Runtime/Online/IOS/OnlineSubsystemIOS/Private/OnlineIdentityInterfaceIOS.cpp b/Engine/Source/Runtime/Online/IOS/OnlineSubsystemIOS/Private/OnlineIdentityInterfaceIOS.cpp
--- a/Engine/Source/Runtime/Online/IOS/OnlineSubsystemIOS/Private/OnlineIdentityInterfaceIOS.cpp
+++ b/Engine/Source/Runtime/Online/IOS/OnlineSubsystemIOS/Private/OnlineIdentityInterfaceIOS.cpp
@@ -196,4 +196,23 @@ FString FOnlineIdentityIOS::GetAuthToken(int32 LocalUserNum) const
 {
 	FString ResultToken;
 	return ResultToken;
-}
\ No newline at end of file
+}
+
+void FOnlineIdentityIOS::GetUserPrivilege(const FUniqueNetId& UserId, EUserPrivileges::Type Privilege, const FOnGetUserPrivilegeCompleteDelegate& Delegate)
+{
+	Delegate.ExecuteIfBound(UserId, Privilege, (uint32)EPrivilegeResults::NoFailures);
+}
+
+FPlatformUserId FOnlineIdentityIOS::GetPlatformUserIdFromUniqueNetId(const FUniqueNetId& UniqueNetId)
+{
+	for (int i = 0; i < MAX_LOCAL_PLAYERS; ++i)
+	{
+		auto CurrentUniqueId = GetUniquePlayerId(i);
+		if (CurrentUniqueId.IsValid() && (*CurrentUniqueId == UniqueNetId))
+		{
+			return i;
+		}
+	}
+
+	return PLATFORMUSERID_NONE;
+}
diff --git a/Engine/Source/Runtime/Online/IOS/OnlineSubsystemIOS/Private/OnlineIdentityInterfaceIOS.h b/Engine/Source/Runtime/Online/IOS/OnlineSubsystemIOS/Private/OnlineIdentityInterfaceIOS.h
--- a/Engine/Source/Runtime/Online/IOS/OnlineSubsystemIOS/Private/OnlineIdentityInterfaceIOS.h
+++ b/Engine/Source/Runtime/Online/IOS/OnlineSubsystemIOS/Private/OnlineIdentityInterfaceIOS.h
@@ -41,6 +41,8 @@ public:
 	virtual FString GetPlayerNickname(int32 LocalUserNum) const override;
 	virtual FString GetPlayerNickname(const FUniqueNetId& UserId) const override;
 	virtual FString GetAuthToken(int32 LocalUserNum) const override;
+	virtual void GetUserPrivilege(const FUniqueNetId& UserId, EUserPrivileges::Type Privilege, const FOnGetUserPrivilegeCompleteDelegate& Delegate) override;
+	virtual FPlatformUserId GetPlatformUserIdFromUniqueNetId(const FUniqueNetId& UniqueNetId) override;
 	// End IOnlineIdentity interface
 
 
diff --git a/Engine/Source/Runtime/Online/IOS/OnlineSubsystemIOS/Private/OnlineSessionInterfaceIOS.cpp b/Engine/Source/Runtime/Online/IOS/OnlineSubsystemIOS/Private/OnlineSessionInterfaceIOS.cpp
--- a/Engine/Source/Runtime/Online/IOS/OnlineSubsystemIOS/Private/OnlineSessionInterfaceIOS.cpp
+++ b/Engine/Source/Runtime/Online/IOS/OnlineSubsystemIOS/Private/OnlineSessionInterfaceIOS.cpp
@@ -402,7 +402,7 @@ bool FOnlineSessionIOS::PingSearchResults(const FOnlineSessionSearchResult& Sear
 
 bool FOnlineSessionIOS::JoinSession(int32 PlayerNum, FName SessionName, const FOnlineSessionSearchResult& DesiredSession)
 {
-	bool bSuccessfullyJointSession = false;
+	EOnJoinSessionCompleteResult::Type JoinSessionResult = EOnJoinSessionCompleteResult::UnknownError;
 	UE_LOG(LogOnline, Display, TEXT("FOnlineSessionIOS::JoinSession"));
 
 	FGameCenterSessionDelegate* SessionDelegate = *GKSessions.Find( SessionName );
@@ -411,12 +411,12 @@ bool FOnlineSessionIOS::JoinSession(int32 PlayerNum, FName SessionName, const FO
 		NSString* PeerID = [[SessionDelegate Session] peerID];
 		[SessionDelegate connectToPeer:PeerID];
 
-		bSuccessfullyJointSession = true;
+		JoinSessionResult = EOnJoinSessionCompleteResult::Success;
 	}
 
-	TriggerOnJoinSessionCompleteDelegates(SessionName, bSuccessfullyJointSession);
+	TriggerOnJoinSessionCompleteDelegates(SessionName, JoinSessionResult);
 
-	return bSuccessfullyJointSession;
+	return JoinSessionResult == EOnJoinSessionCompleteResult::Success;
 }
 
 
@@ -578,4 +578,14 @@ int32 FOnlineSessionIOS::GetNumSessions()
 void FOnlineSessionIOS::DumpSessionState()
 {
 
-}
\ No newline at end of file
+}
+
+void FOnlineSessionIOS::RegisterLocalPlayer(const FUniqueNetId& PlayerId, FName SessionName, const FOnRegisterLocalPlayerCompleteDelegate& Delegate)
+{
+	Delegate.ExecuteIfBound(PlayerId, EOnJoinSessionCompleteResult::Success);
+}
+
+void FOnlineSessionIOS::UnregisterLocalPlayer(const FUniqueNetId& PlayerId, FName SessionName, const FOnUnregisterLocalPlayerCompleteDelegate& Delegate)
+{
+	Delegate.ExecuteIfBound(PlayerId, true);
+}
diff --git a/Engine/Source/Runtime/Online/IOS/OnlineSubsystemIOS/Private/OnlineSessionInterfaceIOS.h b/Engine/Source/Runtime/Online/IOS/OnlineSubsystemIOS/Private/OnlineSessionInterfaceIOS.h
--- a/Engine/Source/Runtime/Online/IOS/OnlineSubsystemIOS/Private/OnlineSessionInterfaceIOS.h
+++ b/Engine/Source/Runtime/Online/IOS/OnlineSubsystemIOS/Private/OnlineSessionInterfaceIOS.h
@@ -138,6 +138,10 @@ public:
 
 	virtual bool UnregisterPlayers(FName SessionName, const TArray< TSharedRef<FUniqueNetId> >& Players) override;
 
+	virtual void RegisterLocalPlayer(const FUniqueNetId& PlayerId, FName SessionName, const FOnRegisterLocalPlayerCompleteDelegate& Delegate) override;
+
+	virtual void UnregisterLocalPlayer(const FUniqueNetId& PlayerId, FName SessionName, const FOnUnregisterLocalPlayerCompleteDelegate& Delegate) override;
+
 	virtual int32 GetNumSessions() override;
 
 	virtual void DumpSessionState() override;
diff --git a/Engine/Source/Runtime/Online/OnlineSubsystem/Public/Interfaces/OnlineEventsInterface.h b/Engine/Source/Runtime/Online/OnlineSubsystem/Public/Interfaces/OnlineEventsInterface.h
--- a/Engine/Source/Runtime/Online/OnlineSubsystem/Public/Interfaces/OnlineEventsInterface.h
+++ b/Engine/Source/Runtime/Online/OnlineSubsystem/Public/Interfaces/OnlineEventsInterface.h
@@ -21,4 +21,12 @@ public:
 	 * @return Whether the event was successfully triggered
 	 */
 	virtual bool TriggerEvent( const FUniqueNetId & PlayerId, const TCHAR * EventName, const FOnlineEventParms & Parms ) = 0;
+
+	/**
+	 * Quick way to send a valid PlayerSessionId with every event, required for Xbox One
+	 *
+	 * @param PlayerId the unique id of the player this session is associated with
+	 * @param PlayerSessionId A GUID unique to this player session
+	 */
+	virtual void SetPlayerSessionId(const FUniqueNetId& PlayerId, const FGuid& PlayerSessionId) = 0;
 };
\ No newline at end of file
diff --git a/Engine/Source/Runtime/Online/OnlineSubsystem/Public/Interfaces/OnlineIdentityInterface.h b/Engine/Source/Runtime/Online/OnlineSubsystem/Public/Interfaces/OnlineIdentityInterface.h
--- a/Engine/Source/Runtime/Online/OnlineSubsystem/Public/Interfaces/OnlineIdentityInterface.h
+++ b/Engine/Source/Runtime/Online/OnlineSubsystem/Public/Interfaces/OnlineIdentityInterface.h
@@ -45,6 +45,17 @@ public:
 	}
 };
 
+namespace EUserPrivileges
+{
+	enum Type
+	{
+		CanPlay, // This one can be used for age restrictions on PS4
+		CanPlayOnline,
+		CanCommunicateOnline, // For voice AND text chat
+		CanUseUserGeneratedContent
+	};
+}
+
 /**
  * Delegate called when a player logs in/out
  *
@@ -61,7 +72,7 @@ typedef FOnLoginChanged::FDelegate FOnLoginChangedDelegate;
  * @param NewStatus the new login status for the user
  * @param NewId the new id to associate with the user
  */
-DECLARE_MULTICAST_DELEGATE_ThreeParams(FOnLoginStatusChanged, int32, ELoginStatus::Type, const FUniqueNetId&);
+DECLARE_MULTICAST_DELEGATE_FourParams(FOnLoginStatusChanged, int32, ELoginStatus::Type, ELoginStatus::Type, const FUniqueNetId&);
 typedef FOnLoginStatusChanged::FDelegate FOnLoginStatusChangedDelegate;
 
 /**
@@ -102,9 +113,23 @@ class IOnlineIdentity
 {
 
 protected:
-	IOnlineIdentity() {};
+	IOnlineIdentity() {};	
 
 public:
+
+	enum class EPrivilegeResults
+	{
+		NoFailures				=	0,
+		RequiredPatchAvailable	=	1 << 0,
+		RequiredSystemUpdate	=	1 << 1,
+		AgeRestrictionFailure	=	1 << 2, // parental control failure usually
+		AccountTypeFailure		=	1 << 3, // XboxLive Gold / PSPlus required but not available
+		UserNotFound			=	1 << 4, // Invalid user.
+		ChatRestriction			=	1 << 5, // User restricted from chat
+		UGCRestriction			=	1 << 6, // User restricted from User Generated Content
+		GenericFailure			=	1 << 7, // Platform failed for unknown reason and handles its own dialogs
+	};
+
 	virtual ~IOnlineIdentity() {};
 
 	/**
@@ -121,7 +146,7 @@ public:
 	 * @param NewStatus the new login status for the user
 	 * @param NewId the new id to associate with the user
 	 */
-	DEFINE_ONLINE_PLAYER_DELEGATE_TWO_PARAM(MAX_LOCAL_PLAYERS, OnLoginStatusChanged, ELoginStatus::Type, const FUniqueNetId&);
+	DEFINE_ONLINE_PLAYER_DELEGATE_THREE_PARAM(MAX_LOCAL_PLAYERS, OnLoginStatusChanged, ELoginStatus::Type /*Previous*/, ELoginStatus::Type /*Current*/, const FUniqueNetId&);
 
 	/**
 	 * Delegate called when a controller-user pairing changes
@@ -285,6 +310,32 @@ public:
 	 */
 	//@todo - remove and use GetUserAccount instead
 	virtual FString GetAuthToken(int32 LocalUserNum) const = 0;
+
+	/**
+	 * Delegate executed when we get a user privilege result.
+	 *
+	 * @param UniqueId The unique id of the user who was queried
+	 * @param Privilege the privilege that was queried
+	 * @param PrivilegeResult bitwise OR of any privilege failures. 0 is success.
+	 */
+	DECLARE_DELEGATE_ThreeParams(FOnGetUserPrivilegeCompleteDelegate, const FUniqueNetId&, EUserPrivileges::Type, uint32);
+
+	/**
+	 * Gets the status of a user's privilege.
+	 *
+	 * @param UserId the unique id of the user to query
+	 * @param Privilege the privilege you want to know about
+	 * @param Delegate delegate to execute when the async task completes
+	 */
+	virtual void GetUserPrivilege(const FUniqueNetId& UserId, EUserPrivileges::Type Privilege, const FOnGetUserPrivilegeCompleteDelegate& Delegate) = 0;
+
+	/**
+	 * Temporary hack to get a corresponding FUniqueNetId from a PlatformUserId
+	 *
+	 * @param UniqueNetId The unique id to look up
+	 * @return The corresponding id or PLATFORMID_NONE if not found
+	 */
+	virtual FPlatformUserId GetPlatformUserIdFromUniqueNetId(const FUniqueNetId& UniqueNetId) = 0;
 };
 
 typedef TSharedPtr<IOnlineIdentity, ESPMode::ThreadSafe> IOnlineIdentityPtr;
\ No newline at end of file
diff --git a/Engine/Source/Runtime/Online/OnlineSubsystem/Public/Interfaces/OnlinePartyInterface.h b/Engine/Source/Runtime/Online/OnlineSubsystem/Public/Interfaces/OnlinePartyInterface.h
--- a/Engine/Source/Runtime/Online/OnlineSubsystem/Public/Interfaces/OnlinePartyInterface.h
+++ b/Engine/Source/Runtime/Online/OnlineSubsystem/Public/Interfaces/OnlinePartyInterface.h
@@ -48,4 +48,34 @@ public:
 
 	/** Returns true if the user represented by LocalUserNum is currently in the party */
 	virtual bool IsInParty(int32 LocalUserNum) = 0;
+
+	/**
+	 * Delegate executed when the local player operation completes.
+	 *
+	 * @param bWasSuccessful whether the operation succeeded or not
+	 */
+	DECLARE_DELEGATE_OneParam(FOnLocalUserOperationCompleteDelegate, bool);
+
+	/**
+	 * Adds local users to the party.
+	 *
+	 * @param ActingUser the id of a user who's already signed in
+	 * @param UsersToAdd the ids if the new users to add to the party
+	 * @param Delegate the delegate to execute when the async operation completes
+	 */
+	virtual void AddLocalUsers(
+		const FUniqueNetId& ActingUser,
+		const TArray<TSharedRef<FUniqueNetId>>& UsersToAdd,
+		const FOnLocalUserOperationCompleteDelegate& Delegate) = 0;
+
+	/**
+	 * Removes local users from the party.
+	 *
+	 * @param ActingUser the id of a user who's already signed in
+	 * @param UsersToAdd the ids if the new users to add to the party
+	 * @param Delegate the delegate to execute when the async operation completes
+	 */
+	virtual void RemoveLocalUsers(
+		const TArray<TSharedRef<FUniqueNetId>>& UsersToRemove,
+		const FOnLocalUserOperationCompleteDelegate& Delegate) = 0;
 };
\ No newline at end of file
diff --git a/Engine/Source/Runtime/Online/OnlineSubsystem/Public/Interfaces/OnlinePresenceInterface.h b/Engine/Source/Runtime/Online/OnlineSubsystem/Public/Interfaces/OnlinePresenceInterface.h
--- a/Engine/Source/Runtime/Online/OnlineSubsystem/Public/Interfaces/OnlinePresenceInterface.h
+++ b/Engine/Source/Runtime/Online/OnlineSubsystem/Public/Interfaces/OnlinePresenceInterface.h
@@ -13,6 +13,9 @@ typedef FOnlineKeyValuePairs<FPresenceKey, FVariantData> FPresenceProperties;
 /** The default key that will update presence text in the platform's UI */
 const FString DefaultPresenceKey = "RichPresence";
 
+/** Custom presence data that is not seen by users but can be polled */
+const FString CustomPresenceDataKey = "CustomData";
+
 /**
  * Presence info for an online user returned via IOnlinePresence interface
  */
@@ -48,21 +51,16 @@ public:
 class IOnlinePresence
 {
 public:
+	/** Virtual destructor to force proper child cleanup */
+	virtual ~IOnlinePresence() {}
+
 	/**
 	 * Delegate executed when setting presence for a user has completed.
 	 *
 	 * @param User The unique id of the user whose presence was set.
 	 * @param bWasSuccessful true if the async action completed without error, false if there was an error.
 	 */
-	DECLARE_DELEGATE_TwoParams(FOnSetCompleteDelegate, const class FUniqueNetId&, const bool);
-
-	/**
-	 * Delegate executed when the presence query request has completed.
-	 *
-	 * @param Users The unique ids of the users whose presence was requested.
-	 * @param bWasSuccessful true if the async action completed without error, false if there was an error.
-	 */
-	DECLARE_DELEGATE_TwoParams(FOnQueryCompleteDelegate, const TArray<TSharedRef<FUniqueNetId>>&, const bool);
+	DECLARE_DELEGATE_TwoParams(FOnPresenceTaskCompleteDelegate, const class FUniqueNetId&, const bool);
 
 	/**
 	 * Starts an async task that sets presence information for the user.
@@ -71,7 +69,7 @@ public:
 	 * @param Presence The collection of key/value pairs to set as the user's presence data.
 	 * @param Delegate The delegate to be executed when the potentially asynchronous set operation completes.
 	 */
-	virtual void SetPresence(const FUniqueNetId& User, const FPresenceProperties& Presence, const FOnSetCompleteDelegate& Delegate = FOnSetCompleteDelegate()) = 0;
+	virtual void SetPresence(const FUniqueNetId& User, const FPresenceProperties& Presence, const FOnPresenceTaskCompleteDelegate& Delegate = FOnPresenceTaskCompleteDelegate()) = 0;
 
 	/**
 	 * Starts an async operation that will update the cache with presence data from all users in the Users array.
@@ -80,7 +78,7 @@ public:
 	 * @param Users The list of unique ids of the users to query for presence information.
 	 * @param Delegate The delegate to be executed when the potentially asynchronous query operation completes.
 	 */
-	virtual void QueryPresence(const TArray<TSharedRef<FUniqueNetId>>& Users, const FOnQueryCompleteDelegate& Delegate = FOnQueryCompleteDelegate()) = 0;
+	virtual void QueryPresence(const FUniqueNetId& User, const FOnPresenceTaskCompleteDelegate& Delegate = FOnPresenceTaskCompleteDelegate()) = 0;
 
 	/**
 	 * Gets the cached presence information for a user.
diff --git a/Engine/Source/Runtime/Online/OnlineSubsystem/Public/Interfaces/OnlineSessionInterface.h b/Engine/Source/Runtime/Online/OnlineSubsystem/Public/Interfaces/OnlineSessionInterface.h
--- a/Engine/Source/Runtime/Online/OnlineSubsystem/Public/Interfaces/OnlineSessionInterface.h
+++ b/Engine/Source/Runtime/Online/OnlineSubsystem/Public/Interfaces/OnlineSessionInterface.h
@@ -90,13 +90,27 @@ typedef FOnFindSessionsComplete::FDelegate FOnFindSessionsCompleteDelegate;
 DECLARE_MULTICAST_DELEGATE_OneParam(FOnCancelFindSessionsComplete, bool);
 typedef FOnCancelFindSessionsComplete::FDelegate FOnCancelFindSessionsCompleteDelegate;
 
+
+
+namespace EOnJoinSessionCompleteResult
+{
+	enum Type
+	{
+		Success,
+		RoomIsFull,
+		RoomDoesNotExist,
+		CouldNotRetrieveAddress,
+		AlreadyInSession,
+		UnknownError
+	};
+}
 /**
  * Delegate fired when the joining process for an online session has completed
  *
  * @param SessionName the name of the session this callback is for
  * @param bWasSuccessful true if the async action completed without error, false if there was an error
  */
-DECLARE_MULTICAST_DELEGATE_TwoParams(FOnJoinSessionComplete, FName, bool);
+DECLARE_MULTICAST_DELEGATE_TwoParams(FOnJoinSessionComplete, FName, EOnJoinSessionCompleteResult::Type);
 typedef FOnJoinSessionComplete::FDelegate FOnJoinSessionCompleteDelegate;
 
 /**
@@ -130,6 +144,19 @@ typedef FOnPingSearchResultsComplete::FDelegate FOnPingSearchResultsCompleteDele
 DECLARE_MULTICAST_DELEGATE_ThreeParams(FOnSessionInviteAccepted, int32, bool, const FOnlineSessionSearchResult&);
 typedef FOnSessionInviteAccepted::FDelegate FOnSessionInviteAcceptedDelegate;
 
+/**
+	* Called when a user accepts a session invitation. Allows the game code a chance
+	* to clean up any existing state before accepting the invite. The invite must be
+	* accepted by calling JoinSession() after clean up has completed
+	*
+	* @param bWasSuccessful true if the async action completed without error, false if there was an error
+	* @param ControllerId the controller number of the accepting user
+	* @param UserId the user being invited
+	* @param InviteResult the search/settings for the session we're joining via invite
+	*/
+DECLARE_MULTICAST_DELEGATE_FourParams(FOnSessionUserInviteAccepted, const bool, const int32, TSharedPtr< FUniqueNetId >, const FOnlineSessionSearchResult&);
+typedef FOnSessionUserInviteAccepted::FDelegate FOnSessionUserInviteAcceptedDelegate;
+
 /**
  * Delegate fired when the session registration process has completed
  *
@@ -150,6 +177,16 @@ typedef FOnRegisterPlayersComplete::FDelegate FOnRegisterPlayersCompleteDelegate
 DECLARE_MULTICAST_DELEGATE_ThreeParams(FOnUnregisterPlayersComplete, FName, const TArray< TSharedRef<class FUniqueNetId> >&, bool);
 typedef FOnUnregisterPlayersComplete::FDelegate FOnUnregisterPlayersCompleteDelegate;
 
+/**
+ * Delegate fired when local player registration has completed
+ */
+DECLARE_DELEGATE_TwoParams(FOnRegisterLocalPlayerCompleteDelegate, const FUniqueNetId&, EOnJoinSessionCompleteResult::Type);
+
+/**
+ * Delegate fired when local player unregistration has completed
+ */
+DECLARE_DELEGATE_TwoParams(FOnUnregisterLocalPlayerCompleteDelegate, const FUniqueNetId&, const bool);	
+
 /**
  * Interface definition for the online services session services 
  * Session services are defined as anything related managing a session 
@@ -473,7 +510,7 @@ public:
 	 * @param SessionName the name of the session this callback is for
 	 * @param bWasSuccessful true if the async action completed without error, false if there was an error
 	 */
-	DEFINE_ONLINE_DELEGATE_TWO_PARAM(OnJoinSessionComplete, FName, bool);
+	DEFINE_ONLINE_DELEGATE_TWO_PARAM(OnJoinSessionComplete, FName, EOnJoinSessionCompleteResult::Type);
 
 	/**
 	 * Allows the local player to follow a friend into a session
@@ -560,6 +597,18 @@ public:
 	 */
 	DEFINE_ONLINE_PLAYER_DELEGATE_TWO_PARAM(MAX_LOCAL_PLAYERS, OnSessionInviteAccepted, bool, const FOnlineSessionSearchResult&);
 
+	/**
+	 * Called when a user accepts a session invitation. Allows the game code a chance
+	 * to clean up any existing state before accepting the invite. The invite must be
+	 * accepted by calling JoinSession() after clean up has completed
+	 *
+	 * @param bWasSuccessful true if the async action completed without error, false if there was an error
+	 * @param ControllerId the controller number of the accepting user
+	 * @param UserId the user being invited
+	 * @param InviteResult the search/settings for the session we're joining via invite
+	 */
+	DEFINE_ONLINE_DELEGATE_FOUR_PARAM(OnSessionUserInviteAccepted, const bool, const int32, TSharedPtr< FUniqueNetId >, const FOnlineSessionSearchResult&);
+
 	/**
 	 * Returns the platform specific connection information for joining the match.
 	 * Call this function from the delegate of join completion
@@ -651,6 +700,24 @@ public:
 	 */
 	DEFINE_ONLINE_DELEGATE_THREE_PARAM(OnUnregisterPlayersComplete, FName, const TArray< TSharedRef<class FUniqueNetId> >&, bool);
 
+	/**
+	 * Registers a local player with a session.
+	 *
+	 * @param PlayerId the player to register
+	 * @param SessionName the session in which to register the player
+	 * @param Delegate the delegate executed when the asynchronous operation completes
+	 */
+	virtual void RegisterLocalPlayer(const FUniqueNetId& PlayerId, FName SessionName, const FOnRegisterLocalPlayerCompleteDelegate& Delegate) = 0;
+
+	/**
+	 * Unregisters a local player with a session.
+	 *
+	 * @param PlayerId the player to unregister
+	 * @param SessionName the session in which to unregister the player
+	 * @param Delegate the delegate executed when the asynchronous operation completes
+	 */
+	virtual void UnregisterLocalPlayer(const FUniqueNetId& PlayerId, FName SessionName, const FOnUnregisterLocalPlayerCompleteDelegate& Delegate) = 0;
+
 	/**
 	 * Gets the number of known sessions registered with the interface
 	 */
diff --git a/Engine/Source/Runtime/Online/OnlineSubsystem/Public/Online.h b/Engine/Source/Runtime/Online/OnlineSubsystem/Public/Online.h
--- a/Engine/Source/Runtime/Online/OnlineSubsystem/Public/Online.h
+++ b/Engine/Source/Runtime/Online/OnlineSubsystem/Public/Online.h
@@ -18,6 +18,7 @@
 #include "OnlineUserCloudInterface.h"
 #include "OnlineUserInterface.h"
 #include "OnlinePartyInterface.h"
+#include "OnlineFriendsInterface.h" //Needed for online friend functionality
 
 /** Macro to handle the boilerplate of accessing the proper online subsystem and getting the requested interface */
 #define IMPLEMENT_GET_INTERFACE(InterfaceType) \
diff --git a/Engine/Source/Runtime/Online/OnlineSubsystem/Public/OnlineAsyncTaskManager.h b/Engine/Source/Runtime/Online/OnlineSubsystem/Public/OnlineAsyncTaskManager.h
--- a/Engine/Source/Runtime/Online/OnlineSubsystem/Public/OnlineAsyncTaskManager.h
+++ b/Engine/Source/Runtime/Online/OnlineSubsystem/Public/OnlineAsyncTaskManager.h
@@ -59,36 +59,6 @@ public:
 	}
 };
 
-/**
- * An async item that can execute any callable type with no parameters.
- * For example, l lambda, or an object with an operator().
- * Useful for calling simple functions that need to run on the game thread,
- * but are invoked from an online service thread.
- */
-template<class CallableType>
-class FOnlineAsyncItemGenericCallable : public FOnlineAsyncItem
-{
-public:
-	/**
-	 * Constructor.
-	 *
-	 * @param InCallable any object that can be called with no parameters, usually a lambda
-	 */
-	explicit FOnlineAsyncItemGenericCallable(const CallableType& InCallable)
-		: CallableObject(InCallable) {}
-
-	virtual void Finalize() override
-	{
-		CallableObject();
-	}
-
-	virtual FString ToString() const override { return FString("FOnlineAsyncItemGenericCallable"); }
-
-private:
-	/** Stored copy of the object to invoke on the game thread. */
-	CallableType CallableObject;
-};
-
 /**
  * An event triggered by the online subsystem to be routed to the game thread for processing
  * Originates on the online thread
@@ -157,6 +127,40 @@ public:
 	}
 };
 
+
+/**
+ * An async task that can execute any callable type with no parameters.
+ * For example, l lambda, or an object with an operator().
+ * Useful for calling simple functions that need to run on the game thread,
+ * but are invoked from an online service thread.
+ */
+template<class CallableType>
+class FOnlineAsyncTaskGenericCallable : public FOnlineAsyncTask
+{
+public:
+	/**
+	 * Constructor.
+	 *
+	 * @param InCallable any object that can be called with no parameters, usually a lambda
+	 */
+	explicit FOnlineAsyncTaskGenericCallable(const CallableType& InCallable)
+		: CallableObject(InCallable) {}
+
+	virtual void Finalize() override
+	{
+		CallableObject();
+	}
+
+	virtual FString ToString() const override { return FString("FOnlineAsyncTaskGenericCallable"); }
+
+	virtual bool IsDone() override { return true; }
+	virtual bool WasSuccessful() override { return true; }
+
+private:
+	/** Stored copy of the object to invoke on the game thread. */
+	CallableType CallableObject;
+};
+
 template<class T>
 class FOnlineAsyncTaskBasic : public FOnlineAsyncTask
 {
@@ -317,7 +321,23 @@ public:
 	template<class CallableType>
 	void AddGenericToOutQueue(const CallableType& InCallable)
 	{
-		AddToOutQueue(new FOnlineAsyncItemGenericCallable<CallableType>(InCallable));
+		AddToOutQueue(new FOnlineAsyncTaskGenericCallable<CallableType>(InCallable));
+	}
+
+	/**
+	 * Add a new item to the in queue that will call InCallable on the game thread.
+	 * Very useful when passing in lambdas as parameters, since this function will
+	 * automatically deduce the template parameter type for FOnlineAsyncItemGenericCallable.
+	 *
+	 * Unlike AddGenericToOutQueue, this version is useful when you want to ensure that
+	 * the callable will only execute after any existing tasks in the queue are complete.
+	 *
+	 * @param InCallable the callable object to execute on the game thread.
+	 */
+	template<class CallableType>
+	void AddGenericToInQueue(const CallableType& InCallable)
+	{
+		AddToInQueue(new FOnlineAsyncTaskGenericCallable<CallableType>(InCallable));
 	}
 
 	/**
diff --git a/Engine/Source/Runtime/Online/OnlineSubsystem/Public/OnlineSessionSettings.h b/Engine/Source/Runtime/Online/OnlineSubsystem/Public/OnlineSessionSettings.h
--- a/Engine/Source/Runtime/Online/OnlineSubsystem/Public/OnlineSessionSettings.h
+++ b/Engine/Source/Runtime/Online/OnlineSubsystem/Public/OnlineSessionSettings.h
@@ -352,6 +352,8 @@ public:
 	const FName SessionName;
 	/** Index of the player who created the session [host] or joined it [client] */
 	int32 HostingPlayerNum;
+	/** Whether or not the local player is hosting this session */
+	bool bHosting;
 	/** List of players registered in the session */
 	TArray< TSharedRef<FUniqueNetId> > RegisteredPlayers;
 	/** State of the session (game thread write only) */
@@ -362,6 +364,7 @@ public:
 		FOnlineSession(InSessionSettings),
 		SessionName(InSessionName),
 		HostingPlayerNum(INDEX_NONE),
+		bHosting(false),
 		SessionState(EOnlineSessionState::NoSession)
 	{
 	}
@@ -371,6 +374,7 @@ public:
 		FOnlineSession(Session),
 		SessionName(InSessionName),
 		HostingPlayerNum(INDEX_NONE),
+		bHosting(false),
 		SessionState(EOnlineSessionState::NoSession)
 	{
 	}
diff --git a/Engine/Source/Runtime/Online/OnlineSubsystem/Public/OnlineSubsystem.h b/Engine/Source/Runtime/Online/OnlineSubsystem/Public/OnlineSubsystem.h
--- a/Engine/Source/Runtime/Online/OnlineSubsystem/Public/OnlineSubsystem.h
+++ b/Engine/Source/Runtime/Online/OnlineSubsystem/Public/OnlineSubsystem.h
@@ -317,6 +317,17 @@ public:
 	 */
 	virtual bool Exec(class UWorld* InWorld, const TCHAR* Cmd, FOutputDevice& Ar) = 0;
 
+
+	/**
+	 * Displays a system dialog to purchase user account upgrades.  e.g. PlaystationPlus, XboxLive GOLD, etc.
+	 */
+	virtual void DisplayAccountUpgradeDialog(const FUniqueNetId& UniqueId) = 0;
+
+	/**
+	 * Some platforms must know when the game is using Multiplayer features so they can do recurring authorization checks.
+	 */
+	virtual void SetUsingMultiplayerFeatures(const FUniqueNetId& UniqueId, bool bUsingMP) = 0;
+
 	/**
 	 * Called when the connection state as reported by the online platform changes
 	 *
diff --git a/Engine/Source/Runtime/Online/OnlineSubsystem/Public/OnlineSubsystemImpl.h b/Engine/Source/Runtime/Online/OnlineSubsystem/Public/OnlineSubsystemImpl.h
--- a/Engine/Source/Runtime/Online/OnlineSubsystem/Public/OnlineSubsystemImpl.h
+++ b/Engine/Source/Runtime/Online/OnlineSubsystem/Public/OnlineSubsystemImpl.h
@@ -42,6 +42,9 @@ public:
 
 	virtual bool IsLocalPlayer(const FUniqueNetId& UniqueId) const override;
 
+	virtual void DisplayAccountUpgradeDialog(const FUniqueNetId& UniqueId) override {};
+	virtual void SetUsingMultiplayerFeatures(const FUniqueNetId& UniqueId, bool bUsingMP) override {};
+
 	// FOnlineSubsystemImpl
 
 	/**
diff --git a/Engine/Source/Runtime/Online/OnlineSubsystemAmazon/Private/OnlineIdentityAmazon.cpp b/Engine/Source/Runtime/Online/OnlineSubsystemAmazon/Private/OnlineIdentityAmazon.cpp
--- a/Engine/Source/Runtime/Online/OnlineSubsystemAmazon/Private/OnlineIdentityAmazon.cpp
+++ b/Engine/Source/Runtime/Online/OnlineSubsystemAmazon/Private/OnlineIdentityAmazon.cpp
@@ -334,4 +334,23 @@ FString FOnlineIdentityAmazon::GetAuthToken(int32 LocalUserNum) const
 		}
 	}
 	return FString();
-}
\ No newline at end of file
+}
+
+void FOnlineIdentityAmazon::GetUserPrivilege(const FUniqueNetId& UserId, EUserPrivileges::Type Privilege, const FOnGetUserPrivilegeCompleteDelegate& Delegate)
+{
+	Delegate.ExecuteIfBound(UserId, Privilege, (uint32)EPrivilegeResults::NoFailures);
+}
+
+FPlatformUserId FOnlineIdentityAmazon::GetPlatformUserIdFromUniqueNetId(const FUniqueNetId& UniqueNetId)
+{
+	for (int i = 0; i < MAX_LOCAL_PLAYERS; ++i)
+	{
+		auto CurrentUniqueId = GetUniquePlayerId(i);
+		if (CurrentUniqueId.IsValid() && (*CurrentUniqueId == UniqueNetId))
+		{
+			return i;
+		}
+	}
+
+	return PLATFORMUSERID_NONE;
+}
diff --git a/Engine/Source/Runtime/Online/OnlineSubsystemAmazon/Private/OnlineIdentityAmazon.h b/Engine/Source/Runtime/Online/OnlineSubsystemAmazon/Private/OnlineIdentityAmazon.h
--- a/Engine/Source/Runtime/Online/OnlineSubsystemAmazon/Private/OnlineIdentityAmazon.h
+++ b/Engine/Source/Runtime/Online/OnlineSubsystemAmazon/Private/OnlineIdentityAmazon.h
@@ -115,6 +115,8 @@ public:
 	virtual FString GetPlayerNickname(int32 LocalUserNum) const override;
 	virtual FString GetPlayerNickname(const FUniqueNetId& UserId) const override;
 	virtual FString GetAuthToken(int32 LocalUserNum) const override;
+	virtual void GetUserPrivilege(const FUniqueNetId& UserId, EUserPrivileges::Type Privilege, const FOnGetUserPrivilegeCompleteDelegate& Delegate) override;
+	virtual FPlatformUserId GetPlatformUserIdFromUniqueNetId(const FUniqueNetId& UniqueNetId) override;
 
 	// FOnlineIdentityAmazon
 
diff --git a/Engine/Source/Runtime/Online/OnlineSubsystemFacebook/Private/IOS/OnlineIdentityFacebook.cpp b/Engine/Source/Runtime/Online/OnlineSubsystemFacebook/Private/IOS/OnlineIdentityFacebook.cpp
--- a/Engine/Source/Runtime/Online/OnlineSubsystemFacebook/Private/IOS/OnlineIdentityFacebook.cpp
+++ b/Engine/Source/Runtime/Online/OnlineSubsystemFacebook/Private/IOS/OnlineIdentityFacebook.cpp
@@ -211,4 +211,23 @@ FString FOnlineIdentityFacebook::GetPlayerNickname(const FUniqueNetId& UserId) c
 FString FOnlineIdentityFacebook::GetAuthToken(int32 LocalUserNum) const
 {
 	return UserAccount->GetAccessToken();
-}
\ No newline at end of file
+}
+
+void FOnlineIdentityFacebook::GetUserPrivilege(const FUniqueNetId& UserId, EUserPrivileges::Type Privilege, const FOnGetUserPrivilegeCompleteDelegate& Delegate)
+{
+	Delegate.ExecuteIfBound(UserId, Privilege, (uint32)EPrivilegeResults::NoFailures);
+}
+
+FPlatformUserId FOnlineIdentityFacebook::GetPlatformUserIdFromUniqueNetId(const FUniqueNetId& UniqueNetId)
+{
+	for (int i = 0; i < MAX_LOCAL_PLAYERS; ++i)
+	{
+		auto CurrentUniqueId = GetUniquePlayerId(i);
+		if (CurrentUniqueId.IsValid() && (*CurrentUniqueId == UniqueNetId))
+		{
+			return i;
+		}
+	}
+
+	return PLATFORMUSERID_NONE;
+}
diff --git a/Engine/Source/Runtime/Online/OnlineSubsystemFacebook/Private/IOS/OnlineIdentityFacebook.h b/Engine/Source/Runtime/Online/OnlineSubsystemFacebook/Private/IOS/OnlineIdentityFacebook.h
--- a/Engine/Source/Runtime/Online/OnlineSubsystemFacebook/Private/IOS/OnlineIdentityFacebook.h
+++ b/Engine/Source/Runtime/Online/OnlineSubsystemFacebook/Private/IOS/OnlineIdentityFacebook.h
@@ -86,6 +86,8 @@ public:
 	virtual FString GetPlayerNickname(int32 LocalUserNum) const override;
 	virtual FString GetPlayerNickname(const FUniqueNetId& UserId) const override;
 	virtual FString GetAuthToken(int32 LocalUserNum) const override;
+	virtual void GetUserPrivilege(const FUniqueNetId& UserId, EUserPrivileges::Type Privilege, const FOnGetUserPrivilegeCompleteDelegate& Delegate) override;
+	virtual FPlatformUserId GetPlatformUserIdFromUniqueNetId(const FUniqueNetId& UniqueNetId) override;
 	// End IOnlineIdentity interface
 
 
diff --git a/Engine/Source/Runtime/Online/OnlineSubsystemFacebook/Private/Windows/OnlineIdentityFacebook.cpp b/Engine/Source/Runtime/Online/OnlineSubsystemFacebook/Private/Windows/OnlineIdentityFacebook.cpp
--- a/Engine/Source/Runtime/Online/OnlineSubsystemFacebook/Private/Windows/OnlineIdentityFacebook.cpp
+++ b/Engine/Source/Runtime/Online/OnlineSubsystemFacebook/Private/Windows/OnlineIdentityFacebook.cpp
@@ -428,4 +428,24 @@ void FOnlineIdentityFacebook::MeUser_HttpRequestComplete(FHttpRequestPtr HttpReq
 	}
 
 	TriggerOnLoginCompleteDelegates(PendingRegisterUser.LocalUserNum, bResult, FUniqueNetIdString(User.UserId), ErrorStr);
-}
\ No newline at end of file
+}
+
+void FOnlineIdentityFacebook::GetUserPrivilege(const FUniqueNetId& UserId, EUserPrivileges::Type Privilege, const FOnGetUserPrivilegeCompleteDelegate& Delegate)
+{
+	Delegate.ExecuteIfBound(UserId, Privilege, (uint32)EPrivilegeResults::NoFailures);
+}	
+
+FPlatformUserId FOnlineIdentityFacebook::GetPlatformUserIdFromUniqueNetId(const FUniqueNetId& UniqueNetId)
+{
+	for (int i = 0; i < MAX_LOCAL_PLAYERS; ++i)
+	{
+		auto CurrentUniqueId = GetUniquePlayerId(i);
+		if (CurrentUniqueId.IsValid() && (*CurrentUniqueId == UniqueNetId))
+		{
+			return i;
+		}
+	}
+
+	return PLATFORMUSERID_NONE;
+}
+
diff --git a/Engine/Source/Runtime/Online/OnlineSubsystemFacebook/Private/Windows/OnlineIdentityFacebook.h b/Engine/Source/Runtime/Online/OnlineSubsystemFacebook/Private/Windows/OnlineIdentityFacebook.h
--- a/Engine/Source/Runtime/Online/OnlineSubsystemFacebook/Private/Windows/OnlineIdentityFacebook.h
+++ b/Engine/Source/Runtime/Online/OnlineSubsystemFacebook/Private/Windows/OnlineIdentityFacebook.h
@@ -113,6 +113,8 @@ public:
 	virtual FString GetPlayerNickname(int32 LocalUserNum) const override;
 	virtual FString GetPlayerNickname(const FUniqueNetId& UserId) const override;
 	virtual FString GetAuthToken(int32 LocalUserNum) const override;
+	virtual void GetUserPrivilege(const FUniqueNetId& UserId, EUserPrivileges::Type Privilege, const FOnGetUserPrivilegeCompleteDelegate& Delegate) override;
+	virtual FPlatformUserId GetPlatformUserIdFromUniqueNetId(const FUniqueNetId& UniqueNetId) override;
 
 	// FOnlineIdentityFacebook
 
diff --git a/Engine/Source/Runtime/Online/OnlineSubsystemNull/Private/OnlineIdentityNull.cpp b/Engine/Source/Runtime/Online/OnlineSubsystemNull/Private/OnlineIdentityNull.cpp
--- a/Engine/Source/Runtime/Online/OnlineSubsystemNull/Private/OnlineIdentityNull.cpp
+++ b/Engine/Source/Runtime/Online/OnlineSubsystemNull/Private/OnlineIdentityNull.cpp
@@ -265,3 +265,22 @@ FOnlineIdentityNull::~FOnlineIdentityNull()
 {
 }
 
+void FOnlineIdentityNull::GetUserPrivilege(const FUniqueNetId& UserId, EUserPrivileges::Type Privilege, const FOnGetUserPrivilegeCompleteDelegate& Delegate)
+{
+	Delegate.ExecuteIfBound(UserId, Privilege, (uint32)EPrivilegeResults::NoFailures);
+}
+
+FPlatformUserId FOnlineIdentityNull::GetPlatformUserIdFromUniqueNetId(const FUniqueNetId& UniqueNetId)
+{
+	for (int i = 0; i < MAX_LOCAL_PLAYERS; ++i)
+	{
+		auto CurrentUniqueId = GetUniquePlayerId(i);
+		if (CurrentUniqueId.IsValid() && (*CurrentUniqueId == UniqueNetId))
+		{
+			return i;
+		}
+	}
+
+	return PLATFORMUSERID_NONE;
+}
+
diff --git a/Engine/Source/Runtime/Online/OnlineSubsystemNull/Private/OnlineIdentityNull.h b/Engine/Source/Runtime/Online/OnlineSubsystemNull/Private/OnlineIdentityNull.h
--- a/Engine/Source/Runtime/Online/OnlineSubsystemNull/Private/OnlineIdentityNull.h
+++ b/Engine/Source/Runtime/Online/OnlineSubsystemNull/Private/OnlineIdentityNull.h
@@ -66,6 +66,8 @@ public:
 	virtual FString GetPlayerNickname(int32 LocalUserNum) const override;
 	virtual FString GetPlayerNickname(const FUniqueNetId& UserId) const override;
 	virtual FString GetAuthToken(int32 LocalUserNum) const override;
+	virtual void GetUserPrivilege(const FUniqueNetId& UserId, EUserPrivileges::Type Privilege, const FOnGetUserPrivilegeCompleteDelegate& Delegate) override;
+	virtual FPlatformUserId GetPlatformUserIdFromUniqueNetId(const FUniqueNetId& UniqueNetId) override;
 
 	// FOnlineIdentityNull
 
diff --git a/Engine/Source/Runtime/Online/OnlineSubsystemNull/Private/OnlineSessionInterfaceNull.cpp b/Engine/Source/Runtime/Online/OnlineSubsystemNull/Private/OnlineSessionInterfaceNull.cpp
--- a/Engine/Source/Runtime/Online/OnlineSubsystemNull/Private/OnlineSessionInterfaceNull.cpp
+++ b/Engine/Source/Runtime/Online/OnlineSubsystemNull/Private/OnlineSessionInterfaceNull.cpp
@@ -606,7 +606,7 @@ bool FOnlineSessionNull::JoinSession(int32 PlayerNum, FName SessionName, const F
 	if (Return != ERROR_IO_PENDING)
 	{
 		// Just trigger the delegate as having failed
-		TriggerOnJoinSessionCompleteDelegates(SessionName, Return == ERROR_SUCCESS ? true : false);
+		TriggerOnJoinSessionCompleteDelegates(SessionName, Return == ERROR_SUCCESS ? EOnJoinSessionCompleteResult::Success : EOnJoinSessionCompleteResult::UnknownError);
 	}
 
 	return Return == ERROR_SUCCESS || Return == ERROR_IO_PENDING;
@@ -1200,3 +1200,13 @@ void FOnlineSessionNull::DumpSessionState()
 		DumpNamedSession(&Sessions[SessionIdx]);
 	}
 }
+
+void FOnlineSessionNull::RegisterLocalPlayer(const FUniqueNetId& PlayerId, FName SessionName, const FOnRegisterLocalPlayerCompleteDelegate& Delegate)
+{
+	Delegate.ExecuteIfBound(PlayerId, EOnJoinSessionCompleteResult::Success);
+}
+
+void FOnlineSessionNull::UnregisterLocalPlayer(const FUniqueNetId& PlayerId, FName SessionName, const FOnUnregisterLocalPlayerCompleteDelegate& Delegate)
+{
+	Delegate.ExecuteIfBound(PlayerId, true);
+}
diff --git a/Engine/Source/Runtime/Online/OnlineSubsystemNull/Private/OnlineSessionInterfaceNull.h b/Engine/Source/Runtime/Online/OnlineSubsystemNull/Private/OnlineSessionInterfaceNull.h
--- a/Engine/Source/Runtime/Online/OnlineSubsystemNull/Private/OnlineSessionInterfaceNull.h
+++ b/Engine/Source/Runtime/Online/OnlineSubsystemNull/Private/OnlineSessionInterfaceNull.h
@@ -293,6 +293,8 @@ public:
 	virtual bool RegisterPlayers(FName SessionName, const TArray< TSharedRef<FUniqueNetId> >& Players, bool bWasInvited = false) override;
 	virtual bool UnregisterPlayer(FName SessionName, const FUniqueNetId& PlayerId) override;
 	virtual bool UnregisterPlayers(FName SessionName, const TArray< TSharedRef<FUniqueNetId> >& Players) override;
+	virtual void RegisterLocalPlayer(const FUniqueNetId& PlayerId, FName SessionName, const FOnRegisterLocalPlayerCompleteDelegate& Delegate) override;
+	virtual void UnregisterLocalPlayer(const FUniqueNetId& PlayerId, FName SessionName, const FOnUnregisterLocalPlayerCompleteDelegate& Delegate) override;
 	virtual int32 GetNumSessions() override;
 	virtual void DumpSessionState() override;
 };
diff --git a/Engine/Source/Runtime/Online/OnlineSubsystemSteam/Private/OnlineIdentityInterfaceSteam.cpp b/Engine/Source/Runtime/Online/OnlineSubsystemSteam/Private/OnlineIdentityInterfaceSteam.cpp
--- a/Engine/Source/Runtime/Online/OnlineSubsystemSteam/Private/OnlineIdentityInterfaceSteam.cpp
+++ b/Engine/Source/Runtime/Online/OnlineSubsystemSteam/Private/OnlineIdentityInterfaceSteam.cpp
@@ -201,3 +201,22 @@ FString FOnlineIdentitySteam::GetAuthToken(int32 LocalUserNum) const
 	return ResultToken;
 }
 
+void FOnlineIdentitySteam::GetUserPrivilege(const FUniqueNetId& UserId, EUserPrivileges::Type Privilege, const FOnGetUserPrivilegeCompleteDelegate& Delegate)
+{
+	Delegate.ExecuteIfBound(UserId, Privilege, (uint32)EPrivilegeResults::NoFailures);
+}
+
+FPlatformUserId FOnlineIdentitySteam::GetPlatformUserIdFromUniqueNetId(const FUniqueNetId& UniqueNetId)
+{
+	for (int i = 0; i < MAX_LOCAL_PLAYERS; ++i)
+	{
+		auto CurrentUniqueId = GetUniquePlayerId(i);
+		if (CurrentUniqueId.IsValid() && (*CurrentUniqueId == UniqueNetId))
+		{
+			return i;
+		}
+	}
+
+	return PLATFORMUSERID_NONE;
+}
+
diff --git a/Engine/Source/Runtime/Online/OnlineSubsystemSteam/Private/OnlineIdentityInterfaceSteam.h b/Engine/Source/Runtime/Online/OnlineSubsystemSteam/Private/OnlineIdentityInterfaceSteam.h
--- a/Engine/Source/Runtime/Online/OnlineSubsystemSteam/Private/OnlineIdentityInterfaceSteam.h
+++ b/Engine/Source/Runtime/Online/OnlineSubsystemSteam/Private/OnlineIdentityInterfaceSteam.h
@@ -37,6 +37,8 @@ public:
 	virtual FString GetPlayerNickname(int32 LocalUserNum) const override;
 	virtual FString GetPlayerNickname(const FUniqueNetId& UserId) const override;
 	virtual FString GetAuthToken(int32 LocalUserNum) const override;
+	virtual void GetUserPrivilege(const FUniqueNetId& UserId, EUserPrivileges::Type Privilege, const FOnGetUserPrivilegeCompleteDelegate& Delegate) override;
+	virtual FPlatformUserId GetPlatformUserIdFromUniqueNetId(const FUniqueNetId& UniqueNetId) override;
 };
 
 typedef TSharedPtr<FOnlineIdentitySteam, ESPMode::ThreadSafe> FOnlineIdentitySteamPtr;
\ No newline at end of file
diff --git a/Engine/Source/Runtime/Online/OnlineSubsystemSteam/Private/OnlineSessionAsyncLobbySteam.cpp b/Engine/Source/Runtime/Online/OnlineSubsystemSteam/Private/OnlineSessionAsyncLobbySteam.cpp
--- a/Engine/Source/Runtime/Online/OnlineSubsystemSteam/Private/OnlineSessionAsyncLobbySteam.cpp
+++ b/Engine/Source/Runtime/Online/OnlineSubsystemSteam/Private/OnlineSessionAsyncLobbySteam.cpp
@@ -750,7 +750,7 @@ void FOnlineAsyncTaskSteamJoinLobby::TriggerDelegates()
 	IOnlineSessionPtr SessionInt = Subsystem->GetSessionInterface();
 	if (SessionInt.IsValid())
 	{
-		SessionInt->TriggerOnJoinSessionCompleteDelegates(SessionName, bWasSuccessful);
+		SessionInt->TriggerOnJoinSessionCompleteDelegates(SessionName, bWasSuccessful ? EOnJoinSessionCompleteResult::Success : EOnJoinSessionCompleteResult::UnknownError);
 	}
 }
 
diff --git a/Engine/Source/Runtime/Online/OnlineSubsystemSteam/Private/OnlineSessionInterfaceSteam.cpp b/Engine/Source/Runtime/Online/OnlineSubsystemSteam/Private/OnlineSessionInterfaceSteam.cpp
--- a/Engine/Source/Runtime/Online/OnlineSubsystemSteam/Private/OnlineSessionInterfaceSteam.cpp
+++ b/Engine/Source/Runtime/Online/OnlineSubsystemSteam/Private/OnlineSessionInterfaceSteam.cpp
@@ -861,7 +861,7 @@ bool FOnlineSessionSteam::JoinSession(int32 PlayerNum, FName SessionName, const
 	if (Return != ERROR_IO_PENDING)
 	{
 		// Just trigger the delegate as having failed
-		TriggerOnJoinSessionCompleteDelegates(SessionName, Return == ERROR_SUCCESS ? true : false);
+		TriggerOnJoinSessionCompleteDelegates(SessionName, Return == ERROR_SUCCESS ? EOnJoinSessionCompleteResult::Success : EOnJoinSessionCompleteResult::UnknownError);
 	}
 
 	return Return == ERROR_SUCCESS || Return == ERROR_IO_PENDING;
@@ -1764,3 +1764,13 @@ void FOnlineSessionSteam::DumpSessionState()
 		DumpNamedSession(&Sessions[SessionIdx]);
 	}
 }
+
+void FOnlineSessionSteam::RegisterLocalPlayer(const FUniqueNetId& PlayerId, FName SessionName, const FOnRegisterLocalPlayerCompleteDelegate& Delegate)
+{
+	Delegate.ExecuteIfBound(PlayerId, EOnJoinSessionCompleteResult::Success);
+}
+
+void FOnlineSessionSteam::UnregisterLocalPlayer(const FUniqueNetId& PlayerId, FName SessionName, const FOnUnregisterLocalPlayerCompleteDelegate& Delegate)
+{
+	Delegate.ExecuteIfBound(PlayerId, true);
+}
diff --git a/Engine/Source/Runtime/Online/OnlineSubsystemSteam/Private/OnlineSessionInterfaceSteam.h b/Engine/Source/Runtime/Online/OnlineSubsystemSteam/Private/OnlineSessionInterfaceSteam.h
--- a/Engine/Source/Runtime/Online/OnlineSubsystemSteam/Private/OnlineSessionInterfaceSteam.h
+++ b/Engine/Source/Runtime/Online/OnlineSubsystemSteam/Private/OnlineSessionInterfaceSteam.h
@@ -519,6 +519,8 @@ public:
 	virtual bool RegisterPlayers(FName SessionName, const TArray< TSharedRef<FUniqueNetId> >& Players, bool bWasInvited = false) override;
 	virtual bool UnregisterPlayer(FName SessionName, const FUniqueNetId& PlayerId) override;
 	virtual bool UnregisterPlayers(FName SessionName, const TArray< TSharedRef<FUniqueNetId> >& Players) override;
+	virtual void RegisterLocalPlayer(const FUniqueNetId& PlayerId, FName SessionName, const FOnRegisterLocalPlayerCompleteDelegate& Delegate) override;
+	virtual void UnregisterLocalPlayer(const FUniqueNetId& PlayerId, FName SessionName, const FOnUnregisterLocalPlayerCompleteDelegate& Delegate) override;
 	virtual int32 GetNumSessions() override;
 	virtual void DumpSessionState() override;
 };
diff --git a/Engine/Source/Runtime/Online/OnlineSubsystemUtils/Classes/OnlineSessionClient.h b/Engine/Source/Runtime/Online/OnlineSubsystemUtils/Classes/OnlineSessionClient.h
--- a/Engine/Source/Runtime/Online/OnlineSubsystemUtils/Classes/OnlineSessionClient.h
+++ b/Engine/Source/Runtime/Online/OnlineSubsystemUtils/Classes/OnlineSessionClient.h
@@ -106,7 +106,7 @@ protected:
 	 * @param SessionName the name of the session this callback is for
 	 * @param bWasSuccessful true if the async action completed without error, false if there was an error
 	 */
-	void OnJoinSessionComplete(FName SessionName, bool bWasSuccessful);
+	void OnJoinSessionComplete(FName SessionName, EOnJoinSessionCompleteResult::Type Result);
 
 	/**
 	 * Join a session of a given name after potentially tearing down an existing one
diff --git a/Engine/Source/Runtime/Online/OnlineSubsystemUtils/Private/OnlineSessionClient.cpp b/Engine/Source/Runtime/Online/OnlineSubsystemUtils/Private/OnlineSessionClient.cpp
--- a/Engine/Source/Runtime/Online/OnlineSubsystemUtils/Private/OnlineSessionClient.cpp
+++ b/Engine/Source/Runtime/Online/OnlineSubsystemUtils/Private/OnlineSessionClient.cpp
@@ -7,6 +7,8 @@
 
 #include "OnlineSubsystemUtilsPrivatePCH.h"
 
+#include "Engine/GameInstance.h"
+
 #define INVALID_CONTROLLERID 255
 
 UOnlineSessionClient::UOnlineSessionClient(const class FPostConstructInitializeProperties& PCIP)
@@ -227,12 +229,12 @@ void UOnlineSessionClient::DestroyExistingSession(FName SessionName, FOnDestroyS
  * @param SessionName the name of the session this callback is for
  * @param bWasSuccessful true if the async action completed without error, false if there was an error
  */
-void UOnlineSessionClient::OnJoinSessionComplete(FName SessionName, bool bWasSuccessful)
+void UOnlineSessionClient::OnJoinSessionComplete(FName SessionName, EOnJoinSessionCompleteResult::Type Result)
 {
-	UE_LOG(LogOnline, Verbose, TEXT("OnJoinSessionComplete %s bSuccess: %d"), *SessionName.ToString(), bWasSuccessful);
+	UE_LOG(LogOnline, Verbose, TEXT("OnJoinSessionComplete %s bSuccess: %d"), *SessionName.ToString(), static_cast<int32>(Result));
 	SessionInt->ClearOnJoinSessionCompleteDelegate(OnJoinSessionCompleteDelegate);
 
-	if (bWasSuccessful)
+	if (Result == EOnJoinSessionCompleteResult::Success)
 	{
 		FString URL;
 		if (SessionInt->GetResolvedConnectString(SessionName, URL))
@@ -273,8 +275,10 @@ void UOnlineSessionClient::JoinSession(int32 LocalUserNum, FName SessionName, co
 	}
 	else
 	{
-		SessionInt->AddOnJoinSessionCompleteDelegate(OnJoinSessionCompleteDelegate);
-		SessionInt->JoinSession(LocalUserNum, SessionName, SearchResult);
+		UGameInstance * const GameInstance = GetPlayerController()->GetGameInstance();
+		GameInstance->JoinSession(static_cast<ULocalPlayer*>(GetPlayerController()->Player), SearchResult);
+		/*SessionInt->AddOnJoinSessionCompleteDelegate(OnJoinSessionCompleteDelegate);
+		SessionInt->JoinSession(LocalUserNum, SessionName, SearchResult);*/
 	}
 }
 
diff --git a/Engine/Source/Runtime/Online/OnlineSubsystemUtils/Private/Tests/TestSessionInterface.cpp b/Engine/Source/Runtime/Online/OnlineSubsystemUtils/Private/Tests/TestSessionInterface.cpp
--- a/Engine/Source/Runtime/Online/OnlineSubsystemUtils/Private/Tests/TestSessionInterface.cpp
+++ b/Engine/Source/Runtime/Online/OnlineSubsystemUtils/Private/Tests/TestSessionInterface.cpp
@@ -285,12 +285,12 @@ void FTestSessionInterface::OnUpdateSessionComplete(FName SessionName, bool bWas
 	SessionInt->ClearOnUpdateSessionCompleteDelegate(OnUpdateSessionCompleteDelegate);
 }
 
-void FTestSessionInterface::OnJoinSessionComplete(FName SessionName, bool bWasSuccessful)
+void FTestSessionInterface::OnJoinSessionComplete(FName SessionName, EOnJoinSessionCompleteResult::Type Result)
 {
-	UE_LOG(LogOnline, Verbose, TEXT("OnJoinSessionComplete %s bSuccess: %d"), *SessionName.ToString(), bWasSuccessful);
+	UE_LOG(LogOnline, Verbose, TEXT("OnJoinSessionComplete %s bSuccess: %d"), *SessionName.ToString(), static_cast<int32>(Result));
 	SessionInt->ClearOnJoinSessionCompleteDelegate(OnJoinSessionCompleteDelegate);
 
-	if (bWasSuccessful)
+	if (Result == EOnJoinSessionCompleteResult::Success)
 	{
 		FString URL;
 		if (World && SessionInt->GetResolvedConnectString(SessionName, URL))
diff --git a/Engine/Source/Runtime/Online/OnlineSubsystemUtils/Private/Tests/TestSessionInterface.h b/Engine/Source/Runtime/Online/OnlineSubsystemUtils/Private/Tests/TestSessionInterface.h
--- a/Engine/Source/Runtime/Online/OnlineSubsystemUtils/Private/Tests/TestSessionInterface.h
+++ b/Engine/Source/Runtime/Online/OnlineSubsystemUtils/Private/Tests/TestSessionInterface.h
@@ -204,7 +204,7 @@
 	 * @param SessionName the name of the session this callback is for
 	 * @param bWasSuccessful true if the async action completed without error, false if there was an error
 	 */
-	void OnJoinSessionComplete(FName SessionName, bool bWasSuccessful);
+	void OnJoinSessionComplete(FName SessionName, EOnJoinSessionCompleteResult::Type Result);
 
 	/**
 	 * Delegate fired when after finding a friend's session to join
diff --git a/Engine/Source/Runtime/PakFile/Private/IPlatformFilePak.cpp b/Engine/Source/Runtime/PakFile/Private/IPlatformFilePak.cpp
--- a/Engine/Source/Runtime/PakFile/Private/IPlatformFilePak.cpp
+++ b/Engine/Source/Runtime/PakFile/Private/IPlatformFilePak.cpp
@@ -11,6 +11,7 @@
 #include "SignedArchiveReader.h"
 #include "PublicKey.inl"
 #include "AES.h"
+#include "GenericPlatformChunkInstall.h"
 
 DEFINE_LOG_CATEGORY(LogPakFile);
 
@@ -510,9 +511,11 @@ void FPakPlatformFile::FindPakFilesInDirectory(IPlatformFile* LowLevelFile, cons
 	class FPakSearchVisitor : public IPlatformFile::FDirectoryVisitor
 	{
 		TArray<FString>& FoundPakFiles;
+		IPlatformChunkInstall* ChunkInstall;
 	public:
-		FPakSearchVisitor(TArray<FString>& InFoundPakFiles)
+		FPakSearchVisitor(TArray<FString>& InFoundPakFiles, IPlatformChunkInstall* InChunkInstall)
 			: FoundPakFiles(InFoundPakFiles)
+			, ChunkInstall(InChunkInstall)
 		{}
 		virtual bool Visit(const TCHAR* FilenameOrDirectory, bool bIsDirectory)
 		{
@@ -521,6 +524,26 @@ void FPakPlatformFile::FindPakFilesInDirectory(IPlatformFile* LowLevelFile, cons
 				FString Filename(FilenameOrDirectory);
 				if (FPaths::GetExtension(Filename) == TEXT("pak"))
 				{
+					// if a platform supports chunk style installs, make sure that the chunk a pak file resides in is actually fully installed before accepting pak files from it
+					if (ChunkInstall)
+					{
+						FString ChunkIdentifier(TEXT("pakchunk"));
+						FString BaseFilename = FPaths::GetBaseFilename(Filename);
+						if (BaseFilename.StartsWith(ChunkIdentifier))
+						{
+							int32 DelimiterIndex = 0;
+							int32 StartOfChunkIndex = ChunkIdentifier.Len();
+
+							BaseFilename.FindChar(TEXT('-'), DelimiterIndex);
+							FString ChunkNumberString = BaseFilename.Mid(StartOfChunkIndex, DelimiterIndex-StartOfChunkIndex);
+							int32 ChunkNumber = 0;
+							TTypeFromString<int32>::FromString(ChunkNumber, *ChunkNumberString);
+							if (ChunkInstall->GetChunkLocation(ChunkNumber) == EChunkLocation::NotAvailable)
+							{
+								return true;
+							}
+						}
+					}
 					FoundPakFiles.Add(Filename);
 				}
 			}
@@ -528,7 +551,7 @@ void FPakPlatformFile::FindPakFilesInDirectory(IPlatformFile* LowLevelFile, cons
 		}
 	};
 	// Find all pak files.
-	FPakSearchVisitor Visitor(OutPakFiles);
+	FPakSearchVisitor Visitor(OutPakFiles, FPlatformMisc::GetPlatformChunkInstall());
 	LowLevelFile->IterateDirectoryRecursively(Directory, Visitor);
 }
 
