Fix android compile.

diff --git a/Engine/Source/Runtime/Online/Android/OnlineSubsystemGooglePlay/Public/OnlineSubsystemGooglePlay.h b/Engine/Source/Runtime/Online/Android/OnlineSubsystemGooglePlay/Public/OnlineSubsystemGooglePlay.h
--- a/Engine/Source/Runtime/Online/Android/OnlineSubsystemGooglePlay/Public/OnlineSubsystemGooglePlay.h
+++ b/Engine/Source/Runtime/Online/Android/OnlineSubsystemGooglePlay/Public/OnlineSubsystemGooglePlay.h
@@ -3,6 +3,7 @@
 #pragma once
 
 #include "OnlineSubsystem.h"
+#include "OnlineSubsystemImpl.h"
 #include "OnlineIdentityInterfaceGooglePlay.h"
 #include "OnlineAchievementsInterfaceGooglePlay.h"
 #include "OnlineLeaderboardInterfaceGooglePlay.h"
