Merging CL 2293116 (DaveN's libcurl changes).

diff --git a/Engine/Source/Runtime/Online/HTTP/Private/Curl/CurlHttp.cpp b/Engine/Source/Runtime/Online/HTTP/Private/Curl/CurlHttp.cpp
--- a/Engine/Source/Runtime/Online/HTTP/Private/Curl/CurlHttp.cpp
+++ b/Engine/Source/Runtime/Online/HTTP/Private/Curl/CurlHttp.cpp
@@ -39,7 +39,7 @@ FCurlHttpRequest::FCurlHttpRequest(CURLM * InMultiHandle)
 		GConfig->GetBool(TEXT("/Script/Engine.NetworkSettings"), TEXT("n.VerifyPeer"), bVerifyPeer, GEngineIni);
 		if (bVerifyPeer)
 		{
-			UE_LOG(LogInit, Display, TEXT("Libcurl will verify peer certificate"));
+			UE_LOG(LogInit, Verbose, TEXT("Libcurl will verify peer certificate"));
 			curl_easy_setopt(EasyHandle, CURLOPT_SSL_VERIFYPEER, 1L);
 		}
 		else
@@ -68,7 +68,7 @@ FCurlHttpRequest::FCurlHttpRequest(CURLM * InMultiHandle)
 			}
 		}
 
-		if (!FParse::Param(FCommandLine::Get(), TEXT("reuseconn")))
+		if (FParse::Param(FCommandLine::Get(), TEXT("noreuseconn")))
 		{
 			curl_easy_setopt(EasyHandle, CURLOPT_FORBID_REUSE, 1L);
 		}
diff --git a/Engine/Source/Runtime/Online/HTTP/Private/Curl/CurlHttpManager.cpp b/Engine/Source/Runtime/Online/HTTP/Private/Curl/CurlHttpManager.cpp
--- a/Engine/Source/Runtime/Online/HTTP/Private/Curl/CurlHttpManager.cpp
+++ b/Engine/Source/Runtime/Online/HTTP/Private/Curl/CurlHttpManager.cpp
@@ -67,7 +67,7 @@ void FCurlHttpManager::InitCurl()
 			UE_LOG(LogInit, Fatal, TEXT("Could not initialize create libcurl multi handle! HTTP transfers will not function properly."));
 		}
 
-		UE_LOG(LogInit, Log, TEXT("Libcurl will %s"), FParse::Param(FCommandLine::Get(), TEXT("reuseconn")) ? TEXT("reuse connections") : TEXT("NOT reuse connections"));
+		UE_LOG(LogInit, Log, TEXT("Libcurl will %s"), FParse::Param(FCommandLine::Get(), TEXT("noreuseconn")) ? TEXT("NOT reuse connections") : TEXT("reuse connections"));
 	}
 	else
 	{
diff --git a/Engine/Source/Runtime/Online/HTTP/Private/Windows/WindowsPlatformHttp.cpp b/Engine/Source/Runtime/Online/HTTP/Private/Windows/WindowsPlatformHttp.cpp
--- a/Engine/Source/Runtime/Online/HTTP/Private/Windows/WindowsPlatformHttp.cpp
+++ b/Engine/Source/Runtime/Online/HTTP/Private/Windows/WindowsPlatformHttp.cpp
@@ -16,6 +16,10 @@ void FWindowsPlatformHttp::Init()
 	{
 		bUseCurl = true;
 	}
+	else if (GConfig && GConfig->GetBool(TEXT("Networking"), TEXT("UseLibCurl"), bUseCurl, GEngineIni))
+	{
+		bUseCurl = true;
+	}
 
 #if WITH_LIBCURL
 	FCurlHttpManager::InitCurl();
