Merging CL 2293952  (changes/improvements to curl http module).

- Properly get bUseCurl from a configuration file.
- Do less work when creating requests (checking commandline settings moved to CurlHttpManager).
- Do not init/shutdown unless actually used (on Windows).

diff --git a/Engine/Source/Runtime/Online/HTTP/Private/Curl/CurlHttp.cpp b/Engine/Source/Runtime/Online/HTTP/Private/Curl/CurlHttp.cpp
--- a/Engine/Source/Runtime/Online/HTTP/Private/Curl/CurlHttp.cpp
+++ b/Engine/Source/Runtime/Online/HTTP/Private/Curl/CurlHttp.cpp
@@ -3,6 +3,7 @@
 #include "HttpPrivatePCH.h"
 #include "CurlHttp.h"
 #include "EngineVersion.h"
+#include "CurlHttpManager.h"
 
 #if WITH_LIBCURL
 
@@ -35,16 +36,12 @@ FCurlHttpRequest::FCurlHttpRequest(CURLM * InMultiHandle)
 #endif // !UE_BUILD_SHIPPING && !UE_BUILD_TEST
 
 		// set certificate verification (disable to allow self-signed certificates)
-		bool bVerifyPeer = true;
-		GConfig->GetBool(TEXT("/Script/Engine.NetworkSettings"), TEXT("n.VerifyPeer"), bVerifyPeer, GEngineIni);
-		if (bVerifyPeer)
+		if (FCurlHttpManager::CurlRequestOptions.bVerifyPeer)
 		{
-			UE_LOG(LogInit, Verbose, TEXT("Libcurl will verify peer certificate"));
 			curl_easy_setopt(EasyHandle, CURLOPT_SSL_VERIFYPEER, 1L);
 		}
 		else
 		{
-			UE_LOG(LogInit, Display, TEXT("Libcurl will NOT verify peer certificate"));
 			curl_easy_setopt(EasyHandle, CURLOPT_SSL_VERIFYPEER, 0L);
 		}
 
@@ -54,21 +51,13 @@ FCurlHttpRequest::FCurlHttpRequest(CURLM * InMultiHandle)
 		// associate with this just in case
 		curl_easy_setopt(EasyHandle, CURLOPT_PRIVATE, this);
 
-		FString ProxyAddress;
-		if (FParse::Value(FCommandLine::Get(), TEXT("httpproxy="), ProxyAddress))
+		if (FCurlHttpManager::CurlRequestOptions.bUseHttpProxy)
 		{
-			if (!ProxyAddress.IsEmpty())
-			{
-				curl_easy_setopt(EasyHandle, CURLOPT_PROXY, TCHAR_TO_ANSI(*ProxyAddress));
-				UE_LOG(LogHttp, Display, TEXT("Using '%s' as HTTP proxy for request %p"), *ProxyAddress, this);
-			}
-			else
-			{
-				UE_LOG(LogHttp, Warning, TEXT("-httpproxy has been passed as a parameter, but address doesn't seem to be valid"));
-			}
+			// guaranteed to be valid at this point
+			curl_easy_setopt(EasyHandle, CURLOPT_PROXY, TCHAR_TO_ANSI(*FCurlHttpManager::CurlRequestOptions.HttpProxyAddress));
 		}
 
-		if (FParse::Param(FCommandLine::Get(), TEXT("noreuseconn")))
+		if (FCurlHttpManager::CurlRequestOptions.bDontReuseConnections)
 		{
 			curl_easy_setopt(EasyHandle, CURLOPT_FORBID_REUSE, 1L);
 		}
diff --git a/Engine/Source/Runtime/Online/HTTP/Private/Curl/CurlHttpManager.cpp b/Engine/Source/Runtime/Online/HTTP/Private/Curl/CurlHttpManager.cpp
--- a/Engine/Source/Runtime/Online/HTTP/Private/Curl/CurlHttpManager.cpp
+++ b/Engine/Source/Runtime/Online/HTTP/Private/Curl/CurlHttpManager.cpp
@@ -7,6 +7,7 @@
 #if WITH_LIBCURL
 
 CURLM * FCurlHttpManager::GMultiHandle = NULL;
+FCurlHttpManager::FCurlRequestOptions FCurlHttpManager::CurlRequestOptions;
 
 void FCurlHttpManager::InitCurl()
 {
@@ -66,15 +67,68 @@ void FCurlHttpManager::InitCurl()
 		{
 			UE_LOG(LogInit, Fatal, TEXT("Could not initialize create libcurl multi handle! HTTP transfers will not function properly."));
 		}
-
-		UE_LOG(LogInit, Log, TEXT("Libcurl will %s"), FParse::Param(FCommandLine::Get(), TEXT("noreuseconn")) ? TEXT("NOT reuse connections") : TEXT("reuse connections"));
 	}
 	else
 	{
 		UE_LOG(LogInit, Fatal, TEXT("Could not initialize libcurl (result=%d), HTTP transfers will not function properly."), (int32)InitResult);
 	}
+
+	// Init curl request options
+
+	// set certificate verification (disable to allow self-signed certificates)
+	bool bVerifyPeer = true;
+	if (GConfig->GetBool(TEXT("/Script/Engine.NetworkSettings"), TEXT("n.VerifyPeer"), bVerifyPeer, GEngineIni))
+	{
+		CurlRequestOptions.bVerifyPeer = bVerifyPeer;
+	}
+
+	FString ProxyAddress;
+	if (FParse::Value(FCommandLine::Get(), TEXT("httpproxy="), ProxyAddress))
+	{
+		if (!ProxyAddress.IsEmpty())
+		{
+			CurlRequestOptions.bUseHttpProxy = true;
+			CurlRequestOptions.HttpProxyAddress = ProxyAddress;
+		}
+		else
+		{
+			UE_LOG(LogInit, Warning, TEXT(" Libcurl: -httpproxy has been passed as a parameter, but the address doesn't seem to be valid"));
+		}
+	}
+
+	if (FParse::Param(FCommandLine::Get(), TEXT("noreuseconn")))
+	{
+		CurlRequestOptions.bDontReuseConnections = true;
+	}
+
+	// print for visibility
+	CurlRequestOptions.Log();
 }
 
+void FCurlHttpManager::FCurlRequestOptions::Log()
+{
+	UE_LOG(LogInit, Log, TEXT(" CurlRequestOptions (configurable via config and command line):"));
+		UE_LOG(LogInit, Log, TEXT(" - bVerifyPeer = %s  - Libcurl will %sverify peer certificate"),
+		bVerifyPeer ? TEXT("true") : TEXT("false"),
+		bVerifyPeer ? TEXT("") : TEXT("NOT ")
+		);
+
+	UE_LOG(LogInit, Log, TEXT(" - bUseHttpProxy = %s  - Libcurl will %suse HTTP proxy"),
+		bUseHttpProxy ? TEXT("true") : TEXT("false"),
+		bUseHttpProxy ? TEXT("") : TEXT("NOT ")
+		);	
+	if (bUseHttpProxy)
+	{
+		UE_LOG(LogInit, Log, TEXT(" - HttpProxyAddress = '%s'"), *CurlRequestOptions.HttpProxyAddress);
+	}
+
+	UE_LOG(LogInit, Log, TEXT(" - bDontReuseConnections = %s  - Libcurl will %sreuse connections"),
+		bDontReuseConnections ? TEXT("true") : TEXT("false"),
+		bDontReuseConnections ? TEXT("NOT ") : TEXT("")
+		);
+}
+
+
 void FCurlHttpManager::ShutdownCurl()
 {
 	if (NULL != GMultiHandle)
diff --git a/Engine/Source/Runtime/Online/HTTP/Private/Curl/CurlHttpManager.h b/Engine/Source/Runtime/Online/HTTP/Private/Curl/CurlHttpManager.h
--- a/Engine/Source/Runtime/Online/HTTP/Private/Curl/CurlHttpManager.h
+++ b/Engine/Source/Runtime/Online/HTTP/Private/Curl/CurlHttpManager.h
@@ -38,6 +38,31 @@ public:
 	static void InitCurl();
 	static void ShutdownCurl();
 	static CURLM * GMultiHandle;
+
+	static struct FCurlRequestOptions
+	{
+		FCurlRequestOptions()
+			:	bVerifyPeer(true)
+			,	bUseHttpProxy(false)
+			,	bDontReuseConnections(false)
+		{}
+
+		/** Prints out the options to the log */
+		void Log();
+
+		/** Whether or not should verify peer certificate (disable to allow self-signed certs) */
+		bool bVerifyPeer;
+
+		/** Whether or not should use HTTP proxy */
+		bool bUseHttpProxy;
+
+		/** Forbid reuse connections (for debugging purposes, since normally it's faster to reuse) */
+		bool bDontReuseConnections;
+
+		/** Address of the HTTP proxy */
+		FString HttpProxyAddress;
+	}
+	CurlRequestOptions;
 };
 
 #endif //WITH_LIBCURL
diff --git a/Engine/Source/Runtime/Online/HTTP/Private/Windows/WindowsPlatformHttp.cpp b/Engine/Source/Runtime/Online/HTTP/Private/Windows/WindowsPlatformHttp.cpp
--- a/Engine/Source/Runtime/Online/HTTP/Private/Windows/WindowsPlatformHttp.cpp
+++ b/Engine/Source/Runtime/Online/HTTP/Private/Windows/WindowsPlatformHttp.cpp
@@ -10,28 +10,43 @@ bool bUseCurl = false;
 
 void FWindowsPlatformHttp::Init()
 {
+	if (GConfig)
+	{
+		bool bUseCurlConfigValue = false;
+		if (GConfig->GetBool(TEXT("Networking"), TEXT("UseLibCurl"), bUseCurlConfigValue, GEngineIni))
+		{
+			bUseCurl = bUseCurlConfigValue;
+		}
+	}
+
+	// allow override on command line
 	FString HttpMode;
 	if (FParse::Value(FCommandLine::Get(), TEXT("HTTP="), HttpMode) &&
 		(HttpMode.Equals(TEXT("Curl"), ESearchCase::IgnoreCase) || HttpMode.Equals(TEXT("LibCurl"), ESearchCase::IgnoreCase)))
 	{
 		bUseCurl = true;
 	}
-	else if (GConfig && GConfig->GetBool(TEXT("Networking"), TEXT("UseLibCurl"), bUseCurl, GEngineIni))
-	{
-		bUseCurl = true;
-	}
 
 #if WITH_LIBCURL
-	FCurlHttpManager::InitCurl();
+	if (bUseCurl)
+	{
+		FCurlHttpManager::InitCurl();
+	}
 #endif
 }
 
 void FWindowsPlatformHttp::Shutdown()
 {
 #if WITH_LIBCURL
-	FCurlHttpManager::ShutdownCurl();
+	if (bUseCurl)
+	{
+		FCurlHttpManager::ShutdownCurl();
+	}
+	else
 #endif
-	FWinInetConnection::Get().ShutdownConnection();
+	{
+		FWinInetConnection::Get().ShutdownConnection();
+	}
 }
 
 FHttpManager * FWindowsPlatformHttp::CreatePlatformHttpManager()
