Change BuildPatchTool.exe to respect maximum reuse age

diff --git a/Engine/Source/Programs/AutomationTool/MCPPublic.cs b/Engine/Source/Programs/AutomationTool/MCPPublic.cs
--- a/Engine/Source/Programs/AutomationTool/MCPPublic.cs
+++ b/Engine/Source/Programs/AutomationTool/MCPPublic.cs
@@ -243,6 +243,16 @@ namespace EpicGames.MCP.Automation
 
         public class PatchGenerationOptions
         {
+            /// <summary>
+            /// By default, we will only consider data modified within 12 days to be reusable
+            /// </summary>
+            public const int DEFAULT_DATA_AGE_THRESHOLD = 12;
+
+            public PatchGenerationOptions()
+            {
+                DataAgeThreshold = DEFAULT_DATA_AGE_THRESHOLD;
+            }
+
             /// <summary>
             /// Staging information
             /// </summary>
@@ -271,18 +281,14 @@ namespace EpicGames.MCP.Automation
             /// Matches the corresponding BuildPatchTool command line argument.
             /// </summary>
             public MCPPlatform Platform;
-        }
-
-		public class CompactifyOptions
-		{
-            /// <summary>
-            /// If specified, BuildPatchTool will run a compactify on this directory.
-            /// </summary>
-            public string CompactifyDirectory;
             /// <summary>
-            /// Corresponds to the -preview parameter
+            /// When identifying existing patch data to reuse in this build, only
+            /// files modified within this number of days will be considered for reuse.
+            /// IMPORTANT: This should always be smaller than the data age threshold for any compactify process which will run on the directory, to ensure
+            /// that we do not reuse any files which could be deleted by a concurrently running compactify. It is recommended that this number be at least
+            /// two days less than the compactify data age threshold.
             /// </summary>
-            public bool bPreviewCompactify;
+            public int DataAgeThreshold;
 		}
 
         static BuildPatchToolBase Handler = null;
@@ -318,11 +324,6 @@ namespace EpicGames.MCP.Automation
 		/// <param name="Opts">Parameters which will be passed to the patch tool generation process</param>
 		public abstract void Execute(PatchGenerationOptions Opts);
 
-		/// <summary>
-		/// Runs the Build Patch Tool executable to compactify a cloud directory using the supplied parameters.
-		/// </summary>
-		/// <param name="Opts">Parameters which will be passed to the compactify process</param>
-		public abstract void Execute(CompactifyOptions Opts);
     }
 
 
diff --git a/Engine/Source/Programs/BuildPatchTool/Private/BuildPatchToolMain.cpp b/Engine/Source/Programs/BuildPatchTool/Private/BuildPatchToolMain.cpp
--- a/Engine/Source/Programs/BuildPatchTool/Private/BuildPatchToolMain.cpp
+++ b/Engine/Source/Programs/BuildPatchTool/Private/BuildPatchToolMain.cpp
@@ -33,7 +33,14 @@
 	-PrereqName=""		Specifies in quotes, the display name for the prerequisites installer
 	-PrereqPath=""		Specifies in quotes, the prerequisites installer to launch on successful product install
 	-PrereqArgs=""		Specifies in quotes, the commandline to send to prerequisites installer on launch
-	
+	-DataAgeThreshold=12.5		Specified the maximum age (in days) of existing patch data which can be reused in the generated manifest
+
+	NB: If -DataAgeThreshold is not specified on the command line, then it *MUST* be present in the [PatchGeneration] section of BuildPatchTool.ini in the cloud directory
+
+	Example command lines:
+	-BuildRoot="D:\Builds\ExampleGame_[07-02_03.00]" -FileIgnoreList="D:\Builds\ExampleGame_[07-02_03.00]\Manifest_DebugFiles.txt" -CloudDir="D:\BuildPatchCloud" -AppID=123456 -AppName="Example" -BuildVersion="ExampleGame_[07-02_03.00]" -AppLaunch=".\ExampleGame\Binaries\Win32\ExampleGame.exe" -AppArgs="-pak -nosteam" -DataAgeThreshold=12
+	-BuildRoot="C:\Program Files (x86)\Community Portal" -CloudDir="E:\BuildPatchCloud" -AppID=0 -AppName="Community Portal" -BuildVersion="++depot+UE4-CL-1791234" -AppLaunch=".\Engine\Binaries\Win32\CommunityPortal.exe" -AppArgs="" -nochunks -DataAgeThreshold=12
+
 	-custom="field=value"	 	Adds a custom string field to the build manifest
 	-customint="field=number"	Adds a custom int64 field to the build manifest 
 	-customfloat="field=number"	Adds a custom double field to the build manifest
@@ -47,12 +54,17 @@
 	Optional arguments:
 	-stdout				Adds stdout logging to the app.
 	-preview			Log all the actions it will take to update internal structures, but don't actually execute them.
+	-ManifestsList=""			Specifies in quotes, the list of manifest filenames to keep following the operation. If omitted, all manifests are kept.
+	-ManifestsFile=""			Specifies in quotes, the name of the file (relative to CloudDir) which contains a list of manifests to keep, one manifest filename per line
+	-DataAgeThreshold=14.25		The maximum age in days of chunk files that will be retained. All older chunks will be deleted.
 
+	NB: If -ManifestsList is specified, then -ManifestsFile is ignored.
+	NB: If -DataAgeThreshold is not specified on the command line, then it *MUST* be present in the [Compactify] section of BuildPatchTool.ini in the cloud directory
 
 	Example command lines:
-	-BuildRoot="D:\Builds\ExampleGame_[07-02_03.00]" -FileIgnoreList="D:\Builds\ExampleGame_[07-02_03.00]\Manifest_DebugFiles.txt" -CloudDir="D:\BuildPatchCloud" -AppID=123456 -AppName="Example" -BuildVersion="ExampleGame_[07-02_03.00]" -AppLaunch=".\ExampleGame\Binaries\Win32\ExampleGame.exe" -AppArgs="-pak -nosteam"
-	-BuildRoot="C:\Program Files (x86)\Community Portal" -CloudDir="E:\BuildPatchCloud" -AppID=0 -AppName="Community Portal" -BuildVersion="++depot+UE4-CL-1791234" -AppLaunch=".\Engine\Binaries\Win32\CommunityPortal.exe" -AppArgs="" -nochunks
-	-CloudDir="E:\BuildPatchCloude" -compactify
+	-CloudDir="E:\BuildPatchCloud" -compactify -DataAgeThreshold=14
+	-CloudDir="E:\BuildPatchCloud" -compactify -DataAgeThreshold=14 -ManifestsList="Example_V1.manifest,Example_V2.manifest"
+	-CloudDir="E:\BuildPatchCloud" -compactify -DataAgeThreshold=14 -ManifestsFile="preserve_manifests.txt"
 =============================================================================*/
 
 #include "RequiredProgramMainCPPInclude.h"
@@ -130,10 +142,15 @@ int32 BuildPatchToolMain( const TCHAR* CommandLine )
 	FString PrereqName;
 	FString PrereqPath;
 	FString PrereqArgs;
+	FString ManifestsList;
+	FString ManifestsFile;
+	float DataAgeThreshold = 0.0f;
+	FString IniFile;
 	TMap<FString, FVariant> CustomFields;
 
 	bool bCompactify = false;
-	bool bPreviewCompactify = false;
+	bool bPatchGeneration = true;
+	bool bPreview = false;
 
 	// Collect all the info from the CommandLine
 	TArray< FString > Tokens, Switches;
@@ -151,16 +168,19 @@ int32 BuildPatchToolMain( const TCHAR* CommandLine )
 		int32 PrereqNameIdx;
 		int32 PrereqPathIdx;
 		int32 PrereqArgsIdx;
-
+		int32 ManifestsListIdx;
+		int32 ManifestsFileIdx;
+		int32 DataAgeThresholdIdx;
 
 		FCommandLineMatcher Matcher;
 		bSuccess = true;
 
 		Matcher.Command = TEXT("compactify");
 		bCompactify = Switches.FindMatch(Matcher) != INDEX_NONE;
+		bPatchGeneration = !bCompactify;
 
 		Matcher.Command = TEXT("preview");
-		bPreviewCompactify = bCompactify && Switches.FindMatch(Matcher) != INDEX_NONE;
+		bPreview = bCompactify && Switches.FindMatch(Matcher) != INDEX_NONE;
 		
 		Matcher.Command = TEXT( "BuildRoot" );
 		BuildRootIdx = Switches.FindMatch( Matcher );
@@ -195,9 +215,18 @@ int32 BuildPatchToolMain( const TCHAR* CommandLine )
 		Matcher.Command = TEXT( "PrereqArgs" );
 		PrereqArgsIdx = Switches.FindMatch( Matcher );
 
+		Matcher.Command = TEXT("ManifestsList");
+		ManifestsListIdx = Switches.FindMatch(Matcher);
+
+		Matcher.Command = TEXT("ManifestsFile");
+		ManifestsFileIdx = Switches.FindMatch(Matcher);
+
+		Matcher.Command = TEXT("DataAgeThreshold");
+		DataAgeThresholdIdx = Switches.FindMatch(Matcher);
+
 		// Check required param indexes
 		bSuccess = bSuccess && CloudDirIdx != INDEX_NONE;
-		if (!bCompactify)
+		if (bPatchGeneration)
 		{
 			bSuccess = bSuccess && BuildRootIdx != INDEX_NONE;
 			bSuccess = bSuccess && AppIDIdx != INDEX_NONE;
@@ -209,7 +238,7 @@ int32 BuildPatchToolMain( const TCHAR* CommandLine )
 
 		// Get required param values
 		bSuccess = bSuccess && FParse::Value( *Switches[CloudDirIdx], TEXT( "CloudDir=" ), CloudDirectory );
-		if (!bCompactify)
+		if (bPatchGeneration)
 		{
 			bSuccess = bSuccess && FParse::Value(*Switches[BuildRootIdx], TEXT("BuildRoot="), RootDirectory);
 			bSuccess = bSuccess && FParse::Value(*Switches[AppIDIdx], TEXT("AppID="), AppID);
@@ -240,6 +269,16 @@ int32 BuildPatchToolMain( const TCHAR* CommandLine )
 			FParse::Value( *Switches[ PrereqArgsIdx ], TEXT( "FileIgnoreList=" ), PrereqArgs );
 		}
 
+		if (ManifestsListIdx != INDEX_NONE)
+		{
+			bool bShouldStopOnComma = false;
+			FParse::Value(*Switches[ManifestsListIdx], TEXT("ManifestsList="), ManifestsList, bShouldStopOnComma);
+		}
+		else if (ManifestsFileIdx != INDEX_NONE)
+		{
+			FParse::Value( *Switches[ ManifestsFileIdx ], TEXT( "ManifestsFile=" ), ManifestsFile);
+		}
+
 		FString CustomValue;
 		FString Left;
 		FString Right;
@@ -292,6 +331,36 @@ int32 BuildPatchToolMain( const TCHAR* CommandLine )
 
 		FPaths::NormalizeDirectoryName( RootDirectory );
 		FPaths::NormalizeDirectoryName( CloudDirectory );
+
+		if (bSuccess)
+		{
+			// Initialize the configuration system, we can only do this reliably if we have CloudDirectory (i.e. bSuccess is true)
+			IniFile = CloudDirectory / TEXT("BuildPatchTool.ini");
+			GConfig->InitializeConfigSystem();
+		}
+
+		if (DataAgeThresholdIdx != INDEX_NONE)
+		{
+			FParse::Value(*Switches[DataAgeThresholdIdx], TEXT("DataAgeThreshold="), DataAgeThreshold);
+		}
+		else if (bSuccess && bCompactify)
+		{
+			// For compactification, if we don't pass in DataAgeThreshold, then it *MUST* be specified in BuildPatchTool.ini
+			if (!GConfig->GetFloat(TEXT("Compactify"), TEXT("DataAgeThreshold"), DataAgeThreshold, IniFile))
+			{
+				GLog->Log(ELogVerbosity::Error, TEXT("DataAgeThreshold must be specified either on the command line or in the Compactify section of BuildPatchTool.ini"));
+				bSuccess = false; // We couldn't get the DataAgeThreshold value from the .ini file
+			}
+		}
+		else if (bSuccess && bPatchGeneration)
+		{
+			// For patch generation, if we don't pass in DataAgeThreshold, then it *MUST* be specified in BuildPatchTool.ini
+			if (!GConfig->GetFloat(TEXT("PatchGeneration"), TEXT("DataAgeThreshold"), DataAgeThreshold, IniFile))
+			{
+				GLog->Log(ELogVerbosity::Error, TEXT("DataAgeThreshold must be specified either on the command line or in the PatchGeneration section of BuildPatchTool.ini"));
+				bSuccess = false; // We couldn't get the DataAgeThreshold value from the .ini file
+			}
+		}
 	}
 
 	// Check for argument error
@@ -317,10 +386,33 @@ int32 BuildPatchToolMain( const TCHAR* CommandLine )
 
 	if (bCompactify)
 	{
+		// Split out our manifests to keep arg (if any) into an array of manifest filenames
+		TArray<FString> ManifestsArr;
+		if (ManifestsList.Len() > 0)
+		{
+			ManifestsList.ParseIntoArray(&ManifestsArr, TEXT(","), true);
+		}
+		else if (ManifestsFile.Len() > 0)
+		{
+			FString ManifestsFilePath = CloudDirectory / ManifestsFile;
+			FString Temp;
+			if (FFileHelper::LoadFileToString(Temp, *ManifestsFilePath))
+			{
+				Temp.ReplaceInline(TEXT("\r"), TEXT("\n"));
+				Temp.ParseIntoArray(&ManifestsArr, TEXT("\n"), true);
+			}
+			else
+			{
+				GLog->Log(ELogVerbosity::Error, TEXT("Could not open specified manifests to keep file"));
+				BuildPatchServicesModule.Reset();
+				return 2;
+			}
+		}
+
 		// Run the compactify routine
-		bSuccess = BuildPatchServicesModule->CompactifyCloudDirectory(bPreviewCompactify);
+		bSuccess = BuildPatchServicesModule->CompactifyCloudDirectory(ManifestsArr, DataAgeThreshold, bPreview);
 	}
-	else
+	else if (bPatchGeneration)
 	{
 		FBuildPatchSettings Settings;
 		Settings.RootDirectory = RootDirectory + TEXT("/");
@@ -333,6 +425,7 @@ int32 BuildPatchToolMain( const TCHAR* CommandLine )
 		Settings.PrereqName = PrereqName;
 		Settings.PrereqPath = PrereqPath;
 		Settings.PrereqArgs = PrereqArgs;
+		Settings.DataAgeThreshold = DataAgeThreshold;
 		Settings.CustomFields = CustomFields;
 
 		// Run the build generation
@@ -345,6 +438,12 @@ int32 BuildPatchToolMain( const TCHAR* CommandLine )
 			bSuccess = BuildPatchServicesModule->GenerateChunksManifestFromDirectory(Settings);
 		}
 	}
+	else
+	{
+		GLog->Log(ELogVerbosity::Error, TEXT("Unknown tool mode"));
+		BuildPatchServicesModule.Reset();
+		return 3;
+	}
 
 	// Release the module ptr
 	BuildPatchServicesModule.Reset();
@@ -352,8 +451,8 @@ int32 BuildPatchToolMain( const TCHAR* CommandLine )
 	// Check for processing error
 	if (!bSuccess)
 	{
-		GLog->Log(ELogVerbosity::Error, TEXT("An fatal error occurred executing BuildPatchTool.exe"));
-		return 2;
+		GLog->Log(ELogVerbosity::Error, TEXT("A fatal error occurred executing BuildPatchTool.exe"));
+		return 4;
 	}
 
 	GLog->Log(TEXT("BuildPatchToolMain completed successfuly"));
diff --git a/Engine/Source/Runtime/Online/BuildPatchServices/Private/BuildPatchCompactifier.cpp b/Engine/Source/Runtime/Online/BuildPatchServices/Private/BuildPatchCompactifier.cpp
--- a/Engine/Source/Runtime/Online/BuildPatchServices/Private/BuildPatchCompactifier.cpp
+++ b/Engine/Source/Runtime/Online/BuildPatchServices/Private/BuildPatchCompactifier.cpp
@@ -9,16 +9,18 @@ data that are no longer referenced by the manifests in a given cloud directory.
 
 #if WITH_BUILDPATCHGENERATION
 
-/* Constructors
-*****************************************************************************/
-
-FBuildDataCompactifier::FBuildDataCompactifier()
+namespace BuildDataCompactifierDefs
 {
-
+	const static uint32 TouchFailureThresholdPercentage = 10;
+	const static int64 TimeStampTolerance = FTimespan::FromMinutes(1).GetTicks();
 }
 
-FBuildDataCompactifier::FBuildDataCompactifier(const FString& CloudDir)
+/* Constructors
+*****************************************************************************/
+
+FBuildDataCompactifier::FBuildDataCompactifier(const FString& CloudDir, const bool bPreview)
 	: CloudDir(CloudDir)
+	, bPreview(bPreview)
 {
 
 }
@@ -26,23 +28,31 @@ FBuildDataCompactifier::FBuildDataCompactifier(const FString& CloudDir)
 /* Public static methods
 *****************************************************************************/
 
-bool FBuildDataCompactifier::CompactifyCloudDirectory(const FString& CloudDir, const bool bPreview)
+bool FBuildDataCompactifier::CompactifyCloudDirectory(const FString& CloudDir, const TArray<FString>& ManifestsToKeep, const float DataAgeThreshold, const bool bPreview)
 {
-	FBuildDataCompactifier Compactifier(CloudDir);
-	return Compactifier.Compactify(bPreview);
+	FBuildDataCompactifier Compactifier(CloudDir, bPreview);
+	return Compactifier.Compactify(ManifestsToKeep, DataAgeThreshold);
 }
 
-bool FBuildDataCompactifier::CompactifyCloudDirectory(const bool bPreview)
+bool FBuildDataCompactifier::CompactifyCloudDirectory(const TArray<FString>& ManifestsToKeep, const float DataAgeThreshold, const bool bPreview)
 {
-	return CompactifyCloudDirectory(FBuildPatchServicesModule::GetCloudDirectory(), bPreview);
+	return CompactifyCloudDirectory(FBuildPatchServicesModule::GetCloudDirectory(), ManifestsToKeep, DataAgeThreshold, bPreview);
 }
 
 /* Private methods
 *****************************************************************************/
 
-bool FBuildDataCompactifier::Compactify(bool bPreview) const
+bool FBuildDataCompactifier::Compactify(const TArray<FString>& ManifestsToKeep, const float DataAgeThreshold) const
 {
 	GLog->Logf(TEXT("Running Compactify on %s%s"), *CloudDir, bPreview ? TEXT(". Preview mode. NO action will be taken.") : TEXT(""));
+	if (ManifestsToKeep.Num() > 0)
+	{
+		GLog->Logf(TEXT("Preserving manifest files: %s"), *FString::Join(ManifestsToKeep, TEXT(", ")));
+	}
+	GLog->Logf(TEXT("Minimum age of deleted chunks: %.3f days"), DataAgeThreshold);
+
+	// We'll work out the date of the oldest unreferenced file we'll keep
+	FDateTime Cutoff = FDateTime::UtcNow() - FTimespan::FromDays(DataAgeThreshold);
 
 	// We'll get ALL files first, so we can use the count to preallocate space within the data filenames array to save excessive reallocs
 	TArray<FString> AllFiles;
@@ -53,14 +63,30 @@ bool FBuildDataCompactifier::Compactify(bool bPreview) const
 	TSet<FGuid> ReferencedGuids; // The master list of *ALL* referenced chunk / file data Guids
 
 	TArray<FString> ManifestFilenames;
+	TArray<FString> DeletedManifestFilenames;
 	TArray<FGuid> DataGuids; // The Guids associated with the data files from a single manifest
 	int32 NumDataFiles = 0;
+	uint64 ManifestBytesDeleted = 0;
 	
 	// Preallocate enough storage in DataGuids, to stop repeatedly expanding the allocation
 	DataGuids.Reserve(AllFiles.Num());
 
 	EnumerateManifests(ManifestFilenames);
 
+	if (!DoAllManifestsExist(ManifestFilenames, ManifestsToKeep))
+	{
+		// At least one of the manifests we want to keep does not exist. This is an error condition
+		GLog->Log(ELogVerbosity::Error, TEXT("Not all manifests to keep exist. Aborting operation"));
+		return false;
+	}
+
+	if (!DeleteNonReferencedManifests(ManifestFilenames, ManifestsToKeep, DeletedManifestFilenames, ManifestBytesDeleted))
+	{
+		// An error occurred deleting one or more of the manifest files. This is an error condition
+		GLog->Log(ELogVerbosity::Error, TEXT("Could not delete one or more manifest files. Aborting operation"));
+		return false;
+	}
+
 	// If we don't have any manifest files, we'll treat that as an error condition
 	if (ManifestFilenames.Num() == 0)
 	{
@@ -68,7 +94,7 @@ bool FBuildDataCompactifier::Compactify(bool bPreview) const
 		return true;  // We're still going to return a success code, as this isn't a hard error
 	}
 
-	// Process each manifest, and build up a list of all referenced files
+	// Process each remaining manifest, and build up a list of all referenced files
 	for (const auto& ManifestFilename : ManifestFilenames)
 	{
 		const FString ManifestPath = CloudDir / ManifestFilename;
@@ -97,16 +123,24 @@ bool FBuildDataCompactifier::Compactify(bool bPreview) const
 		}
 	}
 
-	GLog->Logf(TEXT("Compactify walking %s to remove all unreferenced chunks and compute statistics."), *CloudDir);
+	GLog->Logf(TEXT("Compactify walking %s to touch referenced chunks, remove all aged unreferenced chunks and compute statistics."), *CloudDir);
 
 	uint32 FilesProcessed = 0;
+	uint32 FilesTouched = 0;
+	uint32 FilesSkipped = 0;
 	uint32 NonPatchFilesProcessed = 0;
-	uint32 FilesDeleted = 0;
+	uint32 FilesDeleted = DeletedManifestFilenames.Num();
 	uint64 BytesProcessed = 0;
+	uint64 BytesTouched = 0;
+	uint64 BytesSkipped = 0;
 	uint64 NonPatchBytesProcessed = 0;
-	uint64 BytesDeleted = 0;
+	uint64 BytesDeleted = ManifestBytesDeleted;
 	uint64 CurrentFileSize;
+	uint32 TouchFailureCount = 0;
+	uint64 BytesFailedTouch = 0;
 	FGuid FileGuid;
+	FDateTime Now = FDateTime::UtcNow();
+	const uint32 TouchFailureErrorThreshold = FMath::DivideAndRoundUp((uint32)AllFiles.Num(), BuildDataCompactifierDefs::TouchFailureThresholdPercentage);
 
 	for (const auto& File : AllFiles)
 	{
@@ -118,7 +152,8 @@ bool FBuildDataCompactifier::Compactify(bool bPreview) const
 
 			if (!GetPatchDataGuid(File, FileGuid))
 			{
-				if (!ManifestFilenames.Contains(FPaths::GetCleanFilename(File)))
+				FString CleanFilename = FPaths::GetCleanFilename(File);
+				if (!ManifestFilenames.Contains(CleanFilename) && !DeletedManifestFilenames.Contains(CleanFilename))
 				{
 					++NonPatchFilesProcessed;
 					NonPatchBytesProcessed += CurrentFileSize;
@@ -126,12 +161,39 @@ bool FBuildDataCompactifier::Compactify(bool bPreview) const
 				continue;
 			}
 
-			if (!ReferencedGuids.Contains(FileGuid))
+			if (ReferencedGuids.Contains(FileGuid))
+			{
+				// This is a valid, referenced file, so we need to touch it with the current date
+				++FilesTouched;
+				BytesTouched += CurrentFileSize;
+				if (!bPreview)
+				{
+					if (!SafeSetTimeStamp(*File, Now))
+					{
+						++TouchFailureCount;
+						BytesFailedTouch += CurrentFileSize;
+						GLog->Logf(ELogVerbosity::Warning, TEXT("Failed to set timestamp on file %s"), *File);
+						if (TouchFailureCount >= TouchFailureErrorThreshold)
+						{
+							GLog->Logf(ELogVerbosity::Error, TEXT("Failed to set timestamp on %u%% of files"), BuildDataCompactifierDefs::TouchFailureThresholdPercentage);
+							return false;
+						}
+					}
+				}
+			}
+			else if (IFileManager::Get().GetTimeStamp(*File) < Cutoff)
 			{
-				DeleteFile(File, bPreview);
+				// This file is not referenced by any manifest, is a data file, and is older than we want to keep ...
+				// Let's get rid of it!
+				DeleteFile(File);
 				++FilesDeleted;
 				BytesDeleted += CurrentFileSize;
 			}
+			else
+			{
+				++FilesSkipped;
+				BytesSkipped += CurrentFileSize;
+			}
 		}
 		else
 		{
@@ -139,17 +201,26 @@ bool FBuildDataCompactifier::Compactify(bool bPreview) const
 		}
 	}
 
+	FilesTouched -= TouchFailureCount;
+	BytesTouched -= BytesFailedTouch;
+
 	GLog->Logf(TEXT("Compactify of %s complete!"), *CloudDir);
 	GLog->Logf(TEXT("Compactify found %u files totalling %s."), FilesProcessed, *HumanReadableSize(BytesProcessed));
 	if (NonPatchFilesProcessed > 0)
 	{
-		GLog->Logf(TEXT("Of these, %d (totalling %s) were not chunk/manifest files."), NonPatchFilesProcessed, *HumanReadableSize(NonPatchBytesProcessed));
+		GLog->Logf(TEXT("Of these, %u (totalling %s) were not chunk/manifest files."), NonPatchFilesProcessed, *HumanReadableSize(NonPatchBytesProcessed));
 	}
 	GLog->Logf(TEXT("Compactify deleted %u chunk/manifest files totalling %s."), FilesDeleted, *HumanReadableSize(BytesDeleted));
+	GLog->Logf(TEXT("Compactify touched %u chunk files totalling %s."), FilesTouched, *HumanReadableSize(BytesTouched));
+	if (TouchFailureCount > 0)
+	{
+		GLog->Logf(TEXT("Compactify failed to touch %u files toatlling %s."), TouchFailureCount, *HumanReadableSize(BytesFailedTouch));
+	}
+	GLog->Logf(TEXT("Compactify skipped %u unreferenced chunk files (totalling %s) which have not yet aged out."), FilesSkipped, *HumanReadableSize(BytesSkipped));
 	return true;
 }
 
-void FBuildDataCompactifier::DeleteFile(const FString& FilePath, const bool bPreview) const
+void FBuildDataCompactifier::DeleteFile(const FString& FilePath) const
 {
 	GLog->Logf(TEXT("Deleting file %s"), *FilePath);
 	if (!bPreview)
@@ -165,6 +236,54 @@ void FBuildDataCompactifier::EnumerateManifests(TArray<FString>& OutManifests) c
 	IFileManager::Get().FindFiles(OutManifests, *FilePattern, true, false);
 }
 
+bool FBuildDataCompactifier::DoAllManifestsExist(const TArray<FString>& AllManifests, const TArray<FString>& Manifests) const
+{
+	bool bSuccess = true;
+	for (const auto& Manifest : Manifests)
+	{
+		if (!AllManifests.Contains(Manifest))
+		{
+			GLog->Logf(ELogVerbosity::Error, TEXT("Could not locate specified manifest file %s"), *Manifest);
+			bSuccess = false;
+		}
+	}
+	return bSuccess;
+}
+
+bool FBuildDataCompactifier::DeleteNonReferencedManifests(TArray<FString>& AllManifests, const TArray<FString>& ManifestsToKeep, TArray<FString>& DeletedManifests, uint64& BytesDeleted) const
+{
+	if (ManifestsToKeep.Num() == 0)
+	{
+		return true; // We don't need to delete anything, just return a success response
+	}
+
+	for (const auto& Manifest : AllManifests)
+	{
+		if (!ManifestsToKeep.Contains(Manifest))
+		{
+			GLog->Logf(TEXT("Deleting manifest file %s"), *Manifest);
+			FString ManifestPath = CloudDir / Manifest;
+
+			DeletedManifests.Add(Manifest);
+			BytesDeleted += IFileManager::Get().FileSize(*ManifestPath);
+
+			if (!bPreview)
+			{
+				IFileManager::Get().Delete(*ManifestPath);
+				if (FPaths::FileExists(ManifestPath))
+				{
+					// Something went wrong ... the file still exists!
+					GLog->Logf(ELogVerbosity::Error, TEXT("Compactify could not delete manifest file %s"), *Manifest);
+					return false;
+				}
+			}
+		}
+	}
+
+	AllManifests = ManifestsToKeep;
+	return true; // None of the manifests we need to get rid of still exist, we succeeded
+}
+
 bool FBuildDataCompactifier::GetPatchDataGuid(const FString& FilePath, FGuid& OutGuid) const
 {
 	FString Extension = FPaths::GetExtension(FilePath);
@@ -210,4 +329,18 @@ FString FBuildDataCompactifier::HumanReadableSize(uint64 NumBytes, uint8 Decimal
 
 	return FString::Printf(TEXT("%.*f %s"), DecimalPlaces, DataSize / (FMath::Pow(Base, Index)), Suffixes[bUseBase10 ? 0 : 1][Index]);
 }
+
+bool FBuildDataCompactifier::SafeSetTimeStamp(const FString& FilePath, FDateTime TimeStamp) const
+{
+	bool bSuccess = false;
+	IFileManager& FileMan = IFileManager::Get();
+	bSuccess = FileMan.SetTimeStamp(*FilePath, TimeStamp);
+	if (bSuccess)
+	{
+		FDateTime ActualTimeStamp = FileMan.GetTimeStamp(*FilePath);
+		int64 TickVariance = FMath::Abs(ActualTimeStamp.GetTicks() - TimeStamp.GetTicks());
+		bSuccess = TickVariance < BuildDataCompactifierDefs::TimeStampTolerance;
+	}
+	return bSuccess;
+}
 #endif
diff --git a/Engine/Source/Runtime/Online/BuildPatchServices/Private/BuildPatchCompactifier.h b/Engine/Source/Runtime/Online/BuildPatchServices/Private/BuildPatchCompactifier.h
--- a/Engine/Source/Runtime/Online/BuildPatchServices/Private/BuildPatchCompactifier.h
+++ b/Engine/Source/Runtime/Online/BuildPatchServices/Private/BuildPatchCompactifier.h
@@ -20,46 +20,49 @@ public:
 	 * Processes the Cloud Directory to identify and delete any orphaned chunks or files.
 	 * The CloudDir should already have been set using FBuildPatchServicesModule::SetCloudDirectory
 	 * NOTE: THIS function is blocking and will not return until finished. Don't run on main thread.
+	 * @param ManifestsToKeep   If specified, these manifests will be retained, and all others will be deleted
+	 * @param DataAgeThreshold  Chunks which are not referenced by a valid manifest, and which are older than this age (in days), will be deleted
 	 * @param bPreview      If true, then no actual work will be done, but all operations which would be carried out will be logged.
 	 * @return              true if no file errors occurred
 	 */
-	static bool CompactifyCloudDirectory(const bool bPreview);
+	static bool CompactifyCloudDirectory(const TArray<FString>& ManifestsToKeep, const float DataAgeThreshold, const bool bPreview);
 
 	/**
 	 * Processes the specified Cloud Directory to identify and delete any orphaned chunks or files.
 	 * NOTE: THIS function is blocking and will not return until finished. Don't run on main thread.
 	 * @param CloudDir      The path to the Cloud Dir to compactify
+	 * @param ManifestsToKeep   If specified, these manifests will be retained, and all others will be deleted
+	 * @param DataAgeThreshold  Chunks which are not referenced by a valid manifest, and which are older than this age (in days), will be deleted
 	 * @param bPreview      If true, then no actual work will be done, but all operations which would be carried out will be logged.
 	 * @return              true if no file errors occurred
 	 */
-	static bool CompactifyCloudDirectory(const FString& CloudDir, const bool bPreview);
+	static bool CompactifyCloudDirectory(const FString& CloudDir, const TArray<FString>& ManifestsToKeep, const float DataAgeThreshold, const bool bPreview);
 
 private:
-	FString CloudDir;
-
-	// Hide the default constructor
-	FBuildDataCompactifier();
+	const FString CloudDir;
+	const bool bPreview;
 
 	/**
 	 * Constructor.
 	 * @param CloudDir      The directory that this compactifier should operate on.
+	 * @param bPreview      If true, then no actual work will be done, but all operations which would be carried out will be logged.
 	 */
-	FBuildDataCompactifier(const FString& CloudDir);
+	FBuildDataCompactifier(const FString& CloudDir, const bool bPreview);
 
 	/**
 	 * Processes the Cloud Directory to identify and delete any orphaned chunks or files.
 	 * NOTE: THIS function is blocking and will not return until finished. Don't run on main thread.
-	 * @param bPreview      If true, then no actual work will be done, but all operations which would be carried out will be logged.
+	 * @param ManfestsToKeep      The list of manifest files that we want to keep. If specified, all other manifest files will be deleted.
+	 * @param DataAgeThreshold    The minimum age (in days) of data before it is eligible to be deleted.
 	 * @return              true if no file errors occurred
 	 */
-	bool Compactify(const bool bPreview) const;
+	bool Compactify(const TArray<FString>& ManifestsToKeep, const float DataAgeThreshold) const;
 
 	/**
 	 * Deletes the specified file, and logs the deletion
 	 * @param FilePath     The path to the file to delete
-	 * @param bPreview     If true, then the file is not actually deleted.
 	 */
-	void DeleteFile(const FString& FilePath, const bool bPreview) const;
+	void DeleteFile(const FString& FilePath) const;
 
 	/**
 	 * Obtains a list of manifest filenames, along with the manifest data, in the specified Cloud Directory
@@ -67,6 +70,23 @@ private:
 	 */
 	void EnumerateManifests(TArray<FString>& OutManifests) const;
 
+	/**
+	 * Checks to make sure that all manifests passed in exist in a list of all manifests
+	 * @param AllManifests The list of all manifests in the cloud directory
+	 * @param Manifests    The list of manifests to check for existence
+	 * @return             True if all manifests exist in AllManifests, false otherwise
+	 */
+	bool DoAllManifestsExist(const TArray<FString>& AllManifests, const TArray<FString>& Manifests) const;
+
+	/**
+	 * Determines which manifests we should delete, deletes the files, and removes them from the list of manifests
+	 * @param AllManifests      INOUT  The list of all manifest files in the cloud directory. Deleted manifests will be removed from the array
+	 * @param ManifestsToKeep   IN     The complete list of manifests we keep. If this is empty, we will keep *ALL* manifests
+	 * @param DeletedManifests  OUT    The list of manifests which were deleted by the process
+	 * @param BytesDeleted      OUT    The total number of bytes deleted by the process
+	 * @return                         True if all manifests to delete were successfully deleted, false otherwise
+	 */
+	bool DeleteNonReferencedManifests(TArray<FString>& AllManifests, const TArray<FString>& ManifestsToKeep, TArray<FString>& DeletedManifests, uint64& BytesDeleted) const;
 	/**
 	 * Determines whether or not the filename passed in contains patch data, indicated by a file extension of
 	 * .chunk or .file, and obtains the file's data Guid.
@@ -84,6 +104,14 @@ private:
 	 * @return               A string containing the human readable size
 	 */
 	FString HumanReadableSize(uint64 NumBytes, uint8 DecimalPlaces = 2, bool bUseBase10 = false) const;
+
+	/**
+	 * Sets the modified time on a file, and checks that setting the time was successful
+	 * @param FilePath    The full path to the filename on which to set the modified date
+	 * @param TimeStamp   The timestamp to apply to the file
+	 * @return            True if the operation succeeded, false otherwise
+	 */
+	bool SafeSetTimeStamp(const FString& FilePath, FDateTime TimeStamp) const;
 };
 
 #endif //WITH_BUILDPATCHGENERATION
diff --git a/Engine/Source/Runtime/Online/BuildPatchServices/Private/BuildPatchGeneration.cpp b/Engine/Source/Runtime/Online/BuildPatchServices/Private/BuildPatchGeneration.cpp
--- a/Engine/Source/Runtime/Online/BuildPatchServices/Private/BuildPatchGeneration.cpp
+++ b/Engine/Source/Runtime/Online/BuildPatchServices/Private/BuildPatchGeneration.cpp
@@ -619,7 +619,7 @@ const TMap<FGuid, int64>& FBuildDataChunkProcessor::GetChunkFilesizes()
 
 /* FBuildDataFileProcessor implementation
 *****************************************************************************/
-FBuildDataFileProcessor::FBuildDataFileProcessor( FBuildPatchAppManifestRef InBuildManifest, const FString& InBuildRoot )
+FBuildDataFileProcessor::FBuildDataFileProcessor(FBuildPatchAppManifestRef InBuildManifest, const FString& InBuildRoot, const FDateTime& InDataThresholdTime)
 	: NumNewFiles( 0 )
 	, NumKnownFiles( 0 )
 	, BuildManifest( InBuildManifest )
@@ -628,6 +628,7 @@ FBuildDataFileProcessor::FBuildDataFileProcessor( FBuildPatchAppManifestRef InBu
 	, FileHash()
 	, FileSize( 0 )
 	, IsProcessingFile( false )
+	, DataThresholdTime( InDataThresholdTime )
 {
 }
 
@@ -678,7 +679,7 @@ void FBuildDataFileProcessor::EndFile()
 	// Use hash and full file name to find out if we have a file match
 	FString FullPath = FPaths::Combine( *BuildRoot, *CurrentFile->Filename );
 	FGuid FileGuid = FGuid::NewGuid();
-	bool bFoundSameFile = FBuildDataGenerator::FindExistingFileData( FullPath, CurrentFile->FileHash, FileGuid );
+	bool bFoundSameFile = FBuildDataGenerator::FindExistingFileData( FullPath, CurrentFile->FileHash, DataThresholdTime, FileGuid );
 
 	// Fill the 'chunk part' info
 	FChunkPart& FilePart = CurrentFile->FileChunkParts[0];
@@ -717,12 +718,13 @@ void FBuildDataFileProcessor::GetFileStats( uint32& OutNewFiles, uint32& OutKnow
 
 /* FBuildSimpleChunkCache::FChunkReader implementation
 *****************************************************************************/
-FBuildGenerationChunkCache::FChunkReader::FChunkReader( const FString& InChunkFilePath, TSharedRef< FChunkFile > InChunkFile, uint32* InBytesRead )
+FBuildGenerationChunkCache::FChunkReader::FChunkReader( const FString& InChunkFilePath, TSharedRef< FChunkFile > InChunkFile, uint32* InBytesRead, const FDateTime& InDataAgeThreshold )
 	: ChunkFilePath( InChunkFilePath )
 	, ChunkFileReader( NULL )
 	, ChunkFile( InChunkFile )
 	, FileBytesRead( InBytesRead )
 	, MemoryBytesRead( 0 )
+	, DataAgeThreshold( InDataAgeThreshold )
 {
 	ChunkFile->GetDataLock( &ChunkData, &ChunkHeader );
 }
@@ -758,6 +760,13 @@ FArchive* FBuildGenerationChunkCache::FChunkReader::GetArchive()
 		{
 			*ChunkFileReader << *ChunkHeader;
 		}
+		// Check the file is not too old to reuse
+		if (IFileManager::Get().GetTimeStamp(*ChunkFilePath) < DataAgeThreshold)
+		{
+			// Break the magic to mark as invalid (this chunk is too old to reuse)
+			ChunkHeader->Magic = 0;
+			return ChunkFileReader;
+		}
 		// Check we can seek otherwise bad chunk
 		const int64 ExpectedFileSize = ChunkHeader->DataSize + ChunkHeader->HeaderSize;
 		const int64 ChunkFileSize = ChunkFileReader->TotalSize();
@@ -794,7 +803,7 @@ FArchive* FBuildGenerationChunkCache::FChunkReader::GetArchive()
 				FBuildPatchData::ChunkDataSize,
 				CompressedData.GetData(),
 				ChunkHeader->DataSize );
-			// Mark that we have fully read decompressed data and unpate the chunkfile's data size as we are expanding it
+			// Mark that we have fully read decompressed data and update the chunkfile's data size as we are expanding it
 			*FileBytesRead = FBuildPatchData::ChunkDataSize;
 			ChunkHeader->DataSize = FBuildPatchData::ChunkDataSize;
 			// Check uncompression was OK
@@ -869,8 +878,13 @@ const uint32 FBuildGenerationChunkCache::FChunkReader::BytesLeft()
 	return FBuildPatchData::ChunkDataSize - MemoryBytesRead;
 }
 
-/* FBuildSimpleChunkCache implementation
+/* FBuildGenerationChunkCache implementation
 *****************************************************************************/
+FBuildGenerationChunkCache::FBuildGenerationChunkCache(const FDateTime& DataAgeThreshold)
+	: DataAgeThreshold(DataAgeThreshold)
+{
+}
+
 TSharedRef< FBuildGenerationChunkCache::FChunkReader > FBuildGenerationChunkCache::GetChunkReader( const FString& ChunkFilePath )
 {
 	if( ChunkCache.Contains( ChunkFilePath ) == false )
@@ -898,7 +912,7 @@ TSharedRef< FBuildGenerationChunkCache::FChunkReader > FBuildGenerationChunkCach
 		ChunkCache.Add( ChunkFilePath, MakeShareable( new FChunkFile( 1, true ) ) );
 		BytesReadPerChunk.Add( ChunkFilePath, new uint32( 0 ) );
 	}
-	return MakeShareable( new FChunkReader( ChunkFilePath, ChunkCache[ ChunkFilePath ], BytesReadPerChunk[ ChunkFilePath ] ) );
+	return MakeShareable( new FChunkReader( ChunkFilePath, ChunkCache[ ChunkFilePath ], BytesReadPerChunk[ ChunkFilePath ], DataAgeThreshold ) );
 }
 
 void FBuildGenerationChunkCache::Cleanup()
@@ -911,15 +925,15 @@ void FBuildGenerationChunkCache::Cleanup()
 	BytesReadPerChunk.Empty( 0 );
 }
 
-/* FBuildSimpleChunkCache system singleton setup
+/* FBuildGenerationChunkCache system singleton setup
 *****************************************************************************/
 TSharedPtr< FBuildGenerationChunkCache > FBuildGenerationChunkCache::SingletonInstance = NULL;
 
-void FBuildGenerationChunkCache::Init()
+void FBuildGenerationChunkCache::Init(const FDateTime& DataAgeThreshold)
 {
 	// We won't allow misuse of these functions
 	check( !SingletonInstance.IsValid() );
-	SingletonInstance = MakeShareable( new FBuildGenerationChunkCache() );
+	SingletonInstance = MakeShareable(new FBuildGenerationChunkCache(DataAgeThreshold));
 }
 
 FBuildGenerationChunkCache& FBuildGenerationChunkCache::Get()
@@ -981,6 +995,7 @@ bool FBuildDataGenerator::GenerateChunksManifestFromDirectory( const FBuildPatch
 	const FString& PrereqName = Settings.PrereqName;
 	const FString& PrereqPath = Settings.PrereqPath;
 	const FString& prereqArgs = Settings.PrereqArgs;
+	const float& DataAgeThreshold = Settings.DataAgeThreshold;
 
 	// Output to log for builder info
 	GLog->Logf( TEXT( "Running Chunks Patch Generation for: %u:%s %s" ), InAppID, *AppName, *BuildVersion );
@@ -989,7 +1004,8 @@ bool FBuildDataGenerator::GenerateChunksManifestFromDirectory( const FBuildPatch
 	FScopeLock SingleConcurrentBuild( &SingleConcurrentBuildCS );
 
 	// Create our chunk cache
-	FBuildGenerationChunkCache::Init();
+	const FDateTime Cutoff = FDateTime::UtcNow() - FTimespan::FromDays(DataAgeThreshold);
+	FBuildGenerationChunkCache::Init(Cutoff);
 
 	// Create a manifest
 	FBuildPatchAppManifestRef BuildManifest = MakeShareable( new FBuildPatchAppManifest() );
@@ -1180,6 +1196,7 @@ bool FBuildDataGenerator::GenerateFilesManifestFromDirectory( const FBuildPatchS
 	const FString& LaunchExe = Settings.LaunchExe;
 	const FString& LaunchCommand = Settings.LaunchCommand;
 	const FString& IgnoreListFile = Settings.IgnoreListFile;
+	const float& DataAgeThreshold = Settings.DataAgeThreshold;
 
 	// Output to log for builder info
 	GLog->Logf( TEXT( "Running Files Patch Generation for: %u:%s %s" ), InAppID, *AppName, *BuildVersion );
@@ -1198,7 +1215,8 @@ bool FBuildDataGenerator::GenerateFilesManifestFromDirectory( const FBuildPatchS
 	BuildManifest->DestroyData(); 
 
 	// Declare a build processor
-	FBuildDataFileProcessor DataProcessor( BuildManifest, RootDirectory );
+	const FDateTime Cutoff = FDateTime::UtcNow() - FTimespan::FromDays(DataAgeThreshold);
+	FBuildDataFileProcessor DataProcessor(BuildManifest, RootDirectory, Cutoff);
 
 	// Set the App details
 	BuildManifest->AppID = InAppID;
@@ -1439,7 +1457,7 @@ FString FBuildDataGenerator::DiscoverChunkFilename(const FGuid& ChunkGuid, const
 	return TEXT("");
 }
 
-bool FBuildDataGenerator::FindExistingFileData( const FString& InSourceFile, const FSHAHash& InFileHash, FGuid& OutFileGuid )
+bool FBuildDataGenerator::FindExistingFileData( const FString& InSourceFile, const FSHAHash& InFileHash, const FDateTime& DataThresholdTime, FGuid& OutFileGuid )
 {
 	bool bFoundMatchingFile = false;
 
@@ -1578,6 +1596,13 @@ bool FBuildDataGenerator::FindExistingFileData( const FString& InSourceFile, con
 			for( auto FileIt = FileList->CreateConstIterator(); FileIt && !bFoundMatchingFile ; ++FileIt)
 			{
 				FString CloudFilename = *FileIt;
+				// Check the file date
+				FDateTime ModifiedDate = IFileManager::Get().GetTimeStamp( *CloudFilename );
+				if (ModifiedDate < DataThresholdTime)
+				{
+					// We don't want to reuse this file's Guid, as it's older than any existing files we want to consider
+					continue;
+				}
 				// Compare the files
 				FArchive* SourceFile = IFileManager::Get().CreateFileReader( *InSourceFile );
 				FArchive* FoundFile = IFileManager::Get().CreateFileReader( *CloudFilename );
diff --git a/Engine/Source/Runtime/Online/BuildPatchServices/Private/BuildPatchGeneration.h b/Engine/Source/Runtime/Online/BuildPatchServices/Private/BuildPatchGeneration.h
--- a/Engine/Source/Runtime/Online/BuildPatchServices/Private/BuildPatchGeneration.h
+++ b/Engine/Source/Runtime/Online/BuildPatchServices/Private/BuildPatchGeneration.h
@@ -371,10 +371,11 @@ public:
 
 	/**
 	 * Constructor
-	 * @param	InBuildManifest		The build manifest that receives the manifest data.
-	 * @param	InBuildRoot			The root directory of the build
+	 * @param	InBuildManifest			The build manifest that receives the manifest data.
+	 * @param	InBuildRoot				The root directory of the build
+	 * @param	InDataThresholdTime		The cutoff time used for file reuse. Files older than this will not be reused
 	 */
-	FBuildDataFileProcessor( FBuildPatchAppManifestRef InBuildManifest, const FString& InBuildRoot );
+	FBuildDataFileProcessor( FBuildPatchAppManifestRef InBuildManifest, const FString& InBuildRoot, const FDateTime& InDataThresholdTime );
 
 	/**
 	 * Mark the beginning of a new file
@@ -426,6 +427,9 @@ private:
 
 	// Are we currently processing a file
 	bool IsProcessingFile;
+
+	// The threshold time for old data. Used to determine whether to reuse chunks.
+	FDateTime DataThresholdTime;
 };
 
 /**
@@ -463,14 +467,18 @@ public:
 		// How many bytes we have read out of this chunk, used to calculate when we need more data from file
 		uint32 MemoryBytesRead;
 
+		// The validity period of chunk data. Anything older than this will be considered to be invalid
+		FDateTime DataAgeThreshold;
+
 	public:
 		/**
 		 * Constructs the reader from required data
 		 * @param	InChunkFilePath		The file path of the source file for the chunk
 		 * @param	InChunkFile			The data structure for this chunk
 		 * @param	InBytesRead			Pointer to an unsigned int to keep track of how much data is read from file
+		 * @param	InDataAgeThreshold	The validity period of chunk data. Anything older than this will be considered to be invalid
 		 */
-		FChunkReader( const FString& InChunkFilePath, TSharedRef< FChunkFile > InChunkFile, uint32* InBytesRead );
+		FChunkReader( const FString& InChunkFilePath, TSharedRef< FChunkFile > InChunkFile, uint32* InBytesRead, const FDateTime& InDataAgeThreshold );
 
 		/**
 		 * Destructor. Cleans up and releases data locks.
@@ -511,10 +519,19 @@ public:
 		FArchive* GetArchive();
 	};
 
+private:
+	/**
+	* Constructor
+	* @param	DataAgeThreshold		The cutoff point before which data chunks should be regarded as isvalid
+	*/
+	FBuildGenerationChunkCache(const FDateTime& DataAgeThreshold);
+
 private:
 	// The number of chunks to cache in memory
 	const static int32 NumChunksToCache = 256;
 
+	FDateTime DataAgeThreshold;
+
 	// The map of chunks we have cached
 	TMap< FString, TSharedRef< FChunkFile > > ChunkCache;
 
@@ -533,8 +550,9 @@ public:
 
 	/**
 	 * Call to initialise static singleton so that the cache system can be used
+	 * @param DataAgeThreshold		Any chunks older than this date will be classed as invalid
 	 */
-	static void Init();
+	static void Init(const FDateTime& DataAgeThreshold);
 
 	/**
 	 * Get a reference to the chunk system. Only call between Init and Shutdown calls.
@@ -604,12 +622,13 @@ public:
 	/**
 	 * Checks to see if a given file already exists in the database of known files
 	 * NOTE: This function is blocking and will not return until finished. Don't run on main thread.
-	 * @param InSourceFile	IN		The filename of the file on disk.
-	 * @param InFileHash	IN		The hash for the file.
-	 * @param OutFileGuid	OUT		Will be set to the existing file guid if found.
+	 * @param InSourceFile			IN		The filename of the file on disk.
+	 * @param InFileHash			IN		The hash for the file.
+	 * @param DataThresholdTime		IN		The date/time before which files will not be reused
+	 * @param OutFileGuid			OUT		Will be set to the existing file guid if found.
 	 * @return		Whether this file is new or existing.
 	 */
-	static bool FindExistingFileData( const FString& InSourceFile, const FSHAHash& InFileHash, FGuid& OutFileGuid );
+	static bool FindExistingFileData( const FString& InSourceFile, const FSHAHash& InFileHash, const FDateTime& DataThresholdTime, FGuid& OutFileGuid );
 
 	/**
 	 * Saves out the file data for use with file based patch manifests
@@ -659,6 +678,7 @@ private:
 
 	// A critical section to enforce only one build at a time
 	static FCriticalSection SingleConcurrentBuildCS;
+
 };
 
 #endif //WITH_BUILDPATCHGENERATION
diff --git a/Engine/Source/Runtime/Online/BuildPatchServices/Private/BuildPatchServicesModule.cpp b/Engine/Source/Runtime/Online/BuildPatchServices/Private/BuildPatchServicesModule.cpp
--- a/Engine/Source/Runtime/Online/BuildPatchServices/Private/BuildPatchServicesModule.cpp
+++ b/Engine/Source/Runtime/Online/BuildPatchServices/Private/BuildPatchServicesModule.cpp
@@ -189,9 +189,9 @@ bool FBuildPatchServicesModule::GenerateFilesManifestFromDirectory( const FBuild
 	return FBuildDataGenerator::GenerateFilesManifestFromDirectory( Settings );
 }
 
-bool FBuildPatchServicesModule::CompactifyCloudDirectory(const bool bPreview)
+bool FBuildPatchServicesModule::CompactifyCloudDirectory(const TArray<FString>& ManifestsToKeep, const float DataAgeThreshold, const bool bPreview)
 {
-	return FBuildDataCompactifier::CompactifyCloudDirectory(bPreview);
+	return FBuildDataCompactifier::CompactifyCloudDirectory(ManifestsToKeep, DataAgeThreshold, bPreview);
 }
 
 #endif //WITH_BUILDPATCHGENERATION
diff --git a/Engine/Source/Runtime/Online/BuildPatchServices/Private/BuildPatchServicesModule.h b/Engine/Source/Runtime/Online/BuildPatchServices/Private/BuildPatchServicesModule.h
--- a/Engine/Source/Runtime/Online/BuildPatchServices/Private/BuildPatchServicesModule.h
+++ b/Engine/Source/Runtime/Online/BuildPatchServices/Private/BuildPatchServicesModule.h
@@ -85,7 +85,7 @@ public:
 #if WITH_BUILDPATCHGENERATION
 	virtual bool GenerateChunksManifestFromDirectory( const FBuildPatchSettings& Settings ) override;
 	virtual bool GenerateFilesManifestFromDirectory( const FBuildPatchSettings& Settings ) override;
-	virtual bool CompactifyCloudDirectory(const bool bPreview) override;
+	virtual bool CompactifyCloudDirectory(const TArray<FString>& ManifestsToKeep, const float DataAgeThreshold, const bool bPreview) override;
 #endif
 
 	// End IBuildPatchServicesModule interface
diff --git a/Engine/Source/Runtime/Online/BuildPatchServices/Public/BuildPatchSettings.h b/Engine/Source/Runtime/Online/BuildPatchServices/Public/BuildPatchSettings.h
--- a/Engine/Source/Runtime/Online/BuildPatchServices/Public/BuildPatchSettings.h
+++ b/Engine/Source/Runtime/Online/BuildPatchServices/Public/BuildPatchSettings.h
@@ -26,6 +26,8 @@ struct FBuildPatchSettings
 	FString PrereqPath;
 	// The command line arguments for the prerequisites installer
 	FString PrereqArgs;
+	// The maximum age (in days) of existing data files which can be reused in this build
+	float DataAgeThreshold;
 	// Map of custom fields to add to the manifest
 	TMap<FString, FVariant> CustomFields;
 };
diff --git a/Engine/Source/Runtime/Online/BuildPatchServices/Public/Interfaces/IBuildPatchServicesModule.h b/Engine/Source/Runtime/Online/BuildPatchServices/Public/Interfaces/IBuildPatchServicesModule.h
--- a/Engine/Source/Runtime/Online/BuildPatchServices/Public/Interfaces/IBuildPatchServicesModule.h
+++ b/Engine/Source/Runtime/Online/BuildPatchServices/Public/Interfaces/IBuildPatchServicesModule.h
@@ -111,9 +111,11 @@ public:
 	/**
 	 * Processes a Cloud Directory to identify and delete any orphaned chunks or files.
 	 * NOTE: THIS function is blocking and will not return until finished. Don't run on main thread.
+	 * @param ManifestsToKeep      If specified, these manifests will be retained, and all others will be deleted
+	 * @param DataAgeThreshold     Chunks which are not referenced by a valid manifest, and which are older than this age (in days), will be deleted
 	 * @param bPreview      If true, then no actual work will be done, but all operations which would be carried out will be logged.
 	 * @return              true if no file errors occurred
 	 */
-	virtual bool CompactifyCloudDirectory( const bool bPreview ) = 0;
+	virtual bool CompactifyCloudDirectory(const TArray<FString>& ManifestsToKeep, const float DataAgeThreshold, const bool bPreview) = 0;
 #endif // WITH_BUILDPATCHGENERATION
 };
