App Manifest Build Size query optimisation
Vault Cache size FText formatting optimisation

diff --git a/Engine/Source/Runtime/Online/BuildPatchServices/Private/BuildPatchManifest.cpp b/Engine/Source/Runtime/Online/BuildPatchServices/Private/BuildPatchManifest.cpp
--- a/Engine/Source/Runtime/Online/BuildPatchServices/Private/BuildPatchManifest.cpp
+++ b/Engine/Source/Runtime/Online/BuildPatchServices/Private/BuildPatchManifest.cpp
@@ -123,6 +123,7 @@ FBuildPatchAppManifest::FBuildPatchAppManifest( const uint32& InAppID, const FSt
 	: ManifestFileVersion( EBuildPatchAppManifestVersion::GetLatestVersion() )
 	, AppID( InAppID )
 	, AppNameString( AppName )
+	, TotalBuildSize( INDEX_NONE )
 {
 }
 
@@ -477,6 +478,16 @@ bool FBuildPatchAppManifest::DeserializeFromJSON( const FString& JSONInput )
 		}
 	}
 
+	// Calculate build size
+	TotalBuildSize = 0;
+	for (auto& File : FileManifestList)
+	{
+		if (File.Value.IsValid())
+		{
+			TotalBuildSize += File.Value->GetFileSize();
+		}
+	}
+
 	// Make sure we don't have any half loaded data
 	if( !bSuccess )
 	{
@@ -531,16 +542,7 @@ const int64 FBuildPatchAppManifest::GetDownloadSize() const
 
 const int64 FBuildPatchAppManifest::GetBuildSize() const
 {
-	int64 TotalSize = 0;
-	for ( auto& File : FileManifestList )
-	{
-		if ( File.Value.IsValid() )
-		{
-			TotalSize += File.Value->GetFileSize();
-		}
-	}
-
-	return TotalSize;
+	return TotalBuildSize;
 }
 
 const int64 FBuildPatchAppManifest::GetFileSize( const TArray< FString >& Filenames ) const
diff --git a/Engine/Source/Runtime/Online/BuildPatchServices/Private/BuildPatchManifest.h b/Engine/Source/Runtime/Online/BuildPatchServices/Private/BuildPatchManifest.h
--- a/Engine/Source/Runtime/Online/BuildPatchServices/Private/BuildPatchManifest.h
+++ b/Engine/Source/Runtime/Online/BuildPatchServices/Private/BuildPatchManifest.h
@@ -439,4 +439,9 @@ private:
 	 * Holds the list of custom fields
 	 */
 	TMap< FString, FString > CustomFields;
+
+	/**
+	 * Holds the total build size in bytes
+	 */
+	int64 TotalBuildSize;
 };
