Curl: support cert bundles at different locations.

- CURLOPT_CAPATH did not work for some reason.
- Hardcoded, but seems to cover all practical cases.
- Only used on Linux so far, Android/Windows testing welcome.

diff --git a/Engine/Source/Runtime/Online/HTTP/Private/Curl/CurlHttp.cpp b/Engine/Source/Runtime/Online/HTTP/Private/Curl/CurlHttp.cpp
--- a/Engine/Source/Runtime/Online/HTTP/Private/Curl/CurlHttp.cpp
+++ b/Engine/Source/Runtime/Online/HTTP/Private/Curl/CurlHttp.cpp
@@ -61,6 +61,11 @@ FCurlHttpRequest::FCurlHttpRequest(CURLM * InMultiHandle)
 		{
 			curl_easy_setopt(EasyHandle, CURLOPT_FORBID_REUSE, 1L);
 		}
+
+		if (FCurlHttpManager::CurlRequestOptions.CertBundlePath)
+		{
+			curl_easy_setopt(EasyHandle, CURLOPT_CAINFO, FCurlHttpManager::CurlRequestOptions.CertBundlePath);
+		}
 	}
 }
 
diff --git a/Engine/Source/Runtime/Online/HTTP/Private/Curl/CurlHttpManager.cpp b/Engine/Source/Runtime/Online/HTTP/Private/Curl/CurlHttpManager.cpp
--- a/Engine/Source/Runtime/Online/HTTP/Private/Curl/CurlHttpManager.cpp
+++ b/Engine/Source/Runtime/Online/HTTP/Private/Curl/CurlHttpManager.cpp
@@ -101,6 +101,33 @@ void FCurlHttpManager::InitCurl()
 		CurlRequestOptions.bDontReuseConnections = true;
 	}
 
+	// discover cert location
+	if (PLATFORM_LINUX)	// only relevant to Linux (for now?), not #ifdef'ed to keep the code checked by the compiler when compiling for other platforms
+	{
+		static const char * KnownBundlePaths[] =
+		{
+			"/etc/pki/tls/certs/ca-bundle.crt",
+			"/etc/ssl/certs/ca-certificates.crt",
+			nullptr
+		};
+
+		for (const char ** CurrentBundle = KnownBundlePaths; *CurrentBundle; ++CurrentBundle)
+		{
+			FString FileName(*CurrentBundle);
+			UE_LOG(LogInit, Log, TEXT(" Libcurl: checking if '%s' exists"), *FileName);
+
+			if (FPaths::FileExists(FileName))
+			{
+				CurlRequestOptions.CertBundlePath = *CurrentBundle;
+				break;
+			}
+		}
+		if (CurlRequestOptions.CertBundlePath == nullptr)
+		{
+			UE_LOG(LogInit, Log, TEXT(" Libcurl: did not find a cert bundle in any of known locations, TLS may not work"));
+		}
+	}
+
 	// print for visibility
 	CurlRequestOptions.Log();
 }
@@ -126,6 +153,11 @@ void FCurlHttpManager::FCurlRequestOptions::Log()
 		bDontReuseConnections ? TEXT("true") : TEXT("false"),
 		bDontReuseConnections ? TEXT("NOT ") : TEXT("")
 		);
+
+	UE_LOG(LogInit, Log, TEXT(" - CertBundlePath = %s  - Libcurl will %s"),
+		(CertBundlePath != nullptr) ? *FString(CertBundlePath) : TEXT("nullptr"),
+		(CertBundlePath != nullptr) ? TEXT("set CURLOPT_CAINFO to it") : TEXT("use whatever was configured at build time.")
+		);
 }
 
 
diff --git a/Engine/Source/Runtime/Online/HTTP/Private/Curl/CurlHttpManager.h b/Engine/Source/Runtime/Online/HTTP/Private/Curl/CurlHttpManager.h
--- a/Engine/Source/Runtime/Online/HTTP/Private/Curl/CurlHttpManager.h
+++ b/Engine/Source/Runtime/Online/HTTP/Private/Curl/CurlHttpManager.h
@@ -45,6 +45,7 @@ public:
 			:	bVerifyPeer(true)
 			,	bUseHttpProxy(false)
 			,	bDontReuseConnections(false)
+			,	CertBundlePath(nullptr)
 		{}
 
 		/** Prints out the options to the log */
@@ -61,6 +62,9 @@ public:
 
 		/** Address of the HTTP proxy */
 		FString HttpProxyAddress;
+
+		/** A path to certificate bundle */
+		const char * CertBundlePath;
 	}
 	CurlRequestOptions;
 };
