Fix crashes on exit when using libcurl on Windows.

(merging CL 2295849 from Dev branch).

diff --git a/Engine/Source/Runtime/Online/HTTP/Private/HttpModule.cpp b/Engine/Source/Runtime/Online/HTTP/Private/HttpModule.cpp
--- a/Engine/Source/Runtime/Online/HTTP/Private/HttpModule.cpp
+++ b/Engine/Source/Runtime/Online/HTTP/Private/HttpModule.cpp
@@ -48,26 +48,30 @@ void FHttpModule::StartupModule()
 
 void FHttpModule::ShutdownModule()
 {
-#ifdef PLATFORM_WINDOWS
+#if PLATFORM_WINDOWS
 
-	// due to pecularities of some platforms (notably Windows) we need to shutdown platform http first,
-	// then delete the manager. It is more logical to have reverse order of their creation though. Proper fix
-	// would be refactoring HTTP platform abstraction to make HttpManager a proper part of it.
-	FPlatformHttp::Shutdown();
-
-	delete HttpManager;	// can be passed NULLs
+	extern bool bUseCurl;
+	if (!bUseCurl)
+	{
+		// due to peculiarities of some platforms (notably Windows with WinInet implementation) we need to shutdown platform http first,
+		// then delete the manager. It is more logical to have reverse order of their creation though. Proper fix
+		// would be refactoring HTTP platform abstraction to make HttpManager a proper part of it.
+		FPlatformHttp::Shutdown();
 
+		delete HttpManager;	// can be passed NULLs
+	}
+	else
 #else
+	{
+		// at least on Linux, the code in HTTP manager (e.g. request destructors) expects platform to be initialized yet
+		delete HttpManager;	// can be passed NULLs
 
-	// at least on Linux, the code in HTTP manager (e.g. request destructors) expects platform to be initialized yet
-	delete HttpManager;	// can be passed NULLs
-
-	FPlatformHttp::Shutdown();
-
-#endif
+		FPlatformHttp::Shutdown();
+	}
+#endif	// PLATFORM_WINDOWS
 
-	HttpManager = NULL;
-	Singleton = NULL;
+	HttpManager = nullptr;
+	Singleton = nullptr;
 }
 
 bool FHttpModule::HandleHTTPCommand( const TCHAR* Cmd, FOutputDevice& Ar )
