Fix bug in HTTP shutdown logic on Windows.

(Merging CL 2296485 from Dev branch. Caught by DavidN's code review).

diff --git a/Engine/Source/Runtime/Online/HTTP/Private/HttpModule.cpp b/Engine/Source/Runtime/Online/HTTP/Private/HttpModule.cpp
--- a/Engine/Source/Runtime/Online/HTTP/Private/HttpModule.cpp
+++ b/Engine/Source/Runtime/Online/HTTP/Private/HttpModule.cpp
@@ -61,14 +61,13 @@ void FHttpModule::ShutdownModule()
 		delete HttpManager;	// can be passed NULLs
 	}
 	else
-#else
+#endif	// PLATFORM_WINDOWS
 	{
 		// at least on Linux, the code in HTTP manager (e.g. request destructors) expects platform to be initialized yet
 		delete HttpManager;	// can be passed NULLs
 
 		FPlatformHttp::Shutdown();
 	}
-#endif	// PLATFORM_WINDOWS
 
 	HttpManager = nullptr;
 	Singleton = nullptr;
