Fix some comments

diff --git a/Engine/Source/Runtime/Online/OnlineSubsystem/Public/Interfaces/OnlineIdentityInterface.h b/Engine/Source/Runtime/Online/OnlineSubsystem/Public/Interfaces/OnlineIdentityInterface.h
--- a/Engine/Source/Runtime/Online/OnlineSubsystem/Public/Interfaces/OnlineIdentityInterface.h
+++ b/Engine/Source/Runtime/Online/OnlineSubsystem/Public/Interfaces/OnlineIdentityInterface.h
@@ -49,9 +49,13 @@ namespace EUserPrivileges
 {
 	enum Type
 	{
-		CanPlay, // This one can be used for age restrictions on PS4
+		/** Whether the user can play at all, online or offline - may be age restricted */
+		CanPlay,
+		/** Whether the user can play in online modes */
 		CanPlayOnline,
-		CanCommunicateOnline, // For voice AND text chat
+		/** Whether the user can use voice and text chat */
+		CanCommunicateOnline,
+		/** Whether the user can use content generated by other users */
 		CanUseUserGeneratedContent
 	};
 }
@@ -119,15 +123,24 @@ public:
 
 	enum class EPrivilegeResults
 	{
+		/** The user has the requested privilege */
 		NoFailures				=	0,
+		/** Patch required before the user can use the privilege */
 		RequiredPatchAvailable	=	1 << 0,
+		/** System update required before the user can use the privilege */
 		RequiredSystemUpdate	=	1 << 1,
-		AgeRestrictionFailure	=	1 << 2, // parental control failure usually
-		AccountTypeFailure		=	1 << 3, // XboxLive Gold / PSPlus required but not available
-		UserNotFound			=	1 << 4, // Invalid user.
-		ChatRestriction			=	1 << 5, // User restricted from chat
-		UGCRestriction			=	1 << 6, // User restricted from User Generated Content
-		GenericFailure			=	1 << 7, // Platform failed for unknown reason and handles its own dialogs
+		/** Parental control failure usually */
+		AgeRestrictionFailure	=	1 << 2,
+		/** XboxLive Gold / PSPlus required but not available */
+		AccountTypeFailure		=	1 << 3,
+		/** Invalid user */
+		UserNotFound			=	1 << 4,
+		/** User restricted from chat */
+		ChatRestriction			=	1 << 5,
+		/** User restricted from User Generated Content */
+		UGCRestriction			=	1 << 6,
+		/** Platform failed for unknown reason and handles its own dialogs */
+		GenericFailure			=	1 << 7, 
 	};
 
 	virtual ~IOnlineIdentity() {};
@@ -143,6 +156,7 @@ public:
 	 * Delegate called when a player's login status changes but doesn't change identity
 	 *
 	 * @param LocalUserNum the controller number of the associated user
+	 * @param OldStatus the previous status of the user
 	 * @param NewStatus the new login status for the user
 	 * @param NewId the new id to associate with the user
 	 */
