Merging using REVERSE UE4-To-UE4-LauncherReleases

CL# 2299690 on 16/09/2014 13:09:06
------------------------------------------
Manifest files are now saved out as BOM-less UTF8

diff --git a/Engine/Source/Runtime/Online/BuildPatchServices/Private/BuildPatchManifest.cpp b/Engine/Source/Runtime/Online/BuildPatchServices/Private/BuildPatchManifest.cpp
--- a/Engine/Source/Runtime/Online/BuildPatchServices/Private/BuildPatchManifest.cpp
+++ b/Engine/Source/Runtime/Online/BuildPatchServices/Private/BuildPatchManifest.cpp
@@ -156,7 +156,8 @@ bool FBuildPatchAppManifest::SaveToFile(FString Filename)
 	{
 		FString JSONOutput;
 		SerializeToJSON(JSONOutput);
-		FileOut->Serialize(TCHAR_TO_ANSI(*JSONOutput), JSONOutput.Len());
+		FTCHARToUTF8 JsonUTF8(*JSONOutput);
+		FileOut->Serialize((UTF8CHAR*)JsonUTF8.Get(), JsonUTF8.Length() * sizeof(UTF8CHAR));
 		FileOut->Close();
 		delete FileOut;
 	}
