In app purchases, better name for enabling flag. There was a mismatch

diff --git a/Engine/Source/Runtime/Online/IOS/OnlineSubsystemIOS/Private/OnlineSubsystemIOS.cpp b/Engine/Source/Runtime/Online/IOS/OnlineSubsystemIOS/Private/OnlineSubsystemIOS.cpp
--- a/Engine/Source/Runtime/Online/IOS/OnlineSubsystemIOS/Private/OnlineSubsystemIOS.cpp
+++ b/Engine/Source/Runtime/Online/IOS/OnlineSubsystemIOS/Private/OnlineSubsystemIOS.cpp
@@ -188,6 +188,6 @@ bool FOnlineSubsystemIOS::IsEnabled()
 bool FOnlineSubsystemIOS::IsInAppPurchasingEnabled()
 {
 	bool bEnableIAP = false;
-	GConfig->GetBool(TEXT("OnlineSubsystemIOS.Store"), TEXT("bSupportInAppPurchasing"), bEnableIAP, GEngineIni);
+	GConfig->GetBool(TEXT("OnlineSubsystemIOS.Store"), TEXT("bSupportsInAppPurchasing"), bEnableIAP, GEngineIni);
 	return bEnableIAP;
 }
\ No newline at end of file
