Files being overwritten or removed are now properly verified before backing up to avoid false backup

diff --git a/Engine/Source/Runtime/Online/BuildPatchServices/Private/BuildPatchInstaller.cpp b/Engine/Source/Runtime/Online/BuildPatchServices/Private/BuildPatchInstaller.cpp
--- a/Engine/Source/Runtime/Online/BuildPatchServices/Private/BuildPatchInstaller.cpp
+++ b/Engine/Source/Runtime/Online/BuildPatchServices/Private/BuildPatchInstaller.cpp
@@ -607,19 +607,12 @@ bool FBuildPatchInstaller::BackupFileIfNecessary(const FString& Filename, bool b
 	{
 		return IFileManager::Get().Move(*BackupFilename, *InstalledFilename, true, true, true);
 	}
-	// There are lots of checks to do in order to figure out whether a file needs backing up.
-	TSharedPtr< FBuildPatchFileManifest > OldFileManifest = CurrentBuildManifest.IsValid() ? CurrentBuildManifest->GetFileManifest(Filename) : NULL;
-	TSharedPtr< FBuildPatchFileManifest > NewFileManifest = NewBuildManifest->GetFileManifest(Filename);
-	const bool bCurrentlyInstalled = CurrentBuildManifest.IsValid();
-	const bool bPerformingRepair = CurrentBuildManifest == NewBuildManifest;
-	const bool bPerformingPatch = bCurrentlyInstalled && !bPerformingRepair;
-	const bool bFreshInstall = !bCurrentlyInstalled;
-	const bool bFileIsNewAddition = OldFileManifest.IsValid() == false;
-	const bool bFileRemovedByPatch = bPerformingPatch && !NewFileManifest.IsValid();
 	bool bUserEditedFile = bDiscoveredByVerification;
-	const bool bCheckFileChanges = !bDiscoveredByVerification && ((!bPerformingRepair && !bFileIsNewAddition) || bFileRemovedByPatch);
+	const bool bCheckFileChanges = !bDiscoveredByVerification;
 	if (bCheckFileChanges)
 	{
+		TSharedPtr< FBuildPatchFileManifest > OldFileManifest = CurrentBuildManifest.IsValid() ? CurrentBuildManifest->GetFileManifest(Filename) : NULL;
+		TSharedPtr< FBuildPatchFileManifest > NewFileManifest = NewBuildManifest->GetFileManifest(Filename);
 		const int64 InstalledFilesize = IFileManager::Get().FileSize(*InstalledFilename);
 		const int64 OriginalFileSize = OldFileManifest.IsValid() ? OldFileManifest->GetFileSize() : INDEX_NONE;
 		const int64 NewFileSize = NewFileManifest.IsValid() ? NewFileManifest->GetFileSize() : INDEX_NONE;
@@ -630,11 +623,11 @@ bool FBuildPatchInstaller::BackupFileIfNecessary(const FString& Filename, bool b
 		bUserEditedFile = bFileSizeDiffers || FBuildPatchUtils::VerifyFile(InstalledFilename, HashOld, HashNew) == 0;
 	}
 	// Finally, use the above logic to determine if we must do the backup
-	const bool bNeedBackup = bFileIsNewAddition || bPerformingRepair || bUserEditedFile;
+	const bool bNeedBackup = bUserEditedFile;
 	bool bBackupSuccess = true;
 	if (bNeedBackup)
 	{
-		GLog->Logf(TEXT("BuildPatchServices: Backing up (%d,%d,%d) %s"), bFileIsNewAddition?1:0, bPerformingRepair?1:0, bUserEditedFile?1:0, *Filename);
+		GLog->Logf(TEXT("BuildPatchServices: Backing up %s"), *Filename);
 		bBackupSuccess = IFileManager::Get().Move(*BackupFilename, *InstalledFilename, true, true, true);
 	}
 	return bBackupSuccess;
