Limit number of Google Play login attempts
Fixes being unable to cancel Google Play login

diff --git a/Engine/Source/Runtime/Online/Android/OnlineSubsystemGooglePlay/Private/OnlineAsyncTaskGooglePlayLogin.cpp b/Engine/Source/Runtime/Online/Android/OnlineSubsystemGooglePlay/Private/OnlineAsyncTaskGooglePlayLogin.cpp
--- a/Engine/Source/Runtime/Online/Android/OnlineSubsystemGooglePlay/Private/OnlineAsyncTaskGooglePlayLogin.cpp
+++ b/Engine/Source/Runtime/Online/Android/OnlineSubsystemGooglePlay/Private/OnlineAsyncTaskGooglePlayLogin.cpp
@@ -3,12 +3,18 @@
 #include "OnlineSubsystemGooglePlayPrivatePCH.h"
 #include "OnlineAsyncTaskGooglePlayLogin.h"
 
+namespace
+{
+	const int MaxRetries = 1;
+}
+
 FOnlineAsyncTaskGooglePlayLogin::FOnlineAsyncTaskGooglePlayLogin(
 	FOnlineSubsystemGooglePlay* InSubsystem,
 	int InPlayerId)
 	: FOnlineAsyncTaskBasic(InSubsystem)
 	, PlayerId(InPlayerId)
 	, Status(gpg::AuthStatus::ERROR_NOT_AUTHORIZED)
+	, RetryCount(MaxRetries)
 {
 }
 
@@ -25,8 +31,9 @@ void FOnlineAsyncTaskGooglePlayLogin::OnAuthActionFinished(gpg::AuthOperation In
 {
 	if (InOp == AuthOperation::SIGN_IN)
 	{
-		if (InStatus == AuthStatus::ERROR_NOT_AUTHORIZED)
+		if (InStatus == AuthStatus::ERROR_NOT_AUTHORIZED && RetryCount > 0)
 		{
+			--RetryCount;
 			Subsystem->GetGameServices()->StartAuthorizationUI();
 			return;
 		}
diff --git a/Engine/Source/Runtime/Online/Android/OnlineSubsystemGooglePlay/Private/OnlineAsyncTaskGooglePlayLogin.h b/Engine/Source/Runtime/Online/Android/OnlineSubsystemGooglePlay/Private/OnlineAsyncTaskGooglePlayLogin.h
--- a/Engine/Source/Runtime/Online/Android/OnlineSubsystemGooglePlay/Private/OnlineAsyncTaskGooglePlayLogin.h
+++ b/Engine/Source/Runtime/Online/Android/OnlineSubsystemGooglePlay/Private/OnlineAsyncTaskGooglePlayLogin.h
@@ -39,4 +39,5 @@ PACKAGE_SCOPE:
 private:
 	int PlayerId;
 	gpg::AuthStatus Status;
+	int RetryCount;
 };
