In App Purchase Receipts

Exposed as part of the callbacks for a completed iap.
Platform agnostic IPlatformPurchaseReceipt class added for other store kits to extend their own nplementation of the receipts.
IOS implementation added.

Receipts for IOS are formatted in the format:
Identifier | String
Data | Base64 encoded string

diff --git a/Engine/Source/Runtime/Online/IOS/OnlineSubsystemIOS/Private/OnlineStoreInterfaceIOS.cpp b/Engine/Source/Runtime/Online/IOS/OnlineSubsystemIOS/Private/OnlineStoreInterfaceIOS.cpp
--- a/Engine/Source/Runtime/Online/IOS/OnlineSubsystemIOS/Private/OnlineStoreInterfaceIOS.cpp
+++ b/Engine/Source/Runtime/Online/IOS/OnlineSubsystemIOS/Private/OnlineStoreInterfaceIOS.cpp
@@ -55,19 +55,35 @@
     {
         IOnlineSubsystem* OnlineSub = IOnlineSubsystem::Get(FName(TEXT("IOS")));
         FOnlineStoreInterfaceIOS* StoreInterface = (FOnlineStoreInterfaceIOS*)OnlineSub->GetStoreInterface().Get();
+     
+        FIOSPurchaseReceipt Receipt;
+        Receipt.Identifier = transaction.transactionIdentifier;
+     
+        if( [IOSAppDelegate GetDelegate].OSVersion >= 7.0 )
+        {
+            NSURL* ReceiptUrl = [[NSBundle mainBundle] appStoreReceiptURL];
+            NSData* ReceiptData = [NSData dataWithContentsOfURL:ReceiptUrl];
+            NSString* nsEncodedReceiptData = [ReceiptData base64EncodedStringWithOptions:NSDataBase64EncodingEndLineWithLineFeed];
+     
+            Receipt.Data = nsEncodedReceiptData;
+        }
+        else
+        {
+            // If earlier than IOS 7, we will need to use the transactionReceipt
+            NSString* nsEncodedReceiptData = [transaction.transactionReceipt base64EncodedStringWithOptions:NSDataBase64EncodingEndLineWithLineFeed];
+            Receipt.Data = nsEncodedReceiptData;
+        }
+     
+
         if (StoreInterface->CachedPurchaseStateObject.IsValid())
         {
             StoreInterface->CachedPurchaseStateObject->ReadState = EOnlineAsyncTaskState::Done;
         }
      
-        StoreInterface->TriggerOnInAppPurchaseCompleteDelegates(EInAppPurchaseState::Success);
+        StoreInterface->TriggerOnInAppPurchaseCompleteDelegates(EInAppPurchaseState::Success, &Receipt);
      
         return true;
     }];
-    
-    // Your application should implement these two methods.
-    //[self recordTransaction:transaction];
-	//[self provideContent:transaction.payment.productIdentifier];
  
     // Remove the transaction from the payment queue.
     [[SKPaymentQueue defaultQueue] finishTransaction: transaction];
@@ -112,7 +128,8 @@
         {
             StoreInterface->CachedPurchaseStateObject->ReadState = EOnlineAsyncTaskState::Done;
         }
-        StoreInterface->TriggerOnInAppPurchaseCompleteDelegates(CompletionState);
+
+        StoreInterface->TriggerOnInAppPurchaseCompleteDelegates(CompletionState, nullptr);
      
         return true;
      }];
@@ -242,7 +259,8 @@ bool FOnlineStoreInterfaceIOS::BeginPurchase(const FString& ProductId, FOnlineIn
 	if( bIsPurchasing || bIsProductRequestInFlight )
 	{
 		UE_LOG(LogOnline, Verbose, TEXT( "FOnlineStoreInterfaceIOS::BeginPurchase - cannot make a purchase whilst one is in transaction." ));
-		TriggerOnInAppPurchaseCompleteDelegates(EInAppPurchaseState::Failed);
+        
+		TriggerOnInAppPurchaseCompleteDelegates(EInAppPurchaseState::Failed, nullptr);
 	}
 	else if( IsAllowedToMakePurchases() )
 	{
@@ -273,7 +291,8 @@ bool FOnlineStoreInterfaceIOS::BeginPurchase(const FString& ProductId, FOnlineIn
 	{
 		UE_LOG(LogOnline, Verbose, TEXT("This device is not able to make purchases."));
 		InPurchaseStateObject->ReadState = EOnlineAsyncTaskState::Failed;
-		TriggerOnInAppPurchaseCompleteDelegates(EInAppPurchaseState::NotAllowed);
+        
+        TriggerOnInAppPurchaseCompleteDelegates(EInAppPurchaseState::NotAllowed, nullptr);
 	}
 
 	return bCreatedNewTransaction;
diff --git a/Engine/Source/Runtime/Online/IOS/OnlineSubsystemIOS/Private/OnlineStoreInterfaceIOS.h b/Engine/Source/Runtime/Online/IOS/OnlineSubsystemIOS/Private/OnlineStoreInterfaceIOS.h
--- a/Engine/Source/Runtime/Online/IOS/OnlineSubsystemIOS/Private/OnlineStoreInterfaceIOS.h
+++ b/Engine/Source/Runtime/Online/IOS/OnlineSubsystemIOS/Private/OnlineStoreInterfaceIOS.h
@@ -14,6 +14,35 @@
 #include <StoreKit/SKPaymentQueue.h>
 
 
+/**
+ * The resulting state of an iap transaction
+ */
+namespace EInAppPurchaseResult
+{
+    enum Type
+    {
+        Succeeded = 0,
+        RestoredFromServer,
+        Failed,
+        Cancelled,
+    };
+}
+
+/**
+ * IOS implementation of the Platform Purchase receipt. For this we provide an identifier and the encrypted data.
+ */
+class FIOSPurchaseReceipt : public IPlatformPurchaseReceipt
+{
+public:
+    // Product identifier
+    FString Identifier;
+    
+    // The encrypted receipt data
+    FString Data;
+};
+
+
+
 /** Helper class, which allows us to manage IAP product information requests, AND transactions */
 @interface FStoreKitHelper : NSObject<SKProductsRequestDelegate, SKPaymentTransactionObserver>
 {
diff --git a/Engine/Source/Runtime/Online/OnlineSubsystem/Public/Interfaces/OnlineStoreInterface.h b/Engine/Source/Runtime/Online/OnlineSubsystem/Public/Interfaces/OnlineStoreInterface.h
--- a/Engine/Source/Runtime/Online/OnlineSubsystem/Public/Interfaces/OnlineStoreInterface.h
+++ b/Engine/Source/Runtime/Online/OnlineSubsystem/Public/Interfaces/OnlineStoreInterface.h
@@ -21,6 +21,15 @@ namespace EInAppPurchaseState
 	};
 }
 
+/**
+ * Platform derivitive for their own interpretation of product receipts.
+ * Each Store implementation should provide their own receipt type extended from this.
+ */
+class IPlatformPurchaseReceipt
+{
+};
+
+
 /**
  * Delegate fired when a session create request has completed
  *
@@ -36,7 +45,7 @@ typedef FOnQueryForAvailablePurchasesComplete::FDelegate FOnQueryForAvailablePur
  * @param SessionName the name of the session the that has transitioned to started
  * @param bWasSuccessful true if the async action completed without error, false if there was an error
  */
-DECLARE_MULTICAST_DELEGATE_OneParam(FOnInAppPurchaseComplete, EInAppPurchaseState::Type);
+DECLARE_MULTICAST_DELEGATE_TwoParams(FOnInAppPurchaseComplete, EInAppPurchaseState::Type, const IPlatformPurchaseReceipt*);
 typedef FOnInAppPurchaseComplete::FDelegate FOnInAppPurchaseCompleteDelegate;
 
 
@@ -145,5 +154,5 @@ public:
 	/**
 	 * Delegate which is executed when a Purchase completes
 	 */
-	DEFINE_ONLINE_DELEGATE_ONE_PARAM(OnInAppPurchaseComplete, EInAppPurchaseState::Type);
+	DEFINE_ONLINE_DELEGATE_TWO_PARAM(OnInAppPurchaseComplete, EInAppPurchaseState::Type, const IPlatformPurchaseReceipt*);
 };
\ No newline at end of file
diff --git a/Engine/Source/Runtime/Online/OnlineSubsystemUtils/Classes/InAppPurchaseCallbackProxy.h b/Engine/Source/Runtime/Online/OnlineSubsystemUtils/Classes/InAppPurchaseCallbackProxy.h
--- a/Engine/Source/Runtime/Online/OnlineSubsystemUtils/Classes/InAppPurchaseCallbackProxy.h
+++ b/Engine/Source/Runtime/Online/OnlineSubsystemUtils/Classes/InAppPurchaseCallbackProxy.h
@@ -33,7 +33,7 @@ private:
 
 	/** Called by the InAppPurchase system when the transaction has finished */
 	void OnInAppPurchaseComplete_Delayed();
-	void OnInAppPurchaseComplete(EInAppPurchaseState::Type CompletionState);
+	void OnInAppPurchaseComplete(EInAppPurchaseState::Type CompletionState, const IPlatformPurchaseReceipt* ProductReceipt = nullptr);
 
 	/** Unregisters our delegate from the In-App Purchase system */
 	void RemoveDelegate();
@@ -57,7 +57,7 @@ private:
 
 	/** Did the purchase succeed? */
 	EInAppPurchaseState::Type SavedPurchaseState;
-
-	/** Cached product details of the purchased product */
-	FInAppPurchaseProductInfo SavedProductInformation;
+    
+    /** Copy of the provided purchase receipt */
+    const IPlatformPurchaseReceipt* SavedPurchaseReceipt;
 };
diff --git a/Engine/Source/Runtime/Online/OnlineSubsystemUtils/Private/InAppPurchaseCallbackProxy.cpp b/Engine/Source/Runtime/Online/OnlineSubsystemUtils/Private/InAppPurchaseCallbackProxy.cpp
--- a/Engine/Source/Runtime/Online/OnlineSubsystemUtils/Private/InAppPurchaseCallbackProxy.cpp
+++ b/Engine/Source/Runtime/Online/OnlineSubsystemUtils/Private/InAppPurchaseCallbackProxy.cpp
@@ -10,6 +10,7 @@ UInAppPurchaseCallbackProxy::UInAppPurchaseCallbackProxy(const class FPostConstr
 {
 	PurchaseRequest = nullptr;
 	WorldPtr = nullptr;
+    SavedPurchaseReceipt = nullptr;
 }
 
 
@@ -18,7 +19,7 @@ void UInAppPurchaseCallbackProxy::Trigger(APlayerController* PlayerController, c
 	bFailedToEvenSubmit = true;
 
 	WorldPtr = (PlayerController != nullptr) ? PlayerController->GetWorld() : nullptr;
-	if (APlayerState* PlayerState = (PlayerController != NULL) ? PlayerController->PlayerState : nullptr)
+	if (APlayerState* PlayerState = (PlayerController != nullptr) ? PlayerController->PlayerState : nullptr)
 	{
 		if (IOnlineSubsystem* const OnlineSub = IOnlineSubsystem::Get())
 		{
@@ -58,36 +59,44 @@ void UInAppPurchaseCallbackProxy::Trigger(APlayerController* PlayerController, c
 }
 
 
-void UInAppPurchaseCallbackProxy::OnInAppPurchaseComplete(EInAppPurchaseState::Type CompletionState)
+void UInAppPurchaseCallbackProxy::OnInAppPurchaseComplete(EInAppPurchaseState::Type CompletionState, const IPlatformPurchaseReceipt* ProductReceipt)
 {
 	RemoveDelegate();
-
 	SavedPurchaseState = CompletionState;
-
-	if (SavedPurchaseState == EInAppPurchaseState::Success && PurchaseRequest.IsValid())
-	{
-		SavedProductInformation = PurchaseRequest->ProvidedProductInformation;
-	}
-
+    SavedPurchaseReceipt = ProductReceipt;
+    
 	if (UWorld* World = WorldPtr.Get())
 	{
 		World->GetTimerManager().SetTimer(this, &UInAppPurchaseCallbackProxy::OnInAppPurchaseComplete_Delayed, 0.001f, false);
-	}
-
-	PurchaseRequest = NULL;
+    }
+    else
+    {
+        PurchaseRequest = nullptr;
+    }
 }
 
 
 void UInAppPurchaseCallbackProxy::OnInAppPurchaseComplete_Delayed()
 {
+    /** Cached product details of the purchased product */
+    FInAppPurchaseProductInfo ProductInformation;
+    
+    if (SavedPurchaseState == EInAppPurchaseState::Success && PurchaseRequest.IsValid())
+    {
+        ProductInformation = PurchaseRequest->ProvidedProductInformation;
+    }
+    
 	if (SavedPurchaseState == EInAppPurchaseState::Success)
 	{
-		OnSuccess.Broadcast(SavedPurchaseState, SavedProductInformation);
+		OnSuccess.Broadcast(SavedPurchaseState, ProductInformation);
 	}
 	else
 	{
-		OnFailure.Broadcast(SavedPurchaseState, SavedProductInformation);
+		OnFailure.Broadcast(SavedPurchaseState, ProductInformation);
 	}
+    
+    SavedPurchaseReceipt = nullptr;
+    PurchaseRequest = nullptr;
 }
 
 
