IAP for Google Play - Stub classes added. Now need to add our implementation

diff --git a/Engine/Source/Runtime/Online/Android/OnlineSubsystemGooglePlay/Private/OnlineStoreInterfaceGooglePlay.cpp b/Engine/Source/Runtime/Online/Android/OnlineSubsystemGooglePlay/Private/OnlineStoreInterfaceGooglePlay.cpp
new file mode 100644
--- /dev/null
+++ b/Engine/Source/Runtime/Online/Android/OnlineSubsystemGooglePlay/Private/OnlineStoreInterfaceGooglePlay.cpp
@@ -0,0 +1,48 @@
+// Copyright 1998-2014 Epic Games, Inc. All Rights Reserved.
+
+#include "OnlineSubsystemGooglePlayPrivatePCH.h"
+
+
+////////////////////////////////////////////////////////////////////
+/// FOnlineStoreGooglePlay implementation
+
+
+FOnlineStoreGooglePlay::FOnlineStoreGooglePlay() 
+{
+	UE_LOG(LogOnline, Display, TEXT( "FOnlineStoreGooglePlay::FOnlineStoreGooglePlay" ));
+
+}
+
+
+FOnlineStoreGooglePlay::~FOnlineStoreGooglePlay() 
+{
+	UE_LOG(LogOnline, Display, TEXT( "FOnlineStoreGooglePlay::~FOnlineStoreGooglePlay" ));
+
+}
+
+
+bool FOnlineStoreGooglePlay::QueryForAvailablePurchases(const TArray<FString>& ProductIDs, FOnlineProductInformationReadRef& InReadObject)
+{
+	UE_LOG(LogOnline, Display, TEXT( "FOnlineStoreGooglePlay::QueryForAvailablePurchases" ));
+	bool bSentAQueryRequest = false;
+
+	return bSentAQueryRequest;
+}
+
+
+bool FOnlineStoreGooglePlay::BeginPurchase(const FString& ProductId, FOnlineInAppPurchaseTransactionRef& InPurchaseStateObject)
+{
+	UE_LOG(LogOnline, Display, TEXT( "FOnlineStoreGooglePlay::BeginPurchase" ));
+	bool bCreatedNewTransaction = false;
+
+	return bCreatedNewTransaction;
+}
+
+
+bool FOnlineStoreGooglePlay::IsAllowedToMakePurchases()
+{
+	UE_LOG(LogOnline, Display, TEXT( "FOnlineStoreGooglePlay::IsAllowedToMakePurchases" ));
+	bool bCanMakePurchases = false;
+
+	return bCanMakePurchases;
+}
\ No newline at end of file
diff --git a/Engine/Source/Runtime/Online/Android/OnlineSubsystemGooglePlay/Private/OnlineStoreInterfaceGooglePlay.h b/Engine/Source/Runtime/Online/Android/OnlineSubsystemGooglePlay/Private/OnlineStoreInterfaceGooglePlay.h
new file mode 100644
--- /dev/null
+++ b/Engine/Source/Runtime/Online/Android/OnlineSubsystemGooglePlay/Private/OnlineStoreInterfaceGooglePlay.h
@@ -0,0 +1,59 @@
+// Copyright 1998-2014 Epic Games, Inc. All Rights Reserved.
+
+#pragma once
+
+#include "OnlineStoreInterface.h"
+
+
+/**
+ * The resulting state of an iap transaction
+ */
+namespace EInAppPurchaseResult
+{
+	enum Type
+	{
+		Succeeded = 0,
+		RestoredFromServer,
+		Failed,
+		Cancelled,
+	};
+}
+
+/**
+ * Implementation of the Platform Purchase receipt. For this we provide an identifier and the encrypted data.
+ */
+class FGooglePlayPurchaseReceipt : public IPlatformPurchaseReceipt
+{
+public:
+
+};
+
+
+/**
+ *	FOnlineStoreGooglePlay - Implementation of the online store for GooglePlay
+ */
+class FOnlineStoreGooglePlay : public IOnlineStore
+{
+public:
+	/** C-tor */
+	FOnlineStoreGooglePlay();
+	/** Destructor */
+	virtual ~FOnlineStoreGooglePlay();
+
+	// Begin IOnlineStore 
+	virtual bool QueryForAvailablePurchases(const TArray<FString>& ProductIDs, FOnlineProductInformationReadRef& InReadObject) override;
+	virtual bool BeginPurchase(const FString& ProductId, FOnlineInAppPurchaseTransactionRef& InReadObject) override;
+	virtual bool IsAllowedToMakePurchases() override;
+	// End IOnlineStore 
+
+private:
+
+	/** Delegate fired when a query for purchases has completed, whether successful or unsuccessful */
+	FOnQueryForAvailablePurchasesComplete OnQueryForAvailablePurchasesCompleteDelegate;
+
+	/** Delegate fired when a purchase transaction has completed, whether successful or unsuccessful */
+	FOnInAppPurchaseComplete OnPurchaseCompleteDelegate;
+};
+
+typedef TSharedPtr<FOnlineStoreGooglePlay, ESPMode::ThreadSafe> FOnlineStoreGooglePlayPtr;
+
diff --git a/Engine/Source/Runtime/Online/Android/OnlineSubsystemGooglePlay/Private/OnlineSubsystemGooglePlay.cpp b/Engine/Source/Runtime/Online/Android/OnlineSubsystemGooglePlay/Private/OnlineSubsystemGooglePlay.cpp
--- a/Engine/Source/Runtime/Online/Android/OnlineSubsystemGooglePlay/Private/OnlineSubsystemGooglePlay.cpp
+++ b/Engine/Source/Runtime/Online/Android/OnlineSubsystemGooglePlay/Private/OnlineSubsystemGooglePlay.cpp
@@ -33,7 +33,7 @@ IOnlineIdentityPtr FOnlineSubsystemGooglePlay::GetIdentityInterface() const
 
 IOnlineStorePtr FOnlineSubsystemGooglePlay::GetStoreInterface() const
 {
-	return nullptr;
+	return StoreInterface;
 }
 
 
@@ -114,6 +114,7 @@ bool FOnlineSubsystemGooglePlay::Init()
 	LeaderboardsInterface = MakeShareable(new FOnlineLeaderboardsGooglePlay(this));
 	AchievementsInterface = MakeShareable(new FOnlineAchievementsGooglePlay(this));
 	ExternalUIInterface = MakeShareable(new FOnlineExternalUIGooglePlay(this));
+	StoreInterface = MakeShareable(new FOnlineStoreGooglePlay());
 
 	extern struct android_app* GNativeAndroidApp;
 	check(GNativeAndroidApp != nullptr);
diff --git a/Engine/Source/Runtime/Online/Android/OnlineSubsystemGooglePlay/Private/OnlineSubsystemGooglePlayPrivatePCH.h b/Engine/Source/Runtime/Online/Android/OnlineSubsystemGooglePlay/Private/OnlineSubsystemGooglePlayPrivatePCH.h
--- a/Engine/Source/Runtime/Online/Android/OnlineSubsystemGooglePlay/Private/OnlineSubsystemGooglePlayPrivatePCH.h
+++ b/Engine/Source/Runtime/Online/Android/OnlineSubsystemGooglePlay/Private/OnlineSubsystemGooglePlayPrivatePCH.h
@@ -10,6 +10,7 @@
 #include "ModuleManager.h"
 
 #include "OnlineAchievementsInterfaceGooglePlay.h"
+#include "OnlineStoreInterfaceGooglePlay.h"
 
 #include "OnlineSubsystemGooglePlay.h"
 #include "OnlineSubsystemGooglePlayModule.h"
diff --git a/Engine/Source/Runtime/Online/Android/OnlineSubsystemGooglePlay/Public/OnlineSubsystemGooglePlay.h b/Engine/Source/Runtime/Online/Android/OnlineSubsystemGooglePlay/Public/OnlineSubsystemGooglePlay.h
--- a/Engine/Source/Runtime/Online/Android/OnlineSubsystemGooglePlay/Public/OnlineSubsystemGooglePlay.h
+++ b/Engine/Source/Runtime/Online/Android/OnlineSubsystemGooglePlay/Public/OnlineSubsystemGooglePlay.h
@@ -105,6 +105,8 @@ private:
 	/** Interface to the online identity system */
 	FOnlineIdentityGooglePlayPtr IdentityInterface;
 
+	FOnlineStoreGooglePlayPtr StoreInterface;
+
 	/** Interface to the online leaderboards */
 	FOnlineLeaderboardsGooglePlayPtr LeaderboardsInterface;
 
