OpenSUSE SSL Bundle Path

diff --git a/Engine/Source/Runtime/Online/HTTP/Private/Curl/CurlHttpManager.cpp b/Engine/Source/Runtime/Online/HTTP/Private/Curl/CurlHttpManager.cpp
--- a/Engine/Source/Runtime/Online/HTTP/Private/Curl/CurlHttpManager.cpp
+++ b/Engine/Source/Runtime/Online/HTTP/Private/Curl/CurlHttpManager.cpp
@@ -108,6 +108,7 @@ void FCurlHttpManager::InitCurl()
 		{
 			"/etc/pki/tls/certs/ca-bundle.crt",
 			"/etc/ssl/certs/ca-certificates.crt",
+			"/etc/ssl/ca-bundle.pem",
 			nullptr
 		};
 
