#UE4 - delete WITH_STEAMGC

diff --git a/Engine/Source/Runtime/Online/OnlineSubsystemSteam/Classes/SteamMessagingNetDriver.h b/Engine/Source/Runtime/Online/OnlineSubsystemSteam/Classes/SteamMessagingNetDriver.h
deleted file mode 100644
diff --git a/Engine/Source/Runtime/Online/OnlineSubsystemSteam/OnlineSubsystemSteam.Build.cs b/Engine/Source/Runtime/Online/OnlineSubsystemSteam/OnlineSubsystemSteam.Build.cs
--- a/Engine/Source/Runtime/Online/OnlineSubsystemSteam/OnlineSubsystemSteam.Build.cs
+++ b/Engine/Source/Runtime/Online/OnlineSubsystemSteam/OnlineSubsystemSteam.Build.cs
@@ -27,23 +27,5 @@ public class OnlineSubsystemSteam : ModuleRules
 		);
 
 		AddThirdPartyPrivateStaticDependencies(Target, "Steamworks");
-
-		Definitions.Add("WITH_STEAMGC=0");
-
-#if false
-		if (Directory.Exists(UEBuildConfiguration.UEThirdPartySourceDirectory + "NotForLicensees") == true &&
-			Directory.Exists(UEBuildConfiguration.UEThirdPartySourceDirectory + "NotForLicensees/SteamGC") == true)
-		{
-			AddThirdPartyPrivateStaticDependencies(Target, 
-				"ProtoBuf",
-				"SteamGC"
-				);
-			Definitions.Add("WITH_STEAMGC=1");
-		}
-		else
-		{
-			Definitions.Add("WITH_STEAMGC=0");
-		}
-#endif
 	}
 }
diff --git a/Engine/Source/Runtime/Online/OnlineSubsystemSteam/Private/OnlineAsyncTaskManagerSteam.cpp b/Engine/Source/Runtime/Online/OnlineSubsystemSteam/Private/OnlineAsyncTaskManagerSteam.cpp
--- a/Engine/Source/Runtime/Online/OnlineSubsystemSteam/Private/OnlineAsyncTaskManagerSteam.cpp
+++ b/Engine/Source/Runtime/Online/OnlineSubsystemSteam/Private/OnlineAsyncTaskManagerSteam.cpp
@@ -12,111 +12,8 @@
 #include "IPAddressSteam.h"
 #include "SocketSubsystemSteam.h"
 #include "SocketsSteam.h"
-#include "OnlineMsgSteam.h"
 #include "SteamUtilities.h"
 
-/**
- * Destroy any messages left in the queue
- */
-void FOnlineAsyncTaskManagerSteam::FreeMsgQueue()
-{
-#if WITH_STEAMGC
-	FScopeLock ScopeLock(&InMsgQueueLock);
-	for(TMap<int64, FOnlineAsyncMsgSteam*>::TIterator It(InMsgQueue); It; ++It)
-	{
-		It.Value()->Destroy();
-	}
-
-	InMsgQueue.Empty();
-#endif
-}
-
-void FOnlineAsyncTaskManagerSteam::MessagePump()
-{
-#if WITH_STEAMGC
-	uint32 MsgSize = 0;
-	while (SteamCoordinator.IsMessageAvailable(&MsgSize))
-	{
-		TArray<uint8> Buffer(MsgSize);
-		Buffer.AddUninitialized(MsgSize);
-
-		uint32 MsgType;
-		uint32 NewMsgSize;
-
-		// Retrieve the message type and its payload
-		EGCResults eResult = SteamCoordinator.RetrieveMessage(&MsgType, Buffer.GetData(), MsgSize, &NewMsgSize);
-		if (eResult != k_EGCResultOK)
-		{
-			continue;
-		}
-
-		check(MsgSize == NewMsgSize);
-
-		uint32 ProtoMsgByteOffset;
-		CClientMsgGCHeader MsgHdr;
-		// Retrieve the message header and payload offsets
-		if (SteamCoordinator.ParseMessage(MsgType, Buffer.GetData(), MsgSize, MsgHdr, &ProtoMsgByteOffset))
-		{
-			// Who this message is intended for (job that initiated request)
-			uint64 JobId = MsgHdr.jobid_target();
-
-			FOnlineAsyncMsgSteam* NewJob = NULL;
-			{		
-				FScopeLock ScopeLock(&InMsgQueueLock);
-				FOnlineAsyncMsgSteam** Job = InMsgQueue.Find(JobId);
-
-				if (Job && *Job)
-				{
-					NewJob = *Job;
-					InMsgQueue.Remove(JobId);
-				}
-			}
-
-			if (NewJob)
-			{
-				// Convert the message payload to a usable format and pass to game thread
-				check(NewJob->ResponseMessage);
-				check(NewJob->GetResponseMsgType() == MsgType);
-				if (NewJob->GetResponse()->ParseFromArray(Buffer.GetData() + ProtoMsgByteOffset, MsgSize - ProtoMsgByteOffset))
-				{
-					NewJob->Deserialize();
-				}
-
-				// Always add to out queue for destruction on game thread
-				AddToOutQueue(NewJob);
-				if (!NewJob->bWasSuccessful)
-				{
-					UE_LOG_ONLINE(Warning, TEXT("Failed to parse message.  MessageType: %d JobId: %llu"), MsgType, JobId);
-				}
-
-				UE_LOG_ONLINE(Log, TEXT("Async msg '%s' completed in %f seconds with %d"), *NewJob->ToString(), NewJob->GetElapsedTime(), NewJob->bWasSuccessful);
-			}
-			else
-			{
-				UE_LOG_ONLINE(Warning, TEXT("Unsolicited message from server.  MessageType: %d JobId: %llu"), MsgType, JobId);
-			}
-		}
-	}
-
-	static float CheckInterval = 15.0f;
-	static double LastCheckTime = 0.0;
-	if (FPlatformTime::Seconds() - LastCheckTime > CheckInterval)
-	{
-		LastCheckTime = FPlatformTime::Seconds();
-		FScopeLock ScopeLock(&InMsgQueueLock);
-		for (TMap<int64, FOnlineAsyncMsgSteam*>::TConstIterator It(InMsgQueue); It; ++It)
-		{
-			uint64 JobId = It.Key();
-			FOnlineAsyncMsgSteam* Message = It.Value();
-			if (Message->GetElapsedTime() > CheckInterval)
-			{
-				UE_LOG_ONLINE(Warning, TEXT("Message id %ull type %d in queue for %f without a response."), JobId, Message->GetMsgType(), Message->GetElapsedTime());
-			}
-		}
-	}
-#endif
-}
-
 void FOnlineAsyncTaskManagerSteam::OnlineTick()
 {
 	check(SteamSubsystem);
@@ -131,8 +28,6 @@ void FOnlineAsyncTaskManagerSteam::OnlineTick()
 	{
 		SteamGameServer_RunCallbacks();
 	}
-
-	MessagePump();
 }
 
 /**
@@ -1185,33 +1080,3 @@ void FOnlineAsyncTaskManagerSteam::OnSteamShutdown(SteamShutdown_t* CallbackData
 	UE_LOG_ONLINE(Verbose, TEXT("%s"), *NewEvent->ToString());
 	AddToOutQueue(NewEvent);
 }
-
-void FOnlineAsyncTaskManagerSteam::AddToInMsgQueue(FOnlineAsyncMsgSteam* NewMsg)
-{
-#if WITH_STEAMGC
-	// assert if not game thread
-	check(IsInGameThread());
-
-	int64 LocalJobId = -1;
-	{
-		FScopeLock Lock(&InMsgQueueLock);
-		if (NewMsg->GetResponseMsgType() > 0)
-		{
-			static int64 JobId = 0;
-			LocalJobId = ++JobId;
-			InMsgQueue.Add(LocalJobId, NewMsg);
-		}	
-	}
-
-	if (LocalJobId != -1)
-	{
-		SteamCoordinator.BSendMessageExpectResponse(NewMsg->GetMsgType(), LocalJobId, *NewMsg->GetParams());
-	}
-	else
-	{
-		SteamCoordinator.BSendMessageNoResponse(NewMsg->GetMsgType(), *NewMsg->GetParams());
-	}
-
-	WorkEvent->Trigger();
-#endif
-}
diff --git a/Engine/Source/Runtime/Online/OnlineSubsystemSteam/Private/OnlineAsyncTaskManagerSteam.h b/Engine/Source/Runtime/Online/OnlineSubsystemSteam/Private/OnlineAsyncTaskManagerSteam.h
--- a/Engine/Source/Runtime/Online/OnlineSubsystemSteam/Private/OnlineAsyncTaskManagerSteam.h
+++ b/Engine/Source/Runtime/Online/OnlineSubsystemSteam/Private/OnlineAsyncTaskManagerSteam.h
@@ -96,25 +96,6 @@ protected:
 	/** Cached reference to the main online subsystem */
 	class FOnlineSubsystemSteam* SteamSubsystem;
 
-#if WITH_STEAMGC
-	/** RPC requests to the Steam backend are queued up here for processing on the online thread */
-	TMap<int64, class FOnlineAsyncMsgSteam*> InMsgQueue;
-	/** Critical section for thread safe operation of the message queue */
-	FCriticalSection InMsgQueueLock;
-	/** Steam messaging */
-	GCClientWrapper SteamCoordinator;
-#endif
-
-	/**
-	 * Message polling when communicating with the Steam secure backend
-	 */
-	void MessagePump();
-
-	/**
-	 * Destroy any messages left in the queue
-	 */
-	void FreeMsgQueue();
-
 public:
 
 	FOnlineAsyncTaskManagerSteam(class FOnlineSubsystemSteam* InOnlineSubsystem) :
@@ -144,19 +125,12 @@ public:
 
 	~FOnlineAsyncTaskManagerSteam() 
 	{
-		FreeMsgQueue();
 	}
 
 	// FOnlineAsyncTaskManager
 	virtual void OnlineTick() override;
 
 	// FOnlineAsyncTaskManagerSteam
-
-	/**
-	 * Add online messages sent to the secure backend that are waiting for a response onto the incoming queue
-	 * @param NewMsg - some request of the online services
-	 */
-	virtual void AddToInMsgQueue(class FOnlineAsyncMsgSteam* NewTask);
 };
 
 
diff --git a/Engine/Source/Runtime/Online/OnlineSubsystemSteam/Private/OnlineMsgSteam.cpp b/Engine/Source/Runtime/Online/OnlineSubsystemSteam/Private/OnlineMsgSteam.cpp
deleted file mode 100644
diff --git a/Engine/Source/Runtime/Online/OnlineSubsystemSteam/Private/OnlineMsgSteam.h b/Engine/Source/Runtime/Online/OnlineSubsystemSteam/Private/OnlineMsgSteam.h
deleted file mode 100644
diff --git a/Engine/Source/Runtime/Online/OnlineSubsystemSteam/Private/OnlineSubsystemSteam.cpp b/Engine/Source/Runtime/Online/OnlineSubsystemSteam/Private/OnlineSubsystemSteam.cpp
--- a/Engine/Source/Runtime/Online/OnlineSubsystemSteam/Private/OnlineSubsystemSteam.cpp
+++ b/Engine/Source/Runtime/Online/OnlineSubsystemSteam/Private/OnlineSubsystemSteam.cpp
@@ -261,12 +261,6 @@ void FOnlineSubsystemSteam::QueueAsyncOutgoingItem(FOnlineAsyncItem* AsyncItem)
 	OnlineAsyncTaskThreadRunnable->AddToOutQueue(AsyncItem);
 }
 
-void FOnlineSubsystemSteam::QueueAsyncMsg(FOnlineAsyncMsgSteam* AsyncMsg)
-{
-	check(OnlineAsyncTaskThreadRunnable);
-	OnlineAsyncTaskThreadRunnable->AddToInMsgQueue(AsyncMsg);
-}
-
 bool FOnlineSubsystemSteam::Tick(float DeltaTime)
 {
 	if (OnlineAsyncTaskThreadRunnable)
diff --git a/Engine/Source/Runtime/Online/OnlineSubsystemSteam/Private/OnlineSubsystemSteamPrivatePCH.h b/Engine/Source/Runtime/Online/OnlineSubsystemSteam/Private/OnlineSubsystemSteamPrivatePCH.h
--- a/Engine/Source/Runtime/Online/OnlineSubsystemSteam/Private/OnlineSubsystemSteamPrivatePCH.h
+++ b/Engine/Source/Runtime/Online/OnlineSubsystemSteam/Private/OnlineSubsystemSteamPrivatePCH.h
@@ -44,10 +44,6 @@
 #include "steam/steam_api.h"
 #include "steam/steam_gameserver.h"
 
-#if WITH_STEAMGC
-#include "gc_client.h"
-#endif
-
 // @todo Steam: See above
 #ifdef _MSC_VER
 #pragma warning(pop)
diff --git a/Engine/Source/Runtime/Online/OnlineSubsystemSteam/Private/SteamMessagingNetDriver.cpp b/Engine/Source/Runtime/Online/OnlineSubsystemSteam/Private/SteamMessagingNetDriver.cpp
deleted file mode 100644
diff --git a/Engine/Source/Runtime/Online/OnlineSubsystemSteam/Private/SteamNetDriver.cpp b/Engine/Source/Runtime/Online/OnlineSubsystemSteam/Private/SteamNetDriver.cpp
--- a/Engine/Source/Runtime/Online/OnlineSubsystemSteam/Private/SteamNetDriver.cpp
+++ b/Engine/Source/Runtime/Online/OnlineSubsystemSteam/Private/SteamNetDriver.cpp
@@ -8,7 +8,6 @@
 #include "OnlineSubsystemSteam.h"
 #include "SocketSubsystemSteam.h"
 #include "SocketsSteam.h"
-#include "OnlineMsgSteam.h"
 
 USteamNetDriver::USteamNetDriver(const class FPostConstructInitializeProperties& PCIP) :
 	Super(PCIP),
diff --git a/Engine/Source/Runtime/Online/OnlineSubsystemSteam/Public/OnlineSubsystemSteam.h b/Engine/Source/Runtime/Online/OnlineSubsystemSteam/Public/OnlineSubsystemSteam.h
--- a/Engine/Source/Runtime/Online/OnlineSubsystemSteam/Public/OnlineSubsystemSteam.h
+++ b/Engine/Source/Runtime/Online/OnlineSubsystemSteam/Public/OnlineSubsystemSteam.h
@@ -160,12 +160,6 @@ PACKAGE_SCOPE:
 	 */
 	void QueueAsyncOutgoingItem(class FOnlineAsyncItem* AsyncItem);
 
-	/**
-	 *	Add an async msg onto the msg queue for processing
-	 * @param AsyncMsg - new heap allocated msg to process on the async task thread
-	 */
-	void QueueAsyncMsg(class FOnlineAsyncMsgSteam* AsyncMsg);
-
     /** 
      * **INTERNAL**
      * Get the metadata related to a given user
