Github Pull Request : 345 Feature: Implement a basic HTTP request support for the HTML5 platform

diff --git a/Engine/Source/Runtime/HTML5/HTML5JS/Private/HTML5JavaScriptFx.js b/Engine/Source/Runtime/HTML5/HTML5JS/Private/HTML5JavaScriptFx.js
--- a/Engine/Source/Runtime/HTML5/HTML5JS/Private/HTML5JavaScriptFx.js
+++ b/Engine/Source/Runtime/HTML5/HTML5JS/Private/HTML5JavaScriptFx.js
@@ -110,6 +110,58 @@ var UE_JavaScriptLibary =
         }
         Module.writeAsciiToMemory(culture_name, address);
         return 1;
+    },
+
+    UE_MakeHTTPDataRequest: function (ctx, url, verb, payload, freeBuffer, onload, onerror, onprogress) {
+	    var _url = Pointer_stringify(url);
+	    var _verb = Pointer_stringify(verb);
+	    var _payload = Pointer_stringify(payload);
+
+	    var xhr = new XMLHttpRequest();
+	    xhr.open(_verb, _url, true);
+	    xhr.responseType = 'arraybuffer';
+
+	    // Onload event handler
+	    xhr.addEventListener('load', function (e) {
+   	        if (xhr.status === 200 || _url.substr(0, 4).toLowerCase() !== "http") {
+       	            var byteArray = new Uint8Array(xhr.response);
+       	            var buffer = _malloc(byteArray.length);
+       	            HEAPU8.set(byteArray, buffer);
+       	            if (onload) 
+       	                    Runtime.dynCall('viii', onload, [ctx, buffer, byteArray.length]);
+       	            if (freeBuffer) 
+       	                    _free(buffer);
+   	        } 
+   	        else{
+   	            if (onerror) Runtime.dynCall('viii', onerror, [ctx, xhr.status, xhr.statusText]);
+   	        }
+	    });
+
+	    // Onerror event handler
+	    xhr.addEventListener('error', function (e) {
+	        if (onerror) Runtime.dynCall('viii', onerror, [ctx, xhr.status, xhr.statusText]);
+	    });
+
+	    // Onprogress event handler
+	    xhr.addEventListener('progress', function (e) {
+	        if (onprogress) Runtime.dynCall('viii', onprogress, [ctx, e.loaded, e.lengthComputable || e.lengthComputable === undefined ? e.total : 0]);
+	    });
+
+	    // Bypass possible browser redirection limit
+	    try {
+	        if (xhr.channel instanceof Ci.nsIHttpChannel)
+	            xhr.channel.redirectionLimit = 0;
+	    } catch (ex) { }
+
+	    if (_verb === "POST") {
+	        //Send the proper header information along with the request
+	        xhr.setRequestHeader("Content-type", "application/x-www-form-urlencoded");
+	        xhr.setRequestHeader("Content-length", _payload.length);
+	        xhr.setRequestHeader("Connection", "close");
+	        xhr.send(_payload);
+	    } else {
+	        xhr.send(null);
+	    }
     }
 };
 
diff --git a/Engine/Source/Runtime/HTML5/HTML5JS/Public/HTML5JavaScriptFx.h b/Engine/Source/Runtime/HTML5/HTML5JS/Public/HTML5JavaScriptFx.h
--- a/Engine/Source/Runtime/HTML5/HTML5JS/Public/HTML5JavaScriptFx.h
+++ b/Engine/Source/Runtime/HTML5/HTML5JS/Public/HTML5JavaScriptFx.h
@@ -7,6 +7,7 @@
 		 bool UE_SaveGame(const char* Name, const int UserIndex, const char* indata, int insize);
 		 bool UE_LoadGame(const char* Name, const int UserIndex, char **outdata, int* outsize);
 		 int  UE_MessageBox( int MsgType, const char* Text, const char* Caption);
-		 int UE_GetCurrentCultureName(const char* OutName, int outsize);
+		 int  UE_GetCurrentCultureName(const char* OutName, int outsize);
+		 void UE_MakeHTTPDataRequest(void *ctx, const char* url, const char* verb, const char* payload, int freeBuffer, void(*onload)(void*, void*, unsigned), void(*onerror)(void*, int, const char*), void(*onprogress)(void*, int, int));
 	}
 
diff --git a/Engine/Source/Runtime/Online/HTTP/HTTP.Build.cs b/Engine/Source/Runtime/Online/HTTP/HTTP.Build.cs
--- a/Engine/Source/Runtime/Online/HTTP/HTTP.Build.cs
+++ b/Engine/Source/Runtime/Online/HTTP/HTTP.Build.cs
@@ -4,32 +4,43 @@ using UnrealBuildTool;
 
 public class HTTP : ModuleRules
 {
-	public HTTP(TargetInfo Target)
-	{
-		Definitions.Add("HTTP_PACKAGE=1");
+    public HTTP(TargetInfo Target)
+    {
+        Definitions.Add("HTTP_PACKAGE=1");
 
-		PrivateIncludePaths.AddRange(
-			new string[] {
+        PrivateIncludePaths.AddRange(
+            new string[] {
 				"Runtime/Online/HTTP/Private",
 			}
-			);
+            );
 
-		PrivateDependencyModuleNames.AddRange(new string[] { "Core" });
+        PrivateDependencyModuleNames.AddRange(new string[] { "Core" });
 
-		if (Target.Platform == UnrealTargetPlatform.Win32 ||
-			Target.Platform == UnrealTargetPlatform.Win64)
-		{
-			AddThirdPartyPrivateStaticDependencies(Target, "WinInet");
+        if (Target.Platform == UnrealTargetPlatform.Win32 ||
+            Target.Platform == UnrealTargetPlatform.Win64)
+        {
+            AddThirdPartyPrivateStaticDependencies(Target, "WinInet");
             if (!UnrealBuildTool.UnrealBuildTool.BuildingRocket() && !UnrealBuildTool.UnrealBuildTool.RunningRocket())
             {
                 AddThirdPartyPrivateStaticDependencies(Target, "libcurl");
             }
-		}
+        }
 
-		if (Target.Platform == UnrealTargetPlatform.Linux ||
+        if (Target.Platform == UnrealTargetPlatform.Linux ||
             Target.Platform == UnrealTargetPlatform.Android)
-		{
-			AddThirdPartyPrivateStaticDependencies(Target, "libcurl");
-		}
-	}
+        {
+            AddThirdPartyPrivateStaticDependencies(Target, "libcurl");
+        }
+        if (Target.Platform == UnrealTargetPlatform.HTML5)
+        {
+            if (Target.Architecture == "-win32")
+            {
+                PrivateDependencyModuleNames.Add("HTML5Win32");
+            }
+            else
+            {
+                PrivateDependencyModuleNames.Add("HTML5JS");
+            }
+        }
+    }
 }
diff --git a/Engine/Source/Runtime/Online/HTTP/Private/HTML5/HTML5HTTP.cpp b/Engine/Source/Runtime/Online/HTTP/Private/HTML5/HTML5HTTP.cpp
new file mode 100644
--- /dev/null
+++ b/Engine/Source/Runtime/Online/HTTP/Private/HTML5/HTML5HTTP.cpp
@@ -0,0 +1,662 @@
+// Copyright 1998-2014 Epic Games, Inc. All Rights Reserved.
+
+#include "HttpPrivatePCH.h"
+#include "HTML5HTTP.h"
+#include "EngineVersion.h"
+
+#if PLATFORM_HTML5_BROWSER
+#include "HTML5JavaScriptFx.h"
+#endif
+
+/****************************************************************************
+* FHTML5HttpRequest implementation
+***************************************************************************/
+
+FHTML5HttpRequest::FHTML5HttpRequest()
+	:   bCanceled(false)
+	,	bCompleted(false)
+	,	BytesSent(0)
+	,	CompletionStatus(EHttpRequestStatus::NotStarted)
+	,	ElapsedTime(0.0f)
+{
+	UE_LOG(LogHttp, Verbose, TEXT("FHTML5HttpRequest::FHTML5HttpRequest()"));
+}
+
+
+FHTML5HttpRequest::~FHTML5HttpRequest()
+{
+	UE_LOG(LogHttp, Verbose, TEXT("FHTML5HttpRequest::~FHTML5HttpRequest()"));
+}
+
+
+FString FHTML5HttpRequest::GetURL()
+{
+	UE_LOG(LogHttp, Verbose, TEXT("FHTML5HttpRequest::GetURL() - %s"), *URL);
+	return URL;
+}
+
+
+void FHTML5HttpRequest::SetURL(const FString& InURL)
+{
+	UE_LOG(LogHttp, Verbose, TEXT("FHTML5HttpRequest::SetURL() - %s"), *InURL);
+
+	URL = InURL.Replace(TEXT("%"), TEXT("%25"));
+	URL = URL.Replace(TEXT(" "), TEXT("%20"));
+	URL = URL.Replace(TEXT("\""), TEXT("%22"));
+	URL = URL.Replace(TEXT("<"), TEXT("%3C"));
+	URL = URL.Replace(TEXT(">"), TEXT("%3E"));
+
+	URL = URL.Replace(TEXT("["), TEXT("%5B"));
+	URL = URL.Replace(TEXT("]"), TEXT("%5D"));
+	URL = URL.Replace(TEXT("\\"), TEXT("%5C"));
+	URL = URL.Replace(TEXT("^"), TEXT("%5E"));
+
+	URL = URL.Replace(TEXT("`"), TEXT("%60"));
+	URL = URL.Replace(TEXT("{"), TEXT("%7B"));
+	URL = URL.Replace(TEXT("}"), TEXT("%7D"));
+	URL = URL.Replace(TEXT("|"), TEXT("%7C"));
+}
+
+
+FString FHTML5HttpRequest::GetURLParameter(const FString& ParameterName)
+{
+	UE_LOG(LogHttp, Verbose, TEXT("FHTML5HttpRequest::GetURLParameter() - %s"), *ParameterName);
+
+	TArray<FString> StringElements;
+
+	int32 NumElems = URL.ParseIntoArray(&StringElements, TEXT("&"), true);
+	check(NumElems == StringElements.Num());
+
+	FString ParamValDelimiter(TEXT("="));
+	for (int Idx = 0; Idx < NumElems; ++Idx )
+	{
+		FString Param, Value;
+		if (StringElements[Idx].Split(ParamValDelimiter, &Param, &Value) && Param == ParameterName)
+		{
+			return Value;
+		}
+	}
+
+	return FString();
+}
+
+
+FString FHTML5HttpRequest::GetHeader(const FString& HeaderName)
+{
+	UE_LOG(LogHttp, Verbose, TEXT("FHTML5HttpRequest::GetHeader() - %s"), *HeaderName);
+
+	FString Result;
+
+	FString* Header = Headers.Find(HeaderName);
+
+	if (Header != NULL)
+	{
+		Result = *Header;
+	}
+
+	return Result;
+}
+
+
+void FHTML5HttpRequest::SetHeader(const FString& HeaderName, const FString& HeaderValue)
+{
+	UE_LOG(LogHttp, Verbose, TEXT("FHTML5HttpRequest::SetHeader() - %s / %s"), *HeaderName, *HeaderValue);
+	Headers.Add(HeaderName, HeaderValue);
+}
+
+
+TArray<FString> FHTML5HttpRequest::GetAllHeaders()
+{
+	UE_LOG(LogHttp, Verbose, TEXT("FHTML5HttpRequest::GetAllHeaders()"));
+
+	TArray<FString> Result;
+	for (TMap<FString, FString>::TConstIterator It(Headers); It; ++It)
+	{
+		Result.Add(It.Key() + TEXT(": ") + It.Value());
+	}
+	return Result;
+}
+
+
+const TArray<uint8>& FHTML5HttpRequest::GetContent()
+{
+	UE_LOG(LogHttp, Verbose, TEXT("FHTML5HttpRequest::GetContent()"));
+
+	return RequestPayload;
+}
+
+
+void FHTML5HttpRequest::SetContent(const TArray<uint8>& ContentPayload)
+{
+	UE_LOG(LogHttp, Verbose, TEXT("FHTML5HttpRequest::SetContent()"));
+
+	RequestPayload = ContentPayload;
+}
+
+
+FString FHTML5HttpRequest::GetContentType()
+{
+	return GetHeader(TEXT( "Content-Type" ));
+}
+
+
+int32 FHTML5HttpRequest::GetContentLength()
+{
+	int Len = RequestPayload.Num();
+	UE_LOG(LogHttp, Verbose, TEXT("FHTML5HttpRequest::GetContentLength() - %i"), Len);
+
+	return Len;
+}
+
+
+void FHTML5HttpRequest::SetContentAsString(const FString& ContentString)
+{
+	UE_LOG(LogHttp, Verbose, TEXT("FHTML5HttpRequest::SetContentAsString() - %s"), *ContentString);
+
+	FTCHARToUTF8 Converter(*ContentString);
+	RequestPayload.SetNum(Converter.Length());
+	FMemory::Memcpy(RequestPayload.GetTypedData(), (uint8*)(ANSICHAR*)Converter.Get(), RequestPayload.Num());
+}
+
+
+FString FHTML5HttpRequest::GetVerb()
+{
+	return Verb;
+}
+
+
+void FHTML5HttpRequest::SetVerb(const FString& InVerb)
+{
+	UE_LOG(LogHttp, Verbose, TEXT("FHTML5HttpRequest::SetVerb() - %s"), *InVerb);
+
+	Verb = InVerb.ToUpper();
+}
+
+bool IsURLEncoded(const TArray<uint8> & Payload)
+{
+	static char AllowedChars[] = "ABCDEFGHIJKLMNOPQRSTUVWXYZabcdefghijklmnopqrstuvwxyz0123456789-_.~";
+	static bool bTableFilled = false;
+	static bool AllowedTable[256] = { false };
+
+	if (!bTableFilled)
+	{
+		for (int32 Idx = 0; Idx < ARRAY_COUNT(AllowedChars) - 1; ++Idx)	// -1 to avoid trailing 0
+		{
+			uint8 AllowedCharIdx = static_cast<uint8>(AllowedChars[Idx]);
+			check(AllowedCharIdx < ARRAY_COUNT(AllowedTable));
+			AllowedTable[AllowedCharIdx] = true;
+		}
+
+		bTableFilled = true;
+	}
+
+	const int32 Num = Payload.Num();
+	for (int32 Idx = 0; Idx < Num; ++Idx)
+	{
+		if (!AllowedTable[Payload[Idx]])
+			return false;
+	}
+
+	return true;
+}
+
+void FHTML5HttpRequest::StaticReceiveCallback(void *arg, void *buffer, uint32 size){
+	UE_LOG(LogHttp, Verbose, TEXT("FHTML5HttpRequest::StaticReceiveDataCallback()"));
+
+	FHTML5HttpRequest* Request = reinterpret_cast<FHTML5HttpRequest*>(arg);
+
+	return Request->ReceiveCallback(arg, buffer, size);
+}
+
+void FHTML5HttpRequest::ReceiveCallback(void *arg, void *buffer, uint32 size) {
+	UE_LOG(LogHttp, Verbose, TEXT("FHTML5HttpRequest::ReceiveDataCallback()"));
+	UE_LOG(LogHttp, Verbose, TEXT("Response size: %d"), size);
+
+	check(Response.IsValid());
+
+	if (Response.IsValid())
+	{
+		Response->Payload.AddUninitialized(size);
+
+		// save
+		UE_LOG(LogHttp, Verbose, TEXT("Saving payload..."));
+
+		FMemory::Memcpy(static_cast<uint8*>(Response->Payload.GetData()), buffer, size);
+		Response->TotalBytesRead = size;
+		Response->HttpCode = 200;
+
+		UE_LOG(LogHttp, Verbose, TEXT("Payload length: %d"), Response->Payload.Num());
+
+		MarkAsCompleted();
+	}
+}
+
+void FHTML5HttpRequest::StaticErrorCallback(void* arg, int httpStatusCode, const char* httpStatusText) {
+	UE_LOG(LogHttp, Verbose, TEXT("FHTML5HttpRequest::StaticErrorDataCallback()"));
+
+	FHTML5HttpRequest* Request = reinterpret_cast<FHTML5HttpRequest*>(arg);
+	return Request->ErrorCallback(arg, httpStatusCode, httpStatusText);
+}
+
+void FHTML5HttpRequest::ErrorCallback(void* arg, int httpStatusCode, const char* httpStatusText) {
+	UE_LOG(LogHttp, Verbose, TEXT("FHTML5HttpRequest::ErrorDataCallback()"));
+}
+
+void FHTML5HttpRequest::StaticProgressCallback(void* arg, int Loaded, int Total) {
+	UE_LOG(LogHttp, Verbose, TEXT("FHTML5HttpRequest::StaticProgressCallback()"));
+
+	FHTML5HttpRequest* Request = reinterpret_cast<FHTML5HttpRequest*>(arg);
+
+	return Request->ProgressCallback(arg, Loaded, Total);
+}
+
+void FHTML5HttpRequest::ProgressCallback(void* arg, int Loaded, int Total) {
+	UE_LOG(LogHttp, Verbose, TEXT("FHTML5HttpRequest::ProgressCallback()"));
+
+	if (GetVerb() == TEXT("GET")) {
+		if (Response.IsValid())
+		{
+			Response->TotalBytesRead = Loaded;
+			OnRequestProgress().ExecuteIfBound(SharedThis(this), Response->TotalBytesRead);
+		}
+	}
+	else {
+		BytesSent = Loaded;
+		OnRequestProgress().ExecuteIfBound(SharedThis(this), BytesSent);
+	}
+
+	UE_LOG(LogHttp, Verbose, TEXT("Loaded: %d, Total: %d"), Loaded, Total);
+}
+
+
+bool FHTML5HttpRequest::StartRequest()
+{
+	UE_LOG(LogHttp, Verbose, TEXT("FHTML5HttpRequest::StartRequest()"));
+
+	UE_LOG(LogHttp, Verbose, TEXT("%p: URL='%s'"), this, *URL);
+	UE_LOG(LogHttp, Verbose, TEXT("%p: Verb='%s'"), this, *Verb);
+	UE_LOG(LogHttp, Verbose, TEXT("%p: Custom headers are %s"), this, Headers.Num() ? TEXT("present") : TEXT("NOT present"));
+	UE_LOG(LogHttp, Verbose, TEXT("%p: Payload size=%d"), this, RequestPayload.Num());
+
+	if (!FHttpModule::Get().IsHttpEnabled())
+	{
+		UE_LOG(LogHttp, Verbose, TEXT("Http disabled. Skipping request. url=%s"), *GetURL());
+		return false;
+	}// Prevent overlapped requests using the same instance
+	else if (CompletionStatus == EHttpRequestStatus::Processing)
+	{
+		UE_LOG(LogHttp, Warning, TEXT("ProcessRequest failed. Still processing last request."));
+		return false;
+	}
+	// Nothing to do without a valid URL
+	else if (URL.IsEmpty())
+	{
+		UE_LOG(LogHttp, Log, TEXT("Cannot process HTTP request: URL is empty"));
+		return false;
+	}
+
+	// set up headers
+	if (GetHeader("User-Agent").IsEmpty())
+	{
+		SetHeader(TEXT("User-Agent"), FString::Printf(TEXT("game=%s, engine=UE4, version=%s"), FApp::GetGameName(), *GEngineVersion.ToString()));
+	}
+
+	// Add "Pragma: no-cache" to mimic WinInet behavior
+	if (GetHeader("Pragma").IsEmpty())
+	{
+		SetHeader(TEXT("Pragma"), TEXT("no-cache"));
+	}
+
+	TArray<FString> AllHeaders = GetAllHeaders();
+	const int32 NumAllHeaders = AllHeaders.Num();
+
+	//TODO: Add headers into UE_MakeHTTPDataRequest
+
+
+	// set up verb (note that Verb is expected to be uppercase only)
+	if (Verb == TEXT("POST"))
+	{
+		// If we don't pass any other Content-Type, RequestPayload is assumed to be URL-encoded by this time
+		// (if we pass, don't check here and trust the request)
+		check(!GetHeader("Content-Type").IsEmpty() || IsURLEncoded(RequestPayload));
+
+#if PLATFORM_HTML5_BROWSER
+		UE_MakeHTTPDataRequest(this, TCHAR_TO_ANSI(*URL), "POST", (char*)RequestPayload.GetData(), 0, StaticReceiveCallback, StaticErrorCallback, StaticProgressCallback);
+#else
+		return false;
+#endif
+
+	}
+	else if (Verb == TEXT("PUT"))
+	{
+		UE_LOG(LogHttp, Log, TEXT("TODO: PUT"));
+
+		//TODO: PUT
+
+		// reset the counter
+		BytesSent = 0;
+		return false;
+	}
+	else if (Verb == TEXT("GET"))
+	{
+#if PLATFORM_HTML5_BROWSER
+		UE_MakeHTTPDataRequest(this, TCHAR_TO_ANSI(*URL), "GET", NULL, 1, StaticReceiveCallback, StaticErrorCallback, StaticProgressCallback);
+#else
+		return false;
+#endif
+	}
+	else if (Verb == TEXT("HEAD"))
+	{
+		UE_LOG(LogHttp, Log, TEXT("TODO: HEAD"));
+
+		//TODO: HEAD
+		return false;
+	}
+	else if (Verb == TEXT("DELETE"))
+	{
+
+		// If we don't pass any other Content-Type, RequestPayload is assumed to be URL-encoded by this time
+		// (if we pass, don't check here and trust the request)
+		check(!GetHeader("Content-Type").IsEmpty() || IsURLEncoded(RequestPayload));
+
+		//TODO: DELETE
+
+		UE_LOG(LogHttp, Log, TEXT("TODO: DELETE"));
+		return false;
+	}
+	else
+	{
+		UE_LOG(LogHttp, Fatal, TEXT("Unsupported verb '%s"), *Verb);
+		FPlatformMisc::DebugBreak();
+	}
+
+	return true;
+}
+
+bool FHTML5HttpRequest::ProcessRequest()
+{
+	UE_LOG(LogHttp, Verbose, TEXT("FHTML5HttpRequest::ProcessRequest()"));
+
+	if (!StartRequest())
+	{
+		UE_LOG(LogHttp, Warning, TEXT("Processing HTTP request failed. Increase verbosity for additional information."));
+
+		// No response since connection failed
+		Response = NULL;
+		// Cleanup and call delegate
+		FinishedRequest();
+
+		return false;
+	}
+
+	// Mark as in-flight to prevent overlapped requests using the same object
+	CompletionStatus = EHttpRequestStatus::Processing;
+	// Response object to handle data that comes back after starting this request
+	Response = MakeShareable(new FHTML5HttpResponse(*this));
+	// Add to global list while being processed so that the ref counted request does not get deleted
+	FHttpModule::Get().GetHttpManager().AddRequest(SharedThis(this));
+	// reset timeout
+	ElapsedTime = 0.0f;
+
+	UE_LOG(LogHttp, Verbose, TEXT("Request is waiting for processing"), this );
+
+	return true;
+}
+
+
+FHttpRequestCompleteDelegate& FHTML5HttpRequest::OnProcessRequestComplete()
+{
+	UE_LOG(LogHttp, Verbose, TEXT("FHTML5HttpRequest::OnProcessRequestComplete()"));
+	return RequestCompleteDelegate;
+}
+
+FHttpRequestProgressDelegate& FHTML5HttpRequest::OnRequestProgress()
+{
+	UE_LOG(LogHttp, Verbose, TEXT("FHTML5HttpRequest::OnRequestProgress()"));
+	return RequestProgressDelegate;
+}
+
+void FHTML5HttpRequest::FinishedRequest()
+{
+	UE_LOG(LogHttp, Verbose, TEXT("FHTML5HttpRequest::FinishedRequest()"));
+
+	// if completed, get more info
+	if (bCompleted)
+	{
+		if (Response.IsValid())
+		{
+			Response->bSucceeded = EHttpResponseCodes::IsOk(Response->HttpCode);
+			Response->ContentLength = Response->TotalBytesRead;
+		}
+	}
+
+	// if just finished, mark as stopped async processing
+	if (Response.IsValid())
+	{
+		Response->bIsReady = true;
+	}
+
+	// Clean up session/request handles that may have been created
+	CleanupRequest();
+
+	if (Response.IsValid() &&
+		Response->bSucceeded)
+	{
+		UE_LOG(LogHttp, Verbose, TEXT("%p: request has been successfully processed. HTTP code: %d, content length: %d, actual payload size: %d"), 
+			this, Response->HttpCode, Response->ContentLength, Response->Payload.Num() );
+
+		// Mark last request attempt as completed successfully
+		CompletionStatus = EHttpRequestStatus::Succeeded;
+		// Call delegate with valid request/response objects
+		OnProcessRequestComplete().ExecuteIfBound(SharedThis(this),Response,true);
+	}
+	else
+	{
+		// Mark last request attempt as completed but failed
+		CompletionStatus = EHttpRequestStatus::Failed;
+		// No response since connection failed
+		Response = NULL;
+		// Call delegate with failure
+		OnProcessRequestComplete().ExecuteIfBound(SharedThis(this),NULL,false);
+	}
+	// Remove from global list since processing is now complete
+	FHttpModule::Get().GetHttpManager().RemoveRequest(SharedThis(this));
+}
+
+
+void FHTML5HttpRequest::CleanupRequest()
+{
+	UE_LOG(LogHttp, Verbose, TEXT("FHTML5HttpRequest::CleanupRequest()"));
+	if (CompletionStatus == EHttpRequestStatus::Processing)
+	{
+		CancelRequest();
+	}
+}
+
+
+void FHTML5HttpRequest::CancelRequest()
+{
+	UE_LOG(LogHttp, Verbose, TEXT("FHTML5HttpRequest::CancelRequest()"));
+	bCanceled = true;
+}
+
+
+EHttpRequestStatus::Type FHTML5HttpRequest::GetStatus()
+{
+	UE_LOG(LogHttp, Verbose, TEXT("FHTML5HttpRequest::GetStatus()"));
+	return CompletionStatus;
+}
+
+const FHttpResponsePtr FHTML5HttpRequest::GetResponse() const 
+{
+	UE_LOG(LogHttp, Verbose, TEXT("FHTML5HttpRequest::GetResponse()"));
+	return Response;
+}
+
+
+void FHTML5HttpRequest::Tick(float DeltaSeconds)
+{
+	// check for true completion/cancellation
+	if (bCompleted || bCanceled)
+	{
+		FinishedRequest();
+		return;
+	}
+
+	// keep track of elapsed seconds
+	ElapsedTime += DeltaSeconds;
+	const float HttpTimeout = FHttpModule::Get().GetHttpTimeout();
+	if (HttpTimeout > 0 && ElapsedTime >= HttpTimeout)
+	{
+		UE_LOG(LogHttp, Warning, TEXT("Timeout processing Http request. %p"),
+			this);
+
+		// finish it off since it is timeout
+		FinishedRequest();
+	}
+}
+
+
+/****************************************************************************
+* FHTML5HttpResponse implementation
+**************************************************************************/
+
+FHTML5HttpResponse::FHTML5HttpResponse(FHTML5HttpRequest& InRequest)
+	:	Request(InRequest)
+	,	TotalBytesRead(0)
+	,	HttpCode(EHttpResponseCodes::Unknown)
+	,	ContentLength(0)
+	,	bIsReady(0)
+	,	bSucceeded(0)
+{
+}
+
+
+FHTML5HttpResponse::~FHTML5HttpResponse()
+{
+	UE_LOG(LogHttp, Verbose, TEXT("FHTML5HttpResponse::~FHTML5HttpResponse()"));
+}
+
+
+FString FHTML5HttpResponse::GetURL()
+{
+	UE_LOG(LogHttp, Verbose, TEXT("FHTML5HttpResponse::GetURL()"));
+	return Request.GetURL();
+}
+
+
+FString FHTML5HttpResponse::GetURLParameter(const FString& ParameterName)
+{
+	UE_LOG(LogHttp, Verbose, TEXT("FHTML5HttpResponse::GetURLParameter()"));
+
+	return Request.GetURLParameter(ParameterName);
+}
+
+
+FString FHTML5HttpResponse::GetHeader(const FString& HeaderName)
+{
+	UE_LOG(LogHttp, Verbose, TEXT("FHTML5HttpResponse::GetHeader()"));
+
+	FString Result(TEXT(""));
+
+	if (!bIsReady)
+	{
+		UE_LOG(LogHttp, Warning, TEXT("Can't get cached header [%s]. Response still processing. %p"),
+			*HeaderName, &Request);
+	}
+	else
+	{
+		FString* Header = Headers.Find(HeaderName);
+
+		if (Header != NULL)
+		{
+			return *Header;
+		}
+	}
+
+	return Result;
+}
+
+
+TArray<FString> FHTML5HttpResponse::GetAllHeaders()
+{
+	UE_LOG(LogHttp, Verbose, TEXT("FHTML5HttpResponse::GetAllHeaders()"));
+
+	TArray<FString> Result;
+
+	if (!bIsReady)
+	{
+		UE_LOG(LogHttp, Warning, TEXT("Can't get cached headers. Response still processing. %p"),&Request);
+	}
+	else
+	{
+		for (TMap<FString, FString>::TConstIterator It(Headers); It; ++It)
+		{
+			Result.Add(It.Key() + TEXT(": ") + It.Value());
+		}
+	}
+
+	return Result;
+}
+
+
+FString FHTML5HttpResponse::GetContentType()
+{
+	UE_LOG(LogHttp, Verbose, TEXT("FHTML5HttpResponse::GetContentType()"));
+
+	return GetHeader(TEXT("Content-Type"));
+}
+
+
+int32 FHTML5HttpResponse::GetContentLength()
+{
+	UE_LOG(LogHttp, Verbose, TEXT("FHTML5HttpResponse::GetContentLength()"));
+
+	return ContentLength;
+}
+
+
+const TArray<uint8>& FHTML5HttpResponse::GetContent()
+{
+	UE_LOG(LogHttp, Verbose, TEXT("FHTML5HttpResponse::GetContent()"));
+
+	if (!IsReady())
+	{
+		UE_LOG(LogHttp, Warning, TEXT("Payload is incomplete. Response still processing. %p"), &Request);
+	}
+
+	return Payload;
+}
+
+
+FString FHTML5HttpResponse::GetContentAsString()
+{
+	UE_LOG(LogHttp, Verbose, TEXT("FHTML5HttpResponse::GetContentAsString()"));
+
+	// Fill in our data.
+	GetContent();
+
+	TArray<uint8> ZeroTerminatedPayload;
+	ZeroTerminatedPayload.AddZeroed(Payload.Num() + 1);
+	FMemory::Memcpy(ZeroTerminatedPayload.GetData(), Payload.GetData(), Payload.Num());
+
+	return UTF8_TO_TCHAR(ZeroTerminatedPayload.GetData());
+}
+
+
+int32 FHTML5HttpResponse::GetResponseCode()
+{
+	UE_LOG(LogHttp, Verbose, TEXT("FHTML5HttpResponse::GetResponseCode()"));
+
+	return HttpCode;
+}
+
+bool FHTML5HttpResponse::IsReady()
+{
+	if (bIsReady)
+	{
+		UE_LOG(LogHttp, Verbose, TEXT("FHTML5HttpResponse::IsReady()"));
+	}
+
+	return bIsReady;
+}
\ No newline at end of file
diff --git a/Engine/Source/Runtime/Online/HTTP/Private/HTML5/HTML5HTTP.h b/Engine/Source/Runtime/Online/HTTP/Private/HTML5/HTML5HTTP.h
new file mode 100644
--- /dev/null
+++ b/Engine/Source/Runtime/Online/HTTP/Private/HTML5/HTML5HTTP.h
@@ -0,0 +1,208 @@
+// Copyright 1998-2014 Epic Games, Inc. All Rights Reserved.
+
+#pragma once
+
+/**
+* An utility function for allocation chunk of zero initialized memory
+*
+* @param NumElems number of elements to allocate (may be 0, then NULL should be returned)
+* @param ElemSize size of each element in bytes (may be 0)
+* @return Pointer to memory chunk, filled with zeros or NULL if failed
+*/
+void * CallocZero(size_t NumElems, size_t ElemSize)
+{
+	void * Return = NULL;
+	const size_t Size = NumElems * ElemSize;
+	if (Size)
+	{
+		Return = FMemory::Malloc(Size);
+
+		if (Return)
+		{
+			FMemory::Memzero(Return, Size);
+		}
+	}
+
+	return Return;
+}
+
+/**
+* HTML5 implementation of an Http request
+*/
+class FHTML5HttpRequest : public IHttpRequest
+{
+public:
+	// implementation friends
+	friend class FHTML5HttpResponse;
+
+	// Begin IHttpBase interface
+	virtual FString GetURL() override;
+	virtual FString GetURLParameter(const FString& ParameterName) override;
+	virtual FString GetHeader(const FString& HeaderName) override;
+	virtual TArray<FString> GetAllHeaders() override;
+	virtual FString GetContentType() override;
+	virtual int32 GetContentLength() override;
+	virtual const TArray<uint8>& GetContent() override;
+	// End IHttpBase interface
+
+	// Begin IHttpRequest interface
+	virtual FString GetVerb() override;
+	virtual void SetVerb(const FString& InVerb) override;
+	virtual void SetURL(const FString& InURL) override;
+	virtual void SetContent(const TArray<uint8>& ContentPayload) override;
+	virtual void SetContentAsString(const FString& ContentString) override;
+	virtual void SetHeader(const FString& HeaderName, const FString& HeaderValue) override;
+	virtual bool ProcessRequest() override;
+	virtual FHttpRequestCompleteDelegate& OnProcessRequestComplete() override;
+	virtual FHttpRequestProgressDelegate& OnRequestProgress() override;
+	virtual void CancelRequest() override;
+	virtual EHttpRequestStatus::Type GetStatus() override;
+	virtual const FHttpResponsePtr GetResponse() const override;
+	virtual void Tick(float DeltaSeconds) override;
+	// End IHttpRequest interface
+
+	/**
+	* Marks request as completed.
+	*/
+	inline void MarkAsCompleted()
+	{
+		bCompleted = true;
+	}
+
+	/**
+	* Constructor
+	*/
+	FHTML5HttpRequest();
+
+	/**
+	* Destructor.
+	*/
+	virtual ~FHTML5HttpRequest();
+
+
+private:
+
+	/**
+	* Create the session connection and initiate the web request
+	*
+	* @return true if the request was started
+	*/
+	bool StartRequest();
+
+	/**
+	* Process state for a finished request that no longer needs to be ticked
+	* Calls the completion delegate
+	*/
+	void FinishedRequest();
+
+	/**
+	* Close session/request handles and unregister callbacks
+	*/
+	void CleanupRequest();
+
+	static void StaticReceiveCallback(void *arg, void *buffer, uint32 size);
+	void ReceiveCallback(void *arg, void *buffer, uint32 size);
+	static void StaticErrorCallback(void* arg, int httpStatusCode, const char* httpStatusText);
+	void ErrorCallback(void* arg, int httpStatusCode, const char* httpStatusText);
+	static void StaticProgressCallback(void* arg, int Loaded, int Total);
+	void ProgressCallback(void* arg, int Loaded, int Total);
+
+
+private:
+	/** The response object which we will use to pair with this request */
+	TSharedPtr<class FHTML5HttpResponse, ESPMode::ThreadSafe> Response;
+
+	/** BYTE array payload to use with the request. Typically for a POST */
+	TArray<uint8> RequestPayload;
+
+	/** Delegate that will get called once request completes or on any error */
+	FHttpRequestCompleteDelegate RequestCompleteDelegate;
+
+	/** Delegate that will get called once per tick with bytes downloaded so far */
+	FHttpRequestProgressDelegate RequestProgressDelegate;
+
+	/** Current status of request being processed */
+	EHttpRequestStatus::Type CompletionStatus;
+
+	/** Cached URL */
+	FString			URL;
+	/** Cached verb */
+	FString			Verb;
+	/** Set to true if request has been canceled */
+	bool			bCanceled;
+	/** Set to true when request has been completed */
+	bool			bCompleted;
+	/** Number of bytes sent already */
+	uint32			BytesSent;
+	/** Mapping of header section to values. */
+	TMap<FString, FString> Headers;
+	/** Total elapsed time in seconds since the start of the request */
+	float ElapsedTime;
+};
+
+
+/**
+* HTML5 implementation of an Http response
+*/
+class FHTML5HttpResponse : public IHttpResponse
+{
+
+public:
+	// implementation friends
+	friend class FHTML5HttpRequest;
+
+
+	// Begin IHttpBase interface
+	virtual FString GetURL() override;
+	virtual FString GetURLParameter(const FString& ParameterName) override;
+	virtual FString GetHeader(const FString& HeaderName) override;
+	virtual TArray<FString> GetAllHeaders() override;
+	virtual FString GetContentType() override;
+	virtual int32 GetContentLength() override;
+	virtual const TArray<uint8>& GetContent() override;
+	// End IHttpBase interface
+
+	// Begin IHttpResponse interface
+	virtual int32 GetResponseCode() override;
+	virtual FString GetContentAsString() override;
+	// End IHttpResponse interface
+
+	/**
+	* Check whether a response is ready or not.
+	*/
+	bool IsReady();
+
+	/**
+	* Constructor
+	*
+	* @param InRequest - original request that created this response
+	*/
+	FHTML5HttpResponse(FHTML5HttpRequest& InRequest);
+
+	/**
+	* Destructor
+	*/
+	virtual ~FHTML5HttpResponse();
+
+
+private:
+
+	/** Request that owns this response */
+	FHTML5HttpRequest& Request;
+
+	/** BYTE array to fill in as the response is read via ReceiveCallback */
+	TArray<uint8> Payload;
+	/** Caches how many bytes of the response we've read so far */
+	int32 TotalBytesRead;
+	/** Cached key/value header pairs. Parsed once request completes */
+	TMap<FString, FString> Headers;
+	/** Cached code from completed response */
+	int32 HttpCode;
+	/** Cached content length from completed response */
+	int32 ContentLength;
+	/** True when the response has finished async processing */
+	int32 volatile bIsReady;
+	/** True if the response was successfully received/processed */
+	int32 volatile bSucceeded;
+
+};
\ No newline at end of file
diff --git a/Engine/Source/Runtime/Online/HTTP/Private/HTML5/HTML5PlatformHttp.cpp b/Engine/Source/Runtime/Online/HTTP/Private/HTML5/HTML5PlatformHttp.cpp
--- a/Engine/Source/Runtime/Online/HTTP/Private/HTML5/HTML5PlatformHttp.cpp
+++ b/Engine/Source/Runtime/Online/HTTP/Private/HTML5/HTML5PlatformHttp.cpp
@@ -2,6 +2,7 @@
 
 #include "HttpPrivatePCH.h"
 #include "HTML5PlatformHttp.h"
+#include "HTML5Http.h"
 
 
 void FHTML5PlatformHttp::Init()
@@ -16,5 +17,5 @@ void FHTML5PlatformHttp::Shutdown()
 
 IHttpRequest* FHTML5PlatformHttp::ConstructRequest()
 {
-	return FGenericPlatformHttp::ConstructRequest();
+	return new FHTML5HttpRequest();
 }
\ No newline at end of file
