UObject Serialisation for binary manifest files

diff --git a/Engine/Source/Programs/BuildPatchTool/Private/BuildPatchToolMain.cpp b/Engine/Source/Programs/BuildPatchTool/Private/BuildPatchToolMain.cpp
--- a/Engine/Source/Programs/BuildPatchTool/Private/BuildPatchToolMain.cpp
+++ b/Engine/Source/Programs/BuildPatchTool/Private/BuildPatchToolMain.cpp
@@ -28,7 +28,7 @@
 
 	Optional arguments:
 	-stdout				Adds stdout logging to the app.
-	-nochunks			Creates file based patch data instead of chunk based patch data.
+	-nochunks			Creates file based patch data instead of chunk based patch data. Should not be used with builds that have large files.
 	-FileIgnoreList=""	Specifies in quotes, the path to a text file containing BuildRoot relative files, separated by \r\n line endings, to not be included in the build.
 	-PrereqName=""		Specifies in quotes, the display name for the prerequisites installer
 	-PrereqPath=""		Specifies in quotes, the prerequisites installer to launch on successful product install
@@ -73,6 +73,8 @@
 
 #include "BuildPatchServices.h"
 
+#include "CoreUObject.h"
+
 // Ensure compiled with patch generation code
 #if WITH_BUILDPATCHGENERATION
 #else
@@ -259,17 +261,17 @@ int32 BuildPatchToolMain( const TCHAR* CommandLine )
 
 		if( PrereqNameIdx != INDEX_NONE )
 		{
-			FParse::Value( *Switches[ PrereqNameIdx ], TEXT( "FileIgnoreList=" ), PrereqName );
+			FParse::Value( *Switches[ PrereqNameIdx ], TEXT( "PrereqName=" ), PrereqName );
 		}
 
 		if( PrereqPathIdx != INDEX_NONE )
 		{
-			FParse::Value( *Switches[ PrereqPathIdx ], TEXT( "FileIgnoreList=" ), PrereqPath );
+			FParse::Value( *Switches[ PrereqPathIdx ], TEXT( "PrereqPath=" ), PrereqPath );
 		}
 
 		if( PrereqArgsIdx != INDEX_NONE )
 		{
-			FParse::Value( *Switches[ PrereqArgsIdx ], TEXT( "FileIgnoreList=" ), PrereqArgs );
+			FParse::Value( *Switches[ PrereqArgsIdx ], TEXT( "PrereqArgs=" ), PrereqArgs );
 		}
 
 		if (ManifestsListIdx != INDEX_NONE)
@@ -308,7 +310,7 @@ int32 BuildPatchToolMain( const TCHAR* CommandLine )
 					Right.TrimTrailing();
 					if (!Right.IsNumeric())
 					{
-						GLog->Log(ELogVerbosity::Error, TEXT("An error occurred processing token -customint. Non Numeric character found right of ="));
+						GLog->Log(ELogVerbosity::Error, TEXT("An error occurred processing token -customfloat. Non Numeric character found right of ="));
 						bSuccess = false;
 					}
 					CustomFields.Add(Left, FVariant(TCString<TCHAR>::Atod(*Right)));
@@ -324,7 +326,7 @@ int32 BuildPatchToolMain( const TCHAR* CommandLine )
 					Right.TrimTrailing();
 					if (!Right.IsNumeric())
 					{
-						GLog->Log(ELogVerbosity::Error, TEXT("An error occurred processing token -customfloat. Non Numeric character found right of ="));
+						GLog->Log(ELogVerbosity::Error, TEXT("An error occurred processing token -customint. Non Numeric character found right of ="));
 						bSuccess = false;
 					}
 					CustomFields.Add(Left, FVariant(TCString<TCHAR>::Atoi64(*Right)));
@@ -383,6 +385,11 @@ int32 BuildPatchToolMain( const TCHAR* CommandLine )
 	// Load the BuildPatchServices Module
 	TSharedPtr<IBuildPatchServicesModule> BuildPatchServicesModule = StaticCastSharedPtr<IBuildPatchServicesModule>( FModuleManager::Get().LoadModule( TEXT( "BuildPatchServices" ) ) );
 
+	// Initialise the UObject system and process our uobject classes
+	FModuleManager::Get().LoadModule(TEXT("CoreUObject"));
+	FCoreDelegates::OnInit.Broadcast();
+	ProcessNewlyLoadedUObjects();
+
 	// Setup the module
 	BuildPatchServicesModule->SetCloudDirectory( CloudDirectory + TEXT( "/" ) );
 
@@ -407,6 +414,7 @@ int32 BuildPatchToolMain( const TCHAR* CommandLine )
 			{
 				GLog->Log(ELogVerbosity::Error, TEXT("Could not open specified manifests to keep file"));
 				BuildPatchServicesModule.Reset();
+				FCoreDelegates::OnExit.Broadcast();
 				return 2;
 			}
 		}
@@ -418,7 +426,7 @@ int32 BuildPatchToolMain( const TCHAR* CommandLine )
 	{
 		FBuildPatchSettings Settings;
 		Settings.RootDirectory = RootDirectory + TEXT("/");
-		Settings.InAppID = AppID;
+		Settings.AppID = AppID;
 		Settings.AppName = AppName;
 		Settings.BuildVersion = BuildVersion;
 		Settings.LaunchExe = LaunchExe;
@@ -445,6 +453,7 @@ int32 BuildPatchToolMain( const TCHAR* CommandLine )
 	{
 		GLog->Log(ELogVerbosity::Error, TEXT("Unknown tool mode"));
 		BuildPatchServicesModule.Reset();
+		FCoreDelegates::OnExit.Broadcast();
 		return 3;
 	}
 
@@ -455,9 +464,12 @@ int32 BuildPatchToolMain( const TCHAR* CommandLine )
 	if (!bSuccess)
 	{
 		GLog->Log(ELogVerbosity::Error, TEXT("A fatal error occurred executing BuildPatchTool.exe"));
+		FCoreDelegates::OnExit.Broadcast();
 		return 4;
 	}
 
+	FCoreDelegates::OnExit.Broadcast();
+
 	GLog->Log(TEXT("BuildPatchToolMain completed successfuly"));
 	return 0;
 }
diff --git a/Engine/Source/Runtime/Online/BuildPatchServices/BuildPatchServices.Build.cs b/Engine/Source/Runtime/Online/BuildPatchServices/BuildPatchServices.Build.cs
--- a/Engine/Source/Runtime/Online/BuildPatchServices/BuildPatchServices.Build.cs
+++ b/Engine/Source/Runtime/Online/BuildPatchServices/BuildPatchServices.Build.cs
@@ -10,6 +10,7 @@ public class BuildPatchServices : ModuleRules
 			new string[]
 			{
 				"Core",
+				"CoreUObject",
 			}
 		);
 
diff --git a/Engine/Source/Runtime/Online/BuildPatchServices/Private/BuildPatchChunk.cpp b/Engine/Source/Runtime/Online/BuildPatchServices/Private/BuildPatchChunk.cpp
--- a/Engine/Source/Runtime/Online/BuildPatchServices/Private/BuildPatchChunk.cpp
+++ b/Engine/Source/Runtime/Online/BuildPatchServices/Private/BuildPatchChunk.cpp
@@ -6,6 +6,97 @@
 
 #include "BuildPatchServicesPrivatePCH.h"
 
+// The chunk header magic codeword, for quick checking that the opened file is a chunk file.
+#define CHUNK_HEADER_MAGIC		0xB1FE3AA2
+
+// The chunk header version number.
+#define CHUNK_HEADER_VERSION	2
+
+/* FSHAHashData implementation
+*****************************************************************************/
+FSHAHashData::FSHAHashData()
+{
+	FMemory::Memset(Hash, 0, FSHA1::DigestSize);
+}
+
+FORCEINLINE bool FSHAHashData::operator==(const FSHAHashData& Other) const
+{
+	return FMemory::Memcmp(Hash, Other.Hash, FSHA1::DigestSize) == 0;
+}
+
+FORCEINLINE bool FSHAHashData::operator!=(const FSHAHashData& Other) const
+{
+	return !(*this == Other);
+}
+
+FORCEINLINE FString FSHAHashData::ToString() const
+{
+	return BytesToHex(Hash, FSHA1::DigestSize);
+}
+
+/* FChunkHeader implementation
+*****************************************************************************/
+FChunkHeader::FChunkHeader()
+	: Magic(CHUNK_HEADER_MAGIC)
+	, Version(CHUNK_HEADER_VERSION)
+	, HashType(HASH_ROLLING)
+{
+
+}
+
+const bool FChunkHeader::IsValidMagic() const
+{
+	return Magic == CHUNK_HEADER_MAGIC;
+}
+
+FArchive& operator<<(FArchive& Ar, FChunkHeader& Header)
+{
+	// The constant sizes for each version of a header struct. Must be updated
+	// If new member variables are added the version MUST be bumped and handled properly here,
+	// and these values must never change.
+	static const uint32 Version1Size = 41;
+	static const uint32 Version2Size = 62;
+	// Calculate how much space left in the archive for reading data ( will be 0 when writing )
+	const int64 ArchiveSizeLeft = Ar.TotalSize() - Ar.Tell();
+	// Make sure the archive has enough data to read from, or we are saving instead.
+	bool bSuccess = Ar.IsSaving() || (ArchiveSizeLeft >= Version1Size);
+	if (bSuccess)
+	{
+		Ar << Header.Magic
+			<< Header.Version
+			<< Header.HeaderSize
+			<< Header.DataSize
+			<< Header.Guid
+			<< Header.RollingHash
+			<< Header.StoredAs;
+
+		// From version 2, we have a hash type choice. Previous versions default as only rolling
+		if (Header.Version >= 2)
+		{
+			bSuccess = Ar.IsSaving() || (ArchiveSizeLeft >= Version2Size);
+			if (bSuccess)
+			{
+				Ar.Serialize(Header.SHAHash.Hash, FSHA1::DigestSize);
+				Ar << Header.HashType;
+			}
+		}
+	}
+
+	// If we had a size error, zero out the header values
+	if (!bSuccess)
+	{
+		Header.Magic = 0;
+		Header.Version = 0;
+		Header.HeaderSize = 0;
+		Header.DataSize = 0;
+		Header.Guid.Invalidate();
+		Header.RollingHash = 0;
+		Header.StoredAs = 0;
+	}
+
+	return Ar;
+}
+
 /* FChunkFile implementation
 *****************************************************************************/
 FChunkFile::FChunkFile( const uint32& InReferenceCount, const bool& bInIsFromDisk )
diff --git a/Engine/Source/Runtime/Online/BuildPatchServices/Private/BuildPatchChunk.h b/Engine/Source/Runtime/Online/BuildPatchServices/Private/BuildPatchChunk.h
--- a/Engine/Source/Runtime/Online/BuildPatchServices/Private/BuildPatchChunk.h
+++ b/Engine/Source/Runtime/Online/BuildPatchServices/Private/BuildPatchChunk.h
@@ -6,6 +6,26 @@
 
 #pragma once
 
+#include "BuildPatchChunk.generated.h"
+
+/**
+ * A UStruct wrapping SHA1 hash data for serialization
+ */
+USTRUCT()
+struct FSHAHashData
+{
+	GENERATED_USTRUCT_BODY()
+
+	UPROPERTY()
+	uint8 Hash[FSHA1::DigestSize];
+
+	FSHAHashData();
+
+	FORCEINLINE bool operator==(const FSHAHashData& Other) const;
+	FORCEINLINE bool operator!=(const FSHAHashData& Other) const;
+	FORCEINLINE FString ToString() const;
+};
+
 /**
  * Constant values and typedefs
  */
@@ -26,12 +46,6 @@ namespace FBuildPatchData
 	};
 }
 
-// The chunk header magic codeword, for quick checking that the opened file is a chunk file.
-#define CHUNK_HEADER_MAGIC		0xB1FE3AA2
-
-// The chunk header version number.
-#define CHUNK_HEADER_VERSION	2
-
 /**
  * Declares a struct to store the info for a chunk header
  */
@@ -63,24 +77,20 @@ struct FChunkHeader
 	// The ChunkHash for this chunk
 	uint64 RollingHash;
 	// The FileHash for this chunk
-	FSHAHash SHAHash;
+	FSHAHashData SHAHash;
 	// How the chunk data is stored
 	uint8 StoredAs;
 
 	/**
 	 * Default Constructor sets the magic and version ready for writing out
 	 */
-	FChunkHeader()
-		: Magic( CHUNK_HEADER_MAGIC )
-		, Version( CHUNK_HEADER_VERSION )
-		, HashType( HASH_ROLLING )
-	{}
+	FChunkHeader();
 
 	/**
 	 * Checks if the Header Magic was set to the correct value
 	 * @return	true if the header magic value is correct
 	 */
-	const bool IsValidMagic() const { return Magic == CHUNK_HEADER_MAGIC; }
+	const bool IsValidMagic() const;
 
 	/**
 	 * Serialization operator.
@@ -88,97 +98,7 @@ struct FChunkHeader
 	 * @param	Header		Header to serialize
 	 * @return	Passed in archive
 	 */
-	friend FArchive& operator << ( FArchive& Ar, FChunkHeader& Header )
-	{
-		// The constant sizes for each version of a header struct. Must be updated
-		// If new member variables are added the version MUST be bumped and handled properly here,
-		// and these values must never change.
-		static const uint32 Version1Size = 41;
-		static const uint32 Version2Size = 62;
-		// Calculate how much space left in the archive for reading data ( will be 0 when writing )
-		const int64 ArchiveSizeLeft = Ar.TotalSize() - Ar.Tell();
-		// Make sure the archive has enough data to read from, or we are saving instead.
-		bool bSuccess = Ar.IsSaving() || ( ArchiveSizeLeft >= Version1Size );
-		if( bSuccess )
-		{
-			Ar	<< Header.Magic
-				<< Header.Version
-				<< Header.HeaderSize
-				<< Header.DataSize
-				<< Header.Guid
-				<< Header.RollingHash
-				<< Header.StoredAs;
-
-			// From version 2, we have a hash type choice. Previous versions default as only rolling
-			if( Header.Version >= 2 )
-			{
-				bSuccess = Ar.IsSaving() || ( ArchiveSizeLeft >= Version2Size );
-				if( bSuccess )
-				{
-					Ar << Header.SHAHash;
-					Ar << Header.HashType;
-				}
-			}
-		}
-
-		// If we had a size error, zero out the header values
-		if( !bSuccess )
-		{
-			Header.Magic = 0;
-			Header.Version = 0;
-			Header.HeaderSize = 0;
-			Header.DataSize = 0;
-			Header.Guid.Invalidate();
-			Header.RollingHash = 0;
-			Header.StoredAs = 0;
-		}
-
-		return Ar;
-	}
-};
-
-/**
- * Declares a struct to store the info for a chunk part
- */
-struct FChunkPart
-{
-	// The GUID of the chunk containing this part
-	FGuid  Guid;
-	// The offset of the first byte into the chunk
-	uint32 Offset;
-	// The size of this part
-	uint32 Size;
-
-	/**
-	 * Default Constructor
-	 */
-	FChunkPart()
-		: Guid()
-		, Offset( 0 )
-		, Size( 0 )
-	{}
-};
-
-/**
- * Declares a struct to store the info about a piece of a chunk that is inside a file
- */
-struct FFileChunkPart
-{
-	// The file containing this piece
-	FString Filename;
-	// The offset into the file of this piece
-	uint64 FileOffset;
-	// The FChunkPart that can be salvaged from this file
-	FChunkPart ChunkPart;
-
-	/**
-	 * Default Constructor
-	 */
-	FFileChunkPart()
-		: Filename()
-		, FileOffset( 0 )
-		, ChunkPart()
-	{}
+	friend FArchive& operator<< (FArchive& Ar, FChunkHeader& Header);
 };
 
 /**
diff --git a/Engine/Source/Runtime/Online/BuildPatchServices/Private/BuildPatchChunkCache.cpp b/Engine/Source/Runtime/Online/BuildPatchServices/Private/BuildPatchChunkCache.cpp
--- a/Engine/Source/Runtime/Online/BuildPatchServices/Private/BuildPatchChunkCache.cpp
+++ b/Engine/Source/Runtime/Online/BuildPatchServices/Private/BuildPatchChunkCache.cpp
@@ -665,17 +665,16 @@ void FBuildPatchChunkCache::SkipFile( const FString& Filename )
 	// data lists to avoid downloading them and keep track of the next chunks we need to get hold of
 	GLog->Logf( TEXT("FBuildPatchChunkCache::SkipFile %s"), *Filename );
 	// Get the file manifest
-	TSharedPtr< FBuildPatchFileManifest > FileManifest = InstallManifet->GetFileManifest( Filename );
+	const FFileManifestData* FileManifest = InstallManifet->GetFileManifest( Filename );
 	// Go through each chunk part, and remove it from out data lists
 	FScopeLock ScopeLock( &ChunkInfoLock );
-	for( auto ChunkPartIt = FileManifest->FileChunkParts.CreateConstIterator(); ChunkPartIt; ++ChunkPartIt )
+	for (const auto& ChunkPart : FileManifest->FileChunkParts)
 	{
-		const FChunkPart& ChunkPart = *ChunkPartIt;
-		SkipChunkPart( ChunkPart );
+		SkipChunkPart(ChunkPart);
 	}
 }
 
-void FBuildPatchChunkCache::SkipChunkPart( const FChunkPart& ChunkPart )
+void FBuildPatchChunkCache::SkipChunkPart( const FChunkPartData& ChunkPart )
 {
 	check( bDownloadsStarted.GetValue() == 0 );
 	FScopeLock ScopeLock( &ChunkInfoLock );
diff --git a/Engine/Source/Runtime/Online/BuildPatchServices/Private/BuildPatchChunkCache.h b/Engine/Source/Runtime/Online/BuildPatchServices/Private/BuildPatchChunkCache.h
--- a/Engine/Source/Runtime/Online/BuildPatchServices/Private/BuildPatchChunkCache.h
+++ b/Engine/Source/Runtime/Online/BuildPatchServices/Private/BuildPatchChunkCache.h
@@ -246,7 +246,7 @@ public:
 	 * Notify the cache that a chunk part for a partially complete file is to be skipped
 	 * @param ChunkPart		The chunk part
 	 */
-	void SkipChunkPart( const FChunkPart& ChunkPart );
+	void SkipChunkPart(const FChunkPartData& ChunkPart);
 
 	/**
 	 * Should be after initialization and after all skippable parts have been registered to kick of the chunk downloader
diff --git a/Engine/Source/Runtime/Online/BuildPatchServices/Private/BuildPatchChunkStore.cpp b/Engine/Source/Runtime/Online/BuildPatchServices/Private/BuildPatchChunkStore.cpp
new file mode 100644
--- /dev/null
+++ b/Engine/Source/Runtime/Online/BuildPatchServices/Private/BuildPatchChunkStore.cpp
@@ -0,0 +1,5 @@
+// Copyright 1998-2014 Epic Games, Inc. All Rights Reserved.
+
+#include "BuildPatchServicesPrivatePCH.h"
+#include "BuildPatchChunkStore.h"
+
diff --git a/Engine/Source/Runtime/Online/BuildPatchServices/Private/BuildPatchChunkStore.h b/Engine/Source/Runtime/Online/BuildPatchServices/Private/BuildPatchChunkStore.h
new file mode 100644
--- /dev/null
+++ b/Engine/Source/Runtime/Online/BuildPatchServices/Private/BuildPatchChunkStore.h
@@ -0,0 +1,4 @@
+// Copyright 1998-2014 Epic Games, Inc. All Rights Reserved.
+
+#pragma once
+
diff --git a/Engine/Source/Runtime/Online/BuildPatchServices/Private/BuildPatchDownloader.cpp b/Engine/Source/Runtime/Online/BuildPatchServices/Private/BuildPatchDownloader.cpp
--- a/Engine/Source/Runtime/Online/BuildPatchServices/Private/BuildPatchDownloader.cpp
+++ b/Engine/Source/Runtime/Online/BuildPatchServices/Private/BuildPatchDownloader.cpp
@@ -204,7 +204,7 @@ uint32 FBuildPatchDownloader::Run()
 		}
 		else
 		{
-			FSHAHash FileHash;
+			FSHAHashData FileHash;
 			const bool bFound = InstallManifest->GetFileDataHash( NextGuid, FileHash );
 			// Should be impossible to not exist
 			check( bFound );
diff --git a/Engine/Source/Runtime/Online/BuildPatchServices/Private/BuildPatchFileConstructor.cpp b/Engine/Source/Runtime/Online/BuildPatchServices/Private/BuildPatchFileConstructor.cpp
--- a/Engine/Source/Runtime/Online/BuildPatchServices/Private/BuildPatchFileConstructor.cpp
+++ b/Engine/Source/Runtime/Online/BuildPatchServices/Private/BuildPatchFileConstructor.cpp
@@ -340,20 +340,24 @@ bool FBuildPatchFileConstructor::ConstructFileFromChunks( const FString& Filenam
 
 	// Calculate the hash as we write the data
 	FSHA1 HashState;
-	FSHAHash HashValue;
+	FSHAHashData HashValue;
 
 	// First make sure we can get the file manifest
-	TSharedPtr< FBuildPatchFileManifest > FileManifest = BuildManifest->GetFileManifest( Filename );
-	bSuccess = FileManifest.IsValid();
+	const FFileManifestData* FileManifest = BuildManifest->GetFileManifest(Filename);
+	bSuccess = FileManifest != nullptr;
 	if( bSuccess )
 	{
-#if PLATFORM_MAC
 		if( !FileManifest->SymlinkTarget.IsEmpty() )
 		{
+#if PLATFORM_MAC
 			bSuccess = symlink(TCHAR_TO_UTF8(*FileManifest->SymlinkTarget), TCHAR_TO_UTF8(*NewFilename)) == 0;
+#else
+			const bool bSymlinkNotImplemented = false;
+			check(bSymlinkNotImplemented);
+			bSuccess = false;
+#endif
 			return bSuccess;
 		}
-#endif
 
 		// Check for resuming of existing file
 		int64 StartPosition = 0;
@@ -372,7 +376,7 @@ bool FBuildPatchFileConstructor::ConstructFileFromChunks( const FString& Filenam
 				int64 ByteCounter = 0;
 				for( int32 ChunkPartIdx = StartChunkPart; ChunkPartIdx < FileManifest->FileChunkParts.Num() && !FBuildPatchInstallError::HasFatalError(); ++ChunkPartIdx )
 				{
-					const FChunkPart& ChunkPart = FileManifest->FileChunkParts[ ChunkPartIdx ];
+					const FChunkPartData& ChunkPart = FileManifest->FileChunkParts[ ChunkPartIdx ];
 					const int64 NextBytePosition = ByteCounter + ChunkPart.Size;
 					if( NextBytePosition <= StartPosition )
 					{
@@ -425,7 +429,7 @@ bool FBuildPatchFileConstructor::ConstructFileFromChunks( const FString& Filenam
 			// For each chunk, load it, and place it's data into the file
 			for( int32 ChunkPartIdx = StartChunkPart; ChunkPartIdx < FileManifest->FileChunkParts.Num() && bSuccess && !FBuildPatchInstallError::HasFatalError(); ++ChunkPartIdx )
 			{
-				const FChunkPart& ChunkPart = FileManifest->FileChunkParts[ ChunkPartIdx ];
+				const FChunkPartData& ChunkPart = FileManifest->FileChunkParts[ChunkPartIdx];
 				if( bIsFileData )
 				{
 					bSuccess = InsertFileData( ChunkPart, *NewFile, HashState );
@@ -509,7 +513,7 @@ bool FBuildPatchFileConstructor::ConstructFileFromChunks( const FString& Filenam
 	return bSuccess;
 }
 
-bool FBuildPatchFileConstructor::InsertFileData( const FChunkPart& ChunkPart, FArchive& DestinationFile, FSHA1& HashState )
+bool FBuildPatchFileConstructor::InsertFileData(const FChunkPartData& ChunkPart, FArchive& DestinationFile, FSHA1& HashState)
 {
 	bool bSuccess = false;
 
@@ -604,7 +608,7 @@ bool FBuildPatchFileConstructor::InsertFileData( const FChunkPart& ChunkPart, FA
 	return bSuccess;
 }
 
-bool FBuildPatchFileConstructor::InsertChunkData( const FChunkPart& ChunkPart, FArchive& DestinationFile, FSHA1& HashState )
+bool FBuildPatchFileConstructor::InsertChunkData(const FChunkPartData& ChunkPart, FArchive& DestinationFile, FSHA1& HashState)
 {
 	uint8* Data;
 	uint8* DataStart;
diff --git a/Engine/Source/Runtime/Online/BuildPatchServices/Private/BuildPatchFileConstructor.h b/Engine/Source/Runtime/Online/BuildPatchServices/Private/BuildPatchFileConstructor.h
--- a/Engine/Source/Runtime/Online/BuildPatchServices/Private/BuildPatchFileConstructor.h
+++ b/Engine/Source/Runtime/Online/BuildPatchServices/Private/BuildPatchFileConstructor.h
@@ -175,7 +175,7 @@ private:
 	 * @param HashState			An FSHA1 hash state to update with the data going into the destination file.
 	 * @return	true if no file errors occurred
 	 */
-	bool InsertFileData( const FChunkPart& ChunkPart, FArchive& DestinationFile, FSHA1& HashState );
+	bool InsertFileData(const FChunkPartData& ChunkPart, FArchive& DestinationFile, FSHA1& HashState);
 
 	/**
 	 * Inserts the data data from a chunk into the destination file according to the chunk part info
@@ -184,5 +184,5 @@ private:
 	 * @param HashState			An FSHA1 hash state to update with the data going into the destination file.
 	 * @return true if no errors were detected
 	 */
-	bool InsertChunkData( const FChunkPart& ChunkPart, FArchive& DestinationFile, FSHA1& HashState );
+	bool InsertChunkData(const FChunkPartData& ChunkPart, FArchive& DestinationFile, FSHA1& HashState);
 };
diff --git a/Engine/Source/Runtime/Online/BuildPatchServices/Private/BuildPatchGeneration.cpp b/Engine/Source/Runtime/Online/BuildPatchServices/Private/BuildPatchGeneration.cpp
--- a/Engine/Source/Runtime/Online/BuildPatchServices/Private/BuildPatchGeneration.cpp
+++ b/Engine/Source/Runtime/Online/BuildPatchServices/Private/BuildPatchGeneration.cpp
@@ -3,44 +3,46 @@
 /*=============================================================================
 	BuildPatchGeneration.cpp: Implements the classes that control build
 	installation, and the generation of chunks and manifests from a build image.
-=============================================================================*/
+	=============================================================================*/
 
 #include "BuildPatchServicesPrivatePCH.h"
 
 #if WITH_BUILDPATCHGENERATION
 
 /**
- * Creates a hash value from a FSHAHash. This is so that a TMap can be keyed
- * using a FSHAHash
- * @param Hash	The sha hash
+ * Creates a 32bit hash value from a FSHAHashData. This is so that a TMap can be keyed
+ * using a FSHAHashData
+ * @param Hash	The SHA1 hash
  * @return The crc for the data
  */
-FORCEINLINE uint32 GetTypeHash( const FSHAHash& Hash )
+FORCEINLINE uint32 GetTypeHash(const FSHAHashData& Hash)
 {
-	return FCrc::MemCrc_DEPRECATED( &Hash, sizeof( FSHAHash ) );
+	return FCrc::MemCrc32(Hash.Hash, FSHA1::DigestSize);
 }
 
-#if PLATFORM_MAC
 static bool IsUnixExecutable(const TCHAR* Filename)
 {
+#if PLATFORM_MAC
 	struct stat FileInfo;
 	if (stat(TCHAR_TO_UTF8(Filename), &FileInfo) == 0)
 	{
 		return (FileInfo.st_mode & S_IXUSR) != 0;
 	}
+#endif
 	return false;
 }
 
 static FString GetSymlinkTarget(const TCHAR* Filename)
 {
+#if PLATFORM_MAC
 	ANSICHAR SymlinkTarget[MAX_PATH] = { 0 };
 	if (readlink(TCHAR_TO_UTF8(Filename), SymlinkTarget, MAX_PATH) != -1)
 	{
 		return UTF8_TO_TCHAR(SymlinkTarget);
 	}
+#endif
 	return TEXT("");
 }
-#endif
 
 /* FBuildStreamReader implementation
 *****************************************************************************/
@@ -61,9 +63,9 @@ FBuildStream::FBuildStreamReader::~FBuildStreamReader()
 	}
 }
 
-bool FBuildStream::FBuildStreamReader::Init() 
+bool FBuildStream::FBuildStreamReader::Init()
 {
-	return BuildStream != NULL && IFileManager::Get().DirectoryExists( *DepotDirectory );
+	return BuildStream != NULL && IFileManager::Get().DirectoryExists(*DepotDirectory);
 }
 
 uint32 FBuildStream::FBuildStreamReader::Run()
@@ -72,42 +74,42 @@ uint32 FBuildStream::FBuildStreamReader::Run()
 	BuildStream->Clear();
 
 	TArray< FString > AllFiles;
-	IFileManager::Get().FindFilesRecursive( AllFiles, *DepotDirectory, TEXT("*.*"), true, false );
+	IFileManager::Get().FindFilesRecursive(AllFiles, *DepotDirectory, TEXT("*.*"), true, false);
 	AllFiles.Sort();
 
 	// Remove the files that appear in an ignore list
-	FBuildDataGenerator::StripIgnoredFiles( AllFiles, DepotDirectory, IgnoreListFile );
+	FBuildDataGenerator::StripIgnoredFiles(AllFiles, DepotDirectory, IgnoreListFile);
 
 	// Allocate our file read buffer
-	uint8* FileReadBuffer = new uint8[ FileBufferSize ];
+	uint8* FileReadBuffer = new uint8[FileBufferSize];
 
 	for (auto& SourceFile : AllFiles)
 	{
 		// Read the file
-		FArchive* FileReader = IFileManager::Get().CreateFileReader( *SourceFile );
+		FArchive* FileReader = IFileManager::Get().CreateFileReader(*SourceFile);
 		const bool bBuildFileOpenSuccess = FileReader != NULL;
-		if( bBuildFileOpenSuccess )
+		if (bBuildFileOpenSuccess)
 		{
 			// Make SourceFile the format we want it in and start a new file
-			FPaths::MakePathRelativeTo( SourceFile, *( DepotDirectory + TEXT( "/" ) ) );
+			FPaths::MakePathRelativeTo(SourceFile, *(DepotDirectory + TEXT("/")));
 			int64 FileSize = FileReader->TotalSize();
 			// Process files that have bytes
-			if( FileSize > 0 )
+			if (FileSize > 0)
 			{
-				BuildStream->BeginNewFile( SourceFile, FileSize );
-				while( !FileReader->AtEnd() )
+				BuildStream->BeginNewFile(SourceFile, FileSize);
+				while (!FileReader->AtEnd())
 				{
 					const int64 SizeLeft = FileSize - FileReader->Tell();
-					const uint32 ReadLen = FMath::Min< int64 >( FileBufferSize, SizeLeft );
-					FileReader->Serialize( FileReadBuffer, ReadLen );
+					const uint32 ReadLen = FMath::Min< int64 >(FileBufferSize, SizeLeft);
+					FileReader->Serialize(FileReadBuffer, ReadLen);
 					// Copy into data stream
-					BuildStream->EnqueueData( FileReadBuffer, ReadLen );
+					BuildStream->EnqueueData(FileReadBuffer, ReadLen);
 				}
 			}
 			// Special case zero byte files
-			else if( FileSize == 0 )
+			else if (FileSize == 0)
 			{
-				BuildStream->AddEmptyFile( SourceFile );
+				BuildStream->AddEmptyFile(SourceFile);
 			}
 			FileReader->Close();
 			delete FileReader;
@@ -115,10 +117,10 @@ uint32 FBuildStream::FBuildStreamReader::Run()
 		else
 		{
 			// Not being able to load a required file from the build would be fatal, hard fault.
-			GLog->Logf( TEXT( "FBuildStreamReader: Could not open file from build! %s" ), *SourceFile );
+			GLog->Logf(TEXT("FBuildStreamReader: Could not open file from build! %s"), *SourceFile);
 			GLog->PanicFlushThreadedLogs();
 			// Use bool variable for easier to understand assert message
-			check( bBuildFileOpenSuccess );
+			check(bBuildFileOpenSuccess);
 		}
 	}
 
@@ -138,7 +140,7 @@ void FBuildStream::FBuildStreamReader::StartThread()
 
 /* FBuildStream implementation
 *****************************************************************************/
-FBuildStream::FBuildStream( const FString& RootDirectory, const FString& IgnoreListFile )
+FBuildStream::FBuildStream(const FString& RootDirectory, const FString& IgnoreListFile)
 {
 	BuildStreamReader.DepotDirectory = RootDirectory;
 	BuildStreamReader.IgnoreListFile = IgnoreListFile;
@@ -150,13 +152,13 @@ FBuildStream::~FBuildStream()
 {
 }
 
-bool FBuildStream::GetFileSpan( const uint64& StartingIdx, FString& Filename, uint64& FileSize )
+bool FBuildStream::GetFileSpan(const uint64& StartingIdx, FString& Filename, uint64& FileSize)
 {
 	bool bFound = false;
 	FilesListsCS.Lock();
 	// Find the filename
-	FFileSpan* FileSpan = FilesParsed.Find( StartingIdx );
-	if( FileSpan != NULL )
+	FFileSpan* FileSpan = FilesParsed.Find(StartingIdx);
+	if (FileSpan != NULL)
 	{
 		Filename = FileSpan->Filename;
 		FileSize = FileSpan->Size;
@@ -166,20 +168,20 @@ bool FBuildStream::GetFileSpan( const uint64& StartingIdx, FString& Filename, ui
 	return bFound;
 }
 
-uint32 FBuildStream::DequeueData( uint8* Buffer, const uint32& ReqSize, const bool WaitForData )
+uint32 FBuildStream::DequeueData(uint8* Buffer, const uint32& ReqSize, const bool WaitForData)
 {
 	// Wait for data
-	if( WaitForData )
+	if (WaitForData)
 	{
-		while ( DataAvailable() < ReqSize && !IsEndOfBuild() )
+		while (DataAvailable() < ReqSize && !IsEndOfBuild())
 		{
-			FPlatformProcess::Sleep( 0.01f );
+			FPlatformProcess::Sleep(0.01f);
 		}
 	}
 
 	BuildDataStreamCS.Lock();
-	uint32 ReadLen = FMath::Min< int64 >( ReqSize, DataAvailable() );
-	ReadLen = BuildDataStream.Dequeue( Buffer, ReadLen );
+	uint32 ReadLen = FMath::Min< int64 >(ReqSize, DataAvailable());
+	ReadLen = BuildDataStream.Dequeue(Buffer, ReadLen);
 	BuildDataStreamCS.Unlock();
 
 	return ReadLen;
@@ -187,7 +189,7 @@ uint32 FBuildStream::DequeueData( uint8* Buffer, const uint32& ReqSize, const bo
 
 const TArray< FString > FBuildStream::GetEmptyFiles()
 {
-	FScopeLock ScopeLock( &FilesListsCS );
+	FScopeLock ScopeLock(&FilesListsCS);
 	return EmptyFiles;
 }
 
@@ -214,7 +216,7 @@ bool FBuildStream::IsEndOfData()
 
 void FBuildStream::Clear()
 {
-	EndOfBuild( false );
+	EndOfBuild(false);
 
 	BuildDataStreamCS.Lock();
 	BuildDataStream.Empty();
@@ -243,7 +245,7 @@ uint32 FBuildStream::DataAvailable()
 	return rtn;
 }
 
-void FBuildStream::BeginNewFile( const FString& Filename, const uint64& FileSize )
+void FBuildStream::BeginNewFile(const FString& Filename, const uint64& FileSize)
 {
 	FFileSpan FileSpan;
 	FileSpan.Filename = Filename;
@@ -252,30 +254,30 @@ void FBuildStream::BeginNewFile( const FString& Filename, const uint64& FileSize
 	FileSpan.StartIdx = BuildDataStream.TotalDataPushed();
 	BuildDataStreamCS.Unlock();
 	FilesListsCS.Lock();
-	FilesParsed.Add( FileSpan.StartIdx, FileSpan );
+	FilesParsed.Add(FileSpan.StartIdx, FileSpan);
 	FilesListsCS.Unlock();
 }
 
-void FBuildStream::AddEmptyFile( const FString& Filename )
+void FBuildStream::AddEmptyFile(const FString& Filename)
 {
 	FilesListsCS.Lock();
-	EmptyFiles.Add( Filename );
+	EmptyFiles.Add(Filename);
 	FilesListsCS.Unlock();
 }
 
-void FBuildStream::EnqueueData( const uint8* Buffer, const uint32& Len )
+void FBuildStream::EnqueueData(const uint8* Buffer, const uint32& Len)
 {
 	// Wait for space
-	while ( SpaceLeft() < Len )
+	while (SpaceLeft() < Len)
 	{
-		FPlatformProcess::Sleep( 0.01f );
+		FPlatformProcess::Sleep(0.01f);
 	}
 	BuildDataStreamCS.Lock();
-	BuildDataStream.Enqueue( Buffer, Len );
+	BuildDataStream.Enqueue(Buffer, Len);
 	BuildDataStreamCS.Unlock();
 }
 
-void FBuildStream::EndOfBuild( bool bIsEnd )
+void FBuildStream::EndOfBuild(bool bIsEnd)
 {
 	NoMoreDataCS.Lock();
 	bNoMoreData = bIsEnd;
@@ -284,25 +286,25 @@ void FBuildStream::EndOfBuild( bool bIsEnd )
 
 /* FBuildDataProcessor implementation
 *****************************************************************************/
-FBuildDataChunkProcessor::FBuildDataChunkProcessor( FBuildPatchAppManifestRef InBuildManifest, const FString& InBuildRoot )
-	: NumNewChunks( 0 )
-	, NumKnownChunks( 0 )
-	, BuildRoot( InBuildRoot )
-	, ChunkWriter( FBuildPatchServicesModule::GetCloudDirectory() )
-	, CurrentChunkBufferPos( 0 )
-	, CurrentFile( NULL )
-	, IsProcessingChunk( false )
-	, IsProcessingChunkPart( false )
-	, IsProcessingFile( false )
-	, ChunkIsPushed( false )
-	, BackupChunkBufferPos( 0 )
-	, BackupProcessingChunk( false )
-	, BackupProcessingChunkPart( false )
+FBuildDataChunkProcessor::FBuildDataChunkProcessor(FBuildPatchAppManifestRef InBuildManifest, const FString& InBuildRoot)
+	: NumNewChunks(0)
+	, NumKnownChunks(0)
+	, BuildRoot(InBuildRoot)
+	, ChunkWriter(FBuildPatchServicesModule::GetCloudDirectory())
+	, CurrentChunkBufferPos(0)
+	, CurrentFile(nullptr)
+	, IsProcessingChunk(false)
+	, IsProcessingChunkPart(false)
+	, IsProcessingFile(false)
+	, ChunkIsPushed(false)
+	, BackupChunkBufferPos(0)
+	, BackupProcessingChunk(false)
+	, BackupProcessingChunkPart(false)
 {
 	BuildManifest = InBuildManifest;
 	CurrentChunkGuid.Invalidate();
 	BackupChunkGuid.Invalidate();
-	CurrentChunkBuffer = new uint8[ FBuildPatchData::ChunkDataSize ];
+	CurrentChunkBuffer = new uint8[FBuildPatchData::ChunkDataSize];
 }
 
 FBuildDataChunkProcessor::~FBuildDataChunkProcessor()
@@ -310,19 +312,19 @@ FBuildDataChunkProcessor::~FBuildDataChunkProcessor()
 	delete[] CurrentChunkBuffer;
 }
 
-void FBuildDataChunkProcessor::BeginNewChunk( const bool& bZeroBuffer )
+void FBuildDataChunkProcessor::BeginNewChunk(const bool& bZeroBuffer)
 {
-	check( IsProcessingChunk == false );
-	check( IsProcessingChunkPart == false );
+	check(IsProcessingChunk == false);
+	check(IsProcessingChunkPart == false);
 	IsProcessingChunk = true;
 
 	// NB: If you change this function, make sure you make required changes to 
 	// Backup of chunk if need be at the end of a recognized chunk
 
 	// Erase the chunk buffer data ready for new chunk parts
-	if( bZeroBuffer )
+	if (bZeroBuffer)
 	{
-		FMemory::Memzero( CurrentChunkBuffer, FBuildPatchData::ChunkDataSize );
+		FMemory::Memzero(CurrentChunkBuffer, FBuildPatchData::ChunkDataSize);
 		CurrentChunkBufferPos = 0;
 	}
 
@@ -332,70 +334,54 @@ void FBuildDataChunkProcessor::BeginNewChunk( const bool& bZeroBuffer )
 
 void FBuildDataChunkProcessor::BeginNewChunkPart()
 {
-	check( IsProcessingChunkPart == false );
-	check( IsProcessingChunk == true );
+	check(IsProcessingChunkPart == false);
+	check(IsProcessingChunk == true);
 	IsProcessingChunkPart = true;
 
 	// The current file should have a new chunk part setup
-	if( CurrentFile.IsValid() )
+	if (CurrentFile)
 	{
-		FChunkPart NewPart;
+		CurrentFile->FileChunkParts.Add(FChunkPartData());
+		FChunkPartData& NewPart = CurrentFile->FileChunkParts.Last();
 		NewPart.Guid = CurrentChunkGuid;
 		NewPart.Offset = CurrentChunkBufferPos;
-		CurrentFile->FileChunkParts.Add( NewPart );
 	}
 }
 
 void FBuildDataChunkProcessor::EndNewChunkPart()
 {
-	check( IsProcessingChunkPart == true );
-	check( IsProcessingChunk == true );
+	check(IsProcessingChunkPart == true);
+	check(IsProcessingChunk == true);
 	IsProcessingChunkPart = false;
 
 	// Current file should have it's last chunk part updated with the size value finalized.
-	if( CurrentFile.IsValid() )
+	if (CurrentFile)
 	{
-		FChunkPart* ChunkPart = &CurrentFile->FileChunkParts.Top();
+		FChunkPartData* ChunkPart = &CurrentFile->FileChunkParts.Last();
 		ChunkPart->Size = CurrentChunkBufferPos - ChunkPart->Offset;
-		check( ChunkPart->Size != 0 );
+		check(ChunkPart->Size != 0);
 	}
 }
 
-void FBuildDataChunkProcessor::EndNewChunk( const uint64& ChunkHash, const uint8* ChunkData, const FGuid& ChunkGuid )
+void FBuildDataChunkProcessor::EndNewChunk(const uint64& ChunkHash, const uint8* ChunkData, const FGuid& ChunkGuid)
 {
-	check( IsProcessingChunk == true );
-	check( IsProcessingChunkPart == false );
+	check(IsProcessingChunk == true);
+	check(IsProcessingChunkPart == false);
 	IsProcessingChunk = false;
 
 	// A bool that will state whether we should update and log chunking progress
 	bool bLogProgress = false;
 
 	// If the new chunk was recognized, then we will have got a different Guid through, so fix up guid for files using this chunk
-	if( CurrentChunkGuid != ChunkGuid )
+	if (CurrentChunkGuid != ChunkGuid)
 	{
 		bLogProgress = true;
 		++NumKnownChunks;
-		if( CurrentFile.IsValid() )
-		{
-			const uint32 NumFileChunkParts = CurrentFile->FileChunkParts.Num();
-			for( uint32 ChunkIdx = 0; ChunkIdx < NumFileChunkParts; ++ChunkIdx)
-			{
-				FChunkPart& ChunkPart = CurrentFile->FileChunkParts[ChunkIdx];
-				if( ChunkPart.Guid == CurrentChunkGuid )
-				{
-					ChunkPart.Guid = ChunkGuid;
-				}
-			}
-		}
-		for( auto FileIt = BuildManifest->FileManifestList.CreateConstIterator(); FileIt; ++FileIt)
+		for (auto& FileManifest : BuildManifest->Data->FileManifestList)
 		{
-			TSharedPtr< FBuildPatchFileManifest > ManifestFile = FileIt.Value();
-
-			const uint32 NumFileChunkParts = ManifestFile->FileChunkParts.Num();
-			for( uint32 ChunkIdx = 0; ChunkIdx < NumFileChunkParts; ++ChunkIdx)
+			for (auto& ChunkPart : FileManifest.FileChunkParts)
 			{
-				FChunkPart& ChunkPart = ManifestFile->FileChunkParts[ChunkIdx];
-				if( ChunkPart.Guid == CurrentChunkGuid )
+				if (ChunkPart.Guid == CurrentChunkGuid)
 				{
 					ChunkPart.Guid = ChunkGuid;
 				}
@@ -403,29 +389,43 @@ void FBuildDataChunkProcessor::EndNewChunk( const uint64& ChunkHash, const uint8
 		}
 	}
 	// If the chunk is new, count it if we were passed data
-	else if( ChunkData )
+	else if (ChunkData)
 	{
 		bLogProgress = true;
 		++NumNewChunks;
 	}
 
 	// Always queue the chunk if passed data, it will be skipped automatically if existing already, and we need to make
-	// sure the latest version is saved out when recognising an older version
-	if( ChunkData )
+	// sure the latest version is saved out when recognizing an older version
+	if (ChunkData)
 	{
-		ChunkWriter.QueueChunk( ChunkData, ChunkGuid, ChunkHash );
+		ChunkWriter.QueueChunk(ChunkData, ChunkGuid, ChunkHash);
+
+		// Also add the info to the data
+		if (!BuildManifest->ChunkInfoLookup.Contains(ChunkGuid))
+		{
+			BuildManifest->Data->ChunkList.Add(FChunkInfoData());
+			FChunkInfoData& NewChunk = BuildManifest->Data->ChunkList.Last();
+			NewChunk.Guid = ChunkGuid;
+			NewChunk.Hash = ChunkHash;
+			NewChunk.FileSize = INDEX_NONE;
+			NewChunk.GroupNumber = FCrc::MemCrc32(&ChunkGuid, sizeof(FGuid)) % 100;
+			// Lookup is used just to ensure one copy of each chunk added, we must not use the ptr as array will realloc
+			BuildManifest->ChunkInfoLookup.Add(ChunkGuid, nullptr);
+		}
+
 	}
 
-	if( bLogProgress )
+	if (bLogProgress)
 	{
 		// Output to log for builder info
-		GLog->Logf( TEXT( "%s %s [%d:%d]" ), *BuildManifest->GetAppName(), *BuildManifest->GetVersionString(), NumNewChunks, NumKnownChunks );
+		GLog->Logf(TEXT("%s %s [%d:%d]"), *BuildManifest->GetAppName(), *BuildManifest->GetVersionString(), NumNewChunks, NumKnownChunks);
 	}
 }
 
 void FBuildDataChunkProcessor::PushChunk()
 {
-	check( !ChunkIsPushed );
+	check(!ChunkIsPushed);
 	ChunkIsPushed = true;
 
 	// Then we can skip over this entire chunk using it for current files in processing
@@ -433,32 +433,29 @@ void FBuildDataChunkProcessor::PushChunk()
 	BackupProcessingChunkPart = IsProcessingChunkPart;
 	BackupChunkGuid = CurrentChunkGuid;
 	BackupChunkBufferPos = CurrentChunkBufferPos;
-	if( BackupProcessingChunkPart )
+	if (BackupProcessingChunkPart)
 	{
 		EndNewChunkPart();
 	}
-	if( BackupProcessingChunk )
+	if (BackupProcessingChunk)
 	{
 		// null ChunkData will stop the incomplete chunk being saved out
-		EndNewChunk( 0, NULL, CurrentChunkGuid );
+		EndNewChunk(0, NULL, CurrentChunkGuid);
 	}
 
 	// Start this chunk, must begin from 0
 	CurrentChunkBufferPos = 0;
-	BeginNewChunk( false );
+	BeginNewChunk(false);
 	BeginNewChunkPart();
 }
 
-void FBuildDataChunkProcessor::PopChunk( const uint64& ChunkHash, const uint8* ChunkData, const FGuid& ChunkGuid )
+void FBuildDataChunkProcessor::PopChunk(const uint64& ChunkHash, const uint8* ChunkData, const FGuid& ChunkGuid)
 {
-	check( ChunkIsPushed );
+	check(ChunkIsPushed);
 	ChunkIsPushed = false;
 
 	EndNewChunkPart();
-	EndNewChunk( ChunkHash, ChunkData, ChunkGuid );
-
-	BuildManifest->ChunkHashList.Add( ChunkGuid, ChunkHash );
-	BuildManifest->DataGroupList.Add( ChunkGuid, FCrc::MemCrc32( &ChunkGuid, sizeof( FGuid ) ) % 100 );
+	EndNewChunk(ChunkHash, ChunkData, ChunkGuid);
 
 	// Backup data for previous part chunk
 	IsProcessingChunk = BackupProcessingChunk;
@@ -471,20 +468,18 @@ void FBuildDataChunkProcessor::FinalChunk()
 	// If we're processing a chunk part, then we MUST be processing a chunk too
 	check(IsProcessingChunkPart ? IsProcessingChunk : true);
 
-	if( IsProcessingChunk )
+	if (IsProcessingChunk)
 	{
 		// The there is no more data and the last file is processed.
 		// The final chunk should be finished
-		uint64 NewChunkHash = FRollingHash< FBuildPatchData::ChunkDataSize >::GetHashForDataSet( CurrentChunkBuffer );
+		uint64 NewChunkHash = FRollingHash< FBuildPatchData::ChunkDataSize >::GetHashForDataSet(CurrentChunkBuffer);
 		FGuid ChunkGuid = CurrentChunkGuid;
-		FBuildDataGenerator::FindExistingChunkData( NewChunkHash, CurrentChunkBuffer, ChunkGuid );
-		if( IsProcessingChunkPart )
+		FBuildDataGenerator::FindExistingChunkData(NewChunkHash, CurrentChunkBuffer, ChunkGuid);
+		if (IsProcessingChunkPart)
 		{
 			EndNewChunkPart();
 		}
-		EndNewChunk( NewChunkHash, CurrentChunkBuffer, ChunkGuid );
-		BuildManifest->ChunkHashList.Add( ChunkGuid, NewChunkHash );
-		BuildManifest->DataGroupList.Add( ChunkGuid, FCrc::MemCrc32( &ChunkGuid, sizeof( FGuid ) ) % 100 );
+		EndNewChunk(NewChunkHash, CurrentChunkBuffer, ChunkGuid);
 	}
 
 	// Wait for the chunk writer to finish
@@ -492,120 +487,118 @@ void FBuildDataChunkProcessor::FinalChunk()
 	ChunkWriter.WaitForThread();
 }
 
-void FBuildDataChunkProcessor::BeginFile( const FString& FileName )
+void FBuildDataChunkProcessor::BeginFile(const FString& FileName)
 {
-	check( IsProcessingFile == false );
+	check(IsProcessingFile == false);
 	IsProcessingFile = true;
 
 	// We should start a new file, which begins with the current new chunk and current chunk offset
-	CurrentFile = MakeShareable( new FBuildPatchFileManifest() );
+	BuildManifest->Data->FileManifestList.Add(FFileManifestData());
+	CurrentFile = &BuildManifest->Data->FileManifestList.Last();
 	CurrentFile->Filename = FileName;
-#if PLATFORM_MAC
 	CurrentFile->bIsUnixExecutable = IsUnixExecutable(*(BuildRoot / FileName));
 	CurrentFile->SymlinkTarget = GetSymlinkTarget(*(BuildRoot / FileName));
-#endif
 
 	// Setup for current chunk part
-	if( IsProcessingChunkPart )
+	if (IsProcessingChunkPart)
 	{
-		FChunkPart NewPart;
+		CurrentFile->FileChunkParts.Add(FChunkPartData());
+		FChunkPartData& NewPart = CurrentFile->FileChunkParts.Last();
 		NewPart.Guid = CurrentChunkGuid;
 		NewPart.Offset = CurrentChunkBufferPos;
-		CurrentFile->FileChunkParts.Add( NewPart );
 	}
 }
 
 void FBuildDataChunkProcessor::EndFile()
 {
-	check( IsProcessingFile == true );
+	check(IsProcessingFile == true);
 	IsProcessingFile = false;
 
-	if( CurrentFile.IsValid() )
+	if (CurrentFile)
 	{
-		if( IsProcessingChunkPart )
+		if (IsProcessingChunkPart)
 		{
 			// Current file should have it's last chunk part updated with the size value finalized.
-			FChunkPart* ChunkPart = &CurrentFile->FileChunkParts.Top();
+			FChunkPartData* ChunkPart = &CurrentFile->FileChunkParts.Last();
 			ChunkPart->Size = CurrentChunkBufferPos - ChunkPart->Offset;
 		}
-
-		// And add this file to our list of file manifests
-		BuildManifest->FileManifestList.Add( CurrentFile->Filename, CurrentFile );
+		// Check file size is correct from chunks
+		int64 FileSystemSize = IFileManager::Get().FileSize(*(BuildRoot / CurrentFile->Filename));
+		CurrentFile->Init();
+		check(CurrentFile->GetFileSize() == FileSystemSize);
 	}
 
-	CurrentFile.Reset();
+	CurrentFile = nullptr;
 }
 
-void FBuildDataChunkProcessor::SkipKnownByte( const uint8& NextByte, const bool& bStartOfFile, const bool& bEndOfFile, const FString& Filename )
+void FBuildDataChunkProcessor::SkipKnownByte(const uint8& NextByte, const bool& bStartOfFile, const bool& bEndOfFile, const FString& Filename)
 {
 	// Check for start of new file
-	if( bStartOfFile )
+	if (bStartOfFile)
 	{
-		BeginFile( Filename );
+		BeginFile(Filename);
 		FileHash.Reset();
 	}
 
 	// Increment position between start and end of file
 	++CurrentChunkBufferPos;
-	FileHash.Update( &NextByte, 1 );
+	FileHash.Update(&NextByte, 1);
 
 	// Check for end of file
-	if( bEndOfFile )
+	if (bEndOfFile)
 	{
 		FileHash.Final();
-		FileHash.GetHash( CurrentFile->FileHash.Hash );
+		FileHash.GetHash(CurrentFile->FileHash.Hash);
 		EndFile();
 	}
 }
 
-void FBuildDataChunkProcessor::ProcessNewByte( const uint8& NewByte, const bool& bStartOfFile, const bool& bEndOfFile, const FString& Filename )
+void FBuildDataChunkProcessor::ProcessNewByte(const uint8& NewByte, const bool& bStartOfFile, const bool& bEndOfFile, const FString& Filename)
 {
 	// If we finished a chunk, we begin a new chunk and new chunk part
-	if( !IsProcessingChunk )
+	if (!IsProcessingChunk)
 	{
 		BeginNewChunk();
 		BeginNewChunkPart();
 	}
 	// Or if we finished a chunk part (by recognizing a chunk hash), we begin a new one
-	else if( !IsProcessingChunkPart )
+	else if (!IsProcessingChunkPart)
 	{
 		BeginNewChunkPart();
 	}
 
 	// Check for start of new file
-	if( bStartOfFile )
+	if (bStartOfFile)
 	{
-		BeginFile( Filename );
+		BeginFile(Filename);
 		FileHash.Reset();
 	}
 
 	// We add the old byte to our new chunk which will be used as part of the file it belonged to
-	CurrentChunkBuffer[ CurrentChunkBufferPos++ ] = NewByte;
-	FileHash.Update( &NewByte, 1 );
+	CurrentChunkBuffer[CurrentChunkBufferPos++] = NewByte;
+	FileHash.Update(&NewByte, 1);
 
 	// Check for end of file
-	if( bEndOfFile )
+	if (bEndOfFile)
 	{
 		FileHash.Final();
-		FileHash.GetHash( CurrentFile->FileHash.Hash );
+		FileHash.GetHash(CurrentFile->FileHash.Hash);
 		EndFile();
 	}
 
 	// Do we have a full new chunk?
-	check( CurrentChunkBufferPos <= FBuildPatchData::ChunkDataSize );
-	if( CurrentChunkBufferPos == FBuildPatchData::ChunkDataSize )
+	check(CurrentChunkBufferPos <= FBuildPatchData::ChunkDataSize);
+	if (CurrentChunkBufferPos == FBuildPatchData::ChunkDataSize)
 	{
-		uint64 NewChunkHash = FRollingHash< FBuildPatchData::ChunkDataSize >::GetHashForDataSet( CurrentChunkBuffer );
+		uint64 NewChunkHash = FRollingHash< FBuildPatchData::ChunkDataSize >::GetHashForDataSet(CurrentChunkBuffer);
 		FGuid ChunkGuid = CurrentChunkGuid;
-		FBuildDataGenerator::FindExistingChunkData( NewChunkHash, CurrentChunkBuffer, ChunkGuid );
+		FBuildDataGenerator::FindExistingChunkData(NewChunkHash, CurrentChunkBuffer, ChunkGuid);
 		EndNewChunkPart();
-		EndNewChunk( NewChunkHash, CurrentChunkBuffer, ChunkGuid );
-		BuildManifest->ChunkHashList.Add( ChunkGuid, NewChunkHash );
-		BuildManifest->DataGroupList.Add( ChunkGuid, FCrc::MemCrc32( &ChunkGuid, sizeof( FGuid ) ) % 100 );
+		EndNewChunk(NewChunkHash, CurrentChunkBuffer, ChunkGuid);
 	}
 }
 
-void FBuildDataChunkProcessor::GetChunkStats( uint32& OutNewFiles, uint32& OutKnownFiles )
+void FBuildDataChunkProcessor::GetChunkStats(uint32& OutNewFiles, uint32& OutKnownFiles)
 {
 	OutNewFiles = NumNewChunks;
 	OutKnownFiles = NumKnownChunks;
@@ -620,46 +613,44 @@ const TMap<FGuid, int64>& FBuildDataChunkProcessor::GetChunkFilesizes()
 /* FBuildDataFileProcessor implementation
 *****************************************************************************/
 FBuildDataFileProcessor::FBuildDataFileProcessor(FBuildPatchAppManifestRef InBuildManifest, const FString& InBuildRoot, const FDateTime& InDataThresholdTime)
-	: NumNewFiles( 0 )
-	, NumKnownFiles( 0 )
-	, BuildManifest( InBuildManifest )
-	, BuildRoot( InBuildRoot )
-	, CurrentFile( NULL )
+	: NumNewFiles(0)
+	, NumKnownFiles(0)
+	, BuildManifest(InBuildManifest)
+	, BuildRoot(InBuildRoot)
+	, CurrentFile(nullptr)
 	, FileHash()
-	, FileSize( 0 )
-	, IsProcessingFile( false )
-	, DataThresholdTime( InDataThresholdTime )
+	, FileSize(0)
+	, IsProcessingFile(false)
+	, DataThresholdTime(InDataThresholdTime)
 {
 }
 
-void FBuildDataFileProcessor::BeginFile( const FString& InFileName )
+void FBuildDataFileProcessor::BeginFile(const FString& InFileName)
 {
-	check( !IsProcessingFile );
+	check(!IsProcessingFile);
 	IsProcessingFile = true;
 
 	// Create the new file
-	CurrentFile = MakeShareable( new FBuildPatchFileManifest() );
+	BuildManifest->Data->FileManifestList.Add(FFileManifestData());
+	CurrentFile = &BuildManifest->Data->FileManifestList.Last();
 	CurrentFile->Filename = InFileName;
-#if PLATFORM_MAC
 	CurrentFile->bIsUnixExecutable = IsUnixExecutable(*(BuildRoot / InFileName));
 	CurrentFile->SymlinkTarget = GetSymlinkTarget(*(BuildRoot / InFileName));
-#endif
 
 	// Add the 'chunk part' which will just be the whole file
-	FChunkPart FilePart;
-	CurrentFile->FileChunkParts.Add( FilePart );
+	CurrentFile->FileChunkParts.Add(FChunkPartData());
 
 	// Reset the hash and size counter
 	FileHash.Reset();
 	FileSize = 0;
 }
 
-void FBuildDataFileProcessor::ProcessFileData( const uint8* Data, const uint32& DataLen )
+void FBuildDataFileProcessor::ProcessFileData(const uint8* Data, const uint32& DataLen)
 {
-	check( IsProcessingFile );
+	check(IsProcessingFile);
 
 	// Update the hash
-	FileHash.Update( Data, DataLen );
+	FileHash.Update(Data, DataLen);
 
 	// Count size
 	FileSize += DataLen;
@@ -667,37 +658,48 @@ void FBuildDataFileProcessor::ProcessFileData( const uint8* Data, const uint32&
 
 void FBuildDataFileProcessor::EndFile()
 {
-	check( IsProcessingFile );
-	check( CurrentFile.IsValid() );
-	check( CurrentFile->FileChunkParts.Num() == 1 );
+	check(IsProcessingFile);
+	check(CurrentFile);
+	check(CurrentFile->FileChunkParts.Num() == 1);
 	IsProcessingFile = false;
 
-	// Finalise the hash
+	// Finalize the hash
 	FileHash.Final();
-	FileHash.GetHash( CurrentFile->FileHash.Hash );
+	FileHash.GetHash(CurrentFile->FileHash.Hash);
 
 	// Use hash and full file name to find out if we have a file match
-	FString FullPath = FPaths::Combine( *BuildRoot, *CurrentFile->Filename );
+	FString FullPath = FPaths::Combine(*BuildRoot, *CurrentFile->Filename);
 	FGuid FileGuid = FGuid::NewGuid();
-	bool bFoundSameFile = FBuildDataGenerator::FindExistingFileData( FullPath, CurrentFile->FileHash, DataThresholdTime, FileGuid );
+	bool bFoundSameFile = FBuildDataGenerator::FindExistingFileData(FullPath, CurrentFile->FileHash, DataThresholdTime, FileGuid);
 
 	// Fill the 'chunk part' info
-	FChunkPart& FilePart = CurrentFile->FileChunkParts[0];
+	FChunkPartData& FilePart = CurrentFile->FileChunkParts[0];
 	FilePart.Guid = FileGuid;
 	FilePart.Offset = 0;
 	FilePart.Size = FileSize;
 
-	// Add it to the manifest
-	BuildManifest->FileManifestList.Add( CurrentFile->Filename, CurrentFile );
-
-	// Set the data group
-	BuildManifest->DataGroupList.Add( FileGuid, FCrc::MemCrc32( &FileGuid, sizeof( FGuid ) ) % 100 );
+	// Init and check size
+	CurrentFile->Init();
+	check(CurrentFile->GetFileSize() == FileSize);
 
 	// Always call save to ensure new versions exist when recognising old ones, the save will be skipped if required copies already exist.
-	FBuildDataGenerator::SaveOutFileData( FullPath, CurrentFile->FileHash, FileGuid );
+	FBuildDataGenerator::SaveOutFileData(FullPath, CurrentFile->FileHash, FileGuid);
+
+	// Also add the info to the data
+	if (!BuildManifest->ChunkInfoLookup.Contains(FileGuid))
+	{
+		BuildManifest->Data->ChunkList.Add(FChunkInfoData());
+		FChunkInfoData& DataInfo = BuildManifest->Data->ChunkList.Last();
+		DataInfo.Guid = FileGuid;
+		DataInfo.Hash = INDEX_NONE;
+		DataInfo.FileSize = FileSize;
+		DataInfo.GroupNumber = FCrc::MemCrc32(&FileGuid, sizeof(FGuid)) % 100;
+		// Lookup is used just to ensure one copy of each chunk added, we must not use the ptr as array will realloc
+		BuildManifest->ChunkInfoLookup.Add(FileGuid, nullptr);
+	}
 
 	// Count new/known stats
-	if( bFoundSameFile == false )
+	if (bFoundSameFile == false)
 	{
 		++NumNewFiles;
 	}
@@ -707,10 +709,10 @@ void FBuildDataFileProcessor::EndFile()
 	}
 
 	// Output to log for builder info
-	GLog->Logf( TEXT( "%s %s [%d:%d]" ), *BuildManifest->GetAppName(), *BuildManifest->GetVersionString(), NumNewFiles, NumKnownFiles );
+	GLog->Logf(TEXT("%s %s [%d:%d]"), *BuildManifest->GetAppName(), *BuildManifest->GetVersionString(), NumNewFiles, NumKnownFiles);
 }
 
-void FBuildDataFileProcessor::GetFileStats( uint32& OutNewFiles, uint32& OutKnownFiles )
+void FBuildDataFileProcessor::GetFileStats(uint32& OutNewFiles, uint32& OutKnownFiles)
 {
 	OutNewFiles = NumNewFiles;
 	OutKnownFiles = NumKnownFiles;
@@ -718,15 +720,15 @@ void FBuildDataFileProcessor::GetFileStats( uint32& OutNewFiles, uint32& OutKnow
 
 /* FBuildSimpleChunkCache::FChunkReader implementation
 *****************************************************************************/
-FBuildGenerationChunkCache::FChunkReader::FChunkReader( const FString& InChunkFilePath, TSharedRef< FChunkFile > InChunkFile, uint32* InBytesRead, const FDateTime& InDataAgeThreshold )
-	: ChunkFilePath( InChunkFilePath )
-	, ChunkFileReader( NULL )
-	, ChunkFile( InChunkFile )
-	, FileBytesRead( InBytesRead )
-	, MemoryBytesRead( 0 )
-	, DataAgeThreshold( InDataAgeThreshold )
+FBuildGenerationChunkCache::FChunkReader::FChunkReader(const FString& InChunkFilePath, TSharedRef< FChunkFile > InChunkFile, uint32* InBytesRead, const FDateTime& InDataAgeThreshold)
+	: ChunkFilePath(InChunkFilePath)
+	, ChunkFileReader(NULL)
+	, ChunkFile(InChunkFile)
+	, FileBytesRead(InBytesRead)
+	, MemoryBytesRead(0)
+	, DataAgeThreshold(InDataAgeThreshold)
 {
-	ChunkFile->GetDataLock( &ChunkData, &ChunkHeader );
+	ChunkFile->GetDataLock(&ChunkData, &ChunkHeader);
 }
 
 FBuildGenerationChunkCache::FChunkReader::~FChunkReader()
@@ -734,7 +736,7 @@ FBuildGenerationChunkCache::FChunkReader::~FChunkReader()
 	ChunkFile->ReleaseDataLock();
 
 	// Close file handle
-	if( ChunkFileReader != NULL )
+	if (ChunkFileReader != NULL)
 	{
 		ChunkFileReader->Close();
 		delete ChunkFileReader;
@@ -745,18 +747,18 @@ FBuildGenerationChunkCache::FChunkReader::~FChunkReader()
 FArchive* FBuildGenerationChunkCache::FChunkReader::GetArchive()
 {
 	// Open file handle?
-	if( ChunkFileReader == NULL )
+	if (ChunkFileReader == NULL)
 	{
-		ChunkFileReader = IFileManager::Get().CreateFileReader( *ChunkFilePath );
-		if( ChunkFileReader == NULL )
+		ChunkFileReader = IFileManager::Get().CreateFileReader(*ChunkFilePath);
+		if (ChunkFileReader == NULL)
 		{
 			// Break the magic to mark as invalid
 			ChunkHeader->Magic = 0;
-			GLog->Logf( TEXT( "WARNING: Skipped missing chunk file %s" ), *ChunkFilePath );
+			GLog->Logf(TEXT("WARNING: Skipped missing chunk file %s"), *ChunkFilePath);
 			return NULL;
 		}
 		// Read Header?
-		if( ChunkHeader->Guid.IsValid() == false )
+		if (ChunkHeader->Guid.IsValid() == false)
 		{
 			*ChunkFileReader << *ChunkHeader;
 		}
@@ -771,13 +773,13 @@ FArchive* FBuildGenerationChunkCache::FChunkReader::GetArchive()
 		const int64 ExpectedFileSize = ChunkHeader->DataSize + ChunkHeader->HeaderSize;
 		const int64 ChunkFileSize = ChunkFileReader->TotalSize();
 		const int64 NextByte = ChunkHeader->HeaderSize + *FileBytesRead;
-		if( ChunkFileSize == ExpectedFileSize
-		 && NextByte < ChunkFileSize )
+		if (ChunkFileSize == ExpectedFileSize
+			&& NextByte < ChunkFileSize)
 		{
 			// Seek to next byte
-			ChunkFileReader->Seek( NextByte );
+			ChunkFileReader->Seek(NextByte);
 			// Break the magic to mark as invalid if archive errored, this chunk will get ignored
-			if( ChunkFileReader->GetError() )
+			if (ChunkFileReader->GetError())
 			{
 				ChunkHeader->Magic = 0;
 			}
@@ -789,25 +791,25 @@ FArchive* FBuildGenerationChunkCache::FChunkReader::GetArchive()
 		}
 		// If this chunk is valid and compressed, we must read the entire file and decompress to memory now if we have not already
 		// as we cannot compare to compressed data
-		if( *FileBytesRead == 0 && ChunkHeader->StoredAs & FChunkHeader::STORED_COMPRESSED )
+		if (*FileBytesRead == 0 && ChunkHeader->StoredAs & FChunkHeader::STORED_COMPRESSED)
 		{
 			// Load the compressed chunk data
 			TArray< uint8 > CompressedData;
-			CompressedData.Empty( ChunkHeader->DataSize );
-			CompressedData.AddUninitialized( ChunkHeader->DataSize );
-			ChunkFileReader->Serialize( CompressedData.GetData(), ChunkHeader->DataSize );
+			CompressedData.Empty(ChunkHeader->DataSize);
+			CompressedData.AddUninitialized(ChunkHeader->DataSize);
+			ChunkFileReader->Serialize(CompressedData.GetData(), ChunkHeader->DataSize);
 			// Uncompress
 			bool bSuceess = FCompression::UncompressMemory(
-				static_cast< ECompressionFlags >( COMPRESS_ZLIB | COMPRESS_BiasMemory ),
+				static_cast<ECompressionFlags>(COMPRESS_ZLIB | COMPRESS_BiasMemory),
 				ChunkData,
 				FBuildPatchData::ChunkDataSize,
 				CompressedData.GetData(),
-				ChunkHeader->DataSize );
+				ChunkHeader->DataSize);
 			// Mark that we have fully read decompressed data and update the chunkfile's data size as we are expanding it
 			*FileBytesRead = FBuildPatchData::ChunkDataSize;
 			ChunkHeader->DataSize = FBuildPatchData::ChunkDataSize;
 			// Check uncompression was OK
-			if( !bSuceess )
+			if (!bSuceess)
 			{
 				ChunkHeader->Magic = 0;
 			}
@@ -818,7 +820,7 @@ FArchive* FBuildGenerationChunkCache::FChunkReader::GetArchive()
 
 const bool FBuildGenerationChunkCache::FChunkReader::IsValidChunk()
 {
-	if( ChunkHeader->Guid.IsValid() == false )
+	if (ChunkHeader->Guid.IsValid() == false)
 	{
 		GetArchive();
 	}
@@ -831,38 +833,38 @@ const bool FBuildGenerationChunkCache::FChunkReader::IsValidChunk()
 
 const FGuid& FBuildGenerationChunkCache::FChunkReader::GetChunkGuid()
 {
-	if( ChunkHeader->Guid.IsValid() == false )
+	if (ChunkHeader->Guid.IsValid() == false)
 	{
 		GetArchive();
 	}
 	return ChunkHeader->Guid;
 }
 
-void FBuildGenerationChunkCache::FChunkReader::ReadNextBytes( uint8** OutDataBuffer, const uint32& ReadLength )
+void FBuildGenerationChunkCache::FChunkReader::ReadNextBytes(uint8** OutDataBuffer, const uint32& ReadLength)
 {
-	uint8* BufferNextByte = &ChunkData[ MemoryBytesRead ];
+	uint8* BufferNextByte = &ChunkData[MemoryBytesRead];
 	// Do we need to load from disk?
-	if( ( MemoryBytesRead + ReadLength ) > *FileBytesRead )
+	if ((MemoryBytesRead + ReadLength) > *FileBytesRead)
 	{
 		FArchive* Reader = GetArchive();
 		// Do not allow incorrect usage
-		check( Reader );
-		check( ReadLength <= BytesLeft() );
+		check(Reader);
+		check(ReadLength <= BytesLeft());
 		// Read the number of bytes extra we need
 		const int32 NumFileBytesRead = *FileBytesRead;
 		const int32 NextMemoryBytesRead = MemoryBytesRead + ReadLength;
-		const uint32 FileReadLen = FMath::Max<int32>( 0, NextMemoryBytesRead - NumFileBytesRead );
+		const uint32 FileReadLen = FMath::Max<int32>(0, NextMemoryBytesRead - NumFileBytesRead);
 		*FileBytesRead += FileReadLen;
-		Reader->Serialize( BufferNextByte, FileReadLen );
+		Reader->Serialize(BufferNextByte, FileReadLen);
 		// Assert if read error, if theres some problem accessing chunks then continuing would cause bad patch
 		// ratios, so it's better to hard fault.
 		const bool bChunkReadOK = !Reader->GetError();
-		if( !bChunkReadOK )
+		if (!bChunkReadOK)
 		{
 			// Print something helpful
-			GLog->Logf( TEXT( "FATAL ERROR: Could not read from chunk FArchive %s" ), *ChunkFilePath );
+			GLog->Logf(TEXT("FATAL ERROR: Could not read from chunk FArchive %s"), *ChunkFilePath);
 			// Check with bool variable so that output will be readable
-			check( bChunkReadOK );
+			check(bChunkReadOK);
 		}
 	}
 	MemoryBytesRead += ReadLength;
@@ -871,7 +873,7 @@ void FBuildGenerationChunkCache::FChunkReader::ReadNextBytes( uint8** OutDataBuf
 
 const uint32 FBuildGenerationChunkCache::FChunkReader::BytesLeft()
 {
-	if( ChunkHeader->Guid.IsValid() == false )
+	if (ChunkHeader->Guid.IsValid() == false)
 	{
 		GetArchive();
 	}
@@ -885,44 +887,44 @@ FBuildGenerationChunkCache::FBuildGenerationChunkCache(const FDateTime& DataAgeT
 {
 }
 
-TSharedRef< FBuildGenerationChunkCache::FChunkReader > FBuildGenerationChunkCache::GetChunkReader( const FString& ChunkFilePath )
+TSharedRef< FBuildGenerationChunkCache::FChunkReader > FBuildGenerationChunkCache::GetChunkReader(const FString& ChunkFilePath)
 {
-	if( ChunkCache.Contains( ChunkFilePath ) == false )
+	if (ChunkCache.Contains(ChunkFilePath) == false)
 	{
 		// Remove oldest access from cache?
-		if( ChunkCache.Num() >= NumChunksToCache )
+		if (ChunkCache.Num() >= NumChunksToCache)
 		{
-			FString OldestAccessChunk = TEXT( "" );
+			FString OldestAccessChunk = TEXT("");
 			double OldestAccessTime = FPlatformTime::Seconds();
-			for( auto ChunkCacheIt = ChunkCache.CreateConstIterator(); ChunkCacheIt; ++ChunkCacheIt )
+			for (auto ChunkCacheIt = ChunkCache.CreateConstIterator(); ChunkCacheIt; ++ChunkCacheIt)
 			{
 				const FString& ChunkFilePath = ChunkCacheIt.Key();
 				const FChunkFile& ChunkFile = ChunkCacheIt.Value().Get();
-				if( ChunkFile.GetLastAccessTime() < OldestAccessTime )
+				if (ChunkFile.GetLastAccessTime() < OldestAccessTime)
 				{
 					OldestAccessTime = ChunkFile.GetLastAccessTime();
 					OldestAccessChunk = ChunkFilePath;
 				}
 			}
-			ChunkCache.Remove( OldestAccessChunk );
-			delete BytesReadPerChunk[ OldestAccessChunk ];
-			BytesReadPerChunk.Remove( OldestAccessChunk );
+			ChunkCache.Remove(OldestAccessChunk);
+			delete BytesReadPerChunk[OldestAccessChunk];
+			BytesReadPerChunk.Remove(OldestAccessChunk);
 		}
 		// Add the chunk to cache
-		ChunkCache.Add( ChunkFilePath, MakeShareable( new FChunkFile( 1, true ) ) );
-		BytesReadPerChunk.Add( ChunkFilePath, new uint32( 0 ) );
+		ChunkCache.Add(ChunkFilePath, MakeShareable(new FChunkFile(1, true)));
+		BytesReadPerChunk.Add(ChunkFilePath, new uint32(0));
 	}
-	return MakeShareable( new FChunkReader( ChunkFilePath, ChunkCache[ ChunkFilePath ], BytesReadPerChunk[ ChunkFilePath ], DataAgeThreshold ) );
+	return MakeShareable(new FChunkReader(ChunkFilePath, ChunkCache[ChunkFilePath], BytesReadPerChunk[ChunkFilePath], DataAgeThreshold));
 }
 
 void FBuildGenerationChunkCache::Cleanup()
 {
-	ChunkCache.Empty( 0 );
-	for( auto BytesReadPerChunkIt = BytesReadPerChunk.CreateConstIterator(); BytesReadPerChunkIt; ++BytesReadPerChunkIt )
+	ChunkCache.Empty(0);
+	for (auto BytesReadPerChunkIt = BytesReadPerChunk.CreateConstIterator(); BytesReadPerChunkIt; ++BytesReadPerChunkIt)
 	{
 		delete BytesReadPerChunkIt.Value();
 	}
-	BytesReadPerChunk.Empty( 0 );
+	BytesReadPerChunk.Empty(0);
 }
 
 /* FBuildGenerationChunkCache system singleton setup
@@ -932,21 +934,21 @@ TSharedPtr< FBuildGenerationChunkCache > FBuildGenerationChunkCache::SingletonIn
 void FBuildGenerationChunkCache::Init(const FDateTime& DataAgeThreshold)
 {
 	// We won't allow misuse of these functions
-	check( !SingletonInstance.IsValid() );
+	check(!SingletonInstance.IsValid());
 	SingletonInstance = MakeShareable(new FBuildGenerationChunkCache(DataAgeThreshold));
 }
 
 FBuildGenerationChunkCache& FBuildGenerationChunkCache::Get()
 {
 	// We won't allow misuse of these functions
-	check( SingletonInstance.IsValid() );
+	check(SingletonInstance.IsValid());
 	return *SingletonInstance.Get();
 }
 
 void FBuildGenerationChunkCache::Shutdown()
 {
 	// We won't allow misuse of these functions
-	check( SingletonInstance.IsValid() );
+	check(SingletonInstance.IsValid());
 	SingletonInstance->Cleanup();
 	SingletonInstance.Reset();
 }
@@ -956,7 +958,7 @@ void FBuildGenerationChunkCache::Shutdown()
 TMap< FGuid, FString > FBuildDataGenerator::ExistingChunkGuidInventory;
 TMap< uint64, TArray< FGuid > > FBuildDataGenerator::ExistingChunkHashInventory;
 bool FBuildDataGenerator::ExistingChunksEnumerated = false;
-TMap< FSHAHash, TArray< FString > > FBuildDataGenerator::ExistingFileInventory;
+TMap< FSHAHashData, TArray< FString > > FBuildDataGenerator::ExistingFileInventory;
 bool FBuildDataGenerator::ExistingFilesEnumerated = false;
 FCriticalSection FBuildDataGenerator::SingleConcurrentBuildCS;
 
@@ -983,63 +985,53 @@ static void AddCustomFieldsToBuildManifest(const TMap<FString, FVariant>& Custom
 
 /* FBuildDataGenerator implementation
 *****************************************************************************/
-bool FBuildDataGenerator::GenerateChunksManifestFromDirectory( const FBuildPatchSettings& Settings )
-{
-	const FString& RootDirectory = Settings.RootDirectory;
-	const uint32& InAppID = Settings.InAppID;
-	const FString& AppName = Settings.AppName;
-	const FString& BuildVersion = Settings.BuildVersion;
-	const FString& LaunchExe = Settings.LaunchExe;
-	const FString& LaunchCommand = Settings.LaunchCommand;
-	const FString& IgnoreListFile = Settings.IgnoreListFile;
-	const FString& PrereqName = Settings.PrereqName;
-	const FString& PrereqPath = Settings.PrereqPath;
-	const FString& prereqArgs = Settings.PrereqArgs;
-	const float& DataAgeThreshold = Settings.DataAgeThreshold;
-	const bool& bShouldHonorReuseThreshold = Settings.bShouldHonorReuseThreshold;
-
+bool FBuildDataGenerator::GenerateChunksManifestFromDirectory(const FBuildPatchSettings& Settings)
+{
 	// Output to log for builder info
-	GLog->Logf( TEXT( "Running Chunks Patch Generation for: %u:%s %s" ), InAppID, *AppName, *BuildVersion );
+	GLog->Logf(TEXT("Running Chunks Patch Generation for: %u:%s %s"), Settings.AppID, *Settings.AppName, *Settings.BuildVersion);
 
 	// Take the build CS
-	FScopeLock SingleConcurrentBuild( &SingleConcurrentBuildCS );
+	FScopeLock SingleConcurrentBuild(&SingleConcurrentBuildCS);
 
 	// Create our chunk cache
-	const FDateTime Cutoff = bShouldHonorReuseThreshold ? FDateTime::UtcNow() - FTimespan::FromDays(DataAgeThreshold) : FDateTime::MinValue();
+	const FDateTime Cutoff = Settings.bShouldHonorReuseThreshold ? FDateTime::UtcNow() - FTimespan::FromDays(Settings.DataAgeThreshold) : FDateTime::MinValue();
 	FBuildGenerationChunkCache::Init(Cutoff);
 
 	// Create a manifest
-	FBuildPatchAppManifestRef BuildManifest = MakeShareable( new FBuildPatchAppManifest() );
+	FBuildPatchAppManifestRef BuildManifest = MakeShareable(new FBuildPatchAppManifest());
+
+	// Setup custom fields
+	AddCustomFieldsToBuildManifest(Settings.CustomFields, BuildManifest);
 
 	// Reset chunk inventory
 	ExistingChunksEnumerated = false;
 	ExistingChunkGuidInventory.Empty();
 	ExistingChunkHashInventory.Empty();
 
-	// Clear current data
-	BuildManifest->DestroyData(); 
-
 	// Declare a build processor
-	FBuildDataChunkProcessor DataProcessor( BuildManifest, RootDirectory );
+	FBuildDataChunkProcessor DataProcessor(BuildManifest, Settings.RootDirectory);
 
 	// Create a build streamer
-	FBuildStream* BuildStream = new FBuildStream( RootDirectory, IgnoreListFile );
-
-	// Set the App details
-	BuildManifest->AppID = InAppID;
-	BuildManifest->AppNameString = AppName;
-	BuildManifest->BuildVersionString = BuildVersion;
-	BuildManifest->LaunchExeString = LaunchExe;
-	BuildManifest->LaunchCommandString = LaunchCommand;
-	BuildManifest->bIsFileData = false;
+	FBuildStream* BuildStream = new FBuildStream(Settings.RootDirectory, Settings.IgnoreListFile);
+
+	// Set the basic details
+	BuildManifest->Data->bIsFileData = false;
+	BuildManifest->Data->AppID = Settings.AppID;
+	BuildManifest->Data->AppName = Settings.AppName;
+	BuildManifest->Data->BuildVersion = Settings.BuildVersion;
+	BuildManifest->Data->LaunchExe = Settings.LaunchExe;
+	BuildManifest->Data->LaunchCommand = Settings.LaunchCommand;
+	BuildManifest->Data->PrereqName = Settings.PrereqName;
+	BuildManifest->Data->PrereqPath = Settings.PrereqPath;
+	BuildManifest->Data->PrereqArgs = Settings.PrereqArgs;
 
 	// Create a data buffer
 	const uint32 DataBufferSize = FBuildPatchData::ChunkDataSize;
-	uint8* DataBuffer = new uint8[ DataBufferSize ];
+	uint8* DataBuffer = new uint8[DataBufferSize];
 
 	// We'll need a rolling hash for chunking
 	FRollingHash< FBuildPatchData::ChunkDataSize >* RollingHash = new FRollingHash< FBuildPatchData::ChunkDataSize >();
-	
+
 	// Refers to how much data has been processed (into the FBuildDataProcessor)
 	uint64 ProcessPos = 0;
 
@@ -1057,10 +1049,10 @@ bool FBuildDataGenerator::GenerateChunksManifestFromDirectory( const FBuildPatch
 	const double TimeGenStarted = LastProgressLog;
 
 	// Loop through all data
-	while ( !BuildStream->IsEndOfData() )
+	while (!BuildStream->IsEndOfData())
 	{
 		// Grab some data from the build stream
-		ReadLen = BuildStream->DequeueData( DataBuffer, DataBufferSize );
+		ReadLen = BuildStream->DequeueData(DataBuffer, DataBufferSize);
 
 		// A bool says if there's no more data to come from the Build Stream
 		const bool bNoMoreData = BuildStream->IsEndOfData();
@@ -1072,18 +1064,18 @@ bool FBuildDataGenerator::GenerateChunksManifestFromDirectory( const FBuildPatch
 		uint32 PaddedZeros = 0;
 
 		// Process data while we have more
-		while ( ( DataBufferPos < ReadLen ) || ( bNoMoreData && PaddedZeros < RollingHash->GetWindowSize() ) )
+		while ((DataBufferPos < ReadLen) || (bNoMoreData && PaddedZeros < RollingHash->GetWindowSize()))
 		{
 			// Prime the rolling hash
-			if( RollingHash->GetNumDataNeeded() > 0 )
+			if (RollingHash->GetNumDataNeeded() > 0)
 			{
-				if( DataBufferPos < ReadLen )
+				if (DataBufferPos < ReadLen)
 				{
-					RollingHash->ConsumeByte( DataBuffer[ DataBufferPos++ ] );
+					RollingHash->ConsumeByte(DataBuffer[DataBufferPos++]);
 				}
 				else
 				{
-					RollingHash->ConsumeByte( 0 );
+					RollingHash->ConsumeByte(0);
 					++PaddedZeros;
 				}
 				// Keep looping until primed
@@ -1094,24 +1086,24 @@ bool FBuildDataGenerator::GenerateChunksManifestFromDirectory( const FBuildPatch
 			FGuid ChunkGuid;
 			const uint64 WindowHash = RollingHash->GetWindowHash();
 			const TRingBuffer< uint8, FBuildPatchData::ChunkDataSize >& WindowData = RollingHash->GetWindowData();
-			bool ChunkRecognised = FindExistingChunkData( WindowHash, WindowData, ChunkGuid );
-			if( ChunkRecognised )
+			bool ChunkRecognised = FindExistingChunkData(WindowHash, WindowData, ChunkGuid);
+			if (ChunkRecognised)
 			{
 				// Process all bytes
 				DataProcessor.PushChunk();
 				const uint32 WindowDataSize = RollingHash->GetWindowSize() - PaddedZeros;
-				for( uint32 i = 0; i < WindowDataSize; ++i )
+				for (uint32 i = 0; i < WindowDataSize; ++i)
 				{
-					const bool bStartOfFile = BuildStream->GetFileSpan( ProcessPos, FileName, FileDataCount );
+					const bool bStartOfFile = BuildStream->GetFileSpan(ProcessPos, FileName, FileDataCount);
 					const bool bEndOfFile = FileDataCount <= 1;
-					check( FileDataCount > 0 );// If FileDataCount is ever 0, it means this piece of data belongs to no file, so something is wrong
-					DataProcessor.SkipKnownByte( WindowData[i], bStartOfFile, bEndOfFile, FileName );
+					check(FileDataCount > 0);// If FileDataCount is ever 0, it means this piece of data belongs to no file, so something is wrong
+					DataProcessor.SkipKnownByte(WindowData[i], bStartOfFile, bEndOfFile, FileName);
 					++ProcessPos;
 					--FileDataCount;
 				}
-				uint8* SerialWindowData = new uint8[ FBuildPatchData::ChunkDataSize ];
-				WindowData.Serialize( SerialWindowData );
-				DataProcessor.PopChunk( WindowHash, SerialWindowData, ChunkGuid );
+				uint8* SerialWindowData = new uint8[FBuildPatchData::ChunkDataSize];
+				WindowData.Serialize(SerialWindowData);
+				DataProcessor.PopChunk(WindowHash, SerialWindowData, ChunkGuid);
 				delete[] SerialWindowData;
 
 				// Clear
@@ -1120,29 +1112,29 @@ bool FBuildDataGenerator::GenerateChunksManifestFromDirectory( const FBuildPatch
 			else
 			{
 				// Process one byte
-				const bool bStartOfFile = BuildStream->GetFileSpan( ProcessPos, FileName, FileDataCount );
+				const bool bStartOfFile = BuildStream->GetFileSpan(ProcessPos, FileName, FileDataCount);
 				const bool bEndOfFile = FileDataCount <= 1;
-				DataProcessor.ProcessNewByte( WindowData.Bottom(), bStartOfFile, bEndOfFile, FileName );
+				DataProcessor.ProcessNewByte(WindowData.Bottom(), bStartOfFile, bEndOfFile, FileName);
 				++ProcessPos;
 				--FileDataCount;
 
 				// Roll
-				if( DataBufferPos < ReadLen )
+				if (DataBufferPos < ReadLen)
 				{
-					RollingHash->RollForward( DataBuffer[ DataBufferPos++ ] );
+					RollingHash->RollForward(DataBuffer[DataBufferPos++]);
 				}
-				else if( bNoMoreData )
+				else if (bNoMoreData)
 				{
-					RollingHash->RollForward( 0 );
+					RollingHash->RollForward(0);
 					++PaddedZeros;
 				}
 			}
 
 			// Log processed data
-			if( ( FPlatformTime::Seconds() - LastProgressLog ) >= 10.0 )
+			if ((FPlatformTime::Seconds() - LastProgressLog) >= 10.0)
 			{
 				LastProgressLog = FPlatformTime::Seconds();
-				GLog->Logf( TEXT( "Processed %lld bytes." ), ProcessPos );
+				GLog->Logf(TEXT("Processed %lld bytes."), ProcessPos);
 			}
 		}
 	}
@@ -1157,16 +1149,35 @@ bool FBuildDataGenerator::GenerateChunksManifestFromDirectory( const FBuildPatch
 	const TArray< FString >& EmptyFileList = BuildStream->GetEmptyFiles();
 	for (const auto& EmptyFile : EmptyFileList)
 	{
-		FBuildPatchFileManifest* EmptyFileManifest = new FBuildPatchFileManifest();
-		EmptyHasher.GetHash( EmptyFileManifest->FileHash.Hash );
-		EmptyFileManifest->Filename = EmptyFile;
-		BuildManifest->FileManifestList.Add( EmptyFileManifest->Filename, MakeShareable( EmptyFileManifest ) );
+		BuildManifest->Data->FileManifestList.Add(FFileManifestData());
+		FFileManifestData& EmptyFileManifest = BuildManifest->Data->FileManifestList.Last();
+		EmptyFileManifest.Filename = EmptyFile;
+		EmptyHasher.GetHash(EmptyFileManifest.FileHash.Hash);
 	}
 
 	// Add chunk sizes
 	const TMap<FGuid, int64>& ChunkFilesizes = DataProcessor.GetChunkFilesizes();
-	BuildManifest->ChunkFilesizeList.Empty(ChunkFilesizes.Num());
-	BuildManifest->ChunkFilesizeList.Append(ChunkFilesizes);
+	for (FChunkInfoData& ChunkInfo : BuildManifest->Data->ChunkList)
+	{
+		if (ChunkFilesizes.Contains(ChunkInfo.Guid))
+		{
+			ChunkInfo.FileSize = ChunkFilesizes[ChunkInfo.Guid];
+		}
+	}
+
+	// Fill out lookups
+	BuildManifest->InitLookups();
+
+	// Save manifest into the cloud directory
+	FString BaseFilename = FBuildPatchServicesModule::GetCloudDirectory() / FDefaultValueHelper::RemoveWhitespaces(BuildManifest->Data->AppName + BuildManifest->Data->BuildVersion);
+	FString JsonFilename = BaseFilename + TEXT(".manifest");
+	FString BinaryFilename = BaseFilename + TEXT(".binary.manifest");
+	BuildManifest->SaveToFile(JsonFilename, false);
+	BuildManifest->SaveToFile(BinaryFilename, true);
+
+	// Output to log for builder info
+	GLog->Logf(TEXT("Saved manifest to %s"), *JsonFilename);
+	GLog->Logf(TEXT("Saved manifest to %s"), *BinaryFilename);
 
 	// Clean up memory
 	delete[] DataBuffer;
@@ -1174,100 +1185,83 @@ bool FBuildDataGenerator::GenerateChunksManifestFromDirectory( const FBuildPatch
 	delete RollingHash;
 	FBuildGenerationChunkCache::Shutdown();
 
-	// Save manifest into the cloud directory
-	FString ManifestFilename = FBuildPatchServicesModule::GetCloudDirectory();
-	ManifestFilename /= FDefaultValueHelper::RemoveWhitespaces( BuildManifest->AppNameString + BuildManifest->BuildVersionString ) + TEXT(".manifest");
-	AddCustomFieldsToBuildManifest(Settings.CustomFields, BuildManifest);
-	BuildManifest->SaveToFile(ManifestFilename);
-
-	// Output to log for builder info
-	GLog->Logf( TEXT( "Saved manifest to %s" ), *ManifestFilename );
-	GLog->Logf( TEXT( "Generation took %s" ), *FPlatformTime::PrettyTime( FPlatformTime::Seconds() - TimeGenStarted ) );
-
 	// @TODO LSwift: Detect errors and return false on failure
 	return true;
 }
 
-bool FBuildDataGenerator::GenerateFilesManifestFromDirectory( const FBuildPatchSettings& Settings )
+bool FBuildDataGenerator::GenerateFilesManifestFromDirectory(const FBuildPatchSettings& Settings)
 {
-	const FString& RootDirectory = Settings.RootDirectory;
-	const uint32& InAppID = Settings.InAppID;
-	const FString& AppName = Settings.AppName;
-	const FString& BuildVersion = Settings.BuildVersion;
-	const FString& LaunchExe = Settings.LaunchExe;
-	const FString& LaunchCommand = Settings.LaunchCommand;
-	const FString& IgnoreListFile = Settings.IgnoreListFile;
-	const float& DataAgeThreshold = Settings.DataAgeThreshold;
-	const bool& bShouldHonorReuseThreshold = Settings.bShouldHonorReuseThreshold;
-
 	// Output to log for builder info
-	GLog->Logf( TEXT( "Running Files Patch Generation for: %u:%s %s" ), InAppID, *AppName, *BuildVersion );
+	GLog->Logf(TEXT("Running Files Patch Generation for: %u:%s %s"), Settings.AppID, *Settings.AppName, *Settings.BuildVersion);
 
 	// Take the build CS
-	FScopeLock SingleConcurrentBuild( &SingleConcurrentBuildCS );
+	FScopeLock SingleConcurrentBuild(&SingleConcurrentBuildCS);
 
 	// Create a manifest
-	FBuildPatchAppManifestRef BuildManifest = MakeShareable( new FBuildPatchAppManifest() );
+	FBuildPatchAppManifestRef BuildManifest = MakeShareable(new FBuildPatchAppManifest());
+
+	// Setup custom fields
+	AddCustomFieldsToBuildManifest(Settings.CustomFields, BuildManifest);
 
 	// Reset file inventory
 	ExistingFilesEnumerated = false;
 	ExistingFileInventory.Empty();
 
-	// Clear current data
-	BuildManifest->DestroyData(); 
-
 	// Declare a build processor
-	const FDateTime Cutoff = bShouldHonorReuseThreshold ? FDateTime::UtcNow() - FTimespan::FromDays(DataAgeThreshold) : FDateTime::MinValue();
-	FBuildDataFileProcessor DataProcessor(BuildManifest, RootDirectory, Cutoff);
-
-	// Set the App details
-	BuildManifest->AppID = InAppID;
-	BuildManifest->AppNameString = AppName;
-	BuildManifest->BuildVersionString = BuildVersion;
-	BuildManifest->LaunchExeString = LaunchExe;
-	BuildManifest->LaunchCommandString = LaunchCommand;
-	BuildManifest->bIsFileData = true;
+	const FDateTime Cutoff = Settings.bShouldHonorReuseThreshold ? FDateTime::UtcNow() - FTimespan::FromDays(Settings.DataAgeThreshold) : FDateTime::MinValue();
+	FBuildDataFileProcessor DataProcessor(BuildManifest, Settings.RootDirectory, Cutoff);
+
+	// Set the basic details
+	BuildManifest->Data->bIsFileData = true;
+	BuildManifest->Data->AppID = Settings.AppID;
+	BuildManifest->Data->AppName = Settings.AppName;
+	BuildManifest->Data->BuildVersion = Settings.BuildVersion;
+	BuildManifest->Data->LaunchExe = Settings.LaunchExe;
+	BuildManifest->Data->LaunchCommand = Settings.LaunchCommand;
+	BuildManifest->Data->PrereqName = Settings.PrereqName;
+	BuildManifest->Data->PrereqPath = Settings.PrereqPath;
+	BuildManifest->Data->PrereqArgs = Settings.PrereqArgs;
 
 	// Create a data buffer
-	uint8* FileReadBuffer = new uint8[ FileBufferSize ];
-	
+	uint8* FileReadBuffer = new uint8[FileBufferSize];
+
 	// Refers to how much data has been processed (into the FBuildDataFileProcessor)
 	uint64 ProcessPos = 0;
 
 	// Find all files
 	TArray< FString > AllFiles;
-	IFileManager::Get().FindFilesRecursive( AllFiles, *RootDirectory, TEXT("*.*"), true, false );
+	IFileManager::Get().FindFilesRecursive(AllFiles, *Settings.RootDirectory, TEXT("*.*"), true, false);
 	AllFiles.Sort();
 
 	// Remove the files that appear in an ignore list
-	FBuildDataGenerator::StripIgnoredFiles( AllFiles, RootDirectory, IgnoreListFile );
+	FBuildDataGenerator::StripIgnoredFiles(AllFiles, Settings.RootDirectory, Settings.IgnoreListFile);
 
 	// Loop through all files
-	for( auto FileIt = AllFiles.CreateConstIterator(); FileIt; ++FileIt )
+	for (auto FileIt = AllFiles.CreateConstIterator(); FileIt; ++FileIt)
 	{
 		const FString& FileName = *FileIt;
 		// Read the file
-		FArchive* FileReader = IFileManager::Get().CreateFileReader( *FileName );
-		if( FileReader != NULL )
+		FArchive* FileReader = IFileManager::Get().CreateFileReader(*FileName);
+		if (FileReader != NULL)
 		{
 			// Make SourceFile the format we want it in and start a new file
 			FString SourceFile = FileName;
-			FPaths::MakePathRelativeTo( SourceFile, *( RootDirectory + TEXT( "/" ) ) );
+			FPaths::MakePathRelativeTo(SourceFile, *(Settings.RootDirectory + TEXT("/")));
 			int64 FileSize = FileReader->TotalSize();
-			if( FileSize < 0 )
+			if (FileSize < 0)
 			{
 				// Skip potential error ( INDEX_NONE == -1 )
 				continue;
 			}
-			DataProcessor.BeginFile( SourceFile );
-			while( !FileReader->AtEnd() )
+			DataProcessor.BeginFile(SourceFile);
+			while (!FileReader->AtEnd())
 			{
 				const int64 SizeLeft = FileSize - FileReader->Tell();
-				const uint32 ReadLen = FMath::Min< int64 >( FileBufferSize, SizeLeft );
+				const uint32 ReadLen = FMath::Min< int64 >(FileBufferSize, SizeLeft);
 				ProcessPos += ReadLen;
-				FileReader->Serialize( FileReadBuffer, ReadLen );
+				FileReader->Serialize(FileReadBuffer, ReadLen);
 				// Copy into data stream
-				DataProcessor.ProcessFileData( FileReadBuffer, ReadLen );
+				DataProcessor.ProcessFileData(FileReadBuffer, ReadLen);
 			}
 			FileReader->Close();
 			delete FileReader;
@@ -1279,41 +1273,46 @@ bool FBuildDataGenerator::GenerateFilesManifestFromDirectory( const FBuildPatchS
 		}
 	}
 
-	// Clean up memory
-	delete[] FileReadBuffer;
+	// Fill out lookups
+	BuildManifest->InitLookups();
 
 	// Save manifest into the cloud directory
-	FString ManifestFilename = FBuildPatchServicesModule::GetCloudDirectory();
-	ManifestFilename /= FDefaultValueHelper::RemoveWhitespaces( BuildManifest->AppNameString + BuildManifest->BuildVersionString ) + TEXT(".manifest");
-	AddCustomFieldsToBuildManifest(Settings.CustomFields, BuildManifest);
-	BuildManifest->SaveToFile(ManifestFilename);
+	FString BaseFilename = FBuildPatchServicesModule::GetCloudDirectory() / FDefaultValueHelper::RemoveWhitespaces(BuildManifest->Data->AppName + BuildManifest->Data->BuildVersion);
+	FString JsonFilename = BaseFilename + TEXT(".manifest");
+	FString BinaryFilename = BaseFilename + TEXT(".binary.manifest");
+	BuildManifest->SaveToFile(JsonFilename, false);
+	BuildManifest->SaveToFile(BinaryFilename, true);
 
 	// Output to log for builder info
-	GLog->Logf( TEXT( "Saved manifest to %s" ), *ManifestFilename );
+	GLog->Logf(TEXT("Saved manifest to %s"), *JsonFilename);
+	GLog->Logf(TEXT("Saved manifest to %s"), *BinaryFilename);
+
+	// Clean up memory
+	delete[] FileReadBuffer;
 
 	// @TODO LSwift: Detect errors and return false on failure
 	return true;
 }
 
-bool FBuildDataGenerator::FindExistingChunkData( const uint64& ChunkHash, const uint8* ChunkData, FGuid& ChunkGuid )
+bool FBuildDataGenerator::FindExistingChunkData(const uint64& ChunkHash, const uint8* ChunkData, FGuid& ChunkGuid)
 {
 	// Quick code hack
 	TRingBuffer< uint8, FBuildPatchData::ChunkDataSize > ChunkDataRing;
 	const uint32 ChunkDataLen = FBuildPatchData::ChunkDataSize;
-	ChunkDataRing.Enqueue( ChunkData, ChunkDataLen );
-	return FindExistingChunkData( ChunkHash, ChunkDataRing, ChunkGuid );
+	ChunkDataRing.Enqueue(ChunkData, ChunkDataLen);
+	return FindExistingChunkData(ChunkHash, ChunkDataRing, ChunkGuid);
 }
 
-bool FBuildDataGenerator::FindExistingChunkData( const uint64& ChunkHash, const TRingBuffer< uint8, FBuildPatchData::ChunkDataSize >& ChunkData, FGuid& ChunkGuid )
+bool FBuildDataGenerator::FindExistingChunkData(const uint64& ChunkHash, const TRingBuffer< uint8, FBuildPatchData::ChunkDataSize >& ChunkData, FGuid& ChunkGuid)
 {
 	bool bFoundMatchingChunk = false;
 
 	// Perform an inventory on Cloud chunks if not already done
-	if( ExistingChunksEnumerated == false )
+	if (ExistingChunksEnumerated == false)
 	{
 		IFileManager& FileManager = IFileManager::Get();
 		FString JSONOutput;
-		TSharedRef< TJsonWriter< TCHAR, TPrettyJsonPrintPolicy< TCHAR > > > DebugWriter = TJsonWriterFactory< TCHAR, TPrettyJsonPrintPolicy< TCHAR > >::Create( &JSONOutput );
+		TSharedRef< TJsonWriter< TCHAR, TPrettyJsonPrintPolicy< TCHAR > > > DebugWriter = TJsonWriterFactory< TCHAR, TPrettyJsonPrintPolicy< TCHAR > >::Create(&JSONOutput);
 		DebugWriter->WriteObjectStart();
 
 		// Find all manifest files
@@ -1340,7 +1339,7 @@ bool FBuildDataGenerator::FindExistingChunkData( const uint64& ChunkHash, const
 				{
 					const double LoadManifestTime = FPlatformTime::Seconds() - StartLoadManifest;
 					GLog->Logf(TEXT("BuildDataGenerator: Loaded %s in %.1f seconds"), *ManifestFile, LoadManifestTime);
-					if(!BuildManifest->IsFileDataManifest())
+					if (!BuildManifest->IsFileDataManifest())
 					{
 						TArray<FGuid> ChunksReferenced;
 						BuildManifest->GetDataList(ChunksReferenced);
@@ -1391,32 +1390,32 @@ bool FBuildDataGenerator::FindExistingChunkData( const uint64& ChunkHash, const
 	}
 
 	// Do we have a chunk matching this data?
-	if( ExistingChunkHashInventory.Num() > 0 )
+	if (ExistingChunkHashInventory.Num() > 0)
 	{
-		TArray< FGuid >* ChunkList = ExistingChunkHashInventory.Find( ChunkHash );
-		if( ChunkList != NULL )
+		TArray< FGuid >* ChunkList = ExistingChunkHashInventory.Find(ChunkHash);
+		if (ChunkList != NULL)
 		{
 			// We need to load each chunk in this list and compare data
-			for( auto ChunkIt = ChunkList->CreateConstIterator(); ChunkIt && !bFoundMatchingChunk ; ++ChunkIt)
+			for (auto ChunkIt = ChunkList->CreateConstIterator(); ChunkIt && !bFoundMatchingChunk; ++ChunkIt)
 			{
 				FGuid Guid = *ChunkIt;
 				if (!ExistingChunkGuidInventory.Contains(Guid))
 				{
 					ExistingChunkGuidInventory.Add(Guid, DiscoverChunkFilename(Guid, ChunkHash));
 				}
-				const FString& SourceFile = ExistingChunkGuidInventory[ Guid ];
+				const FString& SourceFile = ExistingChunkGuidInventory[Guid];
 				// Read the file
-				uint8* TempChunkData = new uint8[ FBuildPatchData::ChunkDataSize ];
-				ChunkData.Serialize( TempChunkData );
+				uint8* TempChunkData = new uint8[FBuildPatchData::ChunkDataSize];
+				ChunkData.Serialize(TempChunkData);
 				bool bChunkIsUsable = true;
-				if( CompareDataToChunk( SourceFile, TempChunkData, Guid, bChunkIsUsable ) )
+				if (CompareDataToChunk(SourceFile, TempChunkData, Guid, bChunkIsUsable))
 				{
 					// We have a chunk match!!
 					bFoundMatchingChunk = true;
 					ChunkGuid = Guid;
 				}
 				// Check if this chunk should be dumped
-				if(!bChunkIsUsable)
+				if (!bChunkIsUsable)
 				{
 					GLog->Logf(TEXT("BuildDataGenerator: Chunk %s unusable, removed from inventory."), *ChunkGuid.ToString());
 					ChunkList->Remove(Guid);
@@ -1459,101 +1458,101 @@ FString FBuildDataGenerator::DiscoverChunkFilename(const FGuid& ChunkGuid, const
 	return TEXT("");
 }
 
-bool FBuildDataGenerator::FindExistingFileData( const FString& InSourceFile, const FSHAHash& InFileHash, const FDateTime& DataThresholdTime, FGuid& OutFileGuid )
+bool FBuildDataGenerator::FindExistingFileData(const FString& InSourceFile, const FSHAHashData& InFileHash, const FDateTime& DataThresholdTime, FGuid& OutFileGuid)
 {
 	bool bFoundMatchingFile = false;
 
 	// Perform an inventory on Cloud files if not already done
-	if( ExistingFilesEnumerated == false )
+	if (ExistingFilesEnumerated == false)
 	{
 		IFileManager& FileManager = IFileManager::Get();
 		FString JSONOutput;
-		TSharedRef< TJsonWriter< TCHAR, TPrettyJsonPrintPolicy< TCHAR > > > DebugWriter = TJsonWriterFactory< TCHAR, TPrettyJsonPrintPolicy< TCHAR > >::Create( &JSONOutput );
+		TSharedRef< TJsonWriter< TCHAR, TPrettyJsonPrintPolicy< TCHAR > > > DebugWriter = TJsonWriterFactory< TCHAR, TPrettyJsonPrintPolicy< TCHAR > >::Create(&JSONOutput);
 		DebugWriter->WriteObjectStart();
 
 		TArray< FGuid > FoundFiles;
 
 		// The directory containing old filename version files
-		const FString CloudFileDir = FBuildPatchServicesModule::GetCloudDirectory() / TEXT( "Files" );
+		const FString CloudFileDir = FBuildPatchServicesModule::GetCloudDirectory() / TEXT("Files");
 		// The directory containing new filename version files
-		const FString CloudFile2Dir = FBuildPatchServicesModule::GetCloudDirectory() / TEXT( "FilesV2" );
+		const FString CloudFile2Dir = FBuildPatchServicesModule::GetCloudDirectory() / TEXT("FilesV2");
 
-		GLog->Logf( TEXT( "BuildDataGenerator: Enumerating Files from %s" ), *CloudFile2Dir );
-		if( FileManager.DirectoryExists( *CloudFile2Dir ) )
+		GLog->Logf(TEXT("BuildDataGenerator: Enumerating Files from %s"), *CloudFile2Dir);
+		if (FileManager.DirectoryExists(*CloudFile2Dir))
 		{
 			const double StartEnumerate = FPlatformTime::Seconds();
 
 			// Find all files
 			TArray<FString> AllFiles;
-			FileManager.FindFilesRecursive( AllFiles, *CloudFile2Dir, TEXT("*.file"), true, false );
+			FileManager.FindFilesRecursive(AllFiles, *CloudFile2Dir, TEXT("*.file"), true, false);
 
-			for( auto FileIt = AllFiles.CreateConstIterator(); FileIt; ++FileIt )
+			for (auto FileIt = AllFiles.CreateConstIterator(); FileIt; ++FileIt)
 			{
 				FString SourceFile = *FileIt;
-				FSHAHash FoundHash;
+				FSHAHashData FoundHash;
 				FGuid FoundGuid;
-				FBuildPatchUtils::GetFileDetailFromNewFilename( SourceFile, FoundGuid, FoundHash );
+				FBuildPatchUtils::GetFileDetailFromNewFilename(SourceFile, FoundGuid, FoundHash);
 
 				// Add to inventory
-				FoundFiles.Add( FoundGuid );
-				TArray< FString >& FileList = ExistingFileInventory.FindOrAdd( FoundHash );
-				FileList.Add( SourceFile );
+				FoundFiles.Add(FoundGuid);
+				TArray< FString >& FileList = ExistingFileInventory.FindOrAdd(FoundHash);
+				FileList.Add(SourceFile);
 
 				const void* HashBuffer = &FoundHash;
-				DebugWriter->WriteValue( FoundGuid.ToString(), FString::FromBlob( static_cast<const uint8*>( HashBuffer ), sizeof( FoundHash ) ) );
+				DebugWriter->WriteValue(FoundGuid.ToString(), FString::FromBlob(static_cast<const uint8*>(HashBuffer), sizeof(FoundHash)));
 			}
 
 			const double EnumerateTime = FPlatformTime::Seconds() - StartEnumerate;
-			GLog->Logf( TEXT( "BuildDataGenerator: Found %d new name files in %.1f seconds" ), AllFiles.Num(), EnumerateTime );
+			GLog->Logf(TEXT("BuildDataGenerator: Found %d new name files in %.1f seconds"), AllFiles.Num(), EnumerateTime);
 		}
 		else
 		{
-			GLog->Logf( TEXT( "BuildDataGenerator: Cloud directory does not exist: %s" ), *CloudFile2Dir );
+			GLog->Logf(TEXT("BuildDataGenerator: Cloud directory does not exist: %s"), *CloudFile2Dir);
 		}
 
-		GLog->Logf( TEXT( "BuildDataGenerator: Enumerating Files From %s" ), *CloudFileDir );
-		if( IFileManager::Get().DirectoryExists( *CloudFileDir ) )
+		GLog->Logf(TEXT("BuildDataGenerator: Enumerating Files From %s"), *CloudFileDir);
+		if (IFileManager::Get().DirectoryExists(*CloudFileDir))
 		{
 			const double StartEnumerate = FPlatformTime::Seconds();
 			const int32 PreviousNumFiles = FoundFiles.Num();
 
 			// Find all files
 			TArray<FString> AllFiles;
-			IFileManager::Get().FindFilesRecursive( AllFiles, *CloudFileDir, TEXT("*.file"), true, false );
+			IFileManager::Get().FindFilesRecursive(AllFiles, *CloudFileDir, TEXT("*.file"), true, false);
 
-			for( auto FileIt = AllFiles.CreateConstIterator(); FileIt; ++FileIt )
+			for (auto FileIt = AllFiles.CreateConstIterator(); FileIt; ++FileIt)
 			{
 				FString SourceFile = *FileIt;
 
 				// Skip any for GUID we already found
 				FGuid FoundGuid;
-				FBuildPatchUtils::GetGUIDFromFilename( SourceFile, FoundGuid );
-				if( FoundFiles.Contains( FoundGuid ) )
+				FBuildPatchUtils::GetGUIDFromFilename(SourceFile, FoundGuid);
+				if (FoundFiles.Contains(FoundGuid))
 				{
 					continue;
 				}
 
 				// Read the file
-				FArchive* FileReader = IFileManager::Get().CreateFileReader( *SourceFile );
-				if( FileReader != NULL )
+				FArchive* FileReader = IFileManager::Get().CreateFileReader(*SourceFile);
+				if (FileReader != NULL)
 				{
 					// Read the header
 					FChunkHeader Header;
 					*FileReader << Header;
 
 					// Check magic
-					if( Header.IsValidMagic() && Header.HashType == FChunkHeader::HASH_SHA1 )
+					if (Header.IsValidMagic() && Header.HashType == FChunkHeader::HASH_SHA1)
 					{
-						TArray< FString >& FileList = ExistingFileInventory.FindOrAdd( Header.SHAHash );
-						FileList.Add( SourceFile );
-						FoundFiles.Add( Header.Guid );
+						TArray< FString >& FileList = ExistingFileInventory.FindOrAdd(Header.SHAHash);
+						FileList.Add(SourceFile);
+						FoundFiles.Add(Header.Guid);
 
 						const void* HashBuffer = &Header.SHAHash;
-						DebugWriter->WriteValue( Header.Guid.ToString(), FString::FromBlob( static_cast<const uint8*>( HashBuffer ), sizeof( Header.SHAHash ) ) );
+						DebugWriter->WriteValue(Header.Guid.ToString(), FString::FromBlob(static_cast<const uint8*>(HashBuffer), sizeof(Header.SHAHash)));
 					}
 					else
 					{
-						GLog->Logf( TEXT( "BuildDataGenerator: Failed magic/hashtype check on file [%d:%d] %s" ), Header.Magic, Header.HashType, *SourceFile );
+						GLog->Logf(TEXT("BuildDataGenerator: Failed magic/hashtype check on file [%d:%d] %s"), Header.Magic, Header.HashType, *SourceFile);
 					}
 
 					FileReader->Close();
@@ -1561,24 +1560,24 @@ bool FBuildDataGenerator::FindExistingFileData( const FString& InSourceFile, con
 				}
 				else
 				{
-					GLog->Logf( TEXT( "BuildDataGenerator: Failed to read chunk %s" ), *SourceFile );
+					GLog->Logf(TEXT("BuildDataGenerator: Failed to read chunk %s"), *SourceFile);
 				}
 			}
 
 			const double EnumerateTime = FPlatformTime::Seconds() - StartEnumerate;
 			const int32 NewNumFiles = FoundFiles.Num();
-			GLog->Logf( TEXT( "BuildDataGenerator: Found %d extra old files in %.1f seconds" ), NewNumFiles - PreviousNumFiles, EnumerateTime );
+			GLog->Logf(TEXT("BuildDataGenerator: Found %d extra old files in %.1f seconds"), NewNumFiles - PreviousNumFiles, EnumerateTime);
 		}
 		else
 		{
-			GLog->Logf( TEXT( "BuildDataGenerator: Cloud directory does not exist: %s" ), *CloudFileDir );
+			GLog->Logf(TEXT("BuildDataGenerator: Cloud directory does not exist: %s"), *CloudFileDir);
 		}
 
 		DebugWriter->WriteObjectEnd();
 		DebugWriter->Close();
 
-		FArchive* FileOut = IFileManager::Get().CreateFileWriter( *( CloudFileDir + TEXT( "DebugFileList.txt" ) ) );
-		if( FileOut != NULL )
+		FArchive* FileOut = IFileManager::Get().CreateFileWriter(*(CloudFileDir + TEXT("DebugFileList.txt")));
+		if (FileOut != NULL)
 		{
 			FileOut->Serialize(TCHAR_TO_ANSI(*JSONOutput), JSONOutput.Len());
 			FileOut->Close();
@@ -1589,55 +1588,55 @@ bool FBuildDataGenerator::FindExistingFileData( const FString& InSourceFile, con
 	}
 
 	// Do we have a file matching this data?
-	if( ExistingFileInventory.Num() > 0 )
+	if (ExistingFileInventory.Num() > 0)
 	{
-		TArray< FString >* FileList = ExistingFileInventory.Find( InFileHash );
-		if( FileList != NULL )
+		TArray< FString >* FileList = ExistingFileInventory.Find(InFileHash);
+		if (FileList != NULL)
 		{
 			// We need to load each chunk in this list and compare data
-			for( auto FileIt = FileList->CreateConstIterator(); FileIt && !bFoundMatchingFile ; ++FileIt)
+			for (auto FileIt = FileList->CreateConstIterator(); FileIt && !bFoundMatchingFile; ++FileIt)
 			{
 				FString CloudFilename = *FileIt;
 				// Check the file date
-				FDateTime ModifiedDate = IFileManager::Get().GetTimeStamp( *CloudFilename );
+				FDateTime ModifiedDate = IFileManager::Get().GetTimeStamp(*CloudFilename);
 				if (ModifiedDate < DataThresholdTime)
 				{
 					// We don't want to reuse this file's Guid, as it's older than any existing files we want to consider
 					continue;
 				}
 				// Compare the files
-				FArchive* SourceFile = IFileManager::Get().CreateFileReader( *InSourceFile );
-				FArchive* FoundFile = IFileManager::Get().CreateFileReader( *CloudFilename );
-				if( SourceFile != NULL && FoundFile != NULL)
+				FArchive* SourceFile = IFileManager::Get().CreateFileReader(*InSourceFile);
+				FArchive* FoundFile = IFileManager::Get().CreateFileReader(*CloudFilename);
+				if (SourceFile != NULL && FoundFile != NULL)
 				{
 					FChunkHeader FoundHeader;
 					*FoundFile << FoundHeader;
 					const int64 SourceFileSize = SourceFile->TotalSize();
 					const int64 FoundFileSize = FoundFile->TotalSize();
-					if( SourceFileSize == FoundHeader.DataSize )
+					if (SourceFileSize == FoundHeader.DataSize)
 					{
 						// Currently only support stored raw!
-						check( FoundHeader.StoredAs == FChunkHeader::STORED_RAW );
+						check(FoundHeader.StoredAs == FChunkHeader::STORED_RAW);
 						// Makes no sense here for the sizes to not add up
-						check( FoundFileSize == ( FoundHeader.DataSize + FoundHeader.HeaderSize ) );
+						check(FoundFileSize == (FoundHeader.DataSize + FoundHeader.HeaderSize));
 						// Move FoundFile to start of file data
-						FoundFile->Seek( FoundHeader.HeaderSize );
+						FoundFile->Seek(FoundHeader.HeaderSize);
 						// Compare
 						bool bSameData = true;
-						uint8* TempSourceBuffer = new uint8[ FileBufferSize ];
-						uint8* TempFoundBuffer = new uint8[ FileBufferSize ];
-						while ( bSameData && ( ( SourceFile->AtEnd() || FoundFile->AtEnd() ) == false ) )
+						uint8* TempSourceBuffer = new uint8[FileBufferSize];
+						uint8* TempFoundBuffer = new uint8[FileBufferSize];
+						while (bSameData && ((SourceFile->AtEnd() || FoundFile->AtEnd()) == false))
 						{
 							const int64 SizeLeft = SourceFileSize - SourceFile->Tell();
-							const uint32 ReadLen = FMath::Min< int64 >( FileBufferSize, SizeLeft );
-							SourceFile->Serialize( TempSourceBuffer, ReadLen );
-							FoundFile->Serialize( TempFoundBuffer, ReadLen );
-							bSameData = FMemory::Memcmp( TempSourceBuffer, TempFoundBuffer, ReadLen ) == 0;
+							const uint32 ReadLen = FMath::Min< int64 >(FileBufferSize, SizeLeft);
+							SourceFile->Serialize(TempSourceBuffer, ReadLen);
+							FoundFile->Serialize(TempFoundBuffer, ReadLen);
+							bSameData = FMemory::Memcmp(TempSourceBuffer, TempFoundBuffer, ReadLen) == 0;
 						}
 						delete[] TempSourceBuffer;
 						delete[] TempFoundBuffer;
 						// Did we match?
-						if( bSameData && SourceFile->AtEnd() && FoundFile->AtEnd() )
+						if (bSameData && SourceFile->AtEnd() && FoundFile->AtEnd())
 						{
 							// Yes we did!
 							bFoundMatchingFile = true;
@@ -1645,12 +1644,12 @@ bool FBuildDataGenerator::FindExistingFileData( const FString& InSourceFile, con
 						}
 					}
 				}
-				if( SourceFile != NULL )
+				if (SourceFile != NULL)
 				{
 					SourceFile->Close();
 					delete SourceFile;
 				}
-				if( FoundFile != NULL )
+				if (FoundFile != NULL)
 				{
 					FoundFile->Close();
 					delete FoundFile;
@@ -1662,7 +1661,7 @@ bool FBuildDataGenerator::FindExistingFileData( const FString& InSourceFile, con
 	return bFoundMatchingFile;
 }
 
-bool FBuildDataGenerator::SaveOutFileData( const FString& SourceFile, const FSHAHash& FileHash, const FGuid& FileGuid )
+bool FBuildDataGenerator::SaveOutFileData(const FString& SourceFile, const FSHAHashData& FileHash, const FGuid& FileGuid)
 {
 	bool bAlreadySaved = false;
 	bool bSuccess = false;
@@ -1754,13 +1753,13 @@ bool FBuildDataGenerator::SaveOutFileData( const FString& SourceFile, const FSHA
 	return bSuccess;
 }
 
-bool FBuildDataGenerator::CompareDataToChunk( const FString& ChunkFilePath, uint8* ChunkData, FGuid& ChunkGuid, bool& OutSourceChunkIsValid )
+bool FBuildDataGenerator::CompareDataToChunk(const FString& ChunkFilePath, uint8* ChunkData, FGuid& ChunkGuid, bool& OutSourceChunkIsValid)
 {
 	bool bMatching = false;
 
 	// Read the file
-	TSharedRef< FBuildGenerationChunkCache::FChunkReader > ChunkReader = FBuildGenerationChunkCache::Get().GetChunkReader( ChunkFilePath );
-	if( ChunkReader->IsValidChunk() )
+	TSharedRef< FBuildGenerationChunkCache::FChunkReader > ChunkReader = FBuildGenerationChunkCache::Get().GetChunkReader(ChunkFilePath);
+	if (ChunkReader->IsValidChunk())
 	{
 		ChunkGuid = ChunkReader->GetChunkGuid();
 		// Default true
@@ -1769,11 +1768,11 @@ bool FBuildDataGenerator::CompareDataToChunk( const FString& ChunkFilePath, uint
 		const uint32 CompareSize = 64;
 		uint8* ReadBuffer;
 		uint32 NumCompared = 0;
-		while( bMatching && ChunkReader->BytesLeft() > 0 && NumCompared < FBuildPatchData::ChunkDataSize )
+		while (bMatching && ChunkReader->BytesLeft() > 0 && NumCompared < FBuildPatchData::ChunkDataSize)
 		{
-			const uint32 ReadLen = FMath::Min< uint32 >( CompareSize, ChunkReader->BytesLeft() );
-			ChunkReader->ReadNextBytes( &ReadBuffer, ReadLen );
-			bMatching = FMemory::Memcmp( &ChunkData[ NumCompared ], ReadBuffer, ReadLen ) == 0;
+			const uint32 ReadLen = FMath::Min< uint32 >(CompareSize, ChunkReader->BytesLeft());
+			ChunkReader->ReadNextBytes(&ReadBuffer, ReadLen);
+			bMatching = FMemory::Memcmp(&ChunkData[NumCompared], ReadBuffer, ReadLen) == 0;
 			NumCompared += ReadLen;
 		}
 	}
@@ -1784,37 +1783,37 @@ bool FBuildDataGenerator::CompareDataToChunk( const FString& ChunkFilePath, uint
 	return bMatching;
 }
 
-void FBuildDataGenerator::StripIgnoredFiles( TArray< FString >& AllFiles, const FString& DepotDirectory, const FString& IgnoreListFile )
+void FBuildDataGenerator::StripIgnoredFiles(TArray< FString >& AllFiles, const FString& DepotDirectory, const FString& IgnoreListFile)
 {
 	const int32 OriginalNumFiles = AllFiles.Num();
-	FString IgnoreFileList = TEXT( "" );
-	FFileHelper::LoadFileToString( IgnoreFileList, *IgnoreListFile );
+	FString IgnoreFileList = TEXT("");
+	FFileHelper::LoadFileToString(IgnoreFileList, *IgnoreListFile);
 	TArray< FString > IgnoreFiles;
-	IgnoreFileList.ParseIntoArray( &IgnoreFiles, TEXT( "\r\n" ), true );
+	IgnoreFileList.ParseIntoArray(&IgnoreFiles, TEXT("\r\n"), true);
 	struct FRemoveMatchingStrings
-	{ 
+	{
 		const FString* MatchingString;
-		FRemoveMatchingStrings( const FString* InMatch )
+		FRemoveMatchingStrings(const FString* InMatch)
 			: MatchingString(InMatch) {}
 
 		bool operator()(const FString& RemovalCandidate) const
 		{
 			FString PathA = RemovalCandidate;
 			FString PathB = *MatchingString;
-			FPaths::NormalizeFilename( PathA );
-			FPaths::NormalizeFilename( PathB );
+			FPaths::NormalizeFilename(PathA);
+			FPaths::NormalizeFilename(PathB);
 			return PathA == PathB;
 		}
 	};
 
-	for( int32 IgnoreIdx = 0; IgnoreIdx < IgnoreFiles.Num(); ++IgnoreIdx )
+	for (int32 IgnoreIdx = 0; IgnoreIdx < IgnoreFiles.Num(); ++IgnoreIdx)
 	{
 		const FString& IgnoreFile = IgnoreFiles[IgnoreIdx];
 		const FString FullIgnorePath = DepotDirectory / IgnoreFile;
-		AllFiles.RemoveAll( FRemoveMatchingStrings( &FullIgnorePath ) );
+		AllFiles.RemoveAll(FRemoveMatchingStrings(&FullIgnorePath));
 	}
 	const int32 NewNumFiles = AllFiles.Num();
-	GLog->Logf( TEXT( "Stripped %d ignorable file(s)" ), ( OriginalNumFiles - NewNumFiles ) );
+	GLog->Logf(TEXT("Stripped %d ignorable file(s)"), (OriginalNumFiles - NewNumFiles));
 }
 
 #endif //WITH_BUILDPATCHGENERATION
diff --git a/Engine/Source/Runtime/Online/BuildPatchServices/Private/BuildPatchGeneration.h b/Engine/Source/Runtime/Online/BuildPatchServices/Private/BuildPatchGeneration.h
--- a/Engine/Source/Runtime/Online/BuildPatchServices/Private/BuildPatchGeneration.h
+++ b/Engine/Source/Runtime/Online/BuildPatchServices/Private/BuildPatchGeneration.h
@@ -322,7 +322,7 @@ private:
 	uint32 CurrentChunkBufferPos;
 
 	// Holds the current file being processed. Receives the chunk parts that we pass over.
-	TSharedPtr< class FBuildPatchFileManifest > CurrentFile;
+	FFileManifestData* CurrentFile;
 
 	// Holds an SHA hash calculator
 	FSHA1 FileHash;
@@ -417,7 +417,7 @@ private:
 	FString BuildRoot;
 
 	// Holds the current file being processed. Receives the chunk parts that we pass over.
-	TSharedPtr< class FBuildPatchFileManifest > CurrentFile;
+	FFileManifestData* CurrentFile;
 
 	// Holds an SHA hash calculator
 	FSHA1 FileHash;
@@ -628,17 +628,17 @@ public:
 	 * @param OutFileGuid			OUT		Will be set to the existing file guid if found.
 	 * @return		Whether this file is new or existing.
 	 */
-	static bool FindExistingFileData( const FString& InSourceFile, const FSHAHash& InFileHash, const FDateTime& DataThresholdTime, FGuid& OutFileGuid );
+	static bool FindExistingFileData(const FString& InSourceFile, const FSHAHashData& InFileHash, const FDateTime& DataThresholdTime, FGuid& OutFileGuid);
 
 	/**
 	 * Saves out the file data for use with file based patch manifests
 	 * NOTE: This function is blocking and will not return until finished. Don't run on main thread.
-	 * @param SourceFile	IN		The filename of the source file on disk.
-	 * @param FileHash		IN		The hash for the file.
-	 * @param FileGuid		OUT		The GUID for the file
+	 * @param SourceFile		The filename of the source file on disk.
+	 * @param FileHash			The hash for the file.
+	 * @param FileGuid			The GUID for the file.
 	 * @return		Whether this file is new or existing.
 	 */
-	static bool SaveOutFileData( const FString& SourceFile, const FSHAHash& FileHash, const FGuid& FileGuid );
+	static bool SaveOutFileData(const FString& SourceFile, const FSHAHashData& FileHash, const FGuid& FileGuid);
 
 	/**
 	 * Strips out files from the provided array that appear in the ignore list file
@@ -671,7 +671,7 @@ private:
 	static bool ExistingChunksEnumerated;
 
 	// For File Manifest generation, the existing files
-	static TMap< FSHAHash, TArray< FString > > ExistingFileInventory;
+	static TMap< FSHAHashData, TArray< FString > > ExistingFileInventory;
 
 	// For File Manifest generation, have we enumerated the existing files yet
 	static bool ExistingFilesEnumerated;
diff --git a/Engine/Source/Runtime/Online/BuildPatchServices/Private/BuildPatchInstaller.cpp b/Engine/Source/Runtime/Online/BuildPatchServices/Private/BuildPatchInstaller.cpp
--- a/Engine/Source/Runtime/Online/BuildPatchServices/Private/BuildPatchInstaller.cpp
+++ b/Engine/Source/Runtime/Online/BuildPatchServices/Private/BuildPatchInstaller.cpp
@@ -611,14 +611,14 @@ bool FBuildPatchInstaller::BackupFileIfNecessary(const FString& Filename, bool b
 	const bool bCheckFileChanges = !bDiscoveredByVerification;
 	if (bCheckFileChanges)
 	{
-		TSharedPtr< FBuildPatchFileManifest > OldFileManifest = CurrentBuildManifest.IsValid() ? CurrentBuildManifest->GetFileManifest(Filename) : NULL;
-		TSharedPtr< FBuildPatchFileManifest > NewFileManifest = NewBuildManifest->GetFileManifest(Filename);
+		const FFileManifestData* OldFileManifest = CurrentBuildManifest.IsValid() ? CurrentBuildManifest->GetFileManifest(Filename) : nullptr;
+		const FFileManifestData* NewFileManifest = NewBuildManifest->GetFileManifest(Filename);
 		const int64 InstalledFilesize = IFileManager::Get().FileSize(*InstalledFilename);
-		const int64 OriginalFileSize = OldFileManifest.IsValid() ? OldFileManifest->GetFileSize() : INDEX_NONE;
-		const int64 NewFileSize = NewFileManifest.IsValid() ? NewFileManifest->GetFileSize() : INDEX_NONE;
-		const FSHAHash HashZero;
-		const FSHAHash& HashOld = OldFileManifest.IsValid() ? OldFileManifest->FileHash : HashZero;
-		const FSHAHash& HashNew = NewFileManifest.IsValid() ? NewFileManifest->FileHash : HashZero;
+		const int64 OriginalFileSize = OldFileManifest ? OldFileManifest->GetFileSize() : INDEX_NONE;
+		const int64 NewFileSize = NewFileManifest ? NewFileManifest->GetFileSize() : INDEX_NONE;
+		const FSHAHashData HashZero;
+		const FSHAHashData& HashOld = OldFileManifest ? OldFileManifest->FileHash : HashZero;
+		const FSHAHashData& HashNew = NewFileManifest ? NewFileManifest->FileHash : HashZero;
 		const bool bFileSizeDiffers = OriginalFileSize != InstalledFilesize && NewFileSize != InstalledFilesize;
 		bUserEditedFile = bFileSizeDiffers || FBuildPatchUtils::VerifyFile(InstalledFilename, HashOld, HashNew) == 0;
 	}
diff --git a/Engine/Source/Runtime/Online/BuildPatchServices/Private/BuildPatchManifest.cpp b/Engine/Source/Runtime/Online/BuildPatchServices/Private/BuildPatchManifest.cpp
--- a/Engine/Source/Runtime/Online/BuildPatchServices/Private/BuildPatchManifest.cpp
+++ b/Engine/Source/Runtime/Online/BuildPatchServices/Private/BuildPatchManifest.cpp
@@ -8,6 +8,13 @@
 
 #define LOCTEXT_NAMESPACE "BuildPatchManifest"
 
+// The manifest header magic codeword, for quick checking that the opened file is probably a manifest file.
+#define MANIFEST_HEADER_MAGIC		0x44BEC00C
+
+// The maximum number of FNames that we expect a manifest to generate. This is not a technical limitation, just a sanity check
+// and can be increased if more properties are added to our manifest UObject. FNames are only used by the UOject serialization system.
+#define MANIFEST_MAX_NAMES			50
+
 /**
  * Helper functions that convert generic types to and from string blobs for use with JSON parsing.
  * It's kind of horrible but guarantees no loss of data as the JSON reader/writer only supports float functionality 
@@ -26,6 +33,23 @@ FString ToStringBlob( const DataType& DataVal )
 	return FString::FromBlob( static_cast<const uint8*>( AsBuffer ), sizeof( DataType ) );
 }
 
+/**
+ * Helper functions to decide whether the passed in data is a JSON string we expect to deserialise a manifest from
+ */
+bool BufferIsJsonManifest(const TArray<uint8>& DataInput)
+{
+	// The best we can do is look for the mandatory first character open curly brace,
+	// it will be within the first 4 characters (may have BOM)
+	for (int32 idx = 0; idx < 4 && idx < DataInput.Num(); ++idx)
+	{
+		if (DataInput[idx] == TEXT('{'))
+		{
+			return true;
+		}
+	}
+	return false;
+}
+
 /* EBuildPatchAppManifestVersion implementation
 *****************************************************************************/
 const EBuildPatchAppManifestVersion::Type EBuildPatchAppManifestVersion::GetLatestVersion()
@@ -33,6 +57,11 @@ const EBuildPatchAppManifestVersion::Type EBuildPatchAppManifestVersion::GetLate
 	return static_cast<EBuildPatchAppManifestVersion::Type>(LatestPlusOne - 1);
 }
 
+const EBuildPatchAppManifestVersion::Type EBuildPatchAppManifestVersion::GetLatestJsonVersion()
+{
+	return static_cast<EBuildPatchAppManifestVersion::Type>(StoredAsCompressedUClass - 1);
+}
+
 const FString& EBuildPatchAppManifestVersion::GetChunkSubdir(const EBuildPatchAppManifestVersion::Type ManifestVersion)
 {
 	static const FString ChunksV1 = TEXT("Chunks");
@@ -51,10 +80,348 @@ const FString& EBuildPatchAppManifestVersion::GetFileSubdir(const EBuildPatchApp
 		: FilesV2;
 }
 
+/* FManifestFileHeader - The header for a compressed/encoded manifest file
+*****************************************************************************/
+FManifestFileHeader::FManifestFileHeader()
+	: Magic(MANIFEST_HEADER_MAGIC)
+	, HeaderSize(0)
+	, DataSize(0)
+	, CompressedSize(0)
+	, SHAHash()
+	, StoredAs(0)
+{
+}
+
+bool FManifestFileHeader::CheckMagic() const
+{
+	return Magic == MANIFEST_HEADER_MAGIC;
+}
+
+FArchive& operator<<(FArchive& Ar, FManifestFileHeader& Header)
+{
+	Ar << Header.Magic;
+	Ar << Header.HeaderSize;
+	Ar << Header.DataSize;
+	Ar << Header.CompressedSize;
+	Ar.Serialize(Header.SHAHash.Hash, FSHA1::DigestSize);
+	Ar << Header.StoredAs;
+	return Ar;
+}
+
+/* FCustomFieldData implementation
+*****************************************************************************/
+FCustomFieldData::FCustomFieldData()
+	: Key(TEXT(""))
+	, Value(TEXT(""))
+{
+}
+
+FCustomFieldData::FCustomFieldData(const FString& Key, const FString& Value)
+	: Key(Key)
+	, Value(Value)
+{
+}
+
+/* FChunkInfoData implementation
+*****************************************************************************/
+FChunkInfoData::FChunkInfoData()
+	: Guid()
+	, Hash(0)
+	, FileSize(0)
+	, GroupNumber(0)
+{
+}
+
+/* FChunkPartData implementation
+*****************************************************************************/
+FChunkPartData::FChunkPartData()
+	: Guid()
+	, Offset(0)
+	, Size(0)
+{
+}
+
+/* FFileManifestData implementation
+*****************************************************************************/
+FFileManifestData::FFileManifestData()
+	: Filename(TEXT(""))
+	, FileHash()
+	, FileChunkParts()
+	, bIsUnixExecutable(false)
+	, SymlinkTarget(TEXT(""))
+	, FileSize(INDEX_NONE)
+{
+}
+
+const int64& FFileManifestData::GetFileSize() const
+{
+	return FileSize;
+}
+
+FORCEINLINE bool FFileManifestData::operator<(const FFileManifestData& Other) const
+{
+	return Filename < Other.Filename;
+}
+
+void FFileManifestData::Init()
+{
+	FileSize = 0;
+	for (auto FileChunkPartIt = FileChunkParts.CreateConstIterator(); FileChunkPartIt; ++FileChunkPartIt)
+	{
+		FileSize += (*FileChunkPartIt).Size;
+	}
+}
+
+/**
+ * Archive for writing a manifest into memory
+ */
+class FManifestWriter : public FArchive
+{
+public:
+	FManifestWriter()
+		: FArchive()
+		, Offset(0)
+	{
+		ArIsPersistent = false;
+		ArIsSaving = true;
+	}
+
+	virtual void Seek(int64 InPos) override
+	{
+		Offset = InPos;
+	}
+
+	virtual int64 Tell() override
+	{
+		return Offset;
+	}
+
+	virtual FString GetArchiveName() const override
+	{
+		return TEXT("FManifestWriter");
+	}
+
+	virtual FArchive& operator<<(FName& N) override
+	{
+		const FName NameZero(N, 0);
+		if (!FNameIndexLookup.Contains(NameZero))
+		{
+			const int32 ArNameIndex = FNameIndexLookup.Num();
+			FNameIndexLookup.Add(NameZero, ArNameIndex);
+		}
+		*this << FNameIndexLookup[NameZero];
+		return *this;
+	}
+
+	virtual void Serialize(void* Data, int64 Num) override
+	{
+		if (Num && !ArIsError)
+		{
+			const int64 NumBytesToAdd = Offset + Num - Bytes.Num();
+			if (NumBytesToAdd > 0)
+			{
+				const int64 NewArrayCount = Bytes.Num() + NumBytesToAdd;
+				if (NewArrayCount >= MAX_int32)
+				{
+					ArIsError = true;
+					return;
+				}
+
+				Bytes.AddUninitialized((int32)NumBytesToAdd);
+			}
+
+			check((Offset + Num) <= Bytes.Num());
+
+			if (Num)
+			{
+				FMemory::Memcpy(&Bytes[Offset], Data, Num);
+				Offset += Num;
+			}
+		}
+	}
+
+	virtual int64 TotalSize() override
+	{
+		return Bytes.Num();
+	}
+
+	void Finalize()
+	{
+		TArray<uint8> FinalData;
+		FMemoryWriter NameTableWriter(FinalData);
+
+		int32 NumNames = FNameIndexLookup.Num();
+		check(NumNames <= MANIFEST_MAX_NAMES);
+		NameTableWriter << NumNames;
+		for (auto& MapEntry : FNameIndexLookup)
+		{
+			FName& Name = MapEntry.Key;
+			int32& Index = MapEntry.Value;
+			NameTableWriter << Name;
+			NameTableWriter << Index;
+		}
+
+		FinalData.Append(Bytes);
+
+		Bytes = MoveTemp(FinalData);
+	}
+
+	TArray<uint8>& GetBytes()
+	{
+		return Bytes;
+	}
+
+private:
+	int64 Offset;
+	TArray<uint8> Bytes;
+	TMap<FName, int32> FNameIndexLookup;
+};
+
+/**
+ * Archive for reading a manifest from data in memory
+ */
+class FManifestReader : public FArchive
+{
+public:
+	FManifestReader(const TArray<uint8>& InBytes)
+		: FArchive()
+		, Offset(0)
+		, Bytes(InBytes)
+	{
+		ArIsPersistent = false;
+		ArIsLoading = true;
+
+		// Must load table immediately
+		FMemoryReader NameTableReader(InBytes);
+		int32 NumNames;
+		NameTableReader << NumNames;
+
+		// Check not insane, we know to expect a small number for a manifest
+		if (NumNames < MANIFEST_MAX_NAMES)
+		{
+			FNameLookup.Reserve(NumNames);
+			for (; NumNames > 0; --NumNames)
+			{
+				FName Name;
+				int32 Index;
+				NameTableReader << Name;
+				NameTableReader << Index;
+				FNameLookup.Add(Index, Name);
+			}
+		}
+		else
+		{
+			ArIsError = true;
+		}
+
+		Offset = NameTableReader.Tell();
+	}
+
+	virtual FString GetArchiveName() const override
+	{
+		return TEXT("FManifestReader");
+	}
+
+	virtual void Seek(int64 InPos) override
+	{
+		check(InPos <= Bytes.Num());
+		Offset = InPos;
+	}
+
+	virtual int64 Tell() override
+	{
+		return Offset;
+	}
+
+	virtual FArchive& operator<<(FName& N) override
+	{
+		// Read index and lookup
+		int32 ArNameIndex;
+		*this << ArNameIndex;
+		if (FNameLookup.Contains(ArNameIndex))
+		{
+			N = FNameLookup[ArNameIndex];
+		}
+		else
+		{
+			ArIsError = true;
+		}
+		return *this;
+	}
+
+	virtual void Serialize(void* Data, int64 Num) override
+	{
+		if (Num && !ArIsError)
+		{
+			// Only serialize if we have the requested amount of data
+			if (Offset + Num <= Bytes.Num())
+			{
+				FMemory::Memcpy(Data, &Bytes[Offset], Num);
+				Offset += Num;
+			}
+			else
+			{
+				ArIsError = true;
+			}
+		}
+	}
+
+	virtual int64 TotalSize() override
+	{
+		return Bytes.Num();
+	}
+
+private:
+	int64 Offset;
+	const TArray<uint8>& Bytes;
+	TMap<int32, FName> FNameLookup;
+};
+
+/* UBuildPatchManifest implementation
+*****************************************************************************/
+UBuildPatchManifest::UBuildPatchManifest(const FPostConstructInitializeProperties& PCIP)
+	: Super(PCIP)
+	, ManifestFileVersion(EBuildPatchAppManifestVersion::GetLatestVersion())
+	, bIsFileData(false)
+	, AppID(0)
+	, AppName(TEXT(""))
+	, BuildVersion(TEXT(""))
+	, LaunchExe(TEXT(""))
+	, LaunchCommand(TEXT(""))
+	, PrereqName(TEXT(""))
+	, PrereqPath(TEXT(""))
+	, PrereqArgs(TEXT(""))
+	, FileManifestList()
+	, ChunkList()
+	, CustomFields()
+{
+}
+
+void UBuildPatchManifest::Clear()
+{
+	ManifestFileVersion = EBuildPatchAppManifestVersion::GetLatestVersion();
+	bIsFileData = false;
+	AppID = INDEX_NONE;
+	AppName.Empty();
+	BuildVersion.Empty();
+	LaunchExe.Empty();
+	LaunchCommand.Empty();
+	PrereqName.Empty();
+	PrereqPath.Empty();
+	PrereqArgs.Empty();
+	FileManifestList.Empty();
+	ChunkList.Empty();
+	CustomFields.Empty();
+}
+
+UBuildPatchManifest::~UBuildPatchManifest()
+{
+}
+
 /* FBuildPatchCustomField implementation
 *****************************************************************************/
-FBuildPatchCustomField::FBuildPatchCustomField( const FString& Value )
-	: CustomValue( Value )
+FBuildPatchCustomField::FBuildPatchCustomField(const FString& Value)
+	: CustomValue(Value)
 {
 }
 
@@ -85,149 +452,273 @@ const int64 FBuildPatchCustomField::AsInteger() const
 	return 0;
 }
 
-/* FBuildPatchFileManifest implementation
+/* FBuildPatchAppManifest implementation
 *****************************************************************************/
 
-FBuildPatchFileManifest::FBuildPatchFileManifest()
-	: Filename( TEXT( "" ) )
-#if PLATFORM_MAC
-	, bIsUnixExecutable( false )
-	, SymlinkTarget( TEXT( "" ) )
-#endif
-	, FileSize( INDEX_NONE )
+FBuildPatchAppManifest::FBuildPatchAppManifest()
+	: TotalBuildSize(INDEX_NONE)
+	, TotalDownloadSize(INDEX_NONE)
+{
+	Data = NewObject<UBuildPatchManifest>();
+}
+
+FBuildPatchAppManifest::FBuildPatchAppManifest(const uint32& InAppID, const FString& AppName)
+	: TotalBuildSize(INDEX_NONE)
+	, TotalDownloadSize(INDEX_NONE)
+{
+	Data = NewObject<UBuildPatchManifest>();
+	Data->AppID = InAppID;
+	Data->AppName = AppName;
+}
+
+FBuildPatchAppManifest::~FBuildPatchAppManifest()
 {
+//	Data->MarkPendingKill(); ??? Mark for destory and run GC?
 }
 
-const int64& FBuildPatchFileManifest::GetFileSize()
+bool FBuildPatchAppManifest::SaveToFile(const FString& Filename, bool bUseBinary)
 {
-	if( FileSize == INDEX_NONE )
+	bool bSuccess = false;
+	FArchive* FileOut = IFileManager::Get().CreateFileWriter(*Filename);
+	if (FileOut)
 	{
-		FileSize = 0;
-		for( auto FileChunkPartIt = FileChunkParts.CreateConstIterator(); FileChunkPartIt; ++FileChunkPartIt)
+		if (bUseBinary)
 		{
-			FileSize += (*FileChunkPartIt).Size;
+			FManifestWriter ManifestData;
+			Serialize(ManifestData);
+			ManifestData.Finalize();
+			if (!ManifestData.IsError())
+			{
+				int32 DataSize = ManifestData.TotalSize();
+				TArray<uint8> TempCompressed;
+				TempCompressed.AddUninitialized(DataSize);
+				int32 CompressedSize = DataSize;
+				bool bDataIsCompressed = FCompression::CompressMemory(
+					static_cast<ECompressionFlags>(COMPRESS_ZLIB | COMPRESS_BiasMemory),
+					TempCompressed.GetData(),
+					CompressedSize,
+					ManifestData.GetBytes().GetData(),
+					DataSize);
+				TempCompressed.SetNum(CompressedSize);
+
+				TArray<uint8>& FileData = bDataIsCompressed ? TempCompressed : ManifestData.GetBytes();
+
+				FManifestFileHeader Header;
+				*FileOut << Header;
+				Header.HeaderSize = FileOut->Tell();
+				Header.StoredAs = bDataIsCompressed ? EManifestFileHeader::STORED_COMPRESSED : EManifestFileHeader::STORED_RAW;
+				Header.DataSize = DataSize;
+				Header.CompressedSize = bDataIsCompressed ? CompressedSize : 0;
+				FSHA1::HashBuffer(FileData.GetData(), FileData.Num(), Header.SHAHash.Hash);
+
+				FileOut->Seek(0);
+				*FileOut << Header;
+
+				FileOut->Serialize(FileData.GetData(), FileData.Num());
+				bSuccess = !FileOut->IsError();
+			}
 		}
+		else
+		{
+			FString JSONOutput;
+			SerializeToJSON(JSONOutput);
+			FTCHARToUTF8 JsonUTF8(*JSONOutput);
+			FileOut->Serialize((UTF8CHAR*)JsonUTF8.Get(), JsonUTF8.Length() * sizeof(UTF8CHAR));
+		}
+		FileOut->Close();
+		delete FileOut;
+		FileOut = nullptr;
 	}
-	return FileSize;
-}
 
-/* FBuildPatchAppManifest implementation
-*****************************************************************************/
+	return bSuccess;
+}
 
-FBuildPatchAppManifest::FBuildPatchAppManifest()
-	: ManifestFileVersion( EBuildPatchAppManifestVersion::GetLatestVersion() )
+bool FBuildPatchAppManifest::LoadFromFile(FString Filename)
 {
+	TArray<uint8> FileData;
+	if (FFileHelper::LoadFileToArray(FileData, *Filename))
+	{
+		return DeserializeFromData(FileData);
+	}
+	return false;
 }
 
-FBuildPatchAppManifest::FBuildPatchAppManifest( const uint32& InAppID, const FString& AppName )
-	: ManifestFileVersion( EBuildPatchAppManifestVersion::GetLatestVersion() )
-	, AppID( InAppID )
-	, AppNameString( AppName )
-	, TotalBuildSize( INDEX_NONE )
+bool FBuildPatchAppManifest::DeserializeFromData(const TArray<uint8>& DataInput)
 {
+	if (DataInput.Num())
+	{
+		if (BufferIsJsonManifest(DataInput))
+		{
+			FString JsonManifest;
+			FFileHelper::BufferToString(JsonManifest, DataInput.GetData(), DataInput.Num());
+			return DeserializeFromJSON(JsonManifest);
+		}
+		else
+		{
+			FMemoryReader ManifestFile(DataInput);
+			FManifestFileHeader Header;
+			ManifestFile << Header;
+			if (Header.CheckMagic())
+			{
+				FSHAHashData DataHash;
+				FSHA1::HashBuffer(&DataInput[Header.HeaderSize], DataInput.Num() - Header.HeaderSize, DataHash.Hash);
+				if (DataHash == Header.SHAHash)
+				{
+					TArray<uint8> UncompressedData;
+					if (Header.StoredAs == FChunkHeader::STORED_COMPRESSED && (Header.CompressedSize + Header.HeaderSize) == DataInput.Num())
+					{
+						UncompressedData.AddUninitialized(Header.DataSize);
+						if (!FCompression::UncompressMemory(
+							static_cast<ECompressionFlags>(COMPRESS_ZLIB | COMPRESS_BiasMemory),
+							UncompressedData.GetData(),
+							Header.DataSize,
+							&DataInput[Header.HeaderSize],
+							DataInput.Num() - Header.HeaderSize))
+						{
+							return false;
+						}
+					}
+					else if ((Header.DataSize + Header.HeaderSize) == DataInput.Num())
+					{
+						UncompressedData.Append(&DataInput[Header.HeaderSize], Header.DataSize);
+					}
+					else
+					{
+						return false;
+					}
+					FManifestReader ManifestData(UncompressedData);
+					return Serialize(ManifestData);
+				}
+				else
+				{
+					return false;
+				}
+			}
+			else
+			{
+				return false;
+			}
+		}
+	}
+	return false;
 }
 
-FBuildPatchAppManifest::~FBuildPatchAppManifest()
+bool FBuildPatchAppManifest::Serialize(FArchive& Ar)
 {
-	// Make sure the hash data is cleaned up
-	DestroyData();
+	if (Ar.IsLoading())
+	{
+		DestroyData();
+	}
+
+	Data->Serialize(Ar);
+
+	if (Ar.IsLoading())
+	{
+		InitLookups();
+	}
+
+	return !Ar.IsError();
 }
 
 void FBuildPatchAppManifest::DestroyData()
 {
 	// Clear Manifest data
-	AppNameString = "";
-	BuildVersionString = "";
-	LaunchExeString = "";
-	LaunchCommandString = "";
-	PrereqNameString = "";
-	PrereqPathString = "";
-	PrereqArgsString = "";
-	FileManifestList.Empty();
-	ChunkHashList.Empty();
+	Data->Clear();
+	FileNameLookup.Empty();
+	FileManifestLookup.Empty();
+	ChunkInfoLookup.Empty();
+	CustomFieldLookup.Empty();
+	TotalBuildSize = INDEX_NONE;
+	TotalDownloadSize = INDEX_NONE;
 }
 
-bool FBuildPatchAppManifest::SaveToFile(FString Filename)
+void FBuildPatchAppManifest::InitLookups()
 {
-	bool bSuccess = true;
-	FArchive* FileOut = IFileManager::Get().CreateFileWriter(*Filename);
-	bSuccess &= FileOut != NULL;
-	if (bSuccess)
+	// Make sure file list is sorted
+	Data->FileManifestList.Sort();
+
+	// Setup internals
+	TotalBuildSize = 0;
+	FileManifestLookup.Empty(Data->FileManifestList.Num());
+	FileNameLookup.Empty(Data->bIsFileData ? Data->FileManifestList.Num() : 0);
+	for (auto& File : Data->FileManifestList)
 	{
-		FString JSONOutput;
-		SerializeToJSON(JSONOutput);
-		FTCHARToUTF8 JsonUTF8(*JSONOutput);
-		FileOut->Serialize((UTF8CHAR*)JsonUTF8.Get(), JsonUTF8.Length() * sizeof(UTF8CHAR));
-		FileOut->Close();
-		delete FileOut;
+		for (auto& FilePart : File.FileChunkParts)
+		{
+			TotalBuildSize += FilePart.Size;
+		}
+		FileManifestLookup.Add(File.Filename, &File);
+		if (Data->bIsFileData)
+		{
+			// File data chunk parts should have been checked already
+			FileNameLookup.Add(File.FileChunkParts[0].Guid, &File.Filename);
+		}
 	}
-	return bSuccess;
-}
-
-bool FBuildPatchAppManifest::LoadFromFile( FString Filename )
-{
-	FString JSONString;
-	bool bSuccess = FFileHelper::LoadFileToString( JSONString, *Filename );
-	if( bSuccess )
+	TotalDownloadSize = 0;
+	ChunkInfoLookup.Empty(Data->ChunkList.Num());
+	for (auto& Chunk : Data->ChunkList)
 	{
-		bSuccess &= DeserializeFromJSON( JSONString );
+		ChunkInfoLookup.Add(Chunk.Guid, &Chunk);
+		TotalDownloadSize += Chunk.FileSize;
+	}
+	CustomFieldLookup.Empty(Data->CustomFields.Num());
+	for (auto& CustomField : Data->CustomFields)
+	{
+		CustomFieldLookup.Add(CustomField.Key, &CustomField);
 	}
-	return bSuccess;
 }
 
-void FBuildPatchAppManifest::SerializeToJSON( FString& JSONOutput )
+void FBuildPatchAppManifest::SerializeToJSON(FString& JSONOutput)
 {
 #if UE_BUILD_DEBUG // We'll use this to switch between human readable JSON
 	TSharedRef< TJsonWriter< TCHAR, TPrettyJsonPrintPolicy< TCHAR > > > Writer = TJsonWriterFactory< TCHAR, TPrettyJsonPrintPolicy< TCHAR > >::Create(&JSONOutput);
 #else
-	TSharedRef< TJsonWriter< TCHAR, TCondensedJsonPrintPolicy< TCHAR > > > Writer = TJsonWriterFactory< TCHAR, TCondensedJsonPrintPolicy< TCHAR > >::Create( &JSONOutput );
+	TSharedRef< TJsonWriter< TCHAR, TCondensedJsonPrintPolicy< TCHAR > > > Writer = TJsonWriterFactory< TCHAR, TCondensedJsonPrintPolicy< TCHAR > >::Create(&JSONOutput);
 #endif //ALLOW_DEBUG_FILES
 
 	Writer->WriteObjectStart();
 	{
 		// Write general data
-		Writer->WriteValue( TEXT("ManifestFileVersion"), ToStringBlob( static_cast< int32 >( ManifestFileVersion ) ) );
-		Writer->WriteValue( TEXT("bIsFileData"), bIsFileData );
-		Writer->WriteValue( TEXT("AppID"), ToStringBlob( AppID ) );
-		Writer->WriteValue( TEXT("AppNameString"), AppNameString );
-		Writer->WriteValue( TEXT("BuildVersionString"), BuildVersionString );
-		Writer->WriteValue( TEXT("LaunchExeString"), LaunchExeString );
-		Writer->WriteValue( TEXT("LaunchCommand"), LaunchCommandString );
-		Writer->WriteValue( TEXT("PrereqName"), PrereqNameString );
-		Writer->WriteValue( TEXT("PrereqPath"), PrereqPathString );
-		Writer->WriteValue( TEXT("PrereqArgs"), PrereqArgsString );
+		int32 ManifestVersion = (int32)Data->ManifestFileVersion;
+		int32 JsonVersion = (int32)EBuildPatchAppManifestVersion::GetLatestJsonVersion();
+		Writer->WriteValue(TEXT("ManifestFileVersion"), ToStringBlob(FMath::Min(ManifestVersion, JsonVersion)));
+		Writer->WriteValue(TEXT("bIsFileData"), Data->bIsFileData);
+		Writer->WriteValue(TEXT("AppID"), ToStringBlob(Data->AppID));
+		Writer->WriteValue(TEXT("AppNameString"), Data->AppName);
+		Writer->WriteValue(TEXT("BuildVersionString"), Data->BuildVersion);
+		Writer->WriteValue(TEXT("LaunchExeString"), Data->LaunchExe);
+		Writer->WriteValue(TEXT("LaunchCommand"), Data->LaunchCommand);
+		Writer->WriteValue(TEXT("PrereqName"), Data->PrereqName);
+		Writer->WriteValue(TEXT("PrereqPath"), Data->PrereqPath);
+		Writer->WriteValue(TEXT("PrereqArgs"), Data->PrereqArgs);
 		// Write file manifest data
-		Writer->WriteArrayStart( TEXT("FileManifestList") );
-		for( auto FileManifestIt = FileManifestList.CreateConstIterator(); FileManifestIt; ++FileManifestIt )
+		Writer->WriteArrayStart(TEXT("FileManifestList"));
+		for (const auto& FileManifest : Data->FileManifestList)
 		{
-			const TSharedPtr< FBuildPatchFileManifest > FileManifest = FileManifestIt.Value();
 			Writer->WriteObjectStart();
 			{
-				Writer->WriteValue( TEXT("Filename"), FileManifest->Filename );
-				Writer->WriteValue( TEXT("FileHash"), FString::FromBlob( FileManifestIt.Value()->FileHash.Hash, FSHA1::DigestSize ) );
-#if PLATFORM_MAC
-				if( FileManifest->bIsUnixExecutable )
+				Writer->WriteValue(TEXT("Filename"), FileManifest.Filename);
+				Writer->WriteValue(TEXT("FileHash"), FString::FromBlob(FileManifest.FileHash.Hash, FSHA1::DigestSize));
+				if (FileManifest.bIsUnixExecutable)
 				{
-					Writer->WriteValue( TEXT("bIsUnixExecutable"), FileManifest->bIsUnixExecutable );
+					Writer->WriteValue(TEXT("bIsUnixExecutable"), FileManifest.bIsUnixExecutable);
 				}
-				const bool bIsSymlink = !FileManifest->SymlinkTarget.IsEmpty();
-				if( bIsSymlink )
+				const bool bIsSymlink = !FileManifest.SymlinkTarget.IsEmpty();
+				if (bIsSymlink)
 				{
-					Writer->WriteValue( TEXT("SymlinkTarget"), FileManifest->SymlinkTarget );
+					Writer->WriteValue(TEXT("SymlinkTarget"), FileManifest.SymlinkTarget);
 				}
 				else
-#endif
 				{
-					Writer->WriteArrayStart( TEXT("FileChunkParts") );
+					Writer->WriteArrayStart(TEXT("FileChunkParts"));
 					{
-						for( auto FileChunkPartIt = FileManifestIt.Value()->FileChunkParts.CreateConstIterator(); FileChunkPartIt; ++FileChunkPartIt )
+						for (const auto& FileChunkPart : FileManifest.FileChunkParts)
 						{
-							const FChunkPart& FileChunkPart = *FileChunkPartIt;
 							Writer->WriteObjectStart();
 							{
-								Writer->WriteValue( TEXT("Guid"), FileChunkPart.Guid.ToString() );
-								Writer->WriteValue( TEXT("Offset"), ToStringBlob( FileChunkPart.Offset ) );
-								Writer->WriteValue( TEXT("Size"), ToStringBlob( FileChunkPart.Size ) );
+								Writer->WriteValue(TEXT("Guid"), FileChunkPart.Guid.ToString());
+								Writer->WriteValue(TEXT("Offset"), ToStringBlob(FileChunkPart.Offset));
+								Writer->WriteValue(TEXT("Size"), ToStringBlob(FileChunkPart.Size));
 							}
 							Writer->WriteObjectEnd();
 						}
@@ -239,37 +730,37 @@ void FBuildPatchAppManifest::SerializeToJSON( FString& JSONOutput )
 		}
 		Writer->WriteArrayEnd();
 		// Write chunk hash list
-		Writer->WriteObjectStart( TEXT("ChunkHashList") );
-		for( auto ChunkHashListIt = ChunkHashList.CreateConstIterator(); ChunkHashListIt; ++ChunkHashListIt )
+		Writer->WriteObjectStart(TEXT("ChunkHashList"));
+		for (const auto& ChunkInfo : Data->ChunkList)
 		{
-			const FGuid& ChunkGuid = ChunkHashListIt.Key();
-			const uint64& ChunkHash = ChunkHashListIt.Value();
-			Writer->WriteValue( ChunkGuid.ToString(), ToStringBlob( ChunkHash ) );
+			const FGuid& ChunkGuid = ChunkInfo.Guid;
+			const uint64& ChunkHash = ChunkInfo.Hash;
+			Writer->WriteValue(ChunkGuid.ToString(), ToStringBlob(ChunkHash));
 		}
 		Writer->WriteObjectEnd();
 		// Write data group list
-		Writer->WriteObjectStart( TEXT("DataGroupList") );
-		for( auto DataGroupListIt = DataGroupList.CreateConstIterator(); DataGroupListIt; ++DataGroupListIt )
+		Writer->WriteObjectStart(TEXT("DataGroupList"));
+		for (const auto& ChunkInfo : Data->ChunkList)
 		{
-			const FGuid& DataGuid = DataGroupListIt.Key();
-			const uint8& DataGroup = DataGroupListIt.Value();
-			Writer->WriteValue( DataGuid.ToString(), ToStringBlob( DataGroup ) );
+			const FGuid& DataGuid = ChunkInfo.Guid;
+			const uint8& DataGroup = ChunkInfo.GroupNumber;
+			Writer->WriteValue(DataGuid.ToString(), ToStringBlob(DataGroup));
 		}
 		Writer->WriteObjectEnd();
 		// Write chunk size list
 		Writer->WriteObjectStart(TEXT("ChunkFilesizeList"));
-		for (auto ChunkFilesizeListIt = ChunkFilesizeList.CreateConstIterator(); ChunkFilesizeListIt; ++ChunkFilesizeListIt)
+		for (const auto& ChunkInfo : Data->ChunkList)
 		{
-			const FGuid& ChunkGuid = ChunkFilesizeListIt.Key();
-			const int64& ChunkSize = ChunkFilesizeListIt.Value();
+			const FGuid& ChunkGuid = ChunkInfo.Guid;
+			const int64& ChunkSize = ChunkInfo.FileSize;
 			Writer->WriteValue(ChunkGuid.ToString(), ToStringBlob(ChunkSize));
 		}
 		Writer->WriteObjectEnd();
 		// Write custom fields
-		Writer->WriteObjectStart( TEXT("CustomFields") );
-		for( auto CustomFieldsIt = CustomFields.CreateConstIterator(); CustomFieldsIt; ++CustomFieldsIt )
+		Writer->WriteObjectStart(TEXT("CustomFields"));
+		for (const auto& CustomField : Data->CustomFields)
 		{
-			Writer->WriteValue( CustomFieldsIt.Key(), CustomFieldsIt.Value() );
+			Writer->WriteValue(CustomField.Key, CustomField.Value);
 		}
 		Writer->WriteObjectEnd();
 	}
@@ -282,20 +773,36 @@ void FBuildPatchAppManifest::SerializeToJSON( FString& JSONOutput )
 bool FBuildPatchAppManifest::DeserializeFromJSON( const FString& JSONInput )
 {
 	bool bSuccess = true;
-	TSharedPtr< FJsonObject > JSONManifestObject;
-	TSharedRef< TJsonReader< TCHAR > > Reader = TJsonReaderFactory< TCHAR >::Create( JSONInput );
+	TSharedPtr<FJsonObject> JSONManifestObject;
+	TSharedRef<TJsonReader<TCHAR>> Reader = TJsonReaderFactory<TCHAR>::Create(JSONInput);
 
 	// Clear current data
 	DestroyData();
 
 	// Attempt to deserialize JSON
-	if( !FJsonSerializer::Deserialize( Reader, JSONManifestObject ) || !JSONManifestObject.IsValid() )
+	if (!FJsonSerializer::Deserialize(Reader, JSONManifestObject) || !JSONManifestObject.IsValid())
 	{
 		return false;
 	}
 
+	// Store a list of all data GUID for later use
+	TSet<FGuid> AllDataGuids;
+
 	// Get the values map
-	TMap< FString, TSharedPtr< FJsonValue > >& JsonValueMap = JSONManifestObject->Values;
+	TMap<FString, TSharedPtr<FJsonValue>>& JsonValueMap = JSONManifestObject->Values;
+
+	// Manifest version did not always exist
+	int32 ManifestFileVersionInt = 0;
+	TSharedPtr<FJsonValue> JsonManifestFileVersion = JsonValueMap.FindRef(TEXT("ManifestFileVersion"));
+	if (JsonManifestFileVersion.IsValid() && FromStringBlob(JsonManifestFileVersion->AsString(), ManifestFileVersionInt))
+	{
+		Data->ManifestFileVersion = static_cast<EBuildPatchAppManifestVersion::Type>(ManifestFileVersionInt);
+	}
+	else
+	{
+		// Then we presume version just before we started outputting the version
+		Data->ManifestFileVersion = static_cast<EBuildPatchAppManifestVersion::Type>(EBuildPatchAppManifestVersion::StartStoringVersion - 1);
+	}
 
 	// Get the app and version strings
 	TSharedPtr< FJsonValue > JsonAppID = JsonValueMap.FindRef( TEXT("AppID") );
@@ -309,134 +816,131 @@ bool FBuildPatchAppManifest::DeserializeFromJSON( const FString& JSONInput )
 	bSuccess = bSuccess && JsonAppID.IsValid();
 	if( bSuccess )
 	{
-		bSuccess = bSuccess && FromStringBlob( JsonAppID->AsString(), AppID );
+		bSuccess = bSuccess && FromStringBlob( JsonAppID->AsString(), Data->AppID );
 	}
 	bSuccess = bSuccess && JsonAppNameString.IsValid();
 	if( bSuccess )
 	{
-		AppNameString = JsonAppNameString->AsString();
+		Data->AppName = JsonAppNameString->AsString();
 	}
 	bSuccess = bSuccess && JsonBuildVersionString.IsValid();
 	if( bSuccess )
 	{
-		BuildVersionString = JsonBuildVersionString->AsString();
+		Data->BuildVersion = JsonBuildVersionString->AsString();
 	}
 	bSuccess = bSuccess && JsonLaunchExe.IsValid();
 	if( bSuccess )
 	{
-		LaunchExeString = JsonLaunchExe->AsString();
+		Data->LaunchExe = JsonLaunchExe->AsString();
 	}
 	bSuccess = bSuccess && JsonLaunchCommand.IsValid();
 	if( bSuccess )
 	{
-		LaunchCommandString = JsonLaunchCommand->AsString();
+		Data->LaunchCommand = JsonLaunchCommand->AsString();
 	}
 
 	// Get the prerequisites installer info.  These are optional entries.
-	PrereqNameString = JsonPrereqName.IsValid() ? JsonPrereqName->AsString() : FString();
-	PrereqPathString = JsonPrereqPath.IsValid() ? JsonPrereqPath->AsString() : FString();
-	PrereqArgsString = JsonPrereqArgs.IsValid() ? JsonPrereqArgs->AsString() : FString();
-
-	// Manifest version did not always exist
-	int32 ManifestFileVersionInt = 0;
-	TSharedPtr< FJsonValue > JsonManifestFileVersion = JsonValueMap.FindRef( TEXT("ManifestFileVersion") );
-	if( JsonManifestFileVersion.IsValid() && FromStringBlob( JsonManifestFileVersion->AsString(), ManifestFileVersionInt ) )
-	{
-		ManifestFileVersion = static_cast< EBuildPatchAppManifestVersion::Type >( ManifestFileVersionInt );
-	}
-	else
-	{
-		// Then we presume version just before we started outputting the version
-		ManifestFileVersion = static_cast< EBuildPatchAppManifestVersion::Type >( EBuildPatchAppManifestVersion::StartStoringVersion - 1 );
-	}
-
-	// Store a list of all data GUID for later use
-	TSet< FGuid > AllDataGuids;
+	Data->PrereqName = JsonPrereqName.IsValid() ? JsonPrereqName->AsString() : FString();
+	Data->PrereqPath = JsonPrereqPath.IsValid() ? JsonPrereqPath->AsString() : FString();
+	Data->PrereqArgs = JsonPrereqArgs.IsValid() ? JsonPrereqArgs->AsString() : FString();
 
 	// Get the FileManifestList
-	TSharedPtr< FJsonValue > JsonFileManifestList = JsonValueMap.FindRef( TEXT("FileManifestList") );
+	TSharedPtr<FJsonValue> JsonFileManifestList = JsonValueMap.FindRef(TEXT("FileManifestList"));
 	bSuccess = bSuccess && JsonFileManifestList.IsValid();
 	if( bSuccess )
 	{
-		TArray< TSharedPtr< FJsonValue > > JsonFileManifestArray = JsonFileManifestList->AsArray();
-		for ( auto JsonFileManifestIt = JsonFileManifestArray.CreateConstIterator(); JsonFileManifestIt && bSuccess; ++JsonFileManifestIt )
+		TArray<TSharedPtr<FJsonValue>> JsonFileManifestArray = JsonFileManifestList->AsArray();
+		for (auto JsonFileManifestIt = JsonFileManifestArray.CreateConstIterator(); JsonFileManifestIt && bSuccess; ++JsonFileManifestIt)
 		{
-			TSharedPtr< FJsonObject > JsonFileManifest = (*JsonFileManifestIt)->AsObject();
-			TSharedRef< FBuildPatchFileManifest > FileManifest = MakeShareable( new FBuildPatchFileManifest() );
-			FileManifest->Filename = JsonFileManifest->GetStringField( TEXT("Filename") );
-#if PLATFORM_MAC
-			FileManifest->bIsUnixExecutable = JsonFileManifest->HasField( TEXT("bIsUnixExecutable") ) && JsonFileManifest->GetBoolField( TEXT("bIsUnixExecutable") );
-			FileManifest->SymlinkTarget = JsonFileManifest->HasField( TEXT("SymlinkTarget") ) ? JsonFileManifest->GetStringField( TEXT("SymlinkTarget") ) : TEXT( "" );
-#endif
-			FileManifestList.Add( FileManifest->Filename, FileManifest );
-			bSuccess = bSuccess && FString::ToBlob( JsonFileManifest->GetStringField( TEXT("FileHash") ), FileManifest->FileHash.Hash, FSHA1::DigestSize );
-			TArray< TSharedPtr< FJsonValue > > JsonChunkPartArray = JsonFileManifest->GetArrayField( TEXT("FileChunkParts") );
-			for ( auto JsonChunkPartIt = JsonChunkPartArray.CreateConstIterator(); JsonChunkPartIt && bSuccess; ++JsonChunkPartIt )
+			TSharedPtr<FJsonObject> JsonFileManifest = (*JsonFileManifestIt)->AsObject();
+
+			const int32 FileIndex = Data->FileManifestList.Add(FFileManifestData());
+			FFileManifestData& FileManifest = Data->FileManifestList[FileIndex];
+			FileManifest.Filename = JsonFileManifest->GetStringField(TEXT("Filename"));
+			bSuccess = bSuccess && FString::ToBlob(JsonFileManifest->GetStringField(TEXT("FileHash")), FileManifest.FileHash.Hash, FSHA1::DigestSize);
+			TArray<TSharedPtr<FJsonValue>> JsonChunkPartArray = JsonFileManifest->GetArrayField(TEXT("FileChunkParts"));
+			for (auto JsonChunkPartIt = JsonChunkPartArray.CreateConstIterator(); JsonChunkPartIt && bSuccess; ++JsonChunkPartIt)
 			{
-				FileManifest->FileChunkParts.Add( FChunkPart() );
-				TSharedPtr< FJsonObject > JsonChunkPart = (*JsonChunkPartIt)->AsObject();
-				FChunkPart& FileChunkPart = FileManifest->FileChunkParts.Top();
-				bSuccess = bSuccess && FGuid::Parse( JsonChunkPart->GetStringField( TEXT("Guid") ), FileChunkPart.Guid );
-				bSuccess = bSuccess && FromStringBlob( JsonChunkPart->GetStringField( TEXT("Offset") ), FileChunkPart.Offset );
-				bSuccess = bSuccess && FromStringBlob( JsonChunkPart->GetStringField( TEXT("Size") ), FileChunkPart.Size );
-				AllDataGuids.Add( FileChunkPart.Guid );
+				const int32 ChunkIndex = FileManifest.FileChunkParts.Add(FChunkPartData());
+				FChunkPartData& FileChunkPart = FileManifest.FileChunkParts[ChunkIndex];
+				TSharedPtr<FJsonObject> JsonChunkPart = (*JsonChunkPartIt)->AsObject();
+				bSuccess = bSuccess && FGuid::Parse(JsonChunkPart->GetStringField(TEXT("Guid")), FileChunkPart.Guid);
+				bSuccess = bSuccess && FromStringBlob(JsonChunkPart->GetStringField(TEXT("Offset")), FileChunkPart.Offset);
+				bSuccess = bSuccess && FromStringBlob(JsonChunkPart->GetStringField(TEXT("Size")), FileChunkPart.Size);
+				AllDataGuids.Add(FileChunkPart.Guid);
 			}
+			FileManifest.bIsUnixExecutable = JsonFileManifest->HasField(TEXT("bIsUnixExecutable")) && JsonFileManifest->GetBoolField(TEXT("bIsUnixExecutable"));
+			FileManifest.SymlinkTarget = JsonFileManifest->HasField(TEXT("SymlinkTarget")) ? JsonFileManifest->GetStringField(TEXT("SymlinkTarget")) : TEXT("");
+			FileManifest.Init();
 		}
 	}
+	Data->FileManifestList.Sort();
+	for (auto& FileManifest : Data->FileManifestList)
+	{
+		FileManifestLookup.Add(FileManifest.Filename, &FileManifest);
+	}
+
+	// For each chunk setup it's info
+	for (const auto& DataGuid : AllDataGuids)
+	{
+		int32 ChunkIndex = Data->ChunkList.Add(FChunkInfoData());
+		Data->ChunkList[ChunkIndex].Guid = DataGuid;
+	}
+
+	// Setup chunk info lookup
+	for (auto& ChunkInfo : Data->ChunkList)
+	{
+		ChunkInfoLookup.Add(ChunkInfo.Guid, &ChunkInfo);
+	}
 
 	// Get the ChunkHashList
-	TSharedPtr< FJsonValue > JsonChunkHashList = JsonValueMap.FindRef( TEXT("ChunkHashList") );
+	bool bHasChunkHashList = false;
+	TSharedPtr<FJsonValue> JsonChunkHashList = JsonValueMap.FindRef(TEXT("ChunkHashList"));
 	bSuccess = bSuccess && JsonChunkHashList.IsValid();
-	if( bSuccess )
+	if (bSuccess)
 	{
-		TSharedPtr< FJsonObject > JsonChunkHashListObj = JsonChunkHashList->AsObject();
-		for( auto ChunkHashIt = JsonChunkHashListObj->Values.CreateConstIterator(); ChunkHashIt && bSuccess; ++ChunkHashIt )
+		TSharedPtr<FJsonObject> JsonChunkHashListObj = JsonChunkHashList->AsObject();
+		for (auto ChunkHashIt = JsonChunkHashListObj->Values.CreateConstIterator(); ChunkHashIt && bSuccess; ++ChunkHashIt)
 		{
 			FGuid ChunkGuid;
 			uint64 ChunkHash = 0;
-			bSuccess = bSuccess && FGuid::Parse( ChunkHashIt.Key(), ChunkGuid );
-			bSuccess = bSuccess && FromStringBlob( ChunkHashIt.Value()->AsString(), ChunkHash );
-			ChunkHashList.Add( ChunkGuid, ChunkHash );
-			AllDataGuids.Add( ChunkGuid );
+			bSuccess = bSuccess && FGuid::Parse(ChunkHashIt.Key(), ChunkGuid);
+			bSuccess = bSuccess && FromStringBlob(ChunkHashIt.Value()->AsString(), ChunkHash);
+			if (bSuccess && ChunkInfoLookup.Contains(ChunkGuid))
+			{
+				FChunkInfoData* ChunkInfoData = ChunkInfoLookup[ChunkGuid];
+				ChunkInfoData->Hash = ChunkHash;
+				bHasChunkHashList = true;
+			}
 		}
 	}
 
-	// Get the bIsFileData value. The variable will exist in versions of StoresIfChunkOrFileData or later, otherwise the previous method is to check
-	// if ChunkHashList is empty.
-	TSharedPtr< FJsonValue > JsonIsFileData = JsonValueMap.FindRef( TEXT("bIsFileData") );
-	if( JsonIsFileData.IsValid() && JsonIsFileData->Type == EJson::Boolean )
-	{
-		bIsFileData = JsonIsFileData->AsBool();
-	}
-	else
-	{
-		bIsFileData = ChunkHashList.Num() == 0;
-	}
-
 	// Get the DataGroupList
-	TSharedPtr< FJsonValue > JsonDataGroupList = JsonValueMap.FindRef( TEXT("DataGroupList") );
-	if( JsonDataGroupList.IsValid() )
+	TSharedPtr<FJsonValue> JsonDataGroupList = JsonValueMap.FindRef(TEXT("DataGroupList"));
+	if (JsonDataGroupList.IsValid())
 	{
-		TSharedPtr< FJsonObject > JsonDataGroupListObj = JsonDataGroupList->AsObject();
-		for( auto DataGroupIt = JsonDataGroupListObj->Values.CreateConstIterator(); DataGroupIt && bSuccess; ++DataGroupIt )
+		TSharedPtr<FJsonObject> JsonDataGroupListObj = JsonDataGroupList->AsObject();
+		for (auto DataGroupIt = JsonDataGroupListObj->Values.CreateConstIterator(); DataGroupIt && bSuccess; ++DataGroupIt)
 		{
 			FGuid DataGuid;
-			uint8 DataGroup;
+			uint8 DataGroup = INDEX_NONE;
 			// If the list exists, we must be able to parse it ok otherwise error
-			bSuccess = bSuccess && FGuid::Parse( DataGroupIt.Key(), DataGuid );
-			bSuccess = bSuccess && FromStringBlob( DataGroupIt.Value()->AsString(), DataGroup );
-			DataGroupList.Add( DataGuid, DataGroup );
+			bSuccess = bSuccess && FGuid::Parse(DataGroupIt.Key(), DataGuid);
+			bSuccess = bSuccess && FromStringBlob(DataGroupIt.Value()->AsString(), DataGroup);
+			if (bSuccess && ChunkInfoLookup.Contains(DataGuid))
+			{
+				FChunkInfoData* ChunkInfoData = ChunkInfoLookup[DataGuid];
+				ChunkInfoData->GroupNumber = DataGroup;
+			}
 		}
 	}
-	else if( bSuccess )
+	else if (bSuccess)
 	{
 		// If the list did not exist in the manifest then the grouping is the deprecated crc functionality, as long
-		// as there are no previous parsing errors we can build the group list from the AllDataGuids list.
-		for( auto AllDataGuidsIt = AllDataGuids.CreateConstIterator(); AllDataGuidsIt; ++AllDataGuidsIt )
+		// as there are no previous parsing errors we can build the group list from the Guids.
+		for (auto& ChunkInfo : Data->ChunkList)
 		{
-			const FGuid& DataGuid = *AllDataGuidsIt;
-			uint8 DataGroup = FCrc::MemCrc_DEPRECATED( &DataGuid, sizeof( FGuid ) ) % 100;
-			DataGroupList.Add( DataGuid, DataGroup );
+			ChunkInfo.GroupNumber = FCrc::MemCrc_DEPRECATED(&ChunkInfo.Guid, sizeof(FGuid)) % 100;
 		}
 	}
 
@@ -452,11 +956,27 @@ bool FBuildPatchAppManifest::DeserializeFromJSON( const FString& JSONInput )
 			if (FGuid::Parse(ChunkFilesizeIt.Key(), ChunkGuid))
 			{
 				FromStringBlob(ChunkFilesizeIt.Value()->AsString(), ChunkSize);
-				ChunkFilesizeList.Add(ChunkGuid, ChunkSize);
+				if (ChunkInfoLookup.Contains(ChunkGuid))
+				{
+					FChunkInfoData* ChunkInfoData = ChunkInfoLookup[ChunkGuid];
+					ChunkInfoData->FileSize = ChunkSize;
+				}
 			}
 		}
 	}
 
+	// Get the bIsFileData value. The variable will exist in versions of StoresIfChunkOrFileData or later, otherwise the previous method is to check
+	// if ChunkHashList is empty.
+	TSharedPtr<FJsonValue> JsonIsFileData = JsonValueMap.FindRef(TEXT("bIsFileData"));
+	if (JsonIsFileData.IsValid() && JsonIsFileData->Type == EJson::Boolean)
+	{
+		Data->bIsFileData = JsonIsFileData->AsBool();
+	}
+	else
+	{
+		Data->bIsFileData = !bHasChunkHashList;
+	}
+
 	// Get the custom fields. This is optional, and should not fail if it does not exist
 	TSharedPtr< FJsonValue > JsonCustomFields = JsonValueMap.FindRef( TEXT( "CustomFields" ) );
 	if( JsonCustomFields.IsValid() )
@@ -464,28 +984,52 @@ bool FBuildPatchAppManifest::DeserializeFromJSON( const FString& JSONInput )
 		TSharedPtr< FJsonObject > JsonCustomFieldsObj = JsonCustomFields->AsObject();
 		for( auto CustomFieldIt = JsonCustomFieldsObj->Values.CreateConstIterator(); CustomFieldIt && bSuccess; ++CustomFieldIt )
 		{
-			CustomFields.Add( CustomFieldIt.Key(), CustomFieldIt.Value()->AsString() );
+			Data->CustomFields.Add(FCustomFieldData(CustomFieldIt.Key(), CustomFieldIt.Value()->AsString()));
 		}
 	}
+	CustomFieldLookup.Empty(Data->CustomFields.Num());
+	for (auto& CustomField : Data->CustomFields)
+	{
+		CustomFieldLookup.Add(CustomField.Key, &CustomField);
+	}
 
 	// If this is file data, fill out the guid to filename lookup
-	if( bIsFileData )
+	if (Data->bIsFileData)
 	{
-		for( auto FileManifestIt = FileManifestList.CreateConstIterator(); FileManifestIt; ++FileManifestIt )
+		for (auto& FileManifest : Data->FileManifestList)
 		{
-			const TSharedPtr< FBuildPatchFileManifest > FileManifest = FileManifestIt.Value();
-			check( FileManifest->FileChunkParts.Num() == 1 );
-			FileNameReference.Add( FileManifest->FileChunkParts[ 0 ].Guid, FileManifest->Filename );
+			if (FileManifest.FileChunkParts.Num() == 1)
+			{
+				FileNameLookup.Add(FileManifest.FileChunkParts[0].Guid, &FileManifest.Filename);
+			}
+			else
+			{
+				bSuccess = false;
+			}
 		}
 	}
 
 	// Calculate build size
 	TotalBuildSize = 0;
-	for (auto& File : FileManifestList)
+	if (bSuccess)
+	{
+		for (auto& FileManifest : Data->FileManifestList)
+		{
+			TotalBuildSize += FileManifest.GetFileSize();
+		}
+		for (auto& Chunk : Data->ChunkList)
+		{
+			TotalDownloadSize += Chunk.FileSize;
+		}
+	}
+
+	// Download for json file manifest is determined by build file sizes, except any duplicate identical files
+	if (Data->bIsFileData)
 	{
-		if (File.Value.IsValid())
+		TotalDownloadSize = 0;
+		for (auto& NameLookup : FileNameLookup)
 		{
-			TotalBuildSize += File.Value->GetFileSize();
+			TotalDownloadSize += FileManifestLookup[*NameLookup.Value]->GetFileSize();
 		}
 	}
 
@@ -500,26 +1044,25 @@ bool FBuildPatchAppManifest::DeserializeFromJSON( const FString& JSONInput )
 
 const EBuildPatchAppManifestVersion::Type FBuildPatchAppManifest::GetManifestVersion() const
 {
-	return ManifestFileVersion;
+	return Data->ManifestFileVersion;
 }
 
-void FBuildPatchAppManifest::GetChunksRequiredForFiles( const TArray< FString >& FileList, TArray< FGuid >& RequiredChunks, bool bAddUnique ) const
+void FBuildPatchAppManifest::GetChunksRequiredForFiles(const TArray<FString>& FileList, TArray<FGuid>& RequiredChunks, bool bAddUnique) const
 {
-	for( auto FileIt = FileList.CreateConstIterator(); FileIt ; ++FileIt )
+	for (const auto& File : FileList)
 	{
-		TSharedPtr< FBuildPatchFileManifest > FileManifest = FileManifestList.FindRef( *FileIt );
-		if( FileManifest.IsValid() )
+		FFileManifestData* const * FileManifest = FileManifestLookup.Find(File);
+		if (FileManifest)
 		{
-			for( auto ChunkPartIt = FileManifest->FileChunkParts.CreateConstIterator(); ChunkPartIt; ++ChunkPartIt )
+			for (const auto& ChunkPart : (*FileManifest)->FileChunkParts)
 			{
-				const FChunkPart& ChunkPart = *ChunkPartIt;
-				if( bAddUnique )
+				if (bAddUnique)
 				{
-					RequiredChunks.AddUnique( ChunkPart.Guid );
+					RequiredChunks.AddUnique(ChunkPart.Guid);
 				}
 				else
 				{
-					RequiredChunks.Add( ChunkPart.Guid );
+					RequiredChunks.Add(ChunkPart.Guid);
 				}
 			}
 		}
@@ -528,17 +1071,7 @@ void FBuildPatchAppManifest::GetChunksRequiredForFiles( const TArray< FString >&
 
 const int64 FBuildPatchAppManifest::GetDownloadSize() const
 {
-	// If this is not a chunked build, then return build size
-	if (bIsFileData)
-	{
-		return GetBuildSize();
-	}
-	else
-	{
-		TArray<FGuid> ChunkList;
-		ChunkHashList.GetKeys(ChunkList);
-		return GetDataSize(ChunkList);
-	}
+	return TotalDownloadSize;
 }
 
 const int64 FBuildPatchAppManifest::GetBuildSize() const
@@ -546,42 +1079,38 @@ const int64 FBuildPatchAppManifest::GetBuildSize() const
 	return TotalBuildSize;
 }
 
-const int64 FBuildPatchAppManifest::GetFileSize( const TArray< FString >& Filenames ) const
+const int64 FBuildPatchAppManifest::GetFileSize(const TArray<FString>& Filenames) const
 {
 	int64 TotalSize = 0;
-	for( int32 FileIdx = 0; FileIdx < Filenames.Num(); ++FileIdx )
+	for (const auto& Filename : Filenames)
 	{
-		TSharedPtr< FBuildPatchFileManifest > FileManifest = FileManifestList.FindRef( Filenames[ FileIdx ] );
-		if( FileManifest.IsValid() )
-		{
-			TotalSize += FileManifest->GetFileSize();
-		}
+		TotalSize += GetFileSize(Filename);
 	}
 	return TotalSize;
 }
 
-const int64 FBuildPatchAppManifest::GetFileSize( const FString& Filename ) const
+const int64 FBuildPatchAppManifest::GetFileSize(const FString& Filename) const
 {
-	TSharedPtr< FBuildPatchFileManifest > FileManifest = FileManifestList.FindRef( Filename );
-	if( FileManifest.IsValid() )
+	FFileManifestData* const * FileManifest = FileManifestLookup.Find(Filename);
+	if (FileManifest)
 	{
-		return FileManifest->GetFileSize();
+		return (*FileManifest)->GetFileSize();
 	}
 	return 0;
 }
 
 const int64 FBuildPatchAppManifest::GetDataSize(const FGuid& DataGuid) const
 {
-	if (bIsFileData)
+	if (ChunkInfoLookup.Contains(DataGuid))
 	{
-		// For file data, the file must exist in the list
-		check(FileNameReference.Contains(DataGuid));
-		return GetFileSize(FileNameReference[DataGuid]);
+		// Chunk file sizes are stored in the info
+		return ChunkInfoLookup[DataGuid]->FileSize;
 	}
-	else if (ChunkFilesizeList.Contains(DataGuid))
+	else if (Data->bIsFileData)
 	{
-		// Chunk data sizes are stored in our map
-		return ChunkFilesizeList[DataGuid];
+		// For file data, the file must exist in the list
+		check(FileNameLookup.Contains(DataGuid));
+		return GetFileSize(*FileNameLookup[DataGuid]);
 	}
 	else
 	{
@@ -590,68 +1119,64 @@ const int64 FBuildPatchAppManifest::GetDataSize(const FGuid& DataGuid) const
 	}
 }
 
-const int64 FBuildPatchAppManifest::GetDataSize(const TArray< FGuid >& DataGuids) const
+const int64 FBuildPatchAppManifest::GetDataSize(const TArray<FGuid>& DataGuids) const
 {
 	int64 TotalSize = 0;
-	for (auto DataGuidsIt = DataGuids.CreateConstIterator(); DataGuidsIt; ++DataGuidsIt)
+	for (const auto& DataGuid : DataGuids)
 	{
-		TotalSize += GetDataSize(*DataGuidsIt);
+		TotalSize += GetDataSize(DataGuid);
 	}
 	return TotalSize;
 }
 
 const uint32 FBuildPatchAppManifest::GetNumFiles() const
 {
-	return FileManifestList.Num();
+	return Data->FileManifestList.Num();
 }
 
-void FBuildPatchAppManifest::GetFileList( TArray< FString >& Filenames ) const
+void FBuildPatchAppManifest::GetFileList(TArray<FString>& Filenames) const
 {
-	FileManifestList.GetKeys( Filenames );
+	FileManifestLookup.GetKeys(Filenames);
 }
 
-void FBuildPatchAppManifest::GetDataList(TArray< FGuid >& DataGuids) const
+void FBuildPatchAppManifest::GetDataList(TArray<FGuid>& DataGuids) const
 {
-	ChunkHashList.GetKeys( DataGuids );
+	ChunkInfoLookup.GetKeys(DataGuids);
 }
 
-const TSharedPtr< FBuildPatchFileManifest > FBuildPatchAppManifest::GetFileManifest( const FString& Filename )
+const FFileManifestData* FBuildPatchAppManifest::GetFileManifest(const FString& Filename)
 {
-	return FileManifestList.FindRef( Filename );
+	const FFileManifestData* const * FileManifest = FileManifestLookup.Find(Filename);
+	return (FileManifest) ? (*FileManifest) : nullptr;
 }
 
 const bool FBuildPatchAppManifest::IsFileDataManifest() const
 {
-	return bIsFileData;
+	return Data->bIsFileData;
 }
 
-const bool FBuildPatchAppManifest::GetChunkHash( const FGuid& ChunkGuid, uint64& OutHash ) const
+const bool FBuildPatchAppManifest::GetChunkHash(const FGuid& ChunkGuid, uint64& OutHash) const
 {
-	const uint64* ChunkHash = ChunkHashList.Find( ChunkGuid );
-	if( ChunkHash != NULL )
+	const FChunkInfoData* const * ChunkInfo = ChunkInfoLookup.Find(ChunkGuid);
+	if (ChunkInfo)
 	{
-		OutHash = *ChunkHash;
+		OutHash = (*ChunkInfo)->Hash;
 		return true;
 	}
 	return false;
 }
 
-const bool FBuildPatchAppManifest::GetFileDataHash( const FGuid& FileGuid, FSHAHash& OutHash ) const
+const bool FBuildPatchAppManifest::GetFileDataHash(const FGuid& FileGuid, FSHAHashData& OutHash) const
 {
 	// File data hash is the same as the hash for the file it holds, so we can grab it from there
-	for( auto FileManifestIt = FileManifestList.CreateConstIterator(); FileManifestIt; ++FileManifestIt )
+	for (const auto& FileManifest : Data->FileManifestList)
 	{
-		if( FileManifestIt.Value().IsValid() )
+		for (const auto& FileDataPart : FileManifest.FileChunkParts)
 		{
-			const FBuildPatchFileManifest& FileManifest = *FileManifestIt.Value().Get();
-			for( auto FileChunkPartsIt = FileManifest.FileChunkParts.CreateConstIterator(); FileChunkPartsIt; ++FileChunkPartsIt )
+			if (FileDataPart.Guid == FileGuid)
 			{
-				const FChunkPart& FileDataPart = *FileChunkPartsIt;
-				if( FileDataPart.Guid == FileGuid )
-				{
-					OutHash = FileManifest.FileHash;
-					return true;
-				}
+				FMemory::Memcpy(OutHash.Hash, FileManifest.FileHash.Hash, FSHA1::DigestSize);
+				return true;
 			}
 		}
 	}
@@ -660,42 +1185,42 @@ const bool FBuildPatchAppManifest::GetFileDataHash( const FGuid& FileGuid, FSHAH
 
 const uint32& FBuildPatchAppManifest::GetAppID() const
 {
-	return AppID;
+	return Data->AppID;
 }
 
 const FString& FBuildPatchAppManifest::GetAppName() const
 {
-	return AppNameString;
+	return Data->AppName;
 }
 
 const FString& FBuildPatchAppManifest::GetVersionString() const
 {
-	return BuildVersionString;
+	return Data->BuildVersion;
 }
 
 const FString& FBuildPatchAppManifest::GetLaunchExe() const
 {
-	return LaunchExeString;
+	return Data->LaunchExe;
 }
 
 const FString& FBuildPatchAppManifest::GetLaunchCommand() const
 {
-	return LaunchCommandString;
+	return Data->LaunchCommand;
 }
 
 const FString& FBuildPatchAppManifest::GetPrereqName() const
 {
-	return PrereqNameString;
+	return Data->PrereqName;
 }
 
 const FString& FBuildPatchAppManifest::GetPrereqPath() const
 {
-	return PrereqPathString;
+	return Data->PrereqPath;
 }
 
 const FString& FBuildPatchAppManifest::GetPrereqArgs() const
 {
-	return PrereqArgsString;
+	return Data->PrereqArgs;
 }
 
 IBuildManifestRef FBuildPatchAppManifest::Duplicate() const
@@ -703,29 +1228,42 @@ IBuildManifestRef FBuildPatchAppManifest::Duplicate() const
 	return MakeShareable(new FBuildPatchAppManifest(*this));
 }
 
-const IManifestFieldPtr FBuildPatchAppManifest::GetCustomField( const FString& FieldName ) const
+const IManifestFieldPtr FBuildPatchAppManifest::GetCustomField(const FString& FieldName) const
 {
 	IManifestFieldPtr Return;
-	if( CustomFields.Contains( FieldName ) )
+	if (CustomFieldLookup.Contains(FieldName))
 	{
-		Return = MakeShareable( new FBuildPatchCustomField( CustomFields.FindRef( FieldName ) ) );
+		Return = MakeShareable(new FBuildPatchCustomField(CustomFieldLookup[FieldName]->Value));
 	}
 	return Return;
 }
 
-const IManifestFieldPtr FBuildPatchAppManifest::SetCustomField( const FString& FieldName, const FString& Value )
+const IManifestFieldPtr FBuildPatchAppManifest::SetCustomField(const FString& FieldName, const FString& Value)
 {
-	return MakeShareable( new FBuildPatchCustomField( CustomFields.Add( FieldName, Value ) ) );
+	if (CustomFieldLookup.Contains(FieldName))
+	{
+		CustomFieldLookup[FieldName]->Value = Value;
+	}
+	else
+	{
+		Data->CustomFields.Add(FCustomFieldData(FieldName, Value));
+		CustomFieldLookup.Empty(Data->CustomFields.Num());
+		for (auto& CustomField : Data->CustomFields)
+		{
+			CustomFieldLookup.Add(CustomField.Key, &CustomField);
+		}
+	}
+	return GetCustomField(FieldName);
 }
 
-const IManifestFieldPtr FBuildPatchAppManifest::SetCustomField( const FString& FieldName, const double& Value )
+const IManifestFieldPtr FBuildPatchAppManifest::SetCustomField(const FString& FieldName, const double& Value)
 {
-	return MakeShareable( new FBuildPatchCustomField( CustomFields.Add( FieldName, ToStringBlob( Value ) ) ) );
+	return SetCustomField(FieldName, ToStringBlob(Value));
 }
 
-const IManifestFieldPtr FBuildPatchAppManifest::SetCustomField( const FString& FieldName, const int64& Value )
+const IManifestFieldPtr FBuildPatchAppManifest::SetCustomField(const FString& FieldName, const int64& Value)
 {
-	return MakeShareable( new FBuildPatchCustomField( CustomFields.Add( FieldName, ToStringBlob( Value ) ) ) );
+	return SetCustomField(FieldName, ToStringBlob(Value));
 }
 
 void FBuildPatchAppManifest::EnumerateProducibleChunks( const FString& InstallDirectory, const TArray< FGuid >& ChunksRequired, TArray< FGuid >& ChunksAvailable ) const
@@ -809,7 +1347,7 @@ bool FBuildPatchAppManifest::VerifyAgainstDirectory( const FString& VerifyDirect
 	TimeSpentPaused = 0;
 
 	// Setup progress tracking
-	double TotalBuildSize = GetBuildSize();
+	double TotalBuildSizeDouble = GetBuildSize();
 	double ProcessedBytes = 0;
 	struct FLocalProgress
 	{
@@ -826,7 +1364,7 @@ bool FBuildPatchAppManifest::VerifyAgainstDirectory( const FString& VerifyDirect
 	LocalProgress.ProgressDelegate = ProgressDelegate;
 
 	// For all files in FileManifestList, check that they produce the correct SHA1 hash, adding any that don't to the list
-	for( auto FileManifestIt = FileManifestList.CreateConstIterator(); FileManifestIt && !FBuildPatchInstallError::HasFatalError(); ++FileManifestIt )
+	for( auto FileManifestIt = Data->FileManifestList.CreateConstIterator(); FileManifestIt && !FBuildPatchInstallError::HasFatalError(); ++FileManifestIt )
 	{
 		// Pause if necessary
 		const double PrePauseTime = FPlatformTime::Seconds();
@@ -841,42 +1379,42 @@ bool FBuildPatchAppManifest::VerifyAgainstDirectory( const FString& VerifyDirect
 		// Count up pause time
 		TimeSpentPaused += PostPauseTime - PrePauseTime;
 		// Get file details and construct
-		TSharedPtr< FBuildPatchFileManifest > FileManifest = FileManifestIt.Value();
-		FString FullFilename = VerifyDirectory / FileManifest->Filename;
-		const int64 CurrentFileSize = FileManifest->GetFileSize();
-		LocalProgress.CurrentFileWeight = CurrentFileSize / TotalBuildSize;
-		const bool bSameFile = FBuildPatchUtils::VerifyFile(FullFilename, FileManifest->FileHash, FileManifest->FileHash, FBuildPatchFloatDelegate::CreateRaw(&LocalProgress, &FLocalProgress::PerFileProgress), ShouldPauseDelegate, TimeSpentPaused) != 0;
+		const FFileManifestData& FileManifest = *FileManifestIt;
+		FString FullFilename = VerifyDirectory / FileManifest.Filename;
+		const int64 CurrentFileSize = FileManifest.GetFileSize();
+		LocalProgress.CurrentFileWeight = CurrentFileSize / TotalBuildSizeDouble;
+		const bool bSameFile = FBuildPatchUtils::VerifyFile(FullFilename, FileManifest.FileHash, FileManifest.FileHash, FBuildPatchFloatDelegate::CreateRaw(&LocalProgress, &FLocalProgress::PerFileProgress), ShouldPauseDelegate, TimeSpentPaused) != 0;
 		ProcessedBytes += CurrentFileSize;
-		LocalProgress.CurrentBuildPercentage = ProcessedBytes / TotalBuildSize;
+		LocalProgress.CurrentBuildPercentage = ProcessedBytes / TotalBuildSizeDouble;
 		bAllCorrect &= bSameFile;
 		if( !bSameFile )
 		{
-			OutDatedFiles.Add( FileManifest->Filename );
+			OutDatedFiles.Add( FileManifest.Filename );
 		}
 	}
 
 	return bAllCorrect && !FBuildPatchInstallError::HasFatalError();
 }
 
-void FBuildPatchAppManifest::EnumerateChunkPartInventory( const TArray< FGuid >& ChunksRequired, TMap< FGuid, TArray< FFileChunkPart > >& ChunkPartsAvailable ) const
+void FBuildPatchAppManifest::EnumerateChunkPartInventory(const TArray<FGuid>& ChunksRequired, TMap<FGuid, TArray<FFileChunkPart>>& ChunkPartsAvailable) const
 {
 	ChunkPartsAvailable.Empty();
 	// For each file in the manifest, check what chunks it is made out of, and grab details for the ones in ChunksRequired
-	for (auto FileManifestIt = FileManifestList.CreateConstIterator(); FileManifestIt && !FBuildPatchInstallError::HasFatalError(); ++FileManifestIt)
+	for (auto FileManifestIt = Data->FileManifestList.CreateConstIterator(); FileManifestIt && !FBuildPatchInstallError::HasFatalError(); ++FileManifestIt)
 	{
-		const TSharedPtr< FBuildPatchFileManifest > FileManifest = FileManifestIt.Value();
+		const FFileManifestData& FileManifest = *FileManifestIt;
 		uint64 FileOffset = 0;
-		for (auto ChunkPartIt = FileManifest->FileChunkParts.CreateConstIterator(); ChunkPartIt && !FBuildPatchInstallError::HasFatalError(); ++ChunkPartIt)
+		for (auto ChunkPartIt = FileManifest.FileChunkParts.CreateConstIterator(); ChunkPartIt && !FBuildPatchInstallError::HasFatalError(); ++ChunkPartIt)
 		{
-			const FChunkPart& ChunkPart = *ChunkPartIt;
-			if( ChunksRequired.Contains( ChunkPart.Guid ) )
+			const FChunkPartData& ChunkPart = *ChunkPartIt;
+			if (ChunksRequired.Contains(ChunkPart.Guid))
 			{
-				TArray< FFileChunkPart >& FileChunkParts = ChunkPartsAvailable.FindOrAdd( ChunkPart.Guid );
+				TArray<FFileChunkPart>& FileChunkParts = ChunkPartsAvailable.FindOrAdd(ChunkPart.Guid);
 				FFileChunkPart FileChunkPart;
-				FileChunkPart.Filename = FileManifest->Filename;
+				FileChunkPart.Filename = FileManifest.Filename;
 				FileChunkPart.ChunkPart = ChunkPart;
 				FileChunkPart.FileOffset = FileOffset;
-				FileChunkParts.Add( FileChunkPart );
+				FileChunkParts.Add(FileChunkPart);
 			}
 			FileOffset += ChunkPart.Size;
 		}
@@ -894,12 +1432,11 @@ void FBuildPatchAppManifest::GetRemovableFiles(IBuildManifestRef InOldManifest,
 	FBuildPatchAppManifestRef OldManifest = StaticCastSharedRef< FBuildPatchAppManifest >(InOldManifest);
 	// Simply put, any files that exist in the OldManifest file list, but do not in this manifest's file list, are assumed
 	// to be files no longer required by the build
-	for (auto OldFilesIt = OldManifest->FileManifestList.CreateConstIterator(); OldFilesIt; ++OldFilesIt)
+	for (const auto& OldFile : OldManifest->Data->FileManifestList)
 	{
-		const FString& OldFile = OldFilesIt.Key();
-		if (!this->FileManifestList.Contains(OldFile))
+		if (!FileManifestLookup.Contains(OldFile.Filename))
 		{
-			RemovableFiles.Add(OldFile);
+			RemovableFiles.Add(OldFile.Filename);
 		}
 	}
 }
@@ -921,79 +1458,80 @@ void FBuildPatchAppManifest::GetRemovableFiles(const TCHAR* InstallPath, TArray<
 	
 	for (int32 FileIndex = 0; FileIndex < AllFiles.Num(); ++FileIndex)
 	{
-		const FString& FileName = AllFiles[FileIndex].RightChop(BasePath.Len());
-		if (!FileManifestList.Contains(*FileName))
+		const FString& Filename = AllFiles[FileIndex].RightChop(BasePath.Len());
+		if (!FileManifestLookup.Contains(Filename))
 		{
 			RemovableFiles.Add(AllFiles[FileIndex]);
 		}
 	}
 }
 
-void FBuildPatchAppManifest::GetOutdatedFiles( FBuildPatchAppManifestPtr OldManifest, FBuildPatchAppManifestRef NewManifest, const FString& InstallDirectory, TArray< FString >& OutDatedFiles )
+bool FBuildPatchAppManifest::NeedsResaving() const
+{
+	return true;
+	return Data->ManifestFileVersion < EBuildPatchAppManifestVersion::GetLatestVersion();
+}
+
+void FBuildPatchAppManifest::GetOutdatedFiles(FBuildPatchAppManifestPtr OldManifest, FBuildPatchAppManifestRef NewManifest, const FString& InstallDirectory, TArray< FString >& OutDatedFiles)
 {
-	if( !OldManifest.IsValid() )
+	if (!OldManifest.IsValid())
 	{
 		// All files are outdated if no OldManifest
-		NewManifest->FileManifestList.GetKeys( OutDatedFiles );
+		NewManifest->FileManifestLookup.GetKeys(OutDatedFiles);
 	}
 	else
 	{
 		// Enumerate files in the NewManifest file list, that do not exist, or have different hashes in the OldManifest
 		// to be files no longer required by the build
-		for( auto NewFilesIt = NewManifest->FileManifestList.CreateConstIterator(); NewFilesIt; ++NewFilesIt )
+		for (const auto& NewFile : NewManifest->Data->FileManifestList)
 		{
-			const FString& NewFile = NewFilesIt.Key();
-			const int64 ExistingFileSize = IFileManager::Get().FileSize( *(InstallDirectory / NewFile) );
-			const TSharedPtr< FBuildPatchFileManifest > NewFileManifest = NewFilesIt.Value();
-			const TSharedPtr< FBuildPatchFileManifest > OldFileManifest = OldManifest->FileManifestList.FindRef( NewFile );
+			const int64 ExistingFileSize = IFileManager::Get().FileSize(*(InstallDirectory / NewFile.Filename));
 			// Check changed
-			if( IsFileOutdated( OldManifest.ToSharedRef(), NewManifest, NewFile ) )
+			if (IsFileOutdated(OldManifest.ToSharedRef(), NewManifest, NewFile.Filename))
 			{
-				OutDatedFiles.Add( NewFile );
+				OutDatedFiles.Add(NewFile.Filename);
 			}
 			// Double check an unchanged file is not missing (size will be -1) or is incorrect size
-			else if( ExistingFileSize != NewFileManifest->GetFileSize() )
+			else if (ExistingFileSize != NewFile.GetFileSize())
 			{
-				OutDatedFiles.Add( NewFile );
+				OutDatedFiles.Add(NewFile.Filename);
 			}
 		}
 	}
 }
 
-bool FBuildPatchAppManifest::IsFileOutdated( FBuildPatchAppManifestRef OldManifest, FBuildPatchAppManifestRef NewManifest, const FString& Filename )
+bool FBuildPatchAppManifest::IsFileOutdated(FBuildPatchAppManifestRef OldManifest, FBuildPatchAppManifestRef NewManifest, const FString& Filename)
 {
 	// If both app manifests are the same, return false as only repair would touch the file.
-	if( OldManifest == NewManifest )
+	if (OldManifest == NewManifest)
 	{
 		return false;
 	}
 	// Get file manifests
-	TSharedPtr< FBuildPatchFileManifest > OldFile = OldManifest->GetFileManifest( Filename );
-	TSharedPtr< FBuildPatchFileManifest > NewFile = NewManifest->GetFileManifest( Filename );
+	const FFileManifestData* OldFile = OldManifest->GetFileManifest(Filename);
+	const FFileManifestData* NewFile = NewManifest->GetFileManifest(Filename);
 	// Out of date if not in either manifest
-	if( !OldFile.IsValid() || !NewFile.IsValid() )
+	if (!OldFile || !NewFile)
 	{
 		return true;
 	}
 	// Different hash means different file
-	if( OldFile->FileHash != NewFile->FileHash )
+	if (OldFile->FileHash != NewFile->FileHash)
 	{
 		return true;
 	}
 	return false;
 }
 
-const uint32 FBuildPatchAppManifest::GetNumberOfChunkReferences( const FGuid& ChunkGuid ) const
+const uint32 FBuildPatchAppManifest::GetNumberOfChunkReferences(const FGuid& ChunkGuid) const
 {
 	uint32 RefCount = 0;
 	// For each file in the manifest, check if it has references to this chunk
-	for( auto FileManifestIt = FileManifestList.CreateConstIterator(); FileManifestIt; ++FileManifestIt )
+	for (const auto& FileManifest : Data->FileManifestList)
 	{
-		const TSharedPtr< FBuildPatchFileManifest >& FileManifest = FileManifestIt.Value();
-		for( auto ChunkPartIt = FileManifest->FileChunkParts.CreateConstIterator(); ChunkPartIt; ++ChunkPartIt )
+		for (const auto& ChunkPart : FileManifest.FileChunkParts)
 		{
-			const FChunkPart& ChunkPart = *ChunkPartIt;
-			if( ChunkPart.Guid == ChunkGuid )
+			if (ChunkPart.Guid == ChunkGuid)
 			{
 				++RefCount;
 			}
diff --git a/Engine/Source/Runtime/Online/BuildPatchServices/Private/BuildPatchManifest.h b/Engine/Source/Runtime/Online/BuildPatchServices/Private/BuildPatchManifest.h
--- a/Engine/Source/Runtime/Online/BuildPatchServices/Private/BuildPatchManifest.h
+++ b/Engine/Source/Runtime/Online/BuildPatchServices/Private/BuildPatchManifest.h
@@ -4,6 +4,10 @@
 	BuildPatchManifest.h: Declares the manifest classes.
 =============================================================================*/
 
+#include "CoreUObject.h"
+#include "BuildPatchChunk.h"
+#include "BuildPatchManifest.generated.h"
+
 #pragma once
 
 // A delegate returning a bool. Used to pass a paused state
@@ -17,6 +21,7 @@ typedef TSharedRef< class FBuildPatchAppManifest, ESPMode::ThreadSafe > FBuildPa
 /**
  * An enum type to describe supported features of a certain manifest
  */
+UENUM()
 namespace EBuildPatchAppManifestVersion
 {
 	enum Type
@@ -39,31 +44,223 @@ namespace EBuildPatchAppManifestVersion
 		StoresPrerequisitesInfo,
 		// Manifest stores chunk download sizes
 		StoresChunkFileSizes,
+		// Manifest is now stored using UObject serialization and compressed
+		StoredAsCompressedUClass,
 
 		// Always last, signifies the latest version plus 1 to allow initialization simplicity
 		LatestPlusOne
 	};
+}
 
-	/** @return The last known manifest feature of this codebase. Handy for manifest constructor */
+namespace EBuildPatchAppManifestVersion
+{
+	/** @return The last known manifest feature of this code base. Handy for manifest constructor */
 	const Type GetLatestVersion();
 
+	/** @return The latest version of a manifest support by JSON serialization */
+	const Type GetLatestJsonVersion();
+
 	/**
 	 * Get the chunk subdirectory for used for a specific manifest version, e.g. Chunks, ChunksV2 etc
 	 * @param ManifestVersion     The version of the manifest
 	 * @return The subdirectory name that this manifest version will access
 	 */
-	const FString& GetChunkSubdir( const Type ManifestVersion );
+	const FString& GetChunkSubdir(const Type ManifestVersion);
 
 	/**
 	 * Get the file data subdirectory for used for a specific manifest version, e.g. Files, FilesV2 etc
 	 * @param ManifestVersion     The version of the manifest
 	 * @return The subdirectory name that this manifest version will access
 	 */
-	const FString& GetFileSubdir( const Type ManifestVersion );
+	const FString& GetFileSubdir(const Type ManifestVersion);
+};
+
+
+/**
+ * An enum type to describe the format that manifest data is stored
+ */
+UENUM()
+namespace EManifestFileHeader
+{
+	enum Type
+	{
+		// Storage flags, can be raw or a combination of others.
+		STORED_RAW = 0x0, // Zero means raw data
+		STORED_COMPRESSED = 0x1, // Flag for compressed
+	};
+}
+
+/**
+ * A data structure that hold a manifests custom field. This is a key value pair of strings
+ */
+USTRUCT()
+struct FCustomFieldData
+{
+	GENERATED_USTRUCT_BODY()
+
+	UPROPERTY()
+	FString Key;
+
+	UPROPERTY()
+	FString Value;
+
+	FCustomFieldData();
+	FCustomFieldData(const FString& Key, const FString& Value);
+};
+
+/**
+ * A data structure describing a chunk file
+ */
+USTRUCT()
+struct FChunkInfoData
+{
+	GENERATED_USTRUCT_BODY()
+
+	UPROPERTY()
+	FGuid Guid;
+
+	UPROPERTY()
+	uint64 Hash;
+
+	UPROPERTY()
+	int64 FileSize;
+
+	UPROPERTY()
+	uint8 GroupNumber;
+
+	FChunkInfoData();
+};
+
+/**
+ * A data structure describing the part of a chunk used to construct a file
+ */
+USTRUCT()
+struct FChunkPartData
+{
+	GENERATED_USTRUCT_BODY()
+
+	// The GUID of the chunk containing this part
+	UPROPERTY()
+	FGuid Guid;
+
+	// The offset of the first byte into the chunk
+	UPROPERTY()
+	uint32 Offset;
+
+	// The size of this part
+	UPROPERTY()
+	uint32 Size;
+
+	FChunkPartData();
+};
+
+/**
+ * A data structure that describes a file's construction information
+ */
+USTRUCT()
+struct FFileManifestData
+{
+	GENERATED_USTRUCT_BODY()
+
+	UPROPERTY()
+	FString Filename;
+
+	UPROPERTY()
+	FSHAHashData FileHash;
+
+	UPROPERTY()
+	TArray<FChunkPartData> FileChunkParts;
+
+	UPROPERTY()
+	bool bIsUnixExecutable;
+
+	UPROPERTY()
+	FString SymlinkTarget;
+
+	FFileManifestData();
+
+	void Init();
+
+	const int64& GetFileSize() const;
+
+	FORCEINLINE bool operator<(const FFileManifestData& Other) const;
+
+private:
+	int64 FileSize;
 };
 
 /**
- * Declare the FBuildPatchCustomField object class
+ * This is the manifest UObject where all manifest data is stored
+ */
+UCLASS()
+class UBuildPatchManifest : public UObject
+{
+	GENERATED_UCLASS_BODY()
+
+	UPROPERTY()
+	TEnumAsByte<EBuildPatchAppManifestVersion::Type> ManifestFileVersion;
+
+	UPROPERTY()
+	bool bIsFileData;
+
+	UPROPERTY()
+	uint32 AppID;
+
+	UPROPERTY()
+	FString AppName;
+
+	UPROPERTY()
+	FString BuildVersion;
+
+	UPROPERTY()
+	FString LaunchExe;
+
+	UPROPERTY()
+	FString LaunchCommand;
+
+	UPROPERTY()
+	FString PrereqName;
+
+	UPROPERTY()
+	FString PrereqPath;
+
+	UPROPERTY()
+	FString PrereqArgs;
+
+	UPROPERTY()
+	TArray<FFileManifestData> FileManifestList;
+
+	UPROPERTY()
+	TArray<FChunkInfoData> ChunkList;
+
+	UPROPERTY()
+	TArray<FCustomFieldData> CustomFields;
+
+	void Clear();
+
+	~UBuildPatchManifest();
+};
+
+/**
+ * The manifest file header wraps the manifest data stored in a file on disc to describe how to read it
+ */
+struct FManifestFileHeader
+{
+	uint32 Magic;
+	uint32 HeaderSize;
+	uint32 DataSize;
+	uint32 CompressedSize;
+	FSHAHashData SHAHash;
+	uint8 StoredAs;
+
+	FManifestFileHeader();
+	bool CheckMagic() const;
+	friend FArchive& operator<<(FArchive& Ar, FManifestFileHeader& Header);
+};
+
+/**
+ * Declare the FBuildPatchCustomField object class, which is the implementation of the object we return to
+ * clients of the module
  */
 class FBuildPatchCustomField
 	: public IManifestField
@@ -81,7 +278,7 @@ public:
 	/**
 	 * Constructor taking the custom value
 	 */
-	FBuildPatchCustomField( const FString& Value );
+	FBuildPatchCustomField(const FString& Value);
 
 	// START IBuildManifest Interface
 	virtual const FString AsString() const override;
@@ -91,44 +288,29 @@ public:
 };
 
 /**
- * Declare the FBuildPatchFileManifest object class
+ * Declares a struct to store the info about a piece of a chunk that is inside a file
  */
-class FBuildPatchFileManifest
+struct FFileChunkPart
 {
-public:
-
-	// The filename of the file, local to a build's root
+	// The file containing this piece
 	FString Filename;
+	// The offset into the file of this piece
+	uint64 FileOffset;
+	// The FChunkPartData that can be salvaged from this file
+	FChunkPartData ChunkPart;
 
-	// The Integrity Verification hash
-	FSHAHash FileHash;
-
-	// The sequential list of chunk parts this file is made up of
-	TArray< FChunkPart > FileChunkParts;
-
-#if PLATFORM_MAC
-	// Additional properties for Unix files
-	bool bIsUnixExecutable;
-	FString SymlinkTarget;
-#endif
-
-public:
-	// Default constructor
-	FBuildPatchFileManifest();
-	
 	/**
-	 * Get the total size for this file
-	 * @return		The size in bytes.
-	 */
-	const int64& GetFileSize();
-
-private:
-	// Holds the file size, calculated from the FileChunkParts
-	int64 FileSize;
+	* Default Constructor
+	*/
+	FFileChunkPart()
+		: Filename()
+		, FileOffset(0)
+		, ChunkPart()
+	{}
 };
 
 /**
- * Declare the FBuildPatchAppManifest object class. This will contain FBuildPatchFileManifest and other build information.
+ * Declare the FBuildPatchAppManifest object class. This holds the UObject data, and the implemented build manifest functionality
  */
 class FBuildPatchAppManifest
 	: public IBuildManifest
@@ -139,7 +321,7 @@ class FBuildPatchAppManifest
 	friend class FBuildDataFileProcessor;
 	friend class FBuildPatchInstaller;
 public:
-	
+
 	/**
 	 * Default constructor
 	 */
@@ -148,7 +330,7 @@ public:
 	/**
 	 * Basic details constructor
 	 */
-	FBuildPatchAppManifest( const uint32& InAppID, const FString& AppName );
+	FBuildPatchAppManifest(const uint32& InAppID, const FString& AppName);
 
 	/**
 	 * Default destructor
@@ -168,11 +350,12 @@ public:
 	virtual const int64 GetBuildSize() const override;
 	virtual void GetRemovableFiles(IBuildManifestRef OldManifest, TArray< FString >& RemovableFiles) const override;
 	virtual void GetRemovableFiles(const TCHAR* InstallPath, TArray< FString >& RemovableFiles) const override;
+	virtual bool NeedsResaving() const override;
 	virtual IBuildManifestRef Duplicate() const override;
 	virtual const IManifestFieldPtr GetCustomField(const FString& FieldName) const override;
-	virtual const IManifestFieldPtr SetCustomField( const FString& FieldName, const FString& Value ) override;
-	virtual const IManifestFieldPtr SetCustomField( const FString& FieldName, const double& Value ) override;
-	virtual const IManifestFieldPtr SetCustomField( const FString& FieldName, const int64& Value ) override;
+	virtual const IManifestFieldPtr SetCustomField(const FString& FieldName, const FString& Value) override;
+	virtual const IManifestFieldPtr SetCustomField(const FString& FieldName, const double& Value) override;
+	virtual const IManifestFieldPtr SetCustomField(const FString& FieldName, const int64& Value) override;
 	// END IBuildManifest Interface
 
 	/**
@@ -180,27 +363,35 @@ public:
 	 * @param Filename		The file to load JSON from
 	 * @return		True if successful.
 	 */
-	bool LoadFromFile( FString Filename );
+	bool LoadFromFile(FString Filename);
+
+	/**
+	 * Sets up the object from the passed in data
+	 * @param DataInput		The data to deserialize from
+	 * @return		True if successful.
+	 */
+	bool DeserializeFromData(const TArray<uint8>& DataInput);
 
 	/**
 	 * Sets up the object from the passed in JSON string
 	 * @param JSONInput		The JSON string to deserialize from
 	 * @return		True if successful.
 	 */
-	bool DeserializeFromJSON( const FString& JSONInput );
+	bool DeserializeFromJSON(const FString& JSONInput);
 
 	/**
 	 * Saves out the manifest information
-	 * @param Filename		The file to save a JSON representation to
+	 * @param Filename		The file to save to
+	 * @param bUseBinary	Whether to save out in the new UObject binary format
 	 * @return		True if successful.
 	 */
-	bool SaveToFile( FString Filename );
+	bool SaveToFile(const FString& Filename, bool bUseBinary);
 
 	/**
 	 * Creates the object in JSON format
 	 * @param JSONOutput		A string to receive the JSON representation
 	 */
-	void SerializeToJSON( FString& JSONOutput );
+	void SerializeToJSON(FString& JSONOutput);
 
 	/**
 	 * Gets the version for this manifest. Useful for manifests that were loaded from JSON.
@@ -213,28 +404,28 @@ public:
 	 * @param FileList			IN		The list of files.
 	 * @param RequiredChunks	OUT		The list of chunk GUIDs needed for those files.
 	 */
-	void GetChunksRequiredForFiles( const TArray< FString >& FileList, TArray< FGuid >& RequiredChunks, bool bAddUnique = true ) const;
+	void GetChunksRequiredForFiles(const TArray< FString >& FileList, TArray< FGuid >& RequiredChunks, bool bAddUnique = true) const;
 
 	/**
 	 * Get the number of times a chunks is referenced in this manifest
 	 * @param ChunkGuid		The chunk GUID
 	 * @return	The number of references to this chunk
 	 */
-	const uint32 GetNumberOfChunkReferences( const FGuid& ChunkGuid ) const;
+	const uint32 GetNumberOfChunkReferences(const FGuid& ChunkGuid) const;
 
 	/**
 	 * Returns the size of a particular data file by it's GUID.
 	 * @param DataGuid		The GUID for the data
 	 * @return		File size.
 	 */
-	const int64 GetDataSize( const FGuid& DataGuid ) const;
+	const int64 GetDataSize(const FGuid& DataGuid) const;
 
 	/**
 	 * Returns the total size of all data files in it's list.
 	 * @param DataGuids		The GUID array for the data
 	 * @return		File size.
 	 */
-	const int64 GetDataSize( const TArray< FGuid >& DataGuids ) const;
+	const int64 GetDataSize(const TArray< FGuid >& DataGuids) const;
 
 	/**
 	 * Returns the size of a particular file in the build
@@ -242,7 +433,7 @@ public:
 	 * @param Filename		The file.
 	 * @return		File size.
 	 */
-	const int64 GetFileSize( const FString& Filename ) const;
+	const int64 GetFileSize(const FString& Filename) const;
 
 	/**
 	 * Returns the total size of all files in the array
@@ -250,7 +441,7 @@ public:
 	 * @param Filenames		The array of files.
 	 * @return		Total size of files in array.
 	 */
-	const int64 GetFileSize( const TArray< FString >& Filenames ) const;
+	const int64 GetFileSize(const TArray< FString >& Filenames) const;
 
 	/**
 	 * Returns the number of files in this build.
@@ -262,20 +453,20 @@ public:
 	 * Get the list of files described by this manifest
 	 * @param Filenames		OUT		Receives the array of files.
 	 */
-	void GetFileList( TArray< FString >& Filenames ) const;
+	void GetFileList(TArray< FString >& Filenames) const;
 
 	/**
 	* Get the list of Guids for all files described by this manifest
 	* @param DataGuids		OUT		Receives the array of Guids.
 	*/
-	void GetDataList( TArray< FGuid >& DataGuids ) const;
+	void GetDataList(TArray< FGuid >& DataGuids) const;
 
 	/**
-	 * Returns the manifest for a particular file in the app, invalid if non-existing
+	 * Returns the manifest for a particular file in the app, nullptr if non-existing
 	 * @param Filename	The filename.
 	 * @return	The file manifest, or invalid ptr
 	 */
-	const TSharedPtr< FBuildPatchFileManifest > GetFileManifest( const FString& Filename );
+	const FFileManifestData* GetFileManifest(const FString& Filename);
 
 	/**
 	 * Gets whether this manifest is made up of file data instead of chunk data
@@ -289,7 +480,7 @@ public:
 	 * @param OutHash		OUT		Receives the hash value if found
 	 * @return	true if we had the hash for this chunk
 	 */
-	const bool GetChunkHash( const FGuid& ChunkGuid, uint64& OutHash ) const;
+	const bool GetChunkHash(const FGuid& ChunkGuid, uint64& OutHash) const;
 
 	/**
 	 * Gets the file hash for given file data
@@ -297,7 +488,7 @@ public:
 	 * @param OutHash		OUT		Receives the hash value if found
 	 * @return	true if we had the hash for this file
 	 */
-	const bool GetFileDataHash( const FGuid& FileGuid, FSHAHash& OutHash ) const;
+	const bool GetFileDataHash(const FGuid& FileGuid, FSHAHashData& OutHash) const;
 
 	/**
 	 * Populates an array of chunks that should be producible from this local build, given the list of chunks needed. Also checks source files exist and match size.
@@ -305,7 +496,7 @@ public:
 	 * @param ChunksRequired	IN		A list of chunks that are needed.
 	 * @param ChunksAvailable	OUT		A list to receive the chunks that could be constructed locally.
 	 */
-	void EnumerateProducibleChunks( const FString& InstallDirectory, const TArray< FGuid >& ChunksRequired, TArray< FGuid >& ChunksAvailable ) const;
+	void EnumerateProducibleChunks(const FString& InstallDirectory, const TArray<FGuid>& ChunksRequired, TArray<FGuid>& ChunksAvailable) const;
 
 	/**
 	 * Verifies a local directory structure against a given manifest.
@@ -317,7 +508,7 @@ public:
 	 * @param TimeSpentPaused		OUT		The amount of time we were paused for, in seconds.
 	 * @return		true if no file errors occurred AND the verification was successful
 	 */
-	bool VerifyAgainstDirectory( const FString& VerifyDirectory, TArray < FString >& OutDatedFiles, FBuildPatchFloatDelegate ProgressDelegate, FBuildPatchBoolRetDelegate ShouldPauseDelegate, double& TimeSpentPaused );
+	bool VerifyAgainstDirectory(const FString& VerifyDirectory, TArray < FString >& OutDatedFiles, FBuildPatchFloatDelegate ProgressDelegate, FBuildPatchBoolRetDelegate ShouldPauseDelegate, double& TimeSpentPaused);
 
 	/**
 	 * Gets a list of files that were installed with the Old Manifest, but no longer required by the New Manifest.
@@ -325,7 +516,7 @@ public:
 	 * @param NewManifest		IN		The Build Manifest that is being patched to
 	 * @param RemovableFiles	OUT		A list to receive the files that may be removed.
 	 */
-	static void GetRemovableFiles( FBuildPatchAppManifestRef OldManifest, FBuildPatchAppManifestRef NewManifest, TArray< FString >& RemovableFiles );
+	static void GetRemovableFiles(FBuildPatchAppManifestRef OldManifest, FBuildPatchAppManifestRef NewManifest, TArray< FString >& RemovableFiles);
 
 	/**
 	 * Gets a list of files that have different hash values in the New Manifest, to those in the Old Manifest.
@@ -334,7 +525,7 @@ public:
 	 * @param InstallDirectory	IN		The Build installation directory, so that it can be checked for missing files.
 	 * @param OutDatedFiles		OUT		The array of files that do not match or are new.
 	 */
-	static void GetOutdatedFiles( FBuildPatchAppManifestPtr OldManifest, FBuildPatchAppManifestRef NewManifest, const FString& InstallDirectory, TArray< FString >& OutDatedFiles );
+	static void GetOutdatedFiles(FBuildPatchAppManifestPtr OldManifest, FBuildPatchAppManifestRef NewManifest, const FString& InstallDirectory, TArray< FString >& OutDatedFiles);
 
 	/**
 	 * Check a single file to see if it will be effected by patching
@@ -342,106 +533,47 @@ public:
 	 * @param NewManifest		The Build Manifest that is being patched to. Shared Ref - Implicitly valid.
 	 * @param Filename			The Build installation directory, so that it can be checked for missing files.
 	 */
-	static bool IsFileOutdated( FBuildPatchAppManifestRef OldManifest, FBuildPatchAppManifestRef NewManifest, const FString& Filename );
+	static bool IsFileOutdated(FBuildPatchAppManifestRef OldManifest, FBuildPatchAppManifestRef NewManifest, const FString& Filename);
 
 	/**
 	 * Populates a map of FFileChunkParts referring to the chunks in the given array, that should be accessible locally.
 	 * @param ChunksRequired		IN		A list of chunks that are needed.
 	 * @param ChunkPartsAvailable	OUT		A map to receive the FFileChunkParts keyed by chunk guid.
 	 */
-	void EnumerateChunkPartInventory( const TArray< FGuid >& ChunksRequired, TMap< FGuid, TArray< FFileChunkPart > >& ChunkPartsAvailable ) const;
+	void EnumerateChunkPartInventory(const TArray< FGuid >& ChunksRequired, TMap< FGuid, TArray< FFileChunkPart > >& ChunkPartsAvailable) const;
 
 private:
 
+	bool Serialize(FArchive& Ar);
+
 	/**
 	 * Destroys any memory we have allocated and clears out ready for generation of a new manifest
 	 */
 	void DestroyData();
 
-private:
-
 	/**
-	 * Holds the version support for this manifest.
+	 * Setups the lookup maps that optimize data access, should be called when Data changes
 	 */
-	EBuildPatchAppManifestVersion::Type ManifestFileVersion;
-
-	/**
-	 * Stores whether the manifest was produced from a file data build instead of chunk data
-	 */
-	bool bIsFileData;
+	void InitLookups();
 
-	/**
-	 * Holds the ID of this app.
-	 */
-	uint32 AppID;
-
-	/**
-	 * Holds the name of the app this Build is for
-	 */
-	FString AppNameString;
-
-	/**
-	 * Holds the version string of this Build
-	 */
-	FString BuildVersionString;
-
-	/**
-	 * Holds the app local exe path
-	 */
-	FString LaunchExeString;
-
-	/**
-	 * Holds the name of prerequisites installer
-	 */
-	FString PrereqNameString;
-
-	/**
-	 * Holds the exe path of the prerequisites installer
-	 */
-	FString PrereqPathString;
-
-	/**
-	 * Holds commandline arguments for the prerequisites installer
-	 */
-	FString PrereqArgsString;
-	
-	/**
-	 * Holds the app launch command
-	 */
-	FString LaunchCommandString;
-
-	/**
-	 * References filename by file data GUID
-	 */
-	TMap< FGuid, FString > FileNameReference;
-
-	/**
-	 * Holds the list of files and their generated file manifest
-	 */
-	TMap< FString, TSharedPtr< FBuildPatchFileManifest > > FileManifestList;
-
-	/**
-	 * Holds the list of chunks with their rolling chunk hash, does not apply to file data.
-	 */
-	TMap< FGuid, uint64 > ChunkHashList;
-
-	/**
-	 * Holds the filesize for each chunk, does not apply to file data.
-	 */
-	TMap< FGuid, int64 > ChunkFilesizeList;
+private:
 
 	/**
-	 * Holds which group the chunk (containing chunk or file data) is placed into for file system optimisations
+	 * Holds the actual manifest data. Some other variables point to the memory held by this object
 	 */
-	TMap< FGuid, uint8 > DataGroupList;
+	UBuildPatchManifest* Data;
 
 	/**
-	 * Holds the list of custom fields
+	 * Some lookups to optimize data access
 	 */
-	TMap< FString, FString > CustomFields;
+	TMap<FGuid, FString*> FileNameLookup;
+	TMap<FString, FFileManifestData*> FileManifestLookup;
+	TMap<FGuid, FChunkInfoData*> ChunkInfoLookup;
+	TMap<FString, FCustomFieldData*> CustomFieldLookup;
 
 	/**
 	 * Holds the total build size in bytes
 	 */
 	int64 TotalBuildSize;
+	int64 TotalDownloadSize;
 };
diff --git a/Engine/Source/Runtime/Online/BuildPatchServices/Private/BuildPatchServicesModule.cpp b/Engine/Source/Runtime/Online/BuildPatchServices/Private/BuildPatchServicesModule.cpp
--- a/Engine/Source/Runtime/Online/BuildPatchServices/Private/BuildPatchServicesModule.cpp
+++ b/Engine/Source/Runtime/Online/BuildPatchServices/Private/BuildPatchServicesModule.cpp
@@ -108,6 +108,16 @@ IBuildManifestPtr FBuildPatchServicesModule::LoadManifestFromFile( const FString
 	}
 }
 
+IBuildManifestPtr FBuildPatchServicesModule::MakeManifestFromData(const TArray<uint8>& ManifestData)
+{
+	FBuildPatchAppManifestRef Manifest = MakeShareable(new FBuildPatchAppManifest());
+	if (Manifest->DeserializeFromData(ManifestData))
+	{
+		return Manifest;
+	}
+	return NULL;
+}
+
 IBuildManifestPtr FBuildPatchServicesModule::MakeManifestFromJSON( const FString& ManifestJSON )
 {
 	FBuildPatchAppManifestRef Manifest = MakeShareable( new FBuildPatchAppManifest() );
@@ -118,9 +128,9 @@ IBuildManifestPtr FBuildPatchServicesModule::MakeManifestFromJSON( const FString
 	return NULL;
 }
 
-bool FBuildPatchServicesModule::SaveManifestToFile( const FString& Filename, IBuildManifestRef Manifest )
+bool FBuildPatchServicesModule::SaveManifestToFile(const FString& Filename, IBuildManifestRef Manifest, bool bUseBinary/* = true*/)
 {
-	return StaticCastSharedRef< FBuildPatchAppManifest >( Manifest )->SaveToFile( Filename );
+	return StaticCastSharedRef< FBuildPatchAppManifest >(Manifest)->SaveToFile(Filename, bUseBinary);
 }
 
 IBuildInstallerPtr FBuildPatchServicesModule::StartBuildInstall( IBuildManifestPtr CurrentManifest, IBuildManifestPtr InstallManifest, const FString& InstallDirectory, FBuildPatchBoolManifestDelegate OnCompleteDelegate )
diff --git a/Engine/Source/Runtime/Online/BuildPatchServices/Private/BuildPatchServicesModule.h b/Engine/Source/Runtime/Online/BuildPatchServices/Private/BuildPatchServicesModule.h
--- a/Engine/Source/Runtime/Online/BuildPatchServices/Private/BuildPatchServicesModule.h
+++ b/Engine/Source/Runtime/Online/BuildPatchServices/Private/BuildPatchServicesModule.h
@@ -73,21 +73,20 @@ public:
 
 	// Begin IBuildPatchServicesModule interface
 	virtual IBuildManifestPtr LoadManifestFromFile( const FString& Filename ) override;
-	virtual IBuildManifestPtr MakeManifestFromJSON( const FString& ManifestJSON ) override;
-	virtual bool SaveManifestToFile( const FString& Filename, IBuildManifestRef Manifest ) override;
+	virtual IBuildManifestPtr MakeManifestFromData( const TArray<uint8>& ManifestData ) override;
+	virtual bool SaveManifestToFile( const FString& Filename, IBuildManifestRef Manifest, bool bUseBinary = true ) override;
 	virtual IBuildInstallerPtr StartBuildInstall( IBuildManifestPtr CurrentManifest, IBuildManifestPtr InstallManifest, const FString& InstallDirectory, FBuildPatchBoolManifestDelegate OnCompleteDelegate ) override;
 	virtual void SetStagingDirectory( const FString& StagingDir ) override;
 	virtual void SetCloudDirectory( const FString& CloudDir ) override;
 	virtual void SetBackupDirectory( const FString& BackupDir ) override;
 	virtual void SetAnalyticsProvider( TSharedPtr< IAnalyticsProvider > AnalyticsProvider ) override;
-	virtual void RegisterAppInstallation(IBuildManifestRef AppManifest, const FString AppInstallDirectory) override;
-
+	virtual void RegisterAppInstallation( IBuildManifestRef AppManifest, const FString AppInstallDirectory ) override;
 #if WITH_BUILDPATCHGENERATION
 	virtual bool GenerateChunksManifestFromDirectory( const FBuildPatchSettings& Settings ) override;
 	virtual bool GenerateFilesManifestFromDirectory( const FBuildPatchSettings& Settings ) override;
-	virtual bool CompactifyCloudDirectory(const TArray<FString>& ManifestsToKeep, const float DataAgeThreshold, const bool bPreview) override;
+	virtual bool CompactifyCloudDirectory( const TArray<FString>& ManifestsToKeep, const float DataAgeThreshold, const bool bPreview ) override;
 #endif
-
+	virtual IBuildManifestPtr MakeManifestFromJSON( const FString& ManifestJSON ) override;
 	// End IBuildPatchServicesModule interface
 
 	/**
diff --git a/Engine/Source/Runtime/Online/BuildPatchServices/Private/BuildPatchUtil.cpp b/Engine/Source/Runtime/Online/BuildPatchServices/Private/BuildPatchUtil.cpp
--- a/Engine/Source/Runtime/Online/BuildPatchServices/Private/BuildPatchUtil.cpp
+++ b/Engine/Source/Runtime/Online/BuildPatchServices/Private/BuildPatchUtil.cpp
@@ -14,7 +14,7 @@ const FString FBuildPatchUtils::GetChunkNewFilename( const EBuildPatchAppManifes
 	return FPaths::Combine( *RootDirectory, *FString::Printf( TEXT("%s/%02d/%016llX_%s.chunk"), *EBuildPatchAppManifestVersion::GetChunkSubdir( ManifestVersion ),FCrc::MemCrc32( &ChunkGUID, sizeof( FGuid ) ) % 100, ChunkHash, *ChunkGUID.ToString() ) );
 }
 
-const FString FBuildPatchUtils::GetFileNewFilename( const EBuildPatchAppManifestVersion::Type ManifestVersion, const FString& RootDirectory, const FGuid& FileGUID, const FSHAHash& FileHash )
+const FString FBuildPatchUtils::GetFileNewFilename(const EBuildPatchAppManifestVersion::Type ManifestVersion, const FString& RootDirectory, const FGuid& FileGUID, const FSHAHashData& FileHash)
 {
 	check( FileGUID.IsValid() );
 	return FPaths::Combine( *RootDirectory, *FString::Printf( TEXT("%s/%02d/%s_%s.file"), *EBuildPatchAppManifestVersion::GetFileSubdir( ManifestVersion ), FCrc::MemCrc32( &FileGUID, sizeof( FGuid ) ) % 100, *FileHash.ToString(), *FileGUID.ToString() ) );
@@ -33,7 +33,7 @@ void FBuildPatchUtils::GetChunkDetailFromNewFilename( const FString& ChunkNewFil
 	FGuid::Parse( GuidString, ChunkGUID );
 }
 
-void FBuildPatchUtils::GetFileDetailFromNewFilename( const FString& FileNewFilename, FGuid& FileGUID, FSHAHash& FileHash )
+void FBuildPatchUtils::GetFileDetailFromNewFilename(const FString& FileNewFilename, FGuid& FileGUID, FSHAHashData& FileHash)
 {
 	const FString FileFilename = FPaths::GetBaseFilename( FileNewFilename );
 	FString GuidString;
@@ -91,7 +91,7 @@ bool FBuildPatchUtils::GetGUIDFromFilename( const FString& DataFilename, FGuid&
 	return FGuid::Parse( GuidString, DataGUID );
 }
 
-uint8 FBuildPatchUtils::VerifyFile(const FString& FileToVerify, const FSHAHash& Hash1, const FSHAHash& Hash2)
+uint8 FBuildPatchUtils::VerifyFile(const FString& FileToVerify, const FSHAHashData& Hash1, const FSHAHashData& Hash2)
 {
 	FBuildPatchFloatDelegate NoProgressDelegate;
 	FBuildPatchBoolRetDelegate NoPauseDelegate;
@@ -99,7 +99,7 @@ uint8 FBuildPatchUtils::VerifyFile(const FString& FileToVerify, const FSHAHash&
 	return VerifyFile(FileToVerify, Hash1, Hash2, NoProgressDelegate, NoPauseDelegate, NoDouble);
 }
 
-uint8 FBuildPatchUtils::VerifyFile(const FString& FileToVerify, const FSHAHash& Hash1, const FSHAHash& Hash2, FBuildPatchFloatDelegate ProgressDelegate, FBuildPatchBoolRetDelegate ShouldPauseDelegate, double& TimeSpentPaused)
+uint8 FBuildPatchUtils::VerifyFile(const FString& FileToVerify, const FSHAHashData& Hash1, const FSHAHashData& Hash2, FBuildPatchFloatDelegate ProgressDelegate, FBuildPatchBoolRetDelegate ShouldPauseDelegate, double& TimeSpentPaused)
 {
 	uint8 ReturnValue = 0;
 	FArchive* FileReader = IFileManager::Get().CreateFileReader(*FileToVerify);
@@ -107,7 +107,7 @@ uint8 FBuildPatchUtils::VerifyFile(const FString& FileToVerify, const FSHAHash&
 	if (FileReader != NULL)
 	{
 		FSHA1 HashState;
-		FSHAHash HashValue;
+		FSHAHashData HashValue;
 		const int64 FileSize = FileReader->TotalSize();
 		uint8* FileReadBuffer = new uint8[FileBufferSize];
 		while (!FileReader->AtEnd() && !FBuildPatchInstallError::HasFatalError())
@@ -187,7 +187,7 @@ bool FBuildPatchUtils::VerifyChunkFile( FArchive& ChunkFileData, bool bQuickChec
 		{
 			// Hashes for checking data
 			FSHA1 SHAHasher;
-			FSHAHash SHAHash;
+			FSHAHashData SHAHash;
 			// Load the data to check
 			uint8* FileReadBuffer = new uint8[ FileBufferSize ];
 			int64 DataOffset = 0;
diff --git a/Engine/Source/Runtime/Online/BuildPatchServices/Private/BuildPatchUtil.h b/Engine/Source/Runtime/Online/BuildPatchServices/Private/BuildPatchUtil.h
--- a/Engine/Source/Runtime/Online/BuildPatchServices/Private/BuildPatchUtil.h
+++ b/Engine/Source/Runtime/Online/BuildPatchServices/Private/BuildPatchUtil.h
@@ -19,7 +19,7 @@ struct FBuildPatchUtils
 	 * @param ChunkHash			The chunk rolling hash value
 	 * @return	the chunk path
 	 */
-	static const FString GetChunkNewFilename( const EBuildPatchAppManifestVersion::Type ManifestVersion, const FString& RootDirectory, const FGuid& ChunkGUID, const uint64& ChunkHash );
+	static const FString GetChunkNewFilename(const EBuildPatchAppManifestVersion::Type ManifestVersion, const FString& RootDirectory, const FGuid& ChunkGUID, const uint64& ChunkHash);
 
 	/**
 	 * Gets the filename for a file chunk generated from it's GUID and Hash, which is the new format.
@@ -29,7 +29,7 @@ struct FBuildPatchUtils
 	 * @param FileHash			The file hash value
 	 * @return	The file chunk path
 	 */
-	static const FString GetFileNewFilename( const EBuildPatchAppManifestVersion::Type ManifestVersion, const FString& RootDirectory, const FGuid& FileGUID, const FSHAHash& FileHash );
+	static const FString GetFileNewFilename(const EBuildPatchAppManifestVersion::Type ManifestVersion, const FString& RootDirectory, const FGuid& FileGUID, const FSHAHashData& FileHash);
 
 	/**
 	 * Gets the chunk GUID and Hash, from the filename, which is the new format.
@@ -37,7 +37,7 @@ struct FBuildPatchUtils
 	 * @param ChunkGuid			OUT	The chunk Guid
 	 * @param ChunkHash			OUT	The chunk rolling hash value
 	 */
-	static void GetChunkDetailFromNewFilename( const FString& ChunkNewFilename, FGuid& ChunkGUID, uint64& ChunkHash );
+	static void GetChunkDetailFromNewFilename(const FString& ChunkNewFilename, FGuid& ChunkGUID, uint64& ChunkHash);
 
 	/**
 	 * Gets the file chunk GUID and file hash, from the filename, which is the new format.
@@ -45,7 +45,7 @@ struct FBuildPatchUtils
 	 * @param FileGUID			OUT	The file chunk Guid
 	 * @param FileHash			OUT	The file chunk rolling hash value
 	 */
-	static void GetFileDetailFromNewFilename( const FString& FileNewFilename, FGuid& FileGUID, FSHAHash& FileHash );
+	static void GetFileDetailFromNewFilename(const FString& FileNewFilename, FGuid& FileGUID, FSHAHashData& FileHash);
 
 	/**
 	 * Gets the filename for a chunk generated from it's GUID
@@ -53,7 +53,7 @@ struct FBuildPatchUtils
 	 * @param ChunkGuid			The chunk Guid
 	 * @return	the chunk path
 	 */
-	static const FString GetChunkOldFilename( const FString& RootDirectory, const FGuid& ChunkGUID );
+	static const FString GetChunkOldFilename(const FString& RootDirectory, const FGuid& ChunkGUID);
 
 	/**
 	 * Gets the filename for a file data part generated from it's GUID
@@ -61,7 +61,7 @@ struct FBuildPatchUtils
 	 * @param FileGUID			The file part Guid
 	 * @return	The file data path
 	 */
-	static const FString GetFileOldFilename( const FString& RootDirectory, const FGuid& FileGUID );
+	static const FString GetFileOldFilename(const FString& RootDirectory, const FGuid& FileGUID);
 
 	/**
 	 * Gets the filename for a specific data part type from it's GUID
@@ -70,7 +70,7 @@ struct FBuildPatchUtils
 	 * @param DataGUID			The data Guid
 	 * @return	the data part path
 	 */
-	static const FString GetDataTypeOldFilename( const FBuildPatchData::Type DataType, const FString& RootDirectory, const FGuid& Guid );
+	static const FString GetDataTypeOldFilename(const FBuildPatchData::Type DataType, const FString& RootDirectory, const FGuid& Guid);
 
 	/**
 	 * Gets the GUID for a data file according to it's filename (new or old)
@@ -78,14 +78,14 @@ struct FBuildPatchUtils
 	 * @param DataGUID			OUT		Receives the GUID of the data
 	 * @return  True if successful, false otherwise
 	 */
-	static bool GetGUIDFromFilename( const FString& DataFilename, FGuid& DataGUID );
+	static bool GetGUIDFromFilename(const FString& DataFilename, FGuid& DataGUID);
 
 	/**
 	 * Helper function to uncompress chunk file data. Can be called without knowing if needed and process will be just skipped.
 	 * @param ChunkFileArray	IN OUT		The data array, should contain full chunk - header plus data. Will be overwritten with uncompressed chunk version.
 	 * @return		true if no errors occurred and the data is not corrupted
 	 */
-	static bool UncompressChunkFile( TArray< uint8 >& ChunkFileArray );
+	static bool UncompressChunkFile(TArray<uint8>& ChunkFileArray);
 
 	/**
 	 * Checks a file against SHA1 hashes. The function takes two so that it can return no match, match with Hash1, or match with Hash2, that way we can check the file for being the same as an old manifest or new manifest
@@ -98,7 +98,7 @@ struct FBuildPatchUtils
 	 * @param TimeSpentPaused		OUT		The amount of time we were paused for, in seconds.
 	 * @return		0 if no match, 1 for match with Hash1, and 2 for match with Hash2
 	 */
-	static uint8 VerifyFile(const FString& FileToVerify, const FSHAHash& Hash1, const FSHAHash& Hash2, FBuildPatchFloatDelegate ProgressDelegate, FBuildPatchBoolRetDelegate ShouldPauseDelegate, double& TimeSpentPaused);
+	static uint8 VerifyFile(const FString& FileToVerify, const FSHAHashData& Hash1, const FSHAHashData& Hash2, FBuildPatchFloatDelegate ProgressDelegate, FBuildPatchBoolRetDelegate ShouldPauseDelegate, double& TimeSpentPaused);
 
 	/**
 	 * Checks a file against SHA1 hashes. The function takes two so that it can return no match, match with Hash1, or match with Hash2, that way we can check the file for being the same as an old manifest or new manifest
@@ -108,7 +108,7 @@ struct FBuildPatchUtils
 	 * @param Hash2					IN		A second Hash to match against the file
 	 * @return		0 if no match, 1 for match with Hash1, and 2 for match with Hash2
 	 */
-	static uint8 VerifyFile(const FString& FileToVerify, const FSHAHash& Hash1, const FSHAHash& Hash2);
+	static uint8 VerifyFile(const FString& FileToVerify, const FSHAHashData& Hash1, const FSHAHashData& Hash2);
 
 	/**
 	 * Helper function to verify the integrity of a chunk file that is in memory
@@ -116,7 +116,7 @@ struct FBuildPatchUtils
 	 * @param bQuickCheck		If true, will only check header and file size, not data hash.
 	 * @return		true if no errors occurred and the data is not corrupted
 	 */
-	static bool VerifyChunkFile( const TArray< uint8 >& ChunkFileArray, bool bQuickCheck = false );
+	static bool VerifyChunkFile(const TArray<uint8>& ChunkFileArray, bool bQuickCheck = false);
 
 private:
 
@@ -126,5 +126,5 @@ private:
 	 * @param bQuickCheck		If true, will only check header and file size, not data hash.
 	 * @return		true if no errors occurred and the data is not corrupted
 	 */
-	static bool VerifyChunkFile( FArchive& ChunkFileData, bool bQuickCheck = false );
+	static bool VerifyChunkFile(FArchive& ChunkFileData, bool bQuickCheck = false);
 };
diff --git a/Engine/Source/Runtime/Online/BuildPatchServices/Public/BuildPatchSettings.h b/Engine/Source/Runtime/Online/BuildPatchServices/Public/BuildPatchSettings.h
--- a/Engine/Source/Runtime/Online/BuildPatchServices/Public/BuildPatchSettings.h
+++ b/Engine/Source/Runtime/Online/BuildPatchServices/Public/BuildPatchSettings.h
@@ -9,7 +9,7 @@ struct FBuildPatchSettings
 	// The directory to analyze
 	FString RootDirectory;
 	// The ID of the app of this build
-	uint32 InAppID;
+	uint32 AppID;
 	// The name of the app of this build
 	FString AppName;
 	// The version string for this build
diff --git a/Engine/Source/Runtime/Online/BuildPatchServices/Public/Interfaces/IBuildManifest.h b/Engine/Source/Runtime/Online/BuildPatchServices/Public/Interfaces/IBuildManifest.h
--- a/Engine/Source/Runtime/Online/BuildPatchServices/Public/Interfaces/IBuildManifest.h
+++ b/Engine/Source/Runtime/Online/BuildPatchServices/Public/Interfaces/IBuildManifest.h
@@ -130,6 +130,12 @@ public:
 	 * @param RemovableFiles	OUT		A list to receive the files that may be removed.
 	 */
 	virtual void GetRemovableFiles(const TCHAR* InstallPath, TArray< FString >& RemovableFiles) const = 0;
+
+	/**
+	 * Checks the manifest format version to see if this manifest was loaded from latest data
+	 * @return True if the manifest was created from the latest format
+	 */
+	virtual bool NeedsResaving() const = 0;
 	
 	/**
 	 * Duplicated this manifest to create a copy. Should be used if storing a received manifest as an installed
diff --git a/Engine/Source/Runtime/Online/BuildPatchServices/Public/Interfaces/IBuildPatchServicesModule.h b/Engine/Source/Runtime/Online/BuildPatchServices/Public/Interfaces/IBuildPatchServicesModule.h
--- a/Engine/Source/Runtime/Online/BuildPatchServices/Public/Interfaces/IBuildPatchServicesModule.h
+++ b/Engine/Source/Runtime/Online/BuildPatchServices/Public/Interfaces/IBuildPatchServicesModule.h
@@ -36,19 +36,20 @@ public:
 	virtual IBuildManifestPtr LoadManifestFromFile( const FString& Filename ) = 0;
 	
 	/**
-	 * Constructs a Build Manifest from a JSON string
-	 * @param ManifestJSON		The JSON received from a web api
+	 * Constructs a Build Manifest from a data
+	 * @param ManifestData		The data received from a web api
 	 * @return		a shared pointer to the manifest, which will be invalid if creation failed.
 	 */
-	virtual IBuildManifestPtr MakeManifestFromJSON( const FString& ManifestJSON ) = 0;
+	virtual IBuildManifestPtr MakeManifestFromData( const TArray<uint8>& ManifestData ) = 0;
 	
 	/**
 	 * Saves a Build Manifest to file
 	 * @param Filename		The file to save to
 	 * @param Manifest		The manifest to save out
+	 * @param bUseBinary	Whether to save binary format instead of json
 	 * @return		If the save was successful.
 	 */
-	virtual bool SaveManifestToFile( const FString& Filename, IBuildManifestRef Manifest ) = 0;
+	virtual bool SaveManifestToFile(const FString& Filename, IBuildManifestRef Manifest, bool bUseBinary = true) = 0;
 
 	/**
 	 * Starts an installer thread for the provided manifests
@@ -118,4 +119,9 @@ public:
 	 */
 	virtual bool CompactifyCloudDirectory(const TArray<FString>& ManifestsToKeep, const float DataAgeThreshold, const bool bPreview) = 0;
 #endif // WITH_BUILDPATCHGENERATION
+
+	/**
+	 * Deprecated function, use MakeManifestFromData instead
+	 */
+	virtual IBuildManifestPtr MakeManifestFromJSON(const FString& ManifestJSON) = 0;
 };
