Quick fix remove test code

diff --git a/Engine/Source/Runtime/Online/BuildPatchServices/Private/BuildPatchManifest.cpp b/Engine/Source/Runtime/Online/BuildPatchServices/Private/BuildPatchManifest.cpp
--- a/Engine/Source/Runtime/Online/BuildPatchServices/Private/BuildPatchManifest.cpp
+++ b/Engine/Source/Runtime/Online/BuildPatchServices/Private/BuildPatchManifest.cpp
@@ -1468,7 +1468,6 @@ void FBuildPatchAppManifest::GetRemovableFiles(const TCHAR* InstallPath, TArray<
 
 bool FBuildPatchAppManifest::NeedsResaving() const
 {
-	return true;
 	return Data->ManifestFileVersion < EBuildPatchAppManifestVersion::GetLatestVersion();
 }
 
