IOS and CURL : http cancel request now calls completed delegates, goes through ::FinishedRequest

diff --git a/Engine/Source/Runtime/Online/HTTP/Private/Curl/CurlHttp.cpp b/Engine/Source/Runtime/Online/HTTP/Private/Curl/CurlHttp.cpp
--- a/Engine/Source/Runtime/Online/HTTP/Private/Curl/CurlHttp.cpp
+++ b/Engine/Source/Runtime/Online/HTTP/Private/Curl/CurlHttp.cpp
@@ -621,6 +621,9 @@ FHttpRequestProgressDelegate& FCurlHttpRequest::OnRequestProgress()
 void FCurlHttpRequest::CancelRequest()
 {
 	bCanceled = true;
+    
+    // Cleanup cancel request and fire off any necessary delegates.
+    FinishedRequest();
 }
 
 EHttpRequestStatus::Type FCurlHttpRequest::GetStatus()
diff --git a/Engine/Source/Runtime/Online/HTTP/Private/IOS/IOSHTTP.cpp b/Engine/Source/Runtime/Online/HTTP/Private/IOS/IOSHTTP.cpp
--- a/Engine/Source/Runtime/Online/HTTP/Private/IOS/IOSHTTP.cpp
+++ b/Engine/Source/Runtime/Online/HTTP/Private/IOS/IOSHTTP.cpp
@@ -291,6 +291,9 @@ void FIOSHttpRequest::CancelRequest()
 		[Connection cancel];
 		Connection = NULL;
 	}
+    
+    // Cleanup cancel request and fire off any necessary delegates.
+    FinishedRequest();
 }
 
 
