IOS mem management for request connections bug fix.

diff --git a/Engine/Source/Runtime/Online/HTTP/Private/IOS/IOSHTTP.cpp b/Engine/Source/Runtime/Online/HTTP/Private/IOS/IOSHTTP.cpp
--- a/Engine/Source/Runtime/Online/HTTP/Private/IOS/IOSHTTP.cpp
+++ b/Engine/Source/Runtime/Online/HTTP/Private/IOS/IOSHTTP.cpp
@@ -286,10 +286,11 @@ void FIOSHttpRequest::CleanupRequest()
 void FIOSHttpRequest::CancelRequest() 
 {
 	UE_LOG(LogHttp, Verbose, TEXT("FIOSHttpRequest::CancelRequest()"));
-	if(Connection != NULL)
+	if(Connection != nil)
 	{
 		[Connection cancel];
-		Connection = NULL;
+        [Connection release];
+		Connection = nil;
 	}
     
     // Cleanup cancel request and fire off any necessary delegates.
@@ -401,6 +402,11 @@ FIOSHttpResponse::FIOSHttpResponse(const FIOSHttpRequest& InRequest)
 FIOSHttpResponse::~FIOSHttpResponse()
 {
 	UE_LOG(LogHttp, Verbose, TEXT("FIOSHttpResponse::~FIOSHttpResponse()"));
+    if( ResponseWrapper != nil )
+    {
+        [ResponseWrapper release];
+        ResponseWrapper = nil;
+    }
 }
 
 
