Fixes for manifest resave and missing include

diff --git a/Engine/Source/Runtime/Online/BuildPatchServices/Private/BuildPatchManifest.cpp b/Engine/Source/Runtime/Online/BuildPatchServices/Private/BuildPatchManifest.cpp
--- a/Engine/Source/Runtime/Online/BuildPatchServices/Private/BuildPatchManifest.cpp
+++ b/Engine/Source/Runtime/Online/BuildPatchServices/Private/BuildPatchManifest.cpp
@@ -484,6 +484,7 @@ bool FBuildPatchAppManifest::SaveToFile(const FString& Filename, bool bUseBinary
 	{
 		if (bUseBinary)
 		{
+			Data->ManifestFileVersion = EBuildPatchAppManifestVersion::GetLatestVersion();
 			FManifestWriter ManifestData;
 			Serialize(ManifestData);
 			ManifestData.Finalize();
@@ -520,6 +521,7 @@ bool FBuildPatchAppManifest::SaveToFile(const FString& Filename, bool bUseBinary
 		}
 		else
 		{
+			Data->ManifestFileVersion = EBuildPatchAppManifestVersion::GetLatestJsonVersion();
 			FString JSONOutput;
 			SerializeToJSON(JSONOutput);
 			FTCHARToUTF8 JsonUTF8(*JSONOutput);
