Avoid out of bounds crash due to functions not checking against the correct max value for local talkers

#github Pull request #495 https://github.com/EpicGames/UnrealEngine/pull/495 by alexSilkin

diff --git a/Engine/Source/Runtime/Online/OnlineSubsystemUtils/Private/VoiceInterfaceImpl.cpp b/Engine/Source/Runtime/Online/OnlineSubsystemUtils/Private/VoiceInterfaceImpl.cpp
--- a/Engine/Source/Runtime/Online/OnlineSubsystemUtils/Private/VoiceInterfaceImpl.cpp
+++ b/Engine/Source/Runtime/Online/OnlineSubsystemUtils/Private/VoiceInterfaceImpl.cpp
@@ -111,7 +111,7 @@ void FOnlineVoiceImpl::Tick(float DeltaTime)
 void FOnlineVoiceImpl::StartNetworkedVoice(uint8 LocalUserNum)
 {
 	// Validate the range of the entry
-	if (LocalUserNum >= 0 && LocalUserNum < MAX_LOCAL_PLAYERS)
+	if (LocalUserNum >= 0 && LocalUserNum < MaxLocalTalkers)
 	{
 		LocalTalkers[LocalUserNum].bHasNetworkedVoice = true;
 		UE_LOG(LogVoice, Log, TEXT("Starting networked voice for user: %d"), LocalUserNum);
@@ -126,7 +126,7 @@ void FOnlineVoiceImpl::StartNetworkedVoice(uint8 LocalUserNum)
 void FOnlineVoiceImpl::StopNetworkedVoice(uint8 LocalUserNum)
 {
 	// Validate the range of the entry
-	if (LocalUserNum >= 0 && LocalUserNum < MAX_LOCAL_PLAYERS)
+	if (LocalUserNum >= 0 && LocalUserNum < MaxLocalTalkers)
 	{
 		LocalTalkers[LocalUserNum].bHasNetworkedVoice = false;
 		UE_LOG(LogVoice, Log, TEXT("Stopping networked voice for user: %d"), LocalUserNum);
@@ -141,7 +141,7 @@ void FOnlineVoiceImpl::StopNetworkedVoice(uint8 LocalUserNum)
 bool FOnlineVoiceImpl::RegisterLocalTalker(uint32 LocalUserNum) 
 {
 	uint32 Return = E_FAIL;
-	if (LocalUserNum >= 0 && LocalUserNum < MAX_LOCAL_PLAYERS)
+	if (LocalUserNum >= 0 && LocalUserNum < (uint32)MaxLocalTalkers)
 	{
 		// Get at the local talker's cached data
 		FLocalTalker& Talker = LocalTalkers[LocalUserNum];
@@ -188,7 +188,7 @@ void FOnlineVoiceImpl::RegisterLocalTalkers()
 {
 	UE_LOG(LogVoice, Log, TEXT("Registering all local talkers"));
 	// Loop through the 4 available players and register them
-	for (uint32 Index = 0; Index < MAX_LOCAL_PLAYERS; Index++)
+	for (uint32 Index = 0; Index < (uint32)MaxLocalTalkers ; Index++)
 	{
 		// Register the local player as a local talker
 		RegisterLocalTalker(Index);
@@ -198,7 +198,7 @@ void FOnlineVoiceImpl::RegisterLocalTalkers()
 bool FOnlineVoiceImpl::UnregisterLocalTalker(uint32 LocalUserNum) 
 {
 	uint32 Return = S_OK;
-	if (LocalUserNum >= 0 && LocalUserNum < MAX_LOCAL_PLAYERS)
+	if (LocalUserNum >= 0 && LocalUserNum < (uint32)MaxLocalTalkers)
 	{
 		// Get at the local talker's cached data
 		FLocalTalker& Talker = LocalTalkers[LocalUserNum];
@@ -240,7 +240,7 @@ void FOnlineVoiceImpl::UnregisterLocalTalkers()
 {
 	UE_LOG(LogVoice, Log, TEXT("Unregistering all local talkers"));
 	// Loop through the 4 available players and unregister them
-	for (uint32 Index = 0; Index < MAX_LOCAL_PLAYERS; Index++)
+	for (uint32 Index = 0; Index < (uint32)MaxLocalTalkers; Index++)
 	{
 		// Unregister the local player as a local talker
 		UnregisterLocalTalker(Index);
@@ -387,7 +387,7 @@ bool FOnlineVoiceImpl::IsRemotePlayerTalking(const FUniqueNetId& UniqueId)
 bool FOnlineVoiceImpl::IsMuted(uint32 LocalUserNum, const FUniqueNetId& UniqueId) const
 {
 	int32 Index = INDEX_NONE;
-	if (LocalUserNum >= 0 && LocalUserNum < MAX_LOCAL_PLAYERS)
+	if (LocalUserNum >= 0 && LocalUserNum < (uint32)MaxLocalTalkers)
 	{
 		Index = SystemMuteList.Find((const FUniqueNetIdString&)UniqueId);
 	}
@@ -404,7 +404,7 @@ bool FOnlineVoiceImpl::IsLocallyMuted(const FUniqueNetId& UniqueId) const
 bool FOnlineVoiceImpl::MuteRemoteTalker(uint8 LocalUserNum, const FUniqueNetId& PlayerId, bool bIsSystemWide)
 {
 	uint32 Return = E_FAIL;
-	if (LocalUserNum >= 0 && LocalUserNum < MAX_LOCAL_PLAYERS)
+	if (LocalUserNum >= 0 && LocalUserNum < MaxLocalTalkers )
 	{
 		if (bIsSystemWide)
 		{
@@ -444,7 +444,7 @@ bool FOnlineVoiceImpl::MuteRemoteTalker(uint8 LocalUserNum, const FUniqueNetId&
 bool FOnlineVoiceImpl::UnmuteRemoteTalker(uint8 LocalUserNum, const FUniqueNetId& PlayerId, bool bIsSystemWide)
 {
 	uint32 Return = E_FAIL;
-	if (LocalUserNum >= 0 && LocalUserNum < MAX_LOCAL_PLAYERS)
+	if (LocalUserNum >= 0 && LocalUserNum < MaxLocalTalkers )
 	{
 		if (bIsSystemWide)
 		{
