Fix for progress information when installing via a binary manifest

diff --git a/Engine/Source/Runtime/Online/BuildPatchServices/Private/BuildPatchManifest.cpp b/Engine/Source/Runtime/Online/BuildPatchServices/Private/BuildPatchManifest.cpp
--- a/Engine/Source/Runtime/Online/BuildPatchServices/Private/BuildPatchManifest.cpp
+++ b/Engine/Source/Runtime/Online/BuildPatchServices/Private/BuildPatchManifest.cpp
@@ -645,10 +645,8 @@ void FBuildPatchAppManifest::InitLookups()
 	FileNameLookup.Empty(Data->bIsFileData ? Data->FileManifestList.Num() : 0);
 	for (auto& File : Data->FileManifestList)
 	{
-		for (auto& FilePart : File.FileChunkParts)
-		{
-			TotalBuildSize += FilePart.Size;
-		}
+		File.Init();
+		TotalBuildSize += File.GetFileSize();
 		FileManifestLookup.Add(File.Filename, &File);
 		if (Data->bIsFileData)
 		{
