Fix some issues in the null online subsystem:
- Null OSS sessions now attempt to report the actual port used by the net driver instead of always using the default port from the config or the command line. Fixes multiple instances when one instance has already bound the default port.
- FOnlineSessionNull::NeedsToAdvertise will now correctly return true if the local instance is hosting the session, even if the engine isn't yet in listen server or dedicated server mode.

diff --git a/Engine/Source/Runtime/Online/OnlineSubsystemNull/Private/OnlineSessionInterfaceNull.cpp b/Engine/Source/Runtime/Online/OnlineSubsystemNull/Private/OnlineSessionInterfaceNull.cpp
--- a/Engine/Source/Runtime/Online/OnlineSubsystemNull/Private/OnlineSessionInterfaceNull.cpp
+++ b/Engine/Source/Runtime/Online/OnlineSubsystemNull/Private/OnlineSessionInterfaceNull.cpp
@@ -4,6 +4,7 @@
 #include "OnlineSessionInterfaceNull.h"
 #include "OnlineIdentityInterface.h"
 #include "OnlineSubsystemNull.h"
+#include "OnlineSubsystemUtils.h"
 #include "OnlineAsyncTaskManagerNull.h"
 #include "SocketSubsystem.h"
 #include "LANBeacon.h"
@@ -18,13 +19,13 @@ FOnlineSessionInfoNull::FOnlineSessionInfoNull() :
 {
 }
 
-void FOnlineSessionInfoNull::Init()
+void FOnlineSessionInfoNull::Init(const FOnlineSubsystemNull& Subsystem)
 {
 	// Read the IP from the system
 	bool bCanBindAll;
 	HostAddr = ISocketSubsystem::Get()->GetLocalHostAddr(*GLog, bCanBindAll);
 	// Now set the port that was configured
-	HostAddr->SetPort(FURL::UrlConfig.DefaultPort);
+	HostAddr->SetPort(GetPortFromNetDriver(Subsystem.GetInstanceName()));
 
 	FGuid OwnerGuid;
 	FPlatformMisc::CreateGuid(OwnerGuid);
@@ -201,7 +202,7 @@ bool FOnlineSessionNull::CreateSession(int32 HostingPlayerNum, FName SessionName
 
 		// Setup the host session info
 		FOnlineSessionInfoNull* NewSessionInfo = new FOnlineSessionInfoNull();
-		NewSessionInfo->Init();
+		NewSessionInfo->Init(*NullSubsystem);
 		Session->SessionInfo = MakeShareable(NewSessionInfo);
 
 		Result = UpdateLANStatus();
@@ -266,7 +267,7 @@ bool FOnlineSessionNull::NeedsToAdvertise( FNamedOnlineSession& Session )
 	// b) Not started public LAN session (same as usually)
 	// d) Joinable presence-enabled session that would be advertised with in an online service
 	// (all of that only if we're server)
-	return Session.SessionSettings.bShouldAdvertise && NullSubsystem->IsServer() &&
+	return Session.SessionSettings.bShouldAdvertise && IsHost(Session) &&
 		(
 			(
 			  Session.SessionSettings.bIsLANMatch && 			  
@@ -943,6 +944,9 @@ void FOnlineSessionNull::AppendSessionToPacket(FNboSerializeToBufferNull& Packet
 		<< Session->NumOpenPrivateConnections
 		<< Session->NumOpenPublicConnections;
 
+	// Try to get the actual port the netdriver is using
+	SetPortFromNetDriver(*NullSubsystem, Session->SessionInfo);
+
 	// Write host info (host addr, session id, and key)
 	Packet << *StaticCastSharedPtr<FOnlineSessionInfoNull>(Session->SessionInfo);
 
@@ -1210,3 +1214,30 @@ void FOnlineSessionNull::UnregisterLocalPlayer(const FUniqueNetId& PlayerId, FNa
 {
 	Delegate.ExecuteIfBound(PlayerId, true);
 }
+
+void FOnlineSessionNull::SetPortFromNetDriver(const FOnlineSubsystemNull& Subsystem, const TSharedPtr<FOnlineSessionInfo>& SessionInfo)
+{
+	auto NetDriverPort = GetPortFromNetDriver(Subsystem.GetInstanceName());
+	auto SessionInfoNull = StaticCastSharedPtr<FOnlineSessionInfoNull>(SessionInfo);
+	if (SessionInfoNull.IsValid() && SessionInfoNull->HostAddr.IsValid())
+	{
+		SessionInfoNull->HostAddr->SetPort(NetDriverPort);
+	}
+}
+
+bool FOnlineSessionNull::IsHost(const FNamedOnlineSession& Session) const
+{
+	if (NullSubsystem->IsDedicated())
+	{
+		return true;
+	}
+
+	IOnlineIdentityPtr IdentityInt = NullSubsystem->GetIdentityInterface(); 
+	if (!IdentityInt.IsValid())
+	{
+		return false;
+	}
+
+	TSharedPtr<FUniqueNetId> UserId = IdentityInt->GetUniquePlayerId(Session.HostingPlayerNum);
+	return (UserId.IsValid() && (*UserId == *Session.OwningUserId));
+}
diff --git a/Engine/Source/Runtime/Online/OnlineSubsystemNull/Private/OnlineSessionInterfaceNull.h b/Engine/Source/Runtime/Online/OnlineSubsystemNull/Private/OnlineSessionInterfaceNull.h
--- a/Engine/Source/Runtime/Online/OnlineSubsystemNull/Private/OnlineSessionInterfaceNull.h
+++ b/Engine/Source/Runtime/Online/OnlineSubsystemNull/Private/OnlineSessionInterfaceNull.h
@@ -139,6 +139,16 @@ private:
 	 */
 	void OnLANSearchTimeout();
 
+	/**
+	 * Attempt to set the host port in the session info based on the actual port the netdriver is using.
+	 */
+	static void SetPortFromNetDriver(const FOnlineSubsystemNull& Subsystem, const TSharedPtr<FOnlineSessionInfo>& SessionInfo);
+
+	/**
+	 * Returns true if the session owner is also the host.
+	 */
+	bool IsHost(const FNamedOnlineSession& Session) const;
+
 PACKAGE_SCOPE:
 
 	/** Critical sections for thread safe operation of session lists */
diff --git a/Engine/Source/Runtime/Online/OnlineSubsystemNull/Private/OnlineSubsystemNullTypes.h b/Engine/Source/Runtime/Online/OnlineSubsystemNull/Private/OnlineSubsystemNullTypes.h
--- a/Engine/Source/Runtime/Online/OnlineSubsystemNull/Private/OnlineSubsystemNullTypes.h
+++ b/Engine/Source/Runtime/Online/OnlineSubsystemNull/Private/OnlineSubsystemNullTypes.h
@@ -33,7 +33,7 @@ PACKAGE_SCOPE:
 	 * Initialize a Null session info with the address of this machine
 	 * and an id for the session
 	 */
-	void Init();
+	void Init(const FOnlineSubsystemNull& Subsystem);
 
 	/** The ip & port that the host is listening on (valid for LAN/GameServer) */
 	TSharedPtr<class FInternetAddr> HostAddr;
diff --git a/Engine/Source/Runtime/Online/OnlineSubsystemUtils/Private/OnlineSubsystemUtils.cpp b/Engine/Source/Runtime/Online/OnlineSubsystemUtils/Private/OnlineSubsystemUtils.cpp
--- a/Engine/Source/Runtime/Online/OnlineSubsystemUtils/Private/OnlineSubsystemUtils.cpp
+++ b/Engine/Source/Runtime/Online/OnlineSubsystemUtils/Private/OnlineSubsystemUtils.cpp
@@ -95,6 +95,32 @@ UWorld* GetWorldForOnline(FName InstanceName)
 	return World;
 }
 
+int32 GetPortFromNetDriver(FName InstanceName)
+{
+	int32 Port = 0;
+#if WITH_ENGINE
+	if (GEngine)
+	{
+		UWorld* World = GetWorldForOnline(InstanceName);
+		UNetDriver* NetDriver = World ? GEngine->FindNamedNetDriver(World, NAME_GameNetDriver) : NULL;
+		if (NetDriver && NetDriver->GetNetMode() < NM_Client)
+		{
+			FString AddressStr = NetDriver->LowLevelGetNetworkNumber();
+			int32 Colon = AddressStr.Find(":", ESearchCase::IgnoreCase, ESearchDir::FromEnd);
+			if (Colon != INDEX_NONE)
+			{
+				FString PortStr = AddressStr.Mid(Colon + 1);
+				if (!PortStr.IsEmpty())
+				{
+					Port = FCString::Atoi(*PortStr);
+				}
+			}
+		}
+	}
+#endif
+	return Port;
+}
+
 /**
  * Exec handler that routes online specific execs to the proper subsystem
  *
diff --git a/Engine/Source/Runtime/Online/OnlineSubsystemUtils/Public/OnlineSubsystemUtils.h b/Engine/Source/Runtime/Online/OnlineSubsystemUtils/Public/OnlineSubsystemUtils.h
--- a/Engine/Source/Runtime/Online/OnlineSubsystemUtils/Public/OnlineSubsystemUtils.h
+++ b/Engine/Source/Runtime/Online/OnlineSubsystemUtils/Public/OnlineSubsystemUtils.h
@@ -11,6 +11,11 @@
 #ifdef ONLINESUBSYSTEMUTILS_API
 ONLINESUBSYSTEMUTILS_API class UAudioComponent* CreateVoiceAudioComponent(uint32 SampleRate);
 ONLINESUBSYSTEMUTILS_API UWorld* GetWorldForOnline(FName InstanceName);
+
+/**
+ * Try to retrieve the active listen port for a server session
+ */
+ONLINESUBSYSTEMUTILS_API int32 GetPortFromNetDriver(FName InstanceName);
 #endif
 
 /** Macro to handle the boilerplate of accessing the proper online subsystem and getting the requested interface (UWorld version) */
