Fix non-unity compile

diff --git a/Engine/Source/Runtime/Online/OnlineSubsystemNull/Private/OnlineSubsystemNullTypes.h b/Engine/Source/Runtime/Online/OnlineSubsystemNull/Private/OnlineSubsystemNullTypes.h
--- a/Engine/Source/Runtime/Online/OnlineSubsystemNull/Private/OnlineSubsystemNullTypes.h
+++ b/Engine/Source/Runtime/Online/OnlineSubsystemNull/Private/OnlineSubsystemNullTypes.h
@@ -5,6 +5,7 @@
 #include "OnlineSubsystemTypes.h"
 #include "IPAddress.h"
 
+class FOnlineSubsystemNull;
 
 /** 
  * Implementation of session information
