Merging using UE4-To-UE4-LauncherReleases REVERSE

CL#2328353 on 14/10/2014 05:54:01
------------------------------------------
Fix for JSON self-patch manifest causing broken progress UI

CL#2328331 on 14/10/2014 05:25:07
------------------------------------------
- Fixing typo with binary manifest filename hack.
- Removing obsolete comment from self-pathc code.
- Ignore manifest server hash if redirecting to binary version
- Fix for build error

diff --git a/Engine/Source/Runtime/Online/BuildPatchServices/Private/BuildPatchManifest.cpp b/Engine/Source/Runtime/Online/BuildPatchServices/Private/BuildPatchManifest.cpp
--- a/Engine/Source/Runtime/Online/BuildPatchServices/Private/BuildPatchManifest.cpp
+++ b/Engine/Source/Runtime/Online/BuildPatchServices/Private/BuildPatchManifest.cpp
@@ -203,13 +203,13 @@ public:
 
 	virtual FArchive& operator<<(FName& N) override
 	{
-		const FName NameZero(N, 0);
-		if (!FNameIndexLookup.Contains(NameZero))
+//		const FName NameZero(N, 0);
+		if (!FNameIndexLookup.Contains(N/*ameZero*/))
 		{
 			const int32 ArNameIndex = FNameIndexLookup.Num();
-			FNameIndexLookup.Add(NameZero, ArNameIndex);
+			FNameIndexLookup.Add(N/*ameZero*/, ArNameIndex);
 		}
-		*this << FNameIndexLookup[NameZero];
+		*this << FNameIndexLookup[N/*ameZero*/];
 		return *this;
 	}
 
@@ -299,7 +299,7 @@ public:
 		// Check not insane, we know to expect a small number for a manifest
 		if (NumNames < MANIFEST_MAX_NAMES)
 		{
-			FNameLookup.Reserve(NumNames);
+			FNameLookup.Empty(NumNames);
 			for (; NumNames > 0; --NumNames)
 			{
 				FName Name;
@@ -993,14 +993,20 @@ bool FBuildPatchAppManifest::DeserializeFromJSON( const FString& JSONInput )
 		CustomFieldLookup.Add(CustomField.Key, &CustomField);
 	}
 
-	// If this is file data, fill out the guid to filename lookup
+	// If this is file data, fill out the guid to filename lookup, and chunk file size
 	if (Data->bIsFileData)
 	{
 		for (auto& FileManifest : Data->FileManifestList)
 		{
 			if (FileManifest.FileChunkParts.Num() == 1)
 			{
-				FileNameLookup.Add(FileManifest.FileChunkParts[0].Guid, &FileManifest.Filename);
+				FGuid& Guid = FileManifest.FileChunkParts[0].Guid;
+				FileNameLookup.Add(Guid, &FileManifest.Filename);
+				if (ChunkInfoLookup.Contains(Guid))
+				{
+					FChunkInfoData* ChunkInfoData = ChunkInfoLookup[Guid];
+					ChunkInfoData->FileSize = FileManifest.GetFileSize();
+				}
 			}
 			else
 			{
@@ -1011,6 +1017,7 @@ bool FBuildPatchAppManifest::DeserializeFromJSON( const FString& JSONInput )
 
 	// Calculate build size
 	TotalBuildSize = 0;
+	TotalDownloadSize = 0;
 	if (bSuccess)
 	{
 		for (auto& FileManifest : Data->FileManifestList)
@@ -1023,16 +1030,6 @@ bool FBuildPatchAppManifest::DeserializeFromJSON( const FString& JSONInput )
 		}
 	}
 
-	// Download for json file manifest is determined by build file sizes, except any duplicate identical files
-	if (Data->bIsFileData)
-	{
-		TotalDownloadSize = 0;
-		for (auto& NameLookup : FileNameLookup)
-		{
-			TotalDownloadSize += FileManifestLookup[*NameLookup.Value]->GetFileSize();
-		}
-	}
-
 	// Make sure we don't have any half loaded data
 	if( !bSuccess )
 	{
