OSSNull: Advertise 127.0.0.1 if bound to any address on a loopback interface.

- This solves the problem of inability to join sessions using OnlineSubsystemNull on Debian (and derivatives).
- Not an elegant solution, but the problem itself is caused by OS workaround in the first place.
- TTP #349210 FIXIF: ShooterGame unable to join match.

diff --git a/Engine/Source/Runtime/Online/OnlineSubsystemNull/Private/OnlineSessionInterfaceNull.cpp b/Engine/Source/Runtime/Online/OnlineSubsystemNull/Private/OnlineSessionInterfaceNull.cpp
--- a/Engine/Source/Runtime/Online/OnlineSubsystemNull/Private/OnlineSessionInterfaceNull.cpp
+++ b/Engine/Source/Runtime/Online/OnlineSubsystemNull/Private/OnlineSessionInterfaceNull.cpp
@@ -24,6 +24,19 @@ void FOnlineSessionInfoNull::Init(const FOnlineSubsystemNull& Subsystem)
 	// Read the IP from the system
 	bool bCanBindAll;
 	HostAddr = ISocketSubsystem::Get()->GetLocalHostAddr(*GLog, bCanBindAll);
+
+	// The below is a workaround for systems that set hostname to a distinct address from 127.0.0.1 on a loopback interface.
+	// See e.g. https://www.debian.org/doc/manuals/debian-reference/ch05.en.html#_the_hostname_resolution
+	// and http://serverfault.com/questions/363095/why-does-my-hostname-appear-with-the-address-127-0-1-1-rather-than-127-0-0-1-in
+	// Since we bind to 0.0.0.0, we won't answer on 127.0.1.1, so we need to advertise ourselves as 127.0.0.1 for any other loopback address we may have.
+	uint32 HostIp = 0;
+	HostAddr->GetIp(HostIp); // will return in host order
+	// if this address is on loopback interface, advertise it as 127.0.0.1
+	if ((HostIp & 0xff000000) == 0x7f000000)
+	{
+		HostAddr->SetIp(0x7f000001);	// 127.0.0.1
+	}
+
 	// Now set the port that was configured
 	HostAddr->SetPort(GetPortFromNetDriver(Subsystem.GetInstanceName()));
 
