#online,test,identity
- add logout to OSS identity test

diff --git a/Engine/Source/Runtime/Online/OnlineSubsystemUtils/Private/OnlineSubsystemUtils.cpp b/Engine/Source/Runtime/Online/OnlineSubsystemUtils/Private/OnlineSubsystemUtils.cpp
--- a/Engine/Source/Runtime/Online/OnlineSubsystemUtils/Private/OnlineSubsystemUtils.cpp
+++ b/Engine/Source/Runtime/Online/OnlineSubsystemUtils/Private/OnlineSubsystemUtils.cpp
@@ -232,8 +232,10 @@ static bool OnlineExec( UWorld* InWorld, const TCHAR* Cmd, FOutputDevice& Ar )
 						FString Auth = FParse::Token(Cmd, false);
 						FString Type = FParse::Token(Cmd, false);
 
+						bool bLogout = Id.Equals(TEXT("logout"), ESearchCase::IgnoreCase);
+
 						// This class deletes itself once done
-						(new FTestIdentityInterface(SubName))->Test(InWorld, FOnlineAccountCredentials(Type, Id, Auth));
+						(new FTestIdentityInterface(SubName))->Test(InWorld, FOnlineAccountCredentials(Type, Id, Auth), bLogout);
 						bWasHandled = true;
 					}
 					else if (FParse::Command(&Cmd, TEXT("UNIQUEIDREPL")))
diff --git a/Engine/Source/Runtime/Online/OnlineSubsystemUtils/Private/Tests/TestIdentityInterface.cpp b/Engine/Source/Runtime/Online/OnlineSubsystemUtils/Private/Tests/TestIdentityInterface.cpp
--- a/Engine/Source/Runtime/Online/OnlineSubsystemUtils/Private/Tests/TestIdentityInterface.cpp
+++ b/Engine/Source/Runtime/Online/OnlineSubsystemUtils/Private/Tests/TestIdentityInterface.cpp
@@ -4,10 +4,14 @@
 #include "ModuleManager.h"
 #include "TestIdentityInterface.h"
 
-void FTestIdentityInterface::Test(UWorld* InWorld, const FOnlineAccountCredentials& InAccountCredentials)
+void FTestIdentityInterface::Test(UWorld* InWorld, const FOnlineAccountCredentials& InAccountCredentials, bool bOnlyRunLogoutTest)
 {
 	// Toggle the various tests to run
-	bRunLoginTest = true;
+	if (bOnlyRunLogoutTest)
+	{
+		bRunLoginTest = false;
+		bRunLogoutTest = true;
+	}
 
 	AccountCredentials = InAccountCredentials;
 
@@ -16,6 +20,7 @@ void FTestIdentityInterface::Test(UWorld* InWorld, const FOnlineAccountCredentia
 	{
 		// Add delegates for the various async calls
 		OnlineIdentity->AddOnLoginCompleteDelegate(LocalUserIdx, OnLoginCompleteDelegate);
+		OnlineIdentity->AddOnLogoutCompleteDelegate(LocalUserIdx, OnLogoutCompleteDelegate);
 		
 		// kick off next test
 		StartNextTest();
@@ -37,6 +42,10 @@ void FTestIdentityInterface::StartNextTest()
 		// no external account credentials so use internal secret key for login
 		OnlineIdentity->Login(LocalUserIdx, AccountCredentials);
 	}
+	else if (bRunLogoutTest)
+	{
+		OnlineIdentity->Logout(LocalUserIdx);
+	}
 	else
 	{
 		// done with the test
@@ -50,6 +59,7 @@ void FTestIdentityInterface::FinishTest()
 	{
 		// Clear delegates for the various async calls
 		OnlineIdentity->ClearOnLoginCompleteDelegate(LocalUserIdx, OnLoginCompleteDelegate);
+		OnlineIdentity->ClearOnLogoutCompleteDelegate(LocalUserIdx, OnLogoutCompleteDelegate);
 	}
 	delete this;
 }
@@ -75,3 +85,21 @@ void FTestIdentityInterface::OnLoginComplete(int32 LocalUserNum, bool bWasSucces
 	StartNextTest();
 }
 
+void FTestIdentityInterface::OnLogoutComplete(int32 LocalUserNum, bool bWasSuccessful)
+{
+	if (bWasSuccessful)
+	{
+		UE_LOG(LogOnline, Log, TEXT("Successful logged out user. LocalUserNum=[%d] "),
+			LocalUserNum);
+	}
+	else
+	{
+		UE_LOG(LogOnline, Warning, TEXT("Failed to log out user."));
+	}
+	UserInfo.Reset();
+	// Mark test as done
+	bRunLogoutTest = false;
+	// kick off next test
+	StartNextTest();
+}
+
diff --git a/Engine/Source/Runtime/Online/OnlineSubsystemUtils/Private/Tests/TestIdentityInterface.h b/Engine/Source/Runtime/Online/OnlineSubsystemUtils/Private/Tests/TestIdentityInterface.h
--- a/Engine/Source/Runtime/Online/OnlineSubsystemUtils/Private/Tests/TestIdentityInterface.h
+++ b/Engine/Source/Runtime/Online/OnlineSubsystemUtils/Private/Tests/TestIdentityInterface.h
@@ -16,8 +16,12 @@
 	IOnlineIdentityPtr OnlineIdentity;
 	/** Delegate to use for authenticating new user */
 	FOnLoginCompleteDelegate OnLoginCompleteDelegate;
+	/** Delegate to use for logging out a user */
+	FOnLogoutCompleteDelegate OnLogoutCompleteDelegate;
 	/** true if authentication test should be run */
 	bool bRunLoginTest;
+	/** true if logout test should be run */
+	bool bRunLogoutTest;
 	/** Registered user info */
 	TSharedPtr<FUserOnlineAccount> UserInfo;
 	/** local user to run tests for */
@@ -46,6 +50,7 @@
 	 * See OnlineIdentityInterface
 	 */
 	void OnLoginComplete(int32 LocalUserNum, bool bWasSuccessful, const FUniqueNetId& UserId, const FString& Error);
+	void OnLogoutComplete(int32 LocalUserNum, bool bWasSuccessful);
 
  public:
 	/**
@@ -57,7 +62,9 @@
 		: SubsystemName(InSubsystemName)
 		, OnlineIdentity(NULL)
 		, OnLoginCompleteDelegate(FOnLoginCompleteDelegate::CreateRaw(this, &FTestIdentityInterface::OnLoginComplete))
+		, OnLogoutCompleteDelegate(FOnLogoutCompleteDelegate::CreateRaw(this, &FTestIdentityInterface::OnLogoutComplete))
 		, bRunLoginTest(true)
+		, bRunLogoutTest(false)
 		, LocalUserIdx(0)
 	{
 		
@@ -66,5 +73,5 @@
 	/**
 	 * Kicks off all of the testing process
 	 */
-	void Test(UWorld* InWorld, const FOnlineAccountCredentials& InAccountCredentials);
+	void Test(UWorld* InWorld, const FOnlineAccountCredentials& InAccountCredentials, bool bOnlyRunLogoutTest=false);
  };
