Fix beacon not reporting failure in some cases

diff --git a/Engine/Source/Runtime/Online/OnlineSubsystemUtils/Private/OnlineBeacon.cpp b/Engine/Source/Runtime/Online/OnlineSubsystemUtils/Private/OnlineBeacon.cpp
--- a/Engine/Source/Runtime/Online/OnlineSubsystemUtils/Private/OnlineBeacon.cpp
+++ b/Engine/Source/Runtime/Online/OnlineSubsystemUtils/Private/OnlineBeacon.cpp
@@ -41,7 +41,7 @@ void AOnlineBeacon::DestroyBeacon()
 
 void AOnlineBeacon::HandleNetworkFailure(UWorld *World, UNetDriver *InNetDriver, ENetworkFailure::Type FailureType, const FString& ErrorString)
 {
-	if (InNetDriver && InNetDriver->NetDriverName == NetDriverName)
+	if (InNetDriver && InNetDriver->NetDriverName == BeaconNetDriverName)
 	{
 		OnFailure();
 	}
