#online,presence
- more explicit enum for online presence status

diff --git a/Engine/Source/Runtime/Online/OnlineSubsystem/Public/Interfaces/OnlinePresenceInterface.h b/Engine/Source/Runtime/Online/OnlineSubsystem/Public/Interfaces/OnlinePresenceInterface.h
--- a/Engine/Source/Runtime/Online/OnlineSubsystem/Public/Interfaces/OnlinePresenceInterface.h
+++ b/Engine/Source/Runtime/Online/OnlineSubsystem/Public/Interfaces/OnlinePresenceInterface.h
@@ -16,6 +16,56 @@ const FString DefaultPresenceKey = "RichPresence";
 /** Custom presence data that is not seen by users but can be polled */
 const FString CustomPresenceDataKey = "CustomData";
 
+namespace EOnlinePresenceState
+{
+	enum Type
+	{
+		Online,
+		Offline,
+		Away,
+		ExtendedAway,
+		DoNotDisturb,
+		Chat
+	};
+
+	/** 
+	 * @return the stringified version of the enum passed in 
+	 */
+	inline const TCHAR* ToString(EOnlinePresenceState::Type EnumVal)
+	{
+		switch (EnumVal)
+		{
+		case Online:
+			return TEXT("Online");
+		case Offline:
+			return TEXT("Offline");
+		case Away:
+			return TEXT("Away");
+		case ExtendedAway:
+			return TEXT("ExtendedAway");
+		case DoNotDisturb:
+			return TEXT("DoNotDisturb");
+		case Chat:
+			return TEXT("Chat");
+		}
+		return TEXT("");
+	}
+}
+
+class FOnlineUserPresenceStatus
+{
+public:
+	FString StatusStr;
+	EOnlinePresenceState::Type State;
+	FPresenceProperties Properties;
+
+	FOnlineUserPresenceStatus()
+		: State(EOnlinePresenceState::Offline)
+	{
+
+	}
+};
+
 /**
  * Presence info for an online user returned via IOnlinePresence interface
  */
@@ -23,14 +73,12 @@ class FOnlineUserPresence
 {
 public:
 	uint64 SessionId;
-	FString PresenceStr;
 	uint32 bIsOnline:1;
 	uint32 bIsPlaying:1;
 	uint32 bIsPlayingThisGame:1;
 	uint32 bIsJoinable:1;
 	uint32 bHasVoiceSupport:1;
-
-	FPresenceProperties Presence;
+	FOnlineUserPresenceStatus Status;
 
 	/** Constructor */
 	FOnlineUserPresence()
@@ -45,6 +93,15 @@ public:
 	}
 };
 
+/**
+ * Delegate executed when new presence data is available for a user.
+ *
+ * @param UserId The unique id of the user whose presence was received.
+ * @param Presence The unique id of the user whose presence was received.
+ */
+DECLARE_MULTICAST_DELEGATE_TwoParams(FOnPresenceReceived, const class FUniqueNetId& /*UserId*/, const TSharedRef<FOnlineUserPresence>& /*Presence*/);
+typedef FOnPresenceReceived::FDelegate FOnPresenceReceivedDelegate;
+
 /**
  *	Interface class for getting and setting rich presence information.
  */
@@ -55,21 +112,21 @@ public:
 	virtual ~IOnlinePresence() {}
 
 	/**
-	 * Delegate executed when setting presence for a user has completed.
+	 * Delegate executed when setting or querying presence for a user has completed.
 	 *
-	 * @param User The unique id of the user whose presence was set.
+	 * @param UserId The unique id of the user whose presence was set.
 	 * @param bWasSuccessful true if the async action completed without error, false if there was an error.
 	 */
-	DECLARE_DELEGATE_TwoParams(FOnPresenceTaskCompleteDelegate, const class FUniqueNetId&, const bool);
+	DECLARE_DELEGATE_TwoParams(FOnPresenceTaskCompleteDelegate, const class FUniqueNetId& /*UserId*/, const bool /*bWasSuccessful*/);
 
 	/**
 	 * Starts an async task that sets presence information for the user.
 	 *
 	 * @param User The unique id of the user whose presence is being set.
-	 * @param Presence The collection of key/value pairs to set as the user's presence data.
+	 * @param Status The collection of state and key/value pairs to set as the user's presence data.
 	 * @param Delegate The delegate to be executed when the potentially asynchronous set operation completes.
 	 */
-	virtual void SetPresence(const FUniqueNetId& User, const FPresenceProperties& Presence, const FOnPresenceTaskCompleteDelegate& Delegate = FOnPresenceTaskCompleteDelegate()) = 0;
+	virtual void SetPresence(const FUniqueNetId& User, const FOnlineUserPresenceStatus& Status, const FOnPresenceTaskCompleteDelegate& Delegate = FOnPresenceTaskCompleteDelegate()) = 0;
 
 	/**
 	 * Starts an async operation that will update the cache with presence data from all users in the Users array.
@@ -80,6 +137,14 @@ public:
 	 */
 	virtual void QueryPresence(const FUniqueNetId& User, const FOnPresenceTaskCompleteDelegate& Delegate = FOnPresenceTaskCompleteDelegate()) = 0;
 
+	/**
+	 * Delegate executed when new presence data is available for a user.
+	 *
+	 * @param UserId The unique id of the user whose presence was received.
+	 * @param Presence The unique id of the user whose presence was received.
+	 */
+	DEFINE_ONLINE_DELEGATE_TWO_PARAM(OnPresenceReceived, const class FUniqueNetId& /*UserId*/, const TSharedRef<FOnlineUserPresence>& /*Presence*/);
+
 	/**
 	 * Gets the cached presence information for a user.
 	 *
diff --git a/Engine/Source/Runtime/Online/OnlineSubsystemSteam/Private/OnlineFriendsInterfaceSteam.cpp b/Engine/Source/Runtime/Online/OnlineSubsystemSteam/Private/OnlineFriendsInterfaceSteam.cpp
--- a/Engine/Source/Runtime/Online/OnlineSubsystemSteam/Private/OnlineFriendsInterfaceSteam.cpp
+++ b/Engine/Source/Runtime/Online/OnlineSubsystemSteam/Private/OnlineFriendsInterfaceSteam.cpp
@@ -200,11 +200,29 @@ void FOnlineAsyncTaskSteamReadFriendsList::Finalize()
 			const bool bIsPlayingAGame = SteamFriendsPtr->GetFriendGamePlayed(SteamPlayerId, &FriendGameInfo);
 			// Now fill in the friend info
 			Friend->AccountData.Add(TEXT("nickname"), NickName);
-			Friend->Presence.PresenceStr = UTF8_TO_TCHAR(SteamFriendsPtr->GetFriendRichPresence(SteamPlayerId,"connect"));
+			Friend->Presence.Status.StatusStr = UTF8_TO_TCHAR(SteamFriendsPtr->GetFriendRichPresence(SteamPlayerId,"connect"));
 			FString JoinablePresenceString = UTF8_TO_TCHAR(SteamFriendsPtr->GetFriendRichPresence(SteamPlayerId,"Joinable"));
 			// Remote friend is responsible for updating their presence to have the joinable status
 			Friend->Presence.bIsJoinable = JoinablePresenceString == TEXT("true");
- 			Friend->Presence.bIsOnline = SteamFriendsPtr->GetFriendPersonaState(SteamPlayerId) > k_EPersonaStateOffline;
+			Friend->Presence.bIsOnline = SteamFriendsPtr->GetFriendPersonaState(SteamPlayerId) > k_EPersonaStateOffline;
+			switch (SteamFriendsPtr->GetFriendPersonaState(SteamPlayerId))
+			{
+				case k_EPersonaStateOffline:
+					Friend->Presence.Status.State = EOnlinePresenceState::Offline;
+					break;
+				case k_EPersonaStateBusy:
+					Friend->Presence.Status.State = EOnlinePresenceState::DoNotDisturb;
+					break;
+				case k_EPersonaStateAway:
+					Friend->Presence.Status.State = EOnlinePresenceState::Away;
+					break;
+				case k_EPersonaStateSnooze:
+					Friend->Presence.Status.State = EOnlinePresenceState::ExtendedAway;
+					break;
+				default:
+					Friend->Presence.Status.State = EOnlinePresenceState::Online;
+					break;
+			}
 			Friend->Presence.bIsPlaying = bIsPlayingAGame;
 			// Check the steam id
 			Friend->Presence.bIsPlayingThisGame = FriendGameInfo.m_gameID.AppID() == SteamSubsystem->GetSteamAppId();
diff --git a/Engine/Source/Runtime/Online/OnlineSubsystemUtils/Private/Tests/TestFriendsInterface.cpp b/Engine/Source/Runtime/Online/OnlineSubsystemUtils/Private/Tests/TestFriendsInterface.cpp
--- a/Engine/Source/Runtime/Online/OnlineSubsystemUtils/Private/Tests/TestFriendsInterface.cpp
+++ b/Engine/Source/Runtime/Online/OnlineSubsystemUtils/Private/Tests/TestFriendsInterface.cpp
@@ -124,7 +124,9 @@ void FTestFriendsInterface::OnReadFriendsComplete(int32 LocalPlayer, bool bWasSu
 				UE_LOG(LogOnline, Log,
 					TEXT("\t\t Invite status (%s)"), EInviteStatus::ToString(Friend.GetInviteStatus()));
 				UE_LOG(LogOnline, Log,
-					TEXT("\t\t Presence: %s"), *Presence.PresenceStr);
+					TEXT("\t\t Presence: %s"), *Presence.Status.StatusStr);
+				UE_LOG(LogOnline, Log,
+					TEXT("\t\t State: %s"), EOnlinePresenceState::ToString(Presence.Status.State));
 				UE_LOG(LogOnline, Log,
 					TEXT("\t\t bIsOnline (%s)"), Presence.bIsOnline ? TEXT("true") : TEXT("false"));
 				UE_LOG(LogOnline, Log,
