Add functionality to BuildPatchTool.exe compactify process to update patch timestamps only

* Activated by specifying -touchonly at the command-line
* Cannot be used in conjunction with -preview

diff --git a/Engine/Source/Programs/BuildPatchTool/Private/BuildPatchToolMain.cpp b/Engine/Source/Programs/BuildPatchTool/Private/BuildPatchToolMain.cpp
--- a/Engine/Source/Programs/BuildPatchTool/Private/BuildPatchToolMain.cpp
+++ b/Engine/Source/Programs/BuildPatchTool/Private/BuildPatchToolMain.cpp
@@ -55,6 +55,7 @@
 	Optional arguments:
 	-stdout				Adds stdout logging to the app.
 	-preview			Log all the actions it will take to update internal structures, but don't actually execute them.
+	-touchonly			When specified, will only set the modified date of referenced files to the current time, but will NOT delete any patch data.
 	-ManifestsList=""			Specifies in quotes, the list of manifest filenames to keep following the operation. If omitted, all manifests are kept.
 	-ManifestsFile=""			Specifies in quotes, the name of the file (relative to CloudDir) which contains a list of manifests to keep, one manifest filename per line
 	-DataAgeThreshold=14.25		The maximum age in days of chunk files that will be retained. All older chunks will be deleted.
@@ -155,6 +156,7 @@ int32 BuildPatchToolMain( const TCHAR* CommandLine )
 	bool bCompactify = false;
 	bool bPatchGeneration = true;
 	bool bPreview = false;
+	bool bTouchOnly = false;
 	bool bPatchWithReuseAgeThreshold = true;
 
 	// Collect all the info from the CommandLine
@@ -186,6 +188,9 @@ int32 BuildPatchToolMain( const TCHAR* CommandLine )
 
 		Matcher.Command = TEXT("preview");
 		bPreview = bCompactify && Switches.IndexOfByPredicate(Matcher) != INDEX_NONE;
+
+		Matcher.Command = TEXT("touchonly");
+		bTouchOnly = bCompactify && Switches.IndexOfByPredicate(Matcher) != INDEX_NONE;
 		
 		Matcher.Command = TEXT( "BuildRoot" );
 		BuildRootIdx = Switches.IndexOfByPredicate(Matcher);
@@ -378,6 +383,11 @@ int32 BuildPatchToolMain( const TCHAR* CommandLine )
 		return 1;
 	}
 
+	if (bCompactify && bPreview && bTouchOnly)
+	{
+		GLog->Log(ELogVerbosity::Error, TEXT("Only one of -preview and -touchonly can be specified"));
+		return 5;
+	}
 
 	// Initialize the file manager
 	IFileManager::Get().ProcessCommandLineOptions();
@@ -419,8 +429,19 @@ int32 BuildPatchToolMain( const TCHAR* CommandLine )
 			}
 		}
 
+		// Determine our mode of operation
+		ECompactifyMode::Type CompactifyMode = ECompactifyMode::Full;
+		if (bPreview)
+		{
+			CompactifyMode = ECompactifyMode::Preview;
+		}
+		else if (bTouchOnly)
+		{
+			CompactifyMode = ECompactifyMode::TouchOnly;
+		}
+
 		// Run the compactify routine
-		bSuccess = BuildPatchServicesModule->CompactifyCloudDirectory(ManifestsArr, DataAgeThreshold, bPreview);
+		bSuccess = BuildPatchServicesModule->CompactifyCloudDirectory(ManifestsArr, DataAgeThreshold, CompactifyMode);
 	}
 	else if (bPatchGeneration)
 	{
diff --git a/Engine/Source/Runtime/Online/BuildPatchServices/Private/BuildPatchCompactifier.cpp b/Engine/Source/Runtime/Online/BuildPatchServices/Private/BuildPatchCompactifier.cpp
--- a/Engine/Source/Runtime/Online/BuildPatchServices/Private/BuildPatchCompactifier.cpp
+++ b/Engine/Source/Runtime/Online/BuildPatchServices/Private/BuildPatchCompactifier.cpp
@@ -18,9 +18,10 @@ namespace BuildDataCompactifierDefs
 /* Constructors
 *****************************************************************************/
 
-FBuildDataCompactifier::FBuildDataCompactifier(const FString& CloudDir, const bool bPreview)
+FBuildDataCompactifier::FBuildDataCompactifier(const FString& CloudDir, const bool bPreview, const bool bTouchOnly)
 	: CloudDir(CloudDir)
 	, bPreview(bPreview)
+	, bTouchOnly(bTouchOnly)
 {
 
 }
@@ -28,15 +29,15 @@ FBuildDataCompactifier::FBuildDataCompactifier(const FString& CloudDir, const bo
 /* Public static methods
 *****************************************************************************/
 
-bool FBuildDataCompactifier::CompactifyCloudDirectory(const FString& CloudDir, const TArray<FString>& ManifestsToKeep, const float DataAgeThreshold, const bool bPreview)
+bool FBuildDataCompactifier::CompactifyCloudDirectory(const FString& CloudDir, const TArray<FString>& ManifestsToKeep, const float DataAgeThreshold, const bool bPreview, const bool bTouchOnly)
 {
-	FBuildDataCompactifier Compactifier(CloudDir, bPreview);
+	FBuildDataCompactifier Compactifier(CloudDir, bPreview, bTouchOnly);
 	return Compactifier.Compactify(ManifestsToKeep, DataAgeThreshold);
 }
 
-bool FBuildDataCompactifier::CompactifyCloudDirectory(const TArray<FString>& ManifestsToKeep, const float DataAgeThreshold, const bool bPreview)
+bool FBuildDataCompactifier::CompactifyCloudDirectory(const TArray<FString>& ManifestsToKeep, const float DataAgeThreshold, const bool bPreview, const bool bTouchOnly)
 {
-	return CompactifyCloudDirectory(FBuildPatchServicesModule::GetCloudDirectory(), ManifestsToKeep, DataAgeThreshold, bPreview);
+	return CompactifyCloudDirectory(FBuildPatchServicesModule::GetCloudDirectory(), ManifestsToKeep, DataAgeThreshold, bPreview, bTouchOnly);
 }
 
 /* Private methods
@@ -44,7 +45,7 @@ bool FBuildDataCompactifier::CompactifyCloudDirectory(const TArray<FString>& Man
 
 bool FBuildDataCompactifier::Compactify(const TArray<FString>& ManifestsToKeep, const float DataAgeThreshold) const
 {
-	GLog->Logf(TEXT("Running Compactify on %s%s"), *CloudDir, bPreview ? TEXT(". Preview mode. NO action will be taken.") : TEXT(""));
+	GLog->Logf(TEXT("Running Compactify on %s%s"), *CloudDir, bPreview ? TEXT(". Preview mode. NO action will be taken.") : bTouchOnly ? TEXT(". TouchOnly mode. NO patch data will be deleted.") : TEXT(""));
 	if (ManifestsToKeep.Num() > 0)
 	{
 		GLog->Logf(TEXT("Preserving manifest files: %s"), *FString::Join(ManifestsToKeep, TEXT(", ")));
@@ -223,7 +224,7 @@ bool FBuildDataCompactifier::Compactify(const TArray<FString>& ManifestsToKeep,
 void FBuildDataCompactifier::DeleteFile(const FString& FilePath) const
 {
 	GLog->Logf(TEXT("Deleting file %s"), *FilePath);
-	if (!bPreview)
+	if (!bPreview && !bTouchOnly)
 	{
 		IFileManager::Get().Delete(*FilePath);
 	}
@@ -267,7 +268,7 @@ bool FBuildDataCompactifier::DeleteNonReferencedManifests(TArray<FString>& AllMa
 			DeletedManifests.Add(Manifest);
 			BytesDeleted += IFileManager::Get().FileSize(*ManifestPath);
 
-			if (!bPreview)
+			if (!bPreview && !bTouchOnly)
 			{
 				IFileManager::Get().Delete(*ManifestPath);
 				if (FPaths::FileExists(ManifestPath))
diff --git a/Engine/Source/Runtime/Online/BuildPatchServices/Private/BuildPatchCompactifier.h b/Engine/Source/Runtime/Online/BuildPatchServices/Private/BuildPatchCompactifier.h
--- a/Engine/Source/Runtime/Online/BuildPatchServices/Private/BuildPatchCompactifier.h
+++ b/Engine/Source/Runtime/Online/BuildPatchServices/Private/BuildPatchCompactifier.h
@@ -23,9 +23,10 @@ public:
 	 * @param ManifestsToKeep   If specified, these manifests will be retained, and all others will be deleted
 	 * @param DataAgeThreshold  Chunks which are not referenced by a valid manifest, and which are older than this age (in days), will be deleted
 	 * @param bPreview      If true, then no actual work will be done, but all operations which would be carried out will be logged.
+	 * @param bTouchOnly    If true, then no patch data will be deleted, but referenced data will have its timestamp updated to the current time.
 	 * @return              true if no file errors occurred
 	 */
-	static bool CompactifyCloudDirectory(const TArray<FString>& ManifestsToKeep, const float DataAgeThreshold, const bool bPreview);
+	static bool CompactifyCloudDirectory(const TArray<FString>& ManifestsToKeep, const float DataAgeThreshold, const bool bPreview, const bool bTouchOnly);
 
 	/**
 	 * Processes the specified Cloud Directory to identify and delete any orphaned chunks or files.
@@ -34,20 +35,23 @@ public:
 	 * @param ManifestsToKeep   If specified, these manifests will be retained, and all others will be deleted
 	 * @param DataAgeThreshold  Chunks which are not referenced by a valid manifest, and which are older than this age (in days), will be deleted
 	 * @param bPreview      If true, then no actual work will be done, but all operations which would be carried out will be logged.
+	 * @param bTouchOnly    If true, then no patch data will be deleted, but referenced data will have its timestamp updated to the current time.
 	 * @return              true if no file errors occurred
 	 */
-	static bool CompactifyCloudDirectory(const FString& CloudDir, const TArray<FString>& ManifestsToKeep, const float DataAgeThreshold, const bool bPreview);
+	static bool CompactifyCloudDirectory(const FString& CloudDir, const TArray<FString>& ManifestsToKeep, const float DataAgeThreshold, const bool bPreview, const bool bTouchOnly);
 
 private:
 	const FString CloudDir;
 	const bool bPreview;
+	const bool bTouchOnly;
 
 	/**
 	 * Constructor.
 	 * @param CloudDir      The directory that this compactifier should operate on.
 	 * @param bPreview      If true, then no actual work will be done, but all operations which would be carried out will be logged.
+	 * @param bTouchOnly    If true, then no patch data will be deleted, but referenced data will have its timestamp updated to the current time.
 	 */
-	FBuildDataCompactifier(const FString& CloudDir, const bool bPreview);
+	FBuildDataCompactifier(const FString& CloudDir, const bool bPreview, const bool bTouchOnly);
 
 	/**
 	 * Processes the Cloud Directory to identify and delete any orphaned chunks or files.
diff --git a/Engine/Source/Runtime/Online/BuildPatchServices/Private/BuildPatchServicesModule.cpp b/Engine/Source/Runtime/Online/BuildPatchServices/Private/BuildPatchServicesModule.cpp
--- a/Engine/Source/Runtime/Online/BuildPatchServices/Private/BuildPatchServicesModule.cpp
+++ b/Engine/Source/Runtime/Online/BuildPatchServices/Private/BuildPatchServicesModule.cpp
@@ -199,9 +199,11 @@ bool FBuildPatchServicesModule::GenerateFilesManifestFromDirectory( const FBuild
 	return FBuildDataGenerator::GenerateFilesManifestFromDirectory( Settings );
 }
 
-bool FBuildPatchServicesModule::CompactifyCloudDirectory(const TArray<FString>& ManifestsToKeep, const float DataAgeThreshold, const bool bPreview)
+bool FBuildPatchServicesModule::CompactifyCloudDirectory(const TArray<FString>& ManifestsToKeep, const float DataAgeThreshold, const ECompactifyMode::Type Mode)
 {
-	return FBuildDataCompactifier::CompactifyCloudDirectory(ManifestsToKeep, DataAgeThreshold, bPreview);
+	const bool bPreview = Mode == ECompactifyMode::Preview;
+	const bool bTouchOnly = Mode == ECompactifyMode::TouchOnly;
+	return FBuildDataCompactifier::CompactifyCloudDirectory(ManifestsToKeep, DataAgeThreshold, bPreview, bTouchOnly);
 }
 
 #endif //WITH_BUILDPATCHGENERATION
diff --git a/Engine/Source/Runtime/Online/BuildPatchServices/Private/BuildPatchServicesModule.h b/Engine/Source/Runtime/Online/BuildPatchServices/Private/BuildPatchServicesModule.h
--- a/Engine/Source/Runtime/Online/BuildPatchServices/Private/BuildPatchServicesModule.h
+++ b/Engine/Source/Runtime/Online/BuildPatchServices/Private/BuildPatchServicesModule.h
@@ -84,7 +84,7 @@ public:
 #if WITH_BUILDPATCHGENERATION
 	virtual bool GenerateChunksManifestFromDirectory( const FBuildPatchSettings& Settings ) override;
 	virtual bool GenerateFilesManifestFromDirectory( const FBuildPatchSettings& Settings ) override;
-	virtual bool CompactifyCloudDirectory( const TArray<FString>& ManifestsToKeep, const float DataAgeThreshold, const bool bPreview ) override;
+	virtual bool CompactifyCloudDirectory( const TArray<FString>& ManifestsToKeep, const float DataAgeThreshold, const ECompactifyMode::Type Mode ) override;
 #endif
 	virtual IBuildManifestPtr MakeManifestFromJSON( const FString& ManifestJSON ) override;
 	// End IBuildPatchServicesModule interface
diff --git a/Engine/Source/Runtime/Online/BuildPatchServices/Public/Interfaces/IBuildPatchServicesModule.h b/Engine/Source/Runtime/Online/BuildPatchServices/Public/Interfaces/IBuildPatchServicesModule.h
--- a/Engine/Source/Runtime/Online/BuildPatchServices/Public/Interfaces/IBuildPatchServicesModule.h
+++ b/Engine/Source/Runtime/Online/BuildPatchServices/Public/Interfaces/IBuildPatchServicesModule.h
@@ -16,6 +16,16 @@ DECLARE_DELEGATE_OneParam( FBuildPatchManifestDelegate, IBuildManifestRef );
 DECLARE_DELEGATE_TwoParams( FBuildPatchBoolManifestDelegate, bool, IBuildManifestRef );
 DECLARE_DELEGATE_OneParam( FBuildPatchFloatDelegate, float );
 
+namespace ECompactifyMode
+{
+	enum Type
+	{
+		Preview,
+		TouchOnly,
+		Full
+	};
+}
+
 /**
  * Interface for the services manager.
  */
@@ -114,10 +124,10 @@ public:
 	 * NOTE: THIS function is blocking and will not return until finished. Don't run on main thread.
 	 * @param ManifestsToKeep      If specified, these manifests will be retained, and all others will be deleted
 	 * @param DataAgeThreshold     Chunks which are not referenced by a valid manifest, and which are older than this age (in days), will be deleted
-	 * @param bPreview      If true, then no actual work will be done, but all operations which would be carried out will be logged.
+	 * @param Mode      The mode that compactify will run in. If Preview, then no work will be carried out, if TouchOnly, then no patch-data will be deleted
 	 * @return              true if no file errors occurred
 	 */
-	virtual bool CompactifyCloudDirectory(const TArray<FString>& ManifestsToKeep, const float DataAgeThreshold, const bool bPreview) = 0;
+	virtual bool CompactifyCloudDirectory(const TArray<FString>& ManifestsToKeep, const float DataAgeThreshold, const ECompactifyMode::Type Mode) = 0;
 #endif // WITH_BUILDPATCHGENERATION
 
 	/**
