Merging using UE4-To-UE4-LauncherReleases REVERSE

CL#2337101 on 22/10/2014 07:45:47
------------------------------------------
Fix support for older JSON manifest files

diff --git a/Engine/Source/Runtime/Online/BuildPatchServices/Private/BuildPatchManifest.cpp b/Engine/Source/Runtime/Online/BuildPatchServices/Private/BuildPatchManifest.cpp
--- a/Engine/Source/Runtime/Online/BuildPatchServices/Private/BuildPatchManifest.cpp
+++ b/Engine/Source/Runtime/Online/BuildPatchServices/Private/BuildPatchManifest.cpp
@@ -964,6 +964,14 @@ bool FBuildPatchAppManifest::DeserializeFromJSON( const FString& JSONInput )
 			}
 		}
 	}
+	else
+	{
+		// Missing chunk list, version before we saved them compressed.. Assume chunk size
+		for (FChunkInfoData& ChunkInfo : Data->ChunkList)
+		{
+			ChunkInfo.FileSize = FBuildPatchData::ChunkDataSize;
+		}
+	}
 
 	// Get the bIsFileData value. The variable will exist in versions of StoresIfChunkOrFileData or later, otherwise the previous method is to check
 	// if ChunkHashList is empty.
