Don't check GIsPlayInEditorWorld when getting an OSS instance
This bool can actually be false during PIE while objects derived from FTickerObjectBase are ticking!
Fixes getting the wrong OSS instance during PIE in some cases.

diff --git a/Engine/Source/Runtime/Online/OnlineSubsystemUtils/Public/OnlineSubsystemUtils.h b/Engine/Source/Runtime/Online/OnlineSubsystemUtils/Public/OnlineSubsystemUtils.h
--- a/Engine/Source/Runtime/Online/OnlineSubsystemUtils/Public/OnlineSubsystemUtils.h
+++ b/Engine/Source/Runtime/Online/OnlineSubsystemUtils/Public/OnlineSubsystemUtils.h
@@ -32,8 +32,7 @@ namespace Online
 	{
 #if UE_EDITOR // at present, multiple worlds are only possible in the editor
 		FName Identifier = SubsystemName; 
-		if (GIsPlayInEditorWorld &&
-			World != NULL)
+		if (World != NULL)
 		{
 			FWorldContext& CurrentContext = GEngine->GetWorldContextFromWorldChecked(World);
 			if (CurrentContext.WorldType == EWorldType::PIE)
