FOnlineSubsystemBPCallHelper now requires a UWorld
Allows it to find the correct online subsystem instance in PIE
Blueprint nodes using the helper now get a world context object

diff --git a/Engine/Source/Runtime/Online/OnlineSubsystemUtils/Classes/AchievementBlueprintLibrary.h b/Engine/Source/Runtime/Online/OnlineSubsystemUtils/Classes/AchievementBlueprintLibrary.h
--- a/Engine/Source/Runtime/Online/OnlineSubsystemUtils/Classes/AchievementBlueprintLibrary.h
+++ b/Engine/Source/Runtime/Online/OnlineSubsystemUtils/Classes/AchievementBlueprintLibrary.h
@@ -17,8 +17,8 @@ class ONLINESUBSYSTEMUTILS_API UAchievementBlueprintLibrary : public UBlueprintF
 	 * @param bFoundID - If the ID was found in the cache (if not, none of the other values are meaningful)
 	 * @param Progress - The progress amount of the achievement
 	 */
-	UFUNCTION(BlueprintCallable, Category = "Online|Achievements")
-	static void GetCachedAchievementProgress(APlayerController* PlayerController, FName AchievementID, /*out*/ bool& bFoundID, /*out*/ float& Progress);
+	UFUNCTION(BlueprintCallable, meta=(WorldContext="WorldContextObject"), Category = "Online|Achievements")
+	static void GetCachedAchievementProgress(UObject* WorldContextObject, APlayerController* PlayerController, FName AchievementID, /*out*/ bool& bFoundID, /*out*/ float& Progress);
 
 	/**
 	 * Get the description for an achievement ID (you must call CacheAchievementDescriptions first to cache them)
@@ -30,6 +30,6 @@ class ONLINESUBSYSTEMUTILS_API UAchievementBlueprintLibrary : public UBlueprintF
 	 * @param UnlockedDescription - The localized unlocked description of the achievement
 	 * @param bHidden - Whether the achievement is hidden
 	 */
-	UFUNCTION(BlueprintCallable, Category="Online|Achievements")
-	static void GetCachedAchievementDescription(APlayerController* PlayerController, FName AchievementID, /*out*/ bool& bFoundID, /*out*/ FText& Title, /*out*/ FText& LockedDescription, /*out*/ FText& UnlockedDescription, /*out*/ bool& bHidden);
+	UFUNCTION(BlueprintCallable, meta=(WorldContext="WorldContextObject"), Category="Online|Achievements")
+	static void GetCachedAchievementDescription(UObject* WorldContextObject, APlayerController* PlayerController, FName AchievementID, /*out*/ bool& bFoundID, /*out*/ FText& Title, /*out*/ FText& LockedDescription, /*out*/ FText& UnlockedDescription, /*out*/ bool& bHidden);
 };
diff --git a/Engine/Source/Runtime/Online/OnlineSubsystemUtils/Classes/AchievementQueryCallbackProxy.h b/Engine/Source/Runtime/Online/OnlineSubsystemUtils/Classes/AchievementQueryCallbackProxy.h
--- a/Engine/Source/Runtime/Online/OnlineSubsystemUtils/Classes/AchievementQueryCallbackProxy.h
+++ b/Engine/Source/Runtime/Online/OnlineSubsystemUtils/Classes/AchievementQueryCallbackProxy.h
@@ -18,12 +18,12 @@ class UAchievementQueryCallbackProxy : public UOnlineBlueprintCallProxyBase
 	FEmptyOnlineDelegate OnFailure;
 
 	// Fetches and caches achievement progress from the default online subsystem
-	UFUNCTION(BlueprintCallable, meta=(BlueprintInternalUseOnly = "true"), Category = "Online|Achievements")
-	static UAchievementQueryCallbackProxy* CacheAchievements(class APlayerController* PlayerController);
+	UFUNCTION(BlueprintCallable, meta=(BlueprintInternalUseOnly = "true", WorldContext="WorldContextObject"), Category = "Online|Achievements")
+	static UAchievementQueryCallbackProxy* CacheAchievements(UObject* WorldContextObject, class APlayerController* PlayerController);
 
 	// Fetches and caches achievement descriptions from the default online subsystem
-	UFUNCTION(BlueprintCallable, meta=(BlueprintInternalUseOnly = "true"), Category = "Online|Achievements")
-	static UAchievementQueryCallbackProxy* CacheAchievementDescriptions(class APlayerController* PlayerController);
+	UFUNCTION(BlueprintCallable, meta=(BlueprintInternalUseOnly = "true", WorldContext="WorldContextObject"), Category = "Online|Achievements")
+	static UAchievementQueryCallbackProxy* CacheAchievementDescriptions(UObject* WorldContextObject, class APlayerController* PlayerController);
 
 	// UOnlineBlueprintCallProxyBase interface
 	virtual void Activate() override;
@@ -39,4 +39,7 @@ private:
 
 	// Are we querying achievement progress or achievement descriptions?
 	bool bFetchDescriptions;
+
+	// The world context object in which this call is taking place
+	UObject* WorldContextObject;
 };
diff --git a/Engine/Source/Runtime/Online/OnlineSubsystemUtils/Classes/AchievementWriteCallbackProxy.h b/Engine/Source/Runtime/Online/OnlineSubsystemUtils/Classes/AchievementWriteCallbackProxy.h
--- a/Engine/Source/Runtime/Online/OnlineSubsystemUtils/Classes/AchievementWriteCallbackProxy.h
+++ b/Engine/Source/Runtime/Online/OnlineSubsystemUtils/Classes/AchievementWriteCallbackProxy.h
@@ -23,8 +23,8 @@ class UAchievementWriteCallbackProxy : public UOnlineBlueprintCallProxyBase
 	//   AchievementName is the ID of the achievement to update progress on
 	//   Progress is the reported progress toward accomplishing the achievement
 	//   UserTag is not used internally, but it is returned on success or failure
-	UFUNCTION(BlueprintCallable, meta=(BlueprintInternalUseOnly = "true"), Category="Online|Achievements")
-	static UAchievementWriteCallbackProxy* WriteAchievementProgress(class APlayerController* PlayerController, FName AchievementName, float Progress = 100.0f, int32 UserTag = 0);
+	UFUNCTION(BlueprintCallable, meta=(BlueprintInternalUseOnly = "true", WorldContext="WorldContextObject"), Category="Online|Achievements")
+	static UAchievementWriteCallbackProxy* WriteAchievementProgress(UObject* WorldContextObject, APlayerController* PlayerController, FName AchievementName, float Progress = 100.0f, int32 UserTag = 0);
 
 	// UOnlineBlueprintCallProxyBase interface
 	virtual void Activate() override;
@@ -53,4 +53,8 @@ private:
 
 	/** The specified user tag */
 	int32 UserTag;
+
+	/** The world context object in which this call is taking place */
+	UObject* WorldContextObject;
+
 };
diff --git a/Engine/Source/Runtime/Online/OnlineSubsystemUtils/Classes/ConnectionCallbackProxy.h b/Engine/Source/Runtime/Online/OnlineSubsystemUtils/Classes/ConnectionCallbackProxy.h
--- a/Engine/Source/Runtime/Online/OnlineSubsystemUtils/Classes/ConnectionCallbackProxy.h
+++ b/Engine/Source/Runtime/Online/OnlineSubsystemUtils/Classes/ConnectionCallbackProxy.h
@@ -22,8 +22,8 @@ class UConnectionCallbackProxy : public UOnlineBlueprintCallProxyBase
 	FOnlineConnectionResult OnFailure;
 
 	// Connects to an online service such as Google Play
-	UFUNCTION(BlueprintCallable, meta = (BlueprintInternalUseOnly = "true"), Category = "Online|Achievements")
-		static UConnectionCallbackProxy* ConnectToService(class APlayerController* PlayerController);
+	UFUNCTION(BlueprintCallable, meta = (BlueprintInternalUseOnly = "true", WorldContext="WorldContextObject"), Category = "Online|Achievements")
+		static UConnectionCallbackProxy* ConnectToService(UObject* WorldContextObject, class APlayerController* PlayerController);
 
 	// UOnlineBlueprintCallProxyBase interface
 	virtual void Activate() override;
@@ -39,4 +39,7 @@ private:
 
 	// Connection callback delegate
 	FOnLoginCompleteDelegate OnLoginCompleteDelegate;
+
+	// The world context object in which this call is taking place
+	UObject* WorldContextObject;
 };
\ No newline at end of file
diff --git a/Engine/Source/Runtime/Online/OnlineSubsystemUtils/Private/AchievementBlueprintLibrary.cpp b/Engine/Source/Runtime/Online/OnlineSubsystemUtils/Private/AchievementBlueprintLibrary.cpp
--- a/Engine/Source/Runtime/Online/OnlineSubsystemUtils/Private/AchievementBlueprintLibrary.cpp
+++ b/Engine/Source/Runtime/Online/OnlineSubsystemUtils/Private/AchievementBlueprintLibrary.cpp
@@ -9,12 +9,12 @@ UAchievementBlueprintLibrary::UAchievementBlueprintLibrary(const FObjectInitiali
 {
 }
 
-void UAchievementBlueprintLibrary::GetCachedAchievementProgress(APlayerController* PlayerController, FName AchievementID, /*out*/ bool& bFoundID, /*out*/ float& Progress)
+void UAchievementBlueprintLibrary::GetCachedAchievementProgress(UObject* WorldContextObject, APlayerController* PlayerController, FName AchievementID, /*out*/ bool& bFoundID, /*out*/ float& Progress)
 {
 	bFoundID = false;
 	Progress = 0.0f;
 
-	FOnlineSubsystemBPCallHelper Helper(TEXT("GetCachedAchievementProgress"));
+	FOnlineSubsystemBPCallHelper Helper(TEXT("GetCachedAchievementProgress"), GEngine->GetWorldFromContextObject(WorldContextObject));
 	Helper.QueryIDFromPlayerController(PlayerController);
 
 	if (Helper.IsValid())
@@ -36,7 +36,7 @@ void UAchievementBlueprintLibrary::GetCachedAchievementProgress(APlayerControlle
 	}
 }
 
-void UAchievementBlueprintLibrary::GetCachedAchievementDescription(APlayerController* PlayerController, FName AchievementID, /*out*/ bool& bFoundID, /*out*/ FText& Title, /*out*/ FText& LockedDescription, /*out*/ FText& UnlockedDescription, /*out*/ bool& bHidden)
+void UAchievementBlueprintLibrary::GetCachedAchievementDescription(UObject* WorldContextObject, APlayerController* PlayerController, FName AchievementID, /*out*/ bool& bFoundID, /*out*/ FText& Title, /*out*/ FText& LockedDescription, /*out*/ FText& UnlockedDescription, /*out*/ bool& bHidden)
 {
 	bFoundID = false;
 	Title = FText::GetEmpty();
@@ -44,7 +44,7 @@ void UAchievementBlueprintLibrary::GetCachedAchievementDescription(APlayerContro
 	UnlockedDescription = FText::GetEmpty();
 	bHidden = false;
 
-	FOnlineSubsystemBPCallHelper Helper(TEXT("GetCachedAchievementDescription"));
+	FOnlineSubsystemBPCallHelper Helper(TEXT("GetCachedAchievementDescription"), GEngine->GetWorldFromContextObject(WorldContextObject));
 	Helper.QueryIDFromPlayerController(PlayerController);
 
 	if (Helper.IsValid())
diff --git a/Engine/Source/Runtime/Online/OnlineSubsystemUtils/Private/AchievementQueryCallbackProxy.cpp b/Engine/Source/Runtime/Online/OnlineSubsystemUtils/Private/AchievementQueryCallbackProxy.cpp
--- a/Engine/Source/Runtime/Online/OnlineSubsystemUtils/Private/AchievementQueryCallbackProxy.cpp
+++ b/Engine/Source/Runtime/Online/OnlineSubsystemUtils/Private/AchievementQueryCallbackProxy.cpp
@@ -7,18 +7,20 @@
 
 UAchievementQueryCallbackProxy::UAchievementQueryCallbackProxy(const FObjectInitializer& ObjectInitializer)
 	: Super(ObjectInitializer)
+	, WorldContextObject(nullptr)
 {
 }
 
-UAchievementQueryCallbackProxy* UAchievementQueryCallbackProxy::CacheAchievements(class APlayerController* PlayerController)
+UAchievementQueryCallbackProxy* UAchievementQueryCallbackProxy::CacheAchievements(UObject* WorldContextObject, class APlayerController* PlayerController)
 {
 	UAchievementQueryCallbackProxy* Proxy = NewObject<UAchievementQueryCallbackProxy>();
 	Proxy->PlayerControllerWeakPtr = PlayerController;
 	Proxy->bFetchDescriptions = false;
+	Proxy->WorldContextObject = WorldContextObject;
 	return Proxy;
 }
 
-UAchievementQueryCallbackProxy* UAchievementQueryCallbackProxy::CacheAchievementDescriptions(class APlayerController* PlayerController)
+UAchievementQueryCallbackProxy* UAchievementQueryCallbackProxy::CacheAchievementDescriptions(UObject* WorldContextObject, class APlayerController* PlayerController)
 {
 	UAchievementQueryCallbackProxy* Proxy = NewObject<UAchievementQueryCallbackProxy>();
 	Proxy->PlayerControllerWeakPtr = PlayerController;
@@ -28,7 +30,7 @@ UAchievementQueryCallbackProxy* UAchievementQueryCallbackProxy::CacheAchievement
 
 void UAchievementQueryCallbackProxy::Activate()
 {
-	FOnlineSubsystemBPCallHelper Helper(TEXT("CacheAchievements or CacheAchievementDescriptions"));
+	FOnlineSubsystemBPCallHelper Helper(TEXT("CacheAchievements or CacheAchievementDescriptions"), GEngine->GetWorldFromContextObject(WorldContextObject));
 	Helper.QueryIDFromPlayerController(PlayerControllerWeakPtr.Get());
 
 	if (Helper.IsValid())
diff --git a/Engine/Source/Runtime/Online/OnlineSubsystemUtils/Private/AchievementWriteCallbackProxy.cpp b/Engine/Source/Runtime/Online/OnlineSubsystemUtils/Private/AchievementWriteCallbackProxy.cpp
--- a/Engine/Source/Runtime/Online/OnlineSubsystemUtils/Private/AchievementWriteCallbackProxy.cpp
+++ b/Engine/Source/Runtime/Online/OnlineSubsystemUtils/Private/AchievementWriteCallbackProxy.cpp
@@ -7,10 +7,11 @@
 
 UAchievementWriteCallbackProxy::UAchievementWriteCallbackProxy(const FObjectInitializer& ObjectInitializer)
 	: Super(ObjectInitializer)
+	, WorldContextObject(nullptr)
 {
 }
 
-UAchievementWriteCallbackProxy* UAchievementWriteCallbackProxy::WriteAchievementProgress(class APlayerController* PlayerController, FName AchievementName, float Progress, int32 UserTag)
+UAchievementWriteCallbackProxy* UAchievementWriteCallbackProxy::WriteAchievementProgress(UObject* WorldContextObject, class APlayerController* PlayerController, FName AchievementName, float Progress, int32 UserTag)
 {
 	UAchievementWriteCallbackProxy* Proxy = NewObject<UAchievementWriteCallbackProxy>();
 
@@ -20,13 +21,14 @@ UAchievementWriteCallbackProxy* UAchievementWriteCallbackProxy::WriteAchievement
 	Proxy->AchievementName = AchievementName;
 	Proxy->AchievementProgress = Progress;
 	Proxy->UserTag = UserTag;
+	Proxy->WorldContextObject = WorldContextObject;
 
 	return Proxy;
 }
 
 void UAchievementWriteCallbackProxy::Activate()
 {
-	FOnlineSubsystemBPCallHelper Helper(TEXT("WriteAchievementObject"));
+	FOnlineSubsystemBPCallHelper Helper(TEXT("WriteAchievementObject"), GEngine->GetWorldFromContextObject(WorldContextObject));
 	Helper.QueryIDFromPlayerController(PlayerControllerWeakPtr.Get());
 
 	if (Helper.IsValid())
diff --git a/Engine/Source/Runtime/Online/OnlineSubsystemUtils/Private/ConnectionCallbackProxy.cpp b/Engine/Source/Runtime/Online/OnlineSubsystemUtils/Private/ConnectionCallbackProxy.cpp
--- a/Engine/Source/Runtime/Online/OnlineSubsystemUtils/Private/ConnectionCallbackProxy.cpp
+++ b/Engine/Source/Runtime/Online/OnlineSubsystemUtils/Private/ConnectionCallbackProxy.cpp
@@ -7,22 +7,24 @@
 
 UConnectionCallbackProxy::UConnectionCallbackProxy(const FObjectInitializer& ObjectInitializer)
 	: Super(ObjectInitializer)
+	, WorldContextObject(nullptr)
 {
 
 	OnLoginCompleteDelegate.BindUObject(this, &ThisClass::OnLoginCompleted);
 
 }
 
-UConnectionCallbackProxy* UConnectionCallbackProxy::ConnectToService(class APlayerController* PlayerController)
+UConnectionCallbackProxy* UConnectionCallbackProxy::ConnectToService(UObject* WorldContextObject, class APlayerController* PlayerController)
 {
 	UConnectionCallbackProxy* Proxy = NewObject<UConnectionCallbackProxy>();
 	Proxy->PlayerControllerWeakPtr = PlayerController;
+	Proxy->WorldContextObject = WorldContextObject;
 	return Proxy;
 }
 
 void UConnectionCallbackProxy::Activate()
 {
-	FOnlineSubsystemBPCallHelper Helper(TEXT("ConnectToService"));
+	FOnlineSubsystemBPCallHelper Helper(TEXT("ConnectToService"), GEngine->GetWorldFromContextObject(WorldContextObject));
 	Helper.QueryIDFromPlayerController(PlayerControllerWeakPtr.Get());
 
 	if (Helper.IsValid())
@@ -57,7 +59,7 @@ void UConnectionCallbackProxy::Activate()
 
 void UConnectionCallbackProxy::OnLoginCompleted(int32 LocalUserNum, bool bWasSuccessful, const FUniqueNetId& UserId, const FString& Error)
 {
-	FOnlineSubsystemBPCallHelper Helper(TEXT("ConnectToService"));
+	FOnlineSubsystemBPCallHelper Helper(TEXT("ConnectToService"), GEngine->GetWorldFromContextObject(WorldContextObject));
 	Helper.QueryIDFromPlayerController(PlayerControllerWeakPtr.Get());
 
 	if (Helper.IsValid())
diff --git a/Engine/Source/Runtime/Online/OnlineSubsystemUtils/Private/OnlineSubsystemUtils.cpp b/Engine/Source/Runtime/Online/OnlineSubsystemUtils/Private/OnlineSubsystemUtils.cpp
--- a/Engine/Source/Runtime/Online/OnlineSubsystemUtils/Private/OnlineSubsystemUtils.cpp
+++ b/Engine/Source/Runtime/Online/OnlineSubsystemUtils/Private/OnlineSubsystemUtils.cpp
@@ -327,8 +327,8 @@ FStaticSelfRegisteringExec OnlineExecRegistration(OnlineExec);
 //////////////////////////////////////////////////////////////////////////
 // FOnlineSubsystemBPCallHelper
 
-FOnlineSubsystemBPCallHelper::FOnlineSubsystemBPCallHelper(const TCHAR* CallFunctionContext, FName SystemName)
-	: OnlineSub(IOnlineSubsystem::Get(SystemName))
+FOnlineSubsystemBPCallHelper::FOnlineSubsystemBPCallHelper(const TCHAR* CallFunctionContext, UWorld* World, FName SystemName)
+	: OnlineSub(Online::GetSubsystem(World, SystemName))
 	, FunctionContext(CallFunctionContext)
 {
 	if (OnlineSub == nullptr)
diff --git a/Engine/Source/Runtime/Online/OnlineSubsystemUtils/Private/OnlineSubsystemUtilsPrivatePCH.h b/Engine/Source/Runtime/Online/OnlineSubsystemUtils/Private/OnlineSubsystemUtilsPrivatePCH.h
--- a/Engine/Source/Runtime/Online/OnlineSubsystemUtils/Private/OnlineSubsystemUtilsPrivatePCH.h
+++ b/Engine/Source/Runtime/Online/OnlineSubsystemUtils/Private/OnlineSubsystemUtilsPrivatePCH.h
@@ -16,7 +16,7 @@
 struct FOnlineSubsystemBPCallHelper
 {
 public:
-	FOnlineSubsystemBPCallHelper(const TCHAR* CallFunctionContext, FName SystemName = NAME_None);
+	FOnlineSubsystemBPCallHelper(const TCHAR* CallFunctionContext, UWorld* World, FName SystemName = NAME_None);
 
 	void QueryIDFromPlayerController(APlayerController* PlayerController);
 
