Add basic support for online sessions in blueprints
Add blueprint nodes for creating, destroying, finding, and joining sessions
UE4Game can now use OnlineSubsystemNull and OnlineSubsystemSteam on desktop platforms

diff --git a/Engine/Source/Runtime/Online/OnlineSubsystemUtils/Classes/CreateSessionCallbackProxy.h b/Engine/Source/Runtime/Online/OnlineSubsystemUtils/Classes/CreateSessionCallbackProxy.h
new file mode 100644
--- /dev/null
+++ b/Engine/Source/Runtime/Online/OnlineSubsystemUtils/Classes/CreateSessionCallbackProxy.h
@@ -0,0 +1,52 @@
+// Copyright 1998-2014 Epic Games, Inc. All Rights Reserved.
+#pragma once
+
+#include "OnlineSessionInterface.h"
+#include "CreateSessionCallbackProxy.generated.h"
+
+UCLASS(MinimalAPI)
+class UCreateSessionCallbackProxy : public UOnlineBlueprintCallProxyBase
+{
+	GENERATED_UCLASS_BODY()
+
+	// Called when the session was created successfully
+	UPROPERTY(BlueprintAssignable)
+	FEmptyOnlineDelegate OnSuccess;
+
+	// Called when there was an error creating the session
+	UPROPERTY(BlueprintAssignable)
+	FEmptyOnlineDelegate OnFailure;
+
+	// Creates a session with the default online subsystem
+	UFUNCTION(BlueprintCallable, meta=(BlueprintInternalUseOnly = "true", WorldContext="WorldContextObject"), Category = "Online|Session")
+	static UCreateSessionCallbackProxy* CreateSession(UObject* WorldContextObject, class APlayerController* PlayerController, int32 PublicConnections, bool bUseLAN);
+
+	// UOnlineBlueprintCallProxyBase interface
+	virtual void Activate() override;
+	// End of UOnlineBlueprintCallProxyBase interface
+
+private:
+	// Internal callback when session creation completes, calls StartSession
+	void OnCreateCompleted(FName SessionName, bool bWasSuccessful);
+
+	// Internal callback when session creation completes, calls StartSession
+	void OnStartCompleted(FName SessionName, bool bWasSuccessful);
+
+	// The player controller triggering things
+	TWeakObjectPtr<APlayerController> PlayerControllerWeakPtr;
+
+	// The delegate executed by the online subsystem
+	FOnCreateSessionCompleteDelegate CreateCompleteDelegate;
+
+	// The delegate executed by the online subsystem
+	FOnStartSessionCompleteDelegate StartCompleteDelegate;
+
+	// Number of public connections
+	int NumPublicConnections;
+
+	// Whether or not to search LAN
+	bool bUseLAN;
+
+	// The world context object in which this call is taking place
+	UObject* WorldContextObject;
+};
diff --git a/Engine/Source/Runtime/Online/OnlineSubsystemUtils/Classes/DestroySessionCallbackProxy.h b/Engine/Source/Runtime/Online/OnlineSubsystemUtils/Classes/DestroySessionCallbackProxy.h
new file mode 100644
--- /dev/null
+++ b/Engine/Source/Runtime/Online/OnlineSubsystemUtils/Classes/DestroySessionCallbackProxy.h
@@ -0,0 +1,41 @@
+// Copyright 1998-2014 Epic Games, Inc. All Rights Reserved.
+#pragma once
+
+#include "OnlineSessionInterface.h"
+#include "DestroySessionCallbackProxy.generated.h"
+
+UCLASS(MinimalAPI)
+class UDestroySessionCallbackProxy : public UOnlineBlueprintCallProxyBase
+{
+	GENERATED_UCLASS_BODY()
+
+	// Called when there is a successful destroy
+	UPROPERTY(BlueprintAssignable)
+	FEmptyOnlineDelegate OnSuccess;
+
+	// Called when there is an unsuccessful destroy
+	UPROPERTY(BlueprintAssignable)
+	FEmptyOnlineDelegate OnFailure;
+
+	// Destroys a session previously created by the default online subsystem
+	UFUNCTION(BlueprintCallable, meta=(BlueprintInternalUseOnly = "true", WorldContext="WorldContextObject"), Category = "Online|Session")
+	static UDestroySessionCallbackProxy* DestroySession(UObject* WorldContextObject, class APlayerController* PlayerController);
+
+	// UOnlineBlueprintCallProxyBase interface
+	virtual void Activate() override;
+	// End of UOnlineBlueprintCallProxyBase interface
+
+private:
+	// Internal callback when the operation completes, calls out to the public success/failure callbacks
+	void OnCompleted(FName SessionName, bool bWasSuccessful);
+
+private:
+	// The player controller triggering things
+	TWeakObjectPtr<APlayerController> PlayerControllerWeakPtr;
+
+	// The delegate executed by the online subsystem
+	FOnDestroySessionCompleteDelegate Delegate;
+
+	// The world context object in which this call is taking place
+	UObject* WorldContextObject;
+};
diff --git a/Engine/Source/Runtime/Online/OnlineSubsystemUtils/Classes/FindSessionsCallbackProxy.h b/Engine/Source/Runtime/Online/OnlineSubsystemUtils/Classes/FindSessionsCallbackProxy.h
new file mode 100644
--- /dev/null
+++ b/Engine/Source/Runtime/Online/OnlineSubsystemUtils/Classes/FindSessionsCallbackProxy.h
@@ -0,0 +1,72 @@
+// Copyright 1998-2014 Epic Games, Inc. All Rights Reserved.
+#pragma once
+
+#include "OnlineSessionInterface.h"
+#include "FindSessionsCallbackProxy.generated.h"
+
+USTRUCT(BlueprintType)
+struct FBlueprintSessionResult
+{
+	GENERATED_USTRUCT_BODY()
+
+	FOnlineSessionSearchResult OnlineResult;
+};
+
+DECLARE_DYNAMIC_MULTICAST_DELEGATE_OneParam(FBlueprintFindSessionsResultDelegate, const TArray<FBlueprintSessionResult>&, Results);
+
+UCLASS(MinimalAPI)
+class UFindSessionsCallbackProxy : public UOnlineBlueprintCallProxyBase
+{
+	GENERATED_UCLASS_BODY()
+
+	// Called when there is a successful query
+	UPROPERTY(BlueprintAssignable)
+	FBlueprintFindSessionsResultDelegate OnSuccess;
+
+	// Called when there is an unsuccessful query
+	UPROPERTY(BlueprintAssignable)
+	FBlueprintFindSessionsResultDelegate OnFailure;
+
+	// Searches for advertised sessions with the default online subsystem
+	UFUNCTION(BlueprintCallable, meta=(BlueprintInternalUseOnly = "true", WorldContext="WorldContextObject"), Category = "Online|Session")
+	static UFindSessionsCallbackProxy* FindSessions(UObject* WorldContextObject, class APlayerController* PlayerController, int32 MaxResults, bool bUseLAN);
+
+	UFUNCTION(BlueprintPure, Category = "Online|Session")
+	static int32 GetPingInMs(const FBlueprintSessionResult& Result);
+
+	UFUNCTION(BlueprintPure, Category = "Online|Session")
+	static FString GetServerName(const FBlueprintSessionResult& Result);
+
+	UFUNCTION(BlueprintPure, Category = "Online|Session")
+	static int32 GetCurrentPlayers(const FBlueprintSessionResult& Result);
+
+	UFUNCTION(BlueprintPure, Category = "Online|Session")
+	static int32 GetMaxPlayers(const FBlueprintSessionResult& Result);
+
+	// UOnlineBlueprintCallProxyBase interface
+	virtual void Activate() override;
+	// End of UOnlineBlueprintCallProxyBase interface
+
+private:
+	// Internal callback when the session search completes, calls out to the public success/failure callbacks
+	void OnCompleted(bool bSuccess);
+
+private:
+	// The player controller triggering things
+	TWeakObjectPtr<APlayerController> PlayerControllerWeakPtr;
+
+	// The delegate executed by the online subsystem
+	FOnFindSessionsCompleteDelegate Delegate;
+
+	// Object to track search results
+	TSharedPtr<FOnlineSessionSearch> SearchObject;
+
+	// Whether or not to search LAN
+	bool bUseLAN;
+
+	// Maximum number of results to return
+	int MaxResults;
+
+	// The world context object in which this call is taking place
+	UObject* WorldContextObject;
+};
diff --git a/Engine/Source/Runtime/Online/OnlineSubsystemUtils/Classes/JoinSessionCallbackProxy.h b/Engine/Source/Runtime/Online/OnlineSubsystemUtils/Classes/JoinSessionCallbackProxy.h
new file mode 100644
--- /dev/null
+++ b/Engine/Source/Runtime/Online/OnlineSubsystemUtils/Classes/JoinSessionCallbackProxy.h
@@ -0,0 +1,44 @@
+// Copyright 1998-2014 Epic Games, Inc. All Rights Reserved.
+#pragma once
+
+#include "OnlineSessionInterface.h"
+#include "JoinSessionCallbackProxy.generated.h"
+
+UCLASS(MinimalAPI)
+class UJoinSessionCallbackProxy : public UOnlineBlueprintCallProxyBase
+{
+	GENERATED_UCLASS_BODY()
+
+	// Called when there is a successful join
+	UPROPERTY(BlueprintAssignable)
+	FEmptyOnlineDelegate OnSuccess;
+
+	// Called when there is an unsuccessful join
+	UPROPERTY(BlueprintAssignable)
+	FEmptyOnlineDelegate OnFailure;
+
+	// Joins a remote session with the default online subsystem
+	UFUNCTION(BlueprintCallable, meta=(BlueprintInternalUseOnly = "true", WorldContext="WorldContextObject"), Category = "Online|Session")
+	static UJoinSessionCallbackProxy* JoinSession(UObject* WorldContextObject, class APlayerController* PlayerController, const FBlueprintSessionResult& SearchResult);
+
+	// UOnlineBlueprintCallProxyBase interface
+	virtual void Activate() override;
+	// End of UOnlineBlueprintCallProxyBase interface
+
+private:
+	// Internal callback when the join completes, calls out to the public success/failure callbacks
+	void OnCompleted(FName SessionName, EOnJoinSessionCompleteResult::Type Result);
+
+private:
+	// The player controller triggering things
+	TWeakObjectPtr<APlayerController> PlayerControllerWeakPtr;
+
+	// The search result we are sttempting to join
+	FOnlineSessionSearchResult OnlineSearchResult;
+
+	// The delegate executed by the online subsystem
+	FOnJoinSessionCompleteDelegate Delegate;
+
+	// The world context object in which this call is taking place
+	UObject* WorldContextObject;
+};
diff --git a/Engine/Source/Runtime/Online/OnlineSubsystemUtils/Private/CreateSessionCallbackProxy.cpp b/Engine/Source/Runtime/Online/OnlineSubsystemUtils/Private/CreateSessionCallbackProxy.cpp
new file mode 100644
--- /dev/null
+++ b/Engine/Source/Runtime/Online/OnlineSubsystemUtils/Private/CreateSessionCallbackProxy.cpp
@@ -0,0 +1,113 @@
+//// Copyright 1998-2014 Epic Games, Inc. All Rights Reserved.
+
+#include "OnlineSubsystemUtilsPrivatePCH.h"
+#include "CreateSessionCallbackProxy.h"
+
+//////////////////////////////////////////////////////////////////////////
+// UCreateSessionCallbackProxy
+
+UCreateSessionCallbackProxy::UCreateSessionCallbackProxy(const FObjectInitializer& ObjectInitializer)
+	: Super(ObjectInitializer)
+	, CreateCompleteDelegate(FOnCreateSessionCompleteDelegate::CreateUObject(this, &ThisClass::OnCreateCompleted))
+	, StartCompleteDelegate(FOnStartSessionCompleteDelegate::CreateUObject(this, &ThisClass::OnStartCompleted))
+	, NumPublicConnections(1)
+{
+}
+
+UCreateSessionCallbackProxy* UCreateSessionCallbackProxy::CreateSession(UObject* WorldContextObject, class APlayerController* PlayerController, int32 PublicConnections, bool bUseLAN)
+{
+	UCreateSessionCallbackProxy* Proxy = NewObject<UCreateSessionCallbackProxy>();
+	Proxy->PlayerControllerWeakPtr = PlayerController;
+	Proxy->NumPublicConnections = PublicConnections;
+	Proxy->bUseLAN = bUseLAN;
+	Proxy->WorldContextObject = WorldContextObject;
+	return Proxy;
+}
+
+void UCreateSessionCallbackProxy::Activate()
+{
+	FOnlineSubsystemBPCallHelper Helper(TEXT("CreateSession"), GEngine->GetWorldFromContextObject(WorldContextObject));
+	Helper.QueryIDFromPlayerController(PlayerControllerWeakPtr.Get());
+
+	if (Helper.IsValid())
+	{
+		auto Sessions = Helper.OnlineSub->GetSessionInterface();
+		if (Sessions.IsValid())
+		{
+			Sessions->AddOnCreateSessionCompleteDelegate(CreateCompleteDelegate);
+
+			FOnlineSessionSettings Settings;
+			Settings.NumPublicConnections = NumPublicConnections;
+			Settings.bShouldAdvertise = true;
+			Settings.bAllowJoinInProgress = true;
+			Settings.bIsLANMatch = bUseLAN;
+			Settings.bUsesPresence = true;
+			Settings.bAllowJoinViaPresence = true;
+
+			Sessions->CreateSession(*Helper.UserID, GameSessionName, Settings);
+
+			// OnCreateCompleted will get called, nothing more to do now
+			return;
+		}
+		else
+		{
+			FFrame::KismetExecutionMessage(TEXT("Sessions not supported by Online Subsystem"), ELogVerbosity::Warning);
+		}
+	}
+
+	// Fail immediately
+	OnFailure.Broadcast();
+}
+
+void UCreateSessionCallbackProxy::OnCreateCompleted(FName SessionName, bool bWasSuccessful)
+{
+	FOnlineSubsystemBPCallHelper Helper(TEXT("CreateSessionCallback"), GEngine->GetWorldFromContextObject(WorldContextObject));
+	Helper.QueryIDFromPlayerController(PlayerControllerWeakPtr.Get());
+
+	if (Helper.IsValid())
+	{
+		auto Sessions = Helper.OnlineSub->GetSessionInterface();
+		if (Sessions.IsValid())
+		{
+			Sessions->ClearOnCreateSessionCompleteDelegate(CreateCompleteDelegate);
+			
+			if (bWasSuccessful)
+			{
+				Sessions->AddOnStartSessionCompleteDelegate(StartCompleteDelegate);
+				Sessions->StartSession(GameSessionName);
+
+				// OnStartCompleted will get called, nothing more to do now
+				return;
+			}
+		}
+	}
+
+	if (!bWasSuccessful)
+	{
+		OnFailure.Broadcast();
+	}
+}
+
+void UCreateSessionCallbackProxy::OnStartCompleted(FName SessionName, bool bWasSuccessful)
+{
+	FOnlineSubsystemBPCallHelper Helper(TEXT("StartSessionCallback"), GEngine->GetWorldFromContextObject(WorldContextObject));
+	Helper.QueryIDFromPlayerController(PlayerControllerWeakPtr.Get());
+
+	if (Helper.IsValid())
+	{
+		auto Sessions = Helper.OnlineSub->GetSessionInterface();
+		if (Sessions.IsValid())
+		{
+			Sessions->ClearOnStartSessionCompleteDelegate(StartCompleteDelegate);
+		}
+	}
+
+	if (bWasSuccessful)
+	{
+		OnSuccess.Broadcast();
+	}
+	else
+	{
+		OnFailure.Broadcast();
+	}
+}
diff --git a/Engine/Source/Runtime/Online/OnlineSubsystemUtils/Private/DestroySessionCallbackProxy.cpp b/Engine/Source/Runtime/Online/OnlineSubsystemUtils/Private/DestroySessionCallbackProxy.cpp
new file mode 100644
--- /dev/null
+++ b/Engine/Source/Runtime/Online/OnlineSubsystemUtils/Private/DestroySessionCallbackProxy.cpp
@@ -0,0 +1,71 @@
+// Copyright 1998-2014 Epic Games, Inc. All Rights Reserved.
+
+#include "OnlineSubsystemUtilsPrivatePCH.h"
+#include "DestroySessionCallbackProxy.h"
+
+//////////////////////////////////////////////////////////////////////////
+// UDestroySessionCallbackProxy
+
+UDestroySessionCallbackProxy::UDestroySessionCallbackProxy(const FObjectInitializer& ObjectInitializer)
+	: Super(ObjectInitializer)
+	, Delegate(FOnDestroySessionCompleteDelegate::CreateUObject(this, &ThisClass::OnCompleted))
+{
+}
+
+UDestroySessionCallbackProxy* UDestroySessionCallbackProxy::DestroySession(UObject* WorldContextObject, class APlayerController* PlayerController)
+{
+	UDestroySessionCallbackProxy* Proxy = NewObject<UDestroySessionCallbackProxy>();
+	Proxy->PlayerControllerWeakPtr = PlayerController;
+	Proxy->WorldContextObject = WorldContextObject;
+	return Proxy;
+}
+
+void UDestroySessionCallbackProxy::Activate()
+{
+	FOnlineSubsystemBPCallHelper Helper(TEXT("DestroySession"), GEngine->GetWorldFromContextObject(WorldContextObject));
+	Helper.QueryIDFromPlayerController(PlayerControllerWeakPtr.Get());
+
+	if (Helper.IsValid())
+	{
+		auto Sessions = Helper.OnlineSub->GetSessionInterface();
+		if (Sessions.IsValid())
+		{
+			Sessions->AddOnDestroySessionCompleteDelegate(Delegate);
+			Sessions->DestroySession(GameSessionName);
+
+			// OnCompleted will get called, nothing more to do now
+			return;
+		}
+		else
+		{
+			FFrame::KismetExecutionMessage(TEXT("Sessions not supported by Online Subsystem"), ELogVerbosity::Warning);
+		}
+	}
+
+	// Fail immediately
+	OnFailure.Broadcast();
+}
+
+void UDestroySessionCallbackProxy::OnCompleted(FName SessionName, bool bWasSuccessful)
+{
+	FOnlineSubsystemBPCallHelper Helper(TEXT("DestroySessionCallback"), GEngine->GetWorldFromContextObject(WorldContextObject));
+	Helper.QueryIDFromPlayerController(PlayerControllerWeakPtr.Get());
+
+	if (Helper.IsValid())
+	{
+		auto Sessions = Helper.OnlineSub->GetSessionInterface();
+		if (Sessions.IsValid())
+		{
+			Sessions->ClearOnDestroySessionCompleteDelegate(Delegate);
+		}
+	}
+
+	if (bWasSuccessful)
+	{
+		OnSuccess.Broadcast();
+	}
+	else
+	{
+		OnFailure.Broadcast();
+	}
+}
diff --git a/Engine/Source/Runtime/Online/OnlineSubsystemUtils/Private/FindSessionsCallbackProxy.cpp b/Engine/Source/Runtime/Online/OnlineSubsystemUtils/Private/FindSessionsCallbackProxy.cpp
new file mode 100644
--- /dev/null
+++ b/Engine/Source/Runtime/Online/OnlineSubsystemUtils/Private/FindSessionsCallbackProxy.cpp
@@ -0,0 +1,115 @@
+// Copyright 1998-2014 Epic Games, Inc. All Rights Reserved.
+
+#include "OnlineSubsystemUtilsPrivatePCH.h"
+#include "FindSessionsCallbackProxy.h"
+
+//////////////////////////////////////////////////////////////////////////
+// UFindSessionsCallbackProxy
+
+UFindSessionsCallbackProxy::UFindSessionsCallbackProxy(const FObjectInitializer& ObjectInitializer)
+	: Super(ObjectInitializer)
+	, Delegate(FOnFindSessionsCompleteDelegate::CreateUObject(this, &ThisClass::OnCompleted))
+	, bUseLAN(false)
+{
+}
+
+UFindSessionsCallbackProxy* UFindSessionsCallbackProxy::FindSessions(UObject* WorldContextObject, class APlayerController* PlayerController, int MaxResults, bool bUseLAN)
+{
+	UFindSessionsCallbackProxy* Proxy = NewObject<UFindSessionsCallbackProxy>();
+	Proxy->PlayerControllerWeakPtr = PlayerController;
+	Proxy->bUseLAN = bUseLAN;
+	Proxy->MaxResults = MaxResults;
+	Proxy->WorldContextObject = WorldContextObject;
+	return Proxy;
+}
+
+void UFindSessionsCallbackProxy::Activate()
+{
+	FOnlineSubsystemBPCallHelper Helper(TEXT("FindSessions"), GEngine->GetWorldFromContextObject(WorldContextObject));
+	Helper.QueryIDFromPlayerController(PlayerControllerWeakPtr.Get());
+
+	if (Helper.IsValid())
+	{
+		auto Sessions = Helper.OnlineSub->GetSessionInterface();
+		if (Sessions.IsValid())
+		{
+			Sessions->AddOnFindSessionsCompleteDelegate(Delegate);
+			
+			SearchObject = MakeShareable(new FOnlineSessionSearch);
+			SearchObject->MaxSearchResults = MaxResults;
+			SearchObject->bIsLanQuery = bUseLAN;
+			SearchObject->QuerySettings.Set(SEARCH_PRESENCE, true, EOnlineComparisonOp::Equals);
+
+			Sessions->FindSessions(*Helper.UserID, SearchObject.ToSharedRef());
+
+			// OnQueryCompleted will get called, nothing more to do now
+			return;
+		}
+		else
+		{
+			FFrame::KismetExecutionMessage(TEXT("Sessions not supported by Online Subsystem"), ELogVerbosity::Warning);
+		}
+	}
+
+	// Fail immediately
+	TArray<FBlueprintSessionResult> Results;
+	OnFailure.Broadcast(Results);
+}
+
+void UFindSessionsCallbackProxy::OnCompleted(bool bSuccess)
+{
+	FOnlineSubsystemBPCallHelper Helper(TEXT("FindSessionsCallback"), GEngine->GetWorldFromContextObject(WorldContextObject));
+	Helper.QueryIDFromPlayerController(PlayerControllerWeakPtr.Get());
+
+	if (Helper.IsValid())
+	{
+		auto Sessions = Helper.OnlineSub->GetSessionInterface();
+		if (Sessions.IsValid())
+		{
+			Sessions->ClearOnFindSessionsCompleteDelegate(Delegate);
+		}
+	}
+
+	TArray<FBlueprintSessionResult> Results;
+
+	if (bSuccess && SearchObject.IsValid())
+	{
+		// Just log the results for now, will need to add a blueprint-compatible search result struct
+		for (auto& Result : SearchObject->SearchResults)
+		{
+			FString ResultText = FString::Printf(TEXT("Found a session. Ping is %d"), Result.PingInMs);
+
+			FFrame::KismetExecutionMessage(*ResultText, ELogVerbosity::Log);
+
+			FBlueprintSessionResult BPResult;
+			BPResult.OnlineResult = Result;
+			Results.Add(BPResult);
+		}
+
+		OnSuccess.Broadcast(Results);
+	}
+	else
+	{
+		OnFailure.Broadcast(Results);
+	}
+}
+
+int32 UFindSessionsCallbackProxy::GetPingInMs(const FBlueprintSessionResult& Result)
+{
+	return Result.OnlineResult.PingInMs;
+}
+
+FString UFindSessionsCallbackProxy::GetServerName(const FBlueprintSessionResult& Result)
+{
+	return Result.OnlineResult.Session.OwningUserName;
+}
+
+int32 UFindSessionsCallbackProxy::GetCurrentPlayers(const FBlueprintSessionResult& Result)
+{
+	return Result.OnlineResult.Session.SessionSettings.NumPublicConnections - Result.OnlineResult.Session.NumOpenPublicConnections;
+}
+
+int32 UFindSessionsCallbackProxy::GetMaxPlayers(const FBlueprintSessionResult& Result)
+{
+	return Result.OnlineResult.Session.SessionSettings.NumPublicConnections;
+}
diff --git a/Engine/Source/Runtime/Online/OnlineSubsystemUtils/Private/JoinSessionCallbackProxy.cpp b/Engine/Source/Runtime/Online/OnlineSubsystemUtils/Private/JoinSessionCallbackProxy.cpp
new file mode 100644
--- /dev/null
+++ b/Engine/Source/Runtime/Online/OnlineSubsystemUtils/Private/JoinSessionCallbackProxy.cpp
@@ -0,0 +1,78 @@
+// Copyright 1998-2014 Epic Games, Inc. All Rights Reserved.
+
+#include "OnlineSubsystemUtilsPrivatePCH.h"
+#include "JoinSessionCallbackProxy.h"
+
+//////////////////////////////////////////////////////////////////////////
+// UJoinSessionCallbackProxy
+
+UJoinSessionCallbackProxy::UJoinSessionCallbackProxy(const FObjectInitializer& ObjectInitializer)
+	: Super(ObjectInitializer)
+	, Delegate(FOnJoinSessionCompleteDelegate::CreateUObject(this, &ThisClass::OnCompleted))
+{
+}
+
+UJoinSessionCallbackProxy* UJoinSessionCallbackProxy::JoinSession(UObject* WorldContextObject, class APlayerController* PlayerController, const FBlueprintSessionResult& SearchResult)
+{
+	UJoinSessionCallbackProxy* Proxy = NewObject<UJoinSessionCallbackProxy>();
+	Proxy->PlayerControllerWeakPtr = PlayerController;
+	Proxy->OnlineSearchResult = SearchResult.OnlineResult;
+	Proxy->WorldContextObject = WorldContextObject;
+	return Proxy;
+}
+
+void UJoinSessionCallbackProxy::Activate()
+{
+	FOnlineSubsystemBPCallHelper Helper(TEXT("JoinSession"), GEngine->GetWorldFromContextObject(WorldContextObject));
+	Helper.QueryIDFromPlayerController(PlayerControllerWeakPtr.Get());
+
+	if (Helper.IsValid())
+	{
+		auto Sessions = Helper.OnlineSub->GetSessionInterface();
+		if (Sessions.IsValid())
+		{
+			Sessions->AddOnJoinSessionCompleteDelegate(Delegate);
+			Sessions->JoinSession(*Helper.UserID, GameSessionName, OnlineSearchResult);
+
+			// OnCompleted will get called, nothing more to do now
+			return;
+		}
+		else
+		{
+			FFrame::KismetExecutionMessage(TEXT("Sessions not supported by Online Subsystem"), ELogVerbosity::Warning);
+		}
+	}
+
+	// Fail immediately
+	OnFailure.Broadcast();
+}
+
+void UJoinSessionCallbackProxy::OnCompleted(FName SessionName, EOnJoinSessionCompleteResult::Type Result)
+{
+	FOnlineSubsystemBPCallHelper Helper(TEXT("JoinSessionCallback"), GEngine->GetWorldFromContextObject(WorldContextObject));
+	Helper.QueryIDFromPlayerController(PlayerControllerWeakPtr.Get());
+
+	if (Helper.IsValid())
+	{
+		auto Sessions = Helper.OnlineSub->GetSessionInterface();
+		if (Sessions.IsValid())
+		{
+			Sessions->ClearOnJoinSessionCompleteDelegate(Delegate);
+
+			if (Result == EOnJoinSessionCompleteResult::Success)
+			{
+				// Client travel to the server
+				FString ConnectString;
+				if (Sessions->GetResolvedConnectString(GameSessionName, ConnectString) && PlayerControllerWeakPtr.IsValid())
+				{
+					UE_LOG(LogOnline, Log, TEXT("Join session: traveling to %s"), *ConnectString);
+					PlayerControllerWeakPtr->ClientTravel(ConnectString, TRAVEL_Absolute);
+					OnSuccess.Broadcast();
+					return;
+				}
+			}
+		}
+	}
+
+	OnFailure.Broadcast();
+}
diff --git a/Engine/Source/Runtime/UE4Game/UE4Game.Build.cs b/Engine/Source/Runtime/UE4Game/UE4Game.Build.cs
--- a/Engine/Source/Runtime/UE4Game/UE4Game.Build.cs
+++ b/Engine/Source/Runtime/UE4Game/UE4Game.Build.cs
@@ -8,6 +8,17 @@ public class UE4Game : ModuleRules
 	{
 		PrivateDependencyModuleNames.Add("Core");
 	
+		DynamicallyLoadedModuleNames.Add("OnlineSubsystemNull");
+
+		if (Target.Platform == UnrealTargetPlatform.Win32 ||
+			Target.Platform == UnrealTargetPlatform.Win64 ||
+			Target.Platform == UnrealTargetPlatform.Mac)
+		{
+			if (UEBuildConfiguration.bCompileSteamOSS == true)
+			{
+				DynamicallyLoadedModuleNames.Add("OnlineSubsystemSteam");
+			}
+		}
 		if (Target.Platform == UnrealTargetPlatform.IOS)
 		{
 			PrivateDependencyModuleNames.AddRange(new string[] { "Core", "CoreUObject", "Engine", "OnlineSubsystem", "OnlineSubsystemUtils" });
