Merging using UE4-To-UE4-LauncherReleases REVERSE

CL#2343749 on 29/10/2014 09:26:43
------------------------------------------
- Implement file data splitting and compression support for client
- Fixing a shutdown crash if still scanning installed manifests
- Adding timer logging for manifest loading

diff --git a/Engine/Source/Runtime/Online/BuildPatchServices/Private/BuildPatchChunk.h b/Engine/Source/Runtime/Online/BuildPatchServices/Private/BuildPatchChunk.h
--- a/Engine/Source/Runtime/Online/BuildPatchServices/Private/BuildPatchChunk.h
+++ b/Engine/Source/Runtime/Online/BuildPatchServices/Private/BuildPatchChunk.h
@@ -26,6 +26,8 @@ struct FSHAHashData
 	FString ToString() const;
 };
 
+static_assert(FSHA1::DigestSize == 20, "If this changes a lot of stuff here will break!");
+
 /**
  * Constant values and typedefs
  */
diff --git a/Engine/Source/Runtime/Online/BuildPatchServices/Private/BuildPatchDownloader.cpp b/Engine/Source/Runtime/Online/BuildPatchServices/Private/BuildPatchDownloader.cpp
--- a/Engine/Source/Runtime/Online/BuildPatchServices/Private/BuildPatchDownloader.cpp
+++ b/Engine/Source/Runtime/Online/BuildPatchServices/Private/BuildPatchDownloader.cpp
@@ -88,7 +88,14 @@ uint32 FBuildPatchDownloader::Run()
 				if( bSuccess )
 				{
 					// Uncompress if needed
-					bSuccess = !bIsChunkData || FBuildPatchUtils::UncompressChunkFile( InFlightJob.DataArray );
+					if (bIsChunkData)
+					{
+						bSuccess = FBuildPatchUtils::UncompressChunkFile(InFlightJob.DataArray);
+					}
+					else
+					{
+						bSuccess = FBuildPatchUtils::UncompressFileDataFile(InFlightJob.DataArray);
+					}
 					if( !bSuccess )
 					{
 						FBuildPatchAnalytics::RecordChunkDownloadError( InFlightJob.DownloadUrl, FPlatformMisc::GetLastError(), TEXT( "Uncompress Fail" ) );
@@ -188,29 +195,7 @@ uint32 FBuildPatchDownloader::Run()
 		SetIdle( false );
 
 		// Make the filename and download url
-		FString DownloadUrl;
-		// Support older manifests
-		if( InstallManifest->GetManifestVersion() < EBuildPatchAppManifestVersion::DataFileRenames )
-		{
-			DownloadUrl = FBuildPatchUtils::GetDataTypeOldFilename( DataType, FBuildPatchServicesModule::GetCloudDirectory(), NextGuid );
-		}
-		else if( bIsChunkData )
-		{
-			uint64 ChunkHash;
-			const bool bFound = InstallManifest->GetChunkHash( NextGuid, ChunkHash );
-			// Should be impossible to not exist
-			check( bFound );
-			DownloadUrl = FBuildPatchUtils::GetChunkNewFilename( InstallManifest->GetManifestVersion(), FBuildPatchServicesModule::GetCloudDirectory(), NextGuid, ChunkHash );
-		}
-		else
-		{
-			FSHAHashData FileHash;
-			const bool bFound = InstallManifest->GetFileDataHash( NextGuid, FileHash );
-			// Should be impossible to not exist
-			check( bFound );
-			DownloadUrl = FBuildPatchUtils::GetFileNewFilename( InstallManifest->GetManifestVersion(), FBuildPatchServicesModule::GetCloudDirectory(), NextGuid, FileHash );
-		}
-
+		FString DownloadUrl = FBuildPatchUtils::GetDataFilename(InstallManifest, FBuildPatchServicesModule::GetCloudDirectory(), NextGuid);
 		const bool bIsHTTPRequest = DownloadUrl.Contains( TEXT( "http" ), ESearchCase::IgnoreCase );
 
 		// Start the download
diff --git a/Engine/Source/Runtime/Online/BuildPatchServices/Private/BuildPatchFileConstructor.cpp b/Engine/Source/Runtime/Online/BuildPatchServices/Private/BuildPatchFileConstructor.cpp
--- a/Engine/Source/Runtime/Online/BuildPatchServices/Private/BuildPatchFileConstructor.cpp
+++ b/Engine/Source/Runtime/Online/BuildPatchServices/Private/BuildPatchFileConstructor.cpp
@@ -516,95 +516,81 @@ bool FBuildPatchFileConstructor::ConstructFileFromChunks( const FString& Filenam
 bool FBuildPatchFileConstructor::InsertFileData(const FChunkPartData& ChunkPart, FArchive& DestinationFile, FSHA1& HashState)
 {
 	bool bSuccess = false;
+	bool bLogged = false;
 
 	// Wait for the file data to be available
 	while( IsFileDataAvailable( ChunkPart.Guid ) == false )
 	{
 		FPlatformProcess::Sleep( 0.1f );
 	}
-	FString Filename = GetFileDataFilename( ChunkPart.Guid );
 
 	// Read the file
-	FArchive* FileReader = IFileManager::Get().CreateFileReader( *Filename );
-	bSuccess = FileReader != NULL;
-	if( bSuccess )
+	TArray<uint8> FileData;
+	FChunkHeader Header;
+	const FString DataFilename = GetFileDataFilename(ChunkPart.Guid);
+	bSuccess = FFileHelper::LoadFileToArray(FileData, *DataFilename);
+	if (!bSuccess && !bLogged)
 	{
-		const int64 FileSize = FileReader->TotalSize();
-		// Read the header
-		FChunkHeader Header;
-		*FileReader << Header;
-		// Check header magic
-		if ( !Header.IsValidMagic() )
-		{
-			bSuccess = false;
-			FBuildPatchAnalytics::RecordConstructionError( Filename, INDEX_NONE, TEXT( "File Data Magic Fail" ) );
-			GLog->Logf( TEXT( "BuildPatchFileConstructor: ERROR: InsertFileData: magic failed for %s" ), *FPaths::GetCleanFilename( Filename ) );
-		}
-		// Check we have the right kind of hash; file data must have SHA1
-		if ( bSuccess && ( Header.HashType != FChunkHeader::HASH_SHA1 ) )
-		{
-			bSuccess = false;
-			FBuildPatchAnalytics::RecordConstructionError( Filename, INDEX_NONE, TEXT( "File Data Incorrect Hash" ) );
-			GLog->Logf( TEXT( "BuildPatchFileConstructor: ERROR: InsertFileData: incorrect hashtype info for %s" ), *FPaths::GetCleanFilename( Filename ) );
-		}
-		// Check GUID matches if provided
-		if ( bSuccess && ChunkPart.Guid.IsValid() && ( ChunkPart.Guid != Header.Guid ) )
-		{
-			bSuccess = false;
-			FBuildPatchAnalytics::RecordConstructionError( Filename, INDEX_NONE, TEXT( "File Data Mismatch GUID" ) );
-			GLog->Logf( TEXT( "BuildPatchFileConstructor: ERROR: InsertFileData: mismatch guid for %s" ), *FPaths::GetCleanFilename( Filename ) );
-		}
-		// Continue if all was fine
-		if( bSuccess )
+		bLogged = true;
+		FBuildPatchAnalytics::RecordConstructionError(DataFilename, FPlatformMisc::GetLastError(), TEXT("File Data Missing"));
+		GLog->Logf(TEXT("BuildPatchFileConstructor: ERROR: InsertFileData could not open data file %s"), *DataFilename);
+	}
+	// Decompress data
+	bSuccess = bSuccess && FBuildPatchUtils::UncompressFileDataFile(FileData, &Header);
+	if (!bSuccess && !bLogged)
+	{
+		bLogged = true;
+		FBuildPatchAnalytics::RecordConstructionError(DataFilename, INDEX_NONE, TEXT("File Data Uncompress Fail"));
+		GLog->Logf(TEXT("BuildPatchFileConstructor: ERROR: InsertFileData: could not uncompress %s"), *FPaths::GetCleanFilename(DataFilename));
+	}
+	// Verify integrity
+	bSuccess = bSuccess && FBuildPatchUtils::VerifyChunkFile(FileData);
+	if (!bSuccess && !bLogged)
+	{
+		bLogged = true;
+		FBuildPatchAnalytics::RecordConstructionError(DataFilename, INDEX_NONE, TEXT("File Data Verify Fail"));
+		GLog->Logf(TEXT("BuildPatchFileConstructor: ERROR: InsertFileData: verification failed for %s"), *FPaths::GetCleanFilename(DataFilename));
+	}
+	// Check correct GUID
+	bSuccess = bSuccess && (!ChunkPart.Guid.IsValid() || (ChunkPart.Guid == Header.Guid));
+	if (!bSuccess && !bLogged)
+	{
+		bLogged = true;
+		FBuildPatchAnalytics::RecordConstructionError(DataFilename, INDEX_NONE, TEXT("File Data GUID Mismatch"));
+		GLog->Logf(TEXT("BuildPatchFileConstructor: ERROR: InsertFileData: mismatch GUID for %s"), *FPaths::GetCleanFilename(DataFilename));
+	}
+
+	// Continue if all was fine
+	if (bSuccess)
+	{
+		switch (Header.StoredAs)
 		{
-			// Load the data to check
-			uint8* FileReadBuffer = new uint8[ FileBufferSize ];
-			int64 DataRead = 0;
-			switch ( Header.StoredAs )
-			{
 			case FChunkHeader::STORED_RAW:
+			{
+				// Check we are able to get the chunk part
+				const int64 StartOfPartPos = Header.HeaderSize + ChunkPart.Offset;
+				const int64 EndOfPartPos = StartOfPartPos + ChunkPart.Size;
+				bSuccess = EndOfPartPos <= FileData.Num();
+				if (bSuccess)
 				{
-					// Check we are able to get the chunk part
-					const int64 StartOfPartPos = Header.HeaderSize + ChunkPart.Offset;
-					const int64 EndOfPartPos = StartOfPartPos + ChunkPart.Size;
-					bSuccess = EndOfPartPos <= FileSize;
-					if ( !bSuccess )
-					{
-						FBuildPatchAnalytics::RecordConstructionError( Filename, INDEX_NONE, TEXT( "File Data Part OOB" ) );
-						GLog->Logf( TEXT( "BuildPatchFileConstructor: ERROR: InsertFileData: part out of bounds for %s" ), *FPaths::GetCleanFilename( Filename ) );
-					}
-					else
-					{
-						// Seek to and copy the part across
-						FileReader->Seek( StartOfPartPos );
-						while( DataRead < ChunkPart.Size )
-						{
-							const int64 SizeLeft = ChunkPart.Size - DataRead;
-							const int64 ReadLen = FMath::Min< int64 >( FileBufferSize, SizeLeft );
-							FileReader->Serialize( FileReadBuffer, ReadLen );
-							HashState.Update( FileReadBuffer, ReadLen );
-							DestinationFile.Serialize( FileReadBuffer, ReadLen );
-							DataRead += ReadLen;
-						}
-					}
+					HashState.Update(FileData.GetData() + StartOfPartPos, ChunkPart.Size);
+					DestinationFile.Serialize(FileData.GetData() + StartOfPartPos, ChunkPart.Size);
+				}
+				else
+				{
+					FBuildPatchAnalytics::RecordConstructionError(DataFilename, INDEX_NONE, TEXT("File Data Part OOB"));
+					GLog->Logf(TEXT("BuildPatchFileConstructor: ERROR: InsertFileData: part out of bounds for %s"), *FPaths::GetCleanFilename(DataFilename));
 				}
-				break;
-			default:
-				check( false ); // @TODO LSwift: Implement other storage methods!
-				bSuccess = false;
-				break;
 			}
-			delete[] FileReadBuffer;
+			break;
+		default:
+			FBuildPatchAnalytics::RecordConstructionError(DataFilename, INDEX_NONE, TEXT("File Data Unknown Storage"));
+			GLog->Logf(TEXT("BuildPatchFileConstructor: ERROR: InsertFileData: incorrect storage method %d %s"), Header.StoredAs, *FPaths::GetCleanFilename(DataFilename));
+			bSuccess = false;
+			break;
 		}
-
-		FileReader->Close();
-		delete FileReader;
-	}
-	else
-	{
-		FBuildPatchAnalytics::RecordConstructionError( Filename, FPlatformMisc::GetLastError(), TEXT( "Missing File Data" ) );
-		GLog->Logf( TEXT( "BuildPatchFileConstructor: ERROR: InsertFileData could not open data file %s" ), *Filename );
 	}
+
 	return bSuccess;
 }
 
diff --git a/Engine/Source/Runtime/Online/BuildPatchServices/Private/BuildPatchHash.cpp b/Engine/Source/Runtime/Online/BuildPatchServices/Private/BuildPatchHash.cpp
--- a/Engine/Source/Runtime/Online/BuildPatchServices/Private/BuildPatchHash.cpp
+++ b/Engine/Source/Runtime/Online/BuildPatchServices/Private/BuildPatchHash.cpp
@@ -31,3 +31,14 @@ void FRollingHashConst::Init()
 		HashTable[ TableIdx ] = val;
 	}
 }
+
+uint64 FCycPoly64Hash::GetHashForDataSet(const uint8* DataSet, const uint32 DataSize, const uint64 State/* = 0*/)
+{
+	uint64 HashState = State;
+	for (uint32 i = 0; i < DataSize; ++i)
+	{
+		ROTLEFT_64B(HashState, 1);
+		HashState ^= FRollingHashConst::HashTable[DataSet[i]];
+	}
+	return HashState;
+}
diff --git a/Engine/Source/Runtime/Online/BuildPatchServices/Private/BuildPatchHash.h b/Engine/Source/Runtime/Online/BuildPatchServices/Private/BuildPatchHash.h
--- a/Engine/Source/Runtime/Online/BuildPatchServices/Private/BuildPatchHash.h
+++ b/Engine/Source/Runtime/Online/BuildPatchServices/Private/BuildPatchHash.h
@@ -262,6 +262,17 @@ private:
 	HashRingBuffer WindowData;
 };
 
+namespace FCycPoly64Hash
+{
+	/**
+	 * Static function to simply return the hash for a given data range.
+	 * @param DataSet	The buffer to the data. This must be a buffer of length == DataSize
+	 * @param DataSize	The size of the buffer
+	 * @param State		The current hash state, if continuing hash from previous buffer. 0 for new data.
+	 * @return			The hash state for the provided data
+	 */
+	static uint64 GetHashForDataSet(const uint8* DataSet, const uint32 DataSize, const uint64 State = 0);
+}
 
 /* TRingBuffer implementation
 *****************************************************************************/
diff --git a/Engine/Source/Runtime/Online/BuildPatchServices/Private/BuildPatchInstaller.cpp b/Engine/Source/Runtime/Online/BuildPatchServices/Private/BuildPatchInstaller.cpp
--- a/Engine/Source/Runtime/Online/BuildPatchServices/Private/BuildPatchInstaller.cpp
+++ b/Engine/Source/Runtime/Online/BuildPatchServices/Private/BuildPatchInstaller.cpp
@@ -280,7 +280,6 @@ bool FBuildPatchInstaller::RunInstallation(TArray<FString>& CorruptFiles)
 	const uint32 NumRequiredChunks = bIsFileData ? NumFilesInBuild : FBuildPatchChunkCache::Get().GetStatNumRequiredChunks();
 	const uint32 NumChunksToDownload = bIsFileData ? NumFilesInBuild : FBuildPatchChunkCache::Get().GetStatNumChunksToDownload();
 	const uint32 NumChunksToConstruct = bIsFileData ? 0 : FBuildPatchChunkCache::Get().GetStatNumChunksToRecycle();
-	TotalInitialDownloadSize = bIsFileData ? NewBuildManifest->GetFileSize(FilesToConstruct) : FBuildPatchChunkCache::Get().GetStatTotalChunkDownloadSize();
 
 	// Save stats
 	{
@@ -330,6 +329,11 @@ bool FBuildPatchInstaller::RunInstallation(TArray<FString>& CorruptFiles)
 		TArray< FGuid > RequiredFileData;
 		NewBuildManifest->GetChunksRequiredForFiles(FilesToConstruct, RequiredFileData);
 		FBuildPatchDownloader::Get().AddChunksToDownload(RequiredFileData);
+		TotalInitialDownloadSize = NewBuildManifest->GetDataSize(RequiredFileData);
+	}
+	else
+	{
+		TotalInitialDownloadSize = FBuildPatchChunkCache::Get().GetStatTotalChunkDownloadSize();
 	}
 
 	// Wait for the file constructor to complete
diff --git a/Engine/Source/Runtime/Online/BuildPatchServices/Private/BuildPatchManifest.cpp b/Engine/Source/Runtime/Online/BuildPatchServices/Private/BuildPatchManifest.cpp
--- a/Engine/Source/Runtime/Online/BuildPatchServices/Private/BuildPatchManifest.cpp
+++ b/Engine/Source/Runtime/Online/BuildPatchServices/Private/BuildPatchManifest.cpp
@@ -8,6 +8,9 @@
 
 #define LOCTEXT_NAMESPACE "BuildPatchManifest"
 
+// Remove this when we are to enable creating compressed file data and related manifests
+#define ENABLE_NOCHUNKS_COMPRESSION 0
+
 // The manifest header magic codeword, for quick checking that the opened file is probably a manifest file.
 #define MANIFEST_HEADER_MAGIC		0x44BEC00C
 
@@ -54,7 +57,11 @@ bool BufferIsJsonManifest(const TArray<uint8>& DataInput)
 *****************************************************************************/
 const EBuildPatchAppManifestVersion::Type EBuildPatchAppManifestVersion::GetLatestVersion()
 {
+#if ENABLE_NOCHUNKS_COMPRESSION
 	return static_cast<EBuildPatchAppManifestVersion::Type>(LatestPlusOne - 1);
+#else
+	return EBuildPatchAppManifestVersion::StoredAsCompressedUClass;
+#endif
 }
 
 const EBuildPatchAppManifestVersion::Type EBuildPatchAppManifestVersion::GetLatestJsonVersion()
@@ -76,8 +83,10 @@ const FString& EBuildPatchAppManifestVersion::GetFileSubdir(const EBuildPatchApp
 {
 	static const FString FilesV1 = TEXT("Files");
 	static const FString FilesV2 = TEXT("FilesV2");
+	static const FString FilesV3 = TEXT("FilesV3");
 	return ManifestVersion < DataFileRenames ? FilesV1
-		: FilesV2;
+		: ManifestVersion <= StoredAsCompressedUClass ? FilesV2
+		: FilesV3;
 }
 
 /* FManifestFileHeader - The header for a compressed/encoded manifest file
@@ -381,9 +390,9 @@ private:
 *****************************************************************************/
 UBuildPatchManifest::UBuildPatchManifest(const FObjectInitializer& ObjectInitializer)
 	: Super(ObjectInitializer)
-	, ManifestFileVersion(EBuildPatchAppManifestVersion::GetLatestVersion())
+	, ManifestFileVersion(EBuildPatchAppManifestVersion::Invalid)
 	, bIsFileData(false)
-	, AppID(0)
+	, AppID(INDEX_NONE)
 	, AppName(TEXT(""))
 	, BuildVersion(TEXT(""))
 	, LaunchExe(TEXT(""))
@@ -399,7 +408,7 @@ UBuildPatchManifest::UBuildPatchManifest(const FObjectInitializer& ObjectInitial
 
 void UBuildPatchManifest::Clear()
 {
-	ManifestFileVersion = EBuildPatchAppManifestVersion::GetLatestVersion();
+	ManifestFileVersion = EBuildPatchAppManifestVersion::Invalid;
 	bIsFileData = false;
 	AppID = INDEX_NONE;
 	AppName.Empty();
@@ -458,6 +467,7 @@ const int64 FBuildPatchCustomField::AsInteger() const
 FBuildPatchAppManifest::FBuildPatchAppManifest()
 	: TotalBuildSize(INDEX_NONE)
 	, TotalDownloadSize(INDEX_NONE)
+	, bNeedsResaving(false)
 {
 	Data = NewObject<UBuildPatchManifest>();
 }
@@ -465,6 +475,7 @@ FBuildPatchAppManifest::FBuildPatchAppManifest()
 FBuildPatchAppManifest::FBuildPatchAppManifest(const uint32& InAppID, const FString& AppName)
 	: TotalBuildSize(INDEX_NONE)
 	, TotalDownloadSize(INDEX_NONE)
+	, bNeedsResaving(false)
 {
 	Data = NewObject<UBuildPatchManifest>();
 	Data->AppID = InAppID;
@@ -484,7 +495,6 @@ bool FBuildPatchAppManifest::SaveToFile(const FString& Filename, bool bUseBinary
 	{
 		if (bUseBinary)
 		{
-			Data->ManifestFileVersion = EBuildPatchAppManifestVersion::GetLatestVersion();
 			FManifestWriter ManifestData;
 			Serialize(ManifestData);
 			ManifestData.Finalize();
@@ -521,7 +531,6 @@ bool FBuildPatchAppManifest::SaveToFile(const FString& Filename, bool bUseBinary
 		}
 		else
 		{
-			Data->ManifestFileVersion = EBuildPatchAppManifestVersion::GetLatestJsonVersion();
 			FString JSONOutput;
 			SerializeToJSON(JSONOutput);
 			FTCHARToUTF8 JsonUTF8(*JSONOutput);
@@ -616,6 +625,14 @@ bool FBuildPatchAppManifest::Serialize(FArchive& Ar)
 
 	if (Ar.IsLoading())
 	{
+		// If we didn't load the version number, we know it was skipped when saving therefore must be
+		// the first UObject version
+		if (Data->ManifestFileVersion == static_cast<uint8>(EBuildPatchAppManifestVersion::Invalid))
+		{
+			Data->ManifestFileVersion = EBuildPatchAppManifestVersion::StoredAsCompressedUClass;
+		}
+
+		// Setup internal lookups
 		InitLookups();
 	}
 
@@ -632,6 +649,7 @@ void FBuildPatchAppManifest::DestroyData()
 	CustomFieldLookup.Empty();
 	TotalBuildSize = INDEX_NONE;
 	TotalDownloadSize = INDEX_NONE;
+	bNeedsResaving = false;
 }
 
 void FBuildPatchAppManifest::InitLookups()
@@ -679,9 +697,7 @@ void FBuildPatchAppManifest::SerializeToJSON(FString& JSONOutput)
 	Writer->WriteObjectStart();
 	{
 		// Write general data
-		int32 ManifestVersion = (int32)Data->ManifestFileVersion;
-		int32 JsonVersion = (int32)EBuildPatchAppManifestVersion::GetLatestJsonVersion();
-		Writer->WriteValue(TEXT("ManifestFileVersion"), ToStringBlob(FMath::Min(ManifestVersion, JsonVersion)));
+		Writer->WriteValue(TEXT("ManifestFileVersion"), ToStringBlob(static_cast<int32>(Data->ManifestFileVersion)));
 		Writer->WriteValue(TEXT("bIsFileData"), Data->bIsFileData);
 		Writer->WriteValue(TEXT("AppID"), ToStringBlob(Data->AppID));
 		Writer->WriteValue(TEXT("AppNameString"), Data->AppName);
@@ -1038,6 +1054,9 @@ bool FBuildPatchAppManifest::DeserializeFromJSON( const FString& JSONInput )
 		}
 	}
 
+	// Mark as should be re-saved, client that stores manifests should start using binary
+	bNeedsResaving = true;
+
 	// Make sure we don't have any half loaded data
 	if( !bSuccess )
 	{
@@ -1049,7 +1068,7 @@ bool FBuildPatchAppManifest::DeserializeFromJSON( const FString& JSONInput )
 
 const EBuildPatchAppManifestVersion::Type FBuildPatchAppManifest::GetManifestVersion() const
 {
-	return Data->ManifestFileVersion;
+	return static_cast<EBuildPatchAppManifestVersion::Type>(Data->ManifestFileVersion);
 }
 
 void FBuildPatchAppManifest::GetChunksRequiredForFiles(const TArray<FString>& FileList, TArray<FGuid>& RequiredChunks, bool bAddUnique) const
@@ -1171,19 +1190,34 @@ const bool FBuildPatchAppManifest::GetChunkHash(const FGuid& ChunkGuid, uint64&
 	return false;
 }
 
-const bool FBuildPatchAppManifest::GetFileDataHash(const FGuid& FileGuid, FSHAHashData& OutHash) const
+const bool FBuildPatchAppManifest::GetFileHash(const FGuid& FileGuid, FSHAHashData& OutHash) const
 {
-	// File data hash is the same as the hash for the file it holds, so we can grab it from there
-	for (const auto& FileManifest : Data->FileManifestList)
+	const FString* const * FoundFilename = FileNameLookup.Find(FileGuid);
+	if (FoundFilename)
 	{
-		for (const auto& FileDataPart : FileManifest.FileChunkParts)
-		{
-			if (FileDataPart.Guid == FileGuid)
-			{
-				FMemory::Memcpy(OutHash.Hash, FileManifest.FileHash.Hash, FSHA1::DigestSize);
-				return true;
-			}
-		}
+		return GetFileHash(**FoundFilename, OutHash);
+	}
+	return false;
+}
+
+const bool FBuildPatchAppManifest::GetFileHash(const FString& Filename, FSHAHashData& OutHash) const
+{
+	const FFileManifestData* const * FoundFileManifest = FileManifestLookup.Find(Filename);
+	if (FoundFileManifest)
+	{
+		FMemory::Memcpy(OutHash.Hash, (*FoundFileManifest)->FileHash.Hash, FSHA1::DigestSize);
+		return true;
+	}
+	return false;
+}
+
+const bool FBuildPatchAppManifest::GetFilePartHash(const FGuid& FilePartGuid, uint64& OutHash) const
+{
+	const FChunkInfoData* const * FilePartInfo = ChunkInfoLookup.Find(FilePartGuid);
+	if (FilePartInfo)
+	{
+		OutHash = (*FilePartInfo)->Hash;
+		return true;
 	}
 	return false;
 }
@@ -1473,7 +1507,8 @@ void FBuildPatchAppManifest::GetRemovableFiles(const TCHAR* InstallPath, TArray<
 
 bool FBuildPatchAppManifest::NeedsResaving() const
 {
-	return Data->ManifestFileVersion < EBuildPatchAppManifestVersion::GetLatestVersion();
+	// The bool is marked during file load if we load an old version that should be upgraded
+	return bNeedsResaving;
 }
 
 void FBuildPatchAppManifest::GetOutdatedFiles(FBuildPatchAppManifestPtr OldManifest, FBuildPatchAppManifestRef NewManifest, const FString& InstallDirectory, TArray< FString >& OutDatedFiles)
diff --git a/Engine/Source/Runtime/Online/BuildPatchServices/Private/BuildPatchManifest.h b/Engine/Source/Runtime/Online/BuildPatchServices/Private/BuildPatchManifest.h
--- a/Engine/Source/Runtime/Online/BuildPatchServices/Private/BuildPatchManifest.h
+++ b/Engine/Source/Runtime/Online/BuildPatchServices/Private/BuildPatchManifest.h
@@ -21,7 +21,6 @@ typedef TSharedRef< class FBuildPatchAppManifest, ESPMode::ThreadSafe > FBuildPa
 /**
  * An enum type to describe supported features of a certain manifest
  */
-UENUM()
 namespace EBuildPatchAppManifestVersion
 {
 	enum Type
@@ -46,9 +45,16 @@ namespace EBuildPatchAppManifestVersion
 		StoresChunkFileSizes,
 		// Manifest is now stored using UObject serialization and compressed
 		StoredAsCompressedUClass,
+		// Added support for file data to be split into max individual part size, files now go to FilesV3
+		FileSplittingSupport,
+		// Added support for file data compression, (files still FilesV3)
+		FileCompressionSupport,
+
 
-		// Always last, signifies the latest version plus 1 to allow initialization simplicity
-		LatestPlusOne
+		// Always after the latest version, signifies the latest version plus 1 to allow initialization simplicity
+		LatestPlusOne,
+		// This is for UObject default, so that we always serialize it
+		Invalid = -1
 	};
 }
 
@@ -198,7 +204,7 @@ class UBuildPatchManifest : public UObject
 	GENERATED_UCLASS_BODY()
 
 	UPROPERTY()
-	TEnumAsByte<EBuildPatchAppManifestVersion::Type> ManifestFileVersion;
+	uint8 ManifestFileVersion;
 
 	UPROPERTY()
 	bool bIsFileData;
@@ -488,7 +494,23 @@ public:
 	 * @param OutHash		OUT		Receives the hash value if found
 	 * @return	true if we had the hash for this file
 	 */
-	const bool GetFileDataHash(const FGuid& FileGuid, FSHAHashData& OutHash) const;
+	const bool GetFileHash(const FGuid& FileGuid, FSHAHashData& OutHash) const; // DEPRECATE ME
+
+	/**
+	 * Gets the file hash for a given file
+	 * @param Filename		IN		The filename in the build
+	 * @param OutHash		OUT		Receives the hash value if found
+	 * @return	true if we had the hash for this file
+	 */
+	const bool GetFileHash(const FString& Filename, FSHAHashData& OutHash) const;
+
+	/**
+	 * Gets the file hash for given file data. Valid for non-chunked manifest
+	 * @param FileGuid		IN		The guid of the file data to get hash for
+	 * @param OutHash		OUT		Receives the hash value if found
+	 * @return	true if we had the hash for this file
+	 */
+	const bool GetFilePartHash(const FGuid& FilePartGuid, uint64& OutHash) const;
 
 	/**
 	 * Populates an array of chunks that should be producible from this local build, given the list of chunks needed. Also checks source files exist and match size.
@@ -558,22 +580,19 @@ private:
 
 private:
 
-	/**
-	 * Holds the actual manifest data. Some other variables point to the memory held by this object
-	 */
+	/** Holds the actual manifest data. Some other variables point to the memory held by this object */
 	UBuildPatchManifest* Data;
 
-	/**
-	 * Some lookups to optimize data access
-	 */
+	/** Some lookups to optimize data access */
 	TMap<FGuid, FString*> FileNameLookup;
 	TMap<FString, FFileManifestData*> FileManifestLookup;
 	TMap<FGuid, FChunkInfoData*> ChunkInfoLookup;
 	TMap<FString, FCustomFieldData*> CustomFieldLookup;
 
-	/**
-	 * Holds the total build size in bytes
-	 */
+	/** Holds the total build size in bytes */
 	int64 TotalBuildSize;
 	int64 TotalDownloadSize;
+
+	/** Flag marked true if we loaded from disk as an old manifest version that should be updated */
+	bool bNeedsResaving;
 };
diff --git a/Engine/Source/Runtime/Online/BuildPatchServices/Private/BuildPatchUtil.cpp b/Engine/Source/Runtime/Online/BuildPatchServices/Private/BuildPatchUtil.cpp
--- a/Engine/Source/Runtime/Online/BuildPatchServices/Private/BuildPatchUtil.cpp
+++ b/Engine/Source/Runtime/Online/BuildPatchServices/Private/BuildPatchUtil.cpp
@@ -20,6 +20,12 @@ const FString FBuildPatchUtils::GetFileNewFilename(const EBuildPatchAppManifestV
 	return FPaths::Combine( *RootDirectory, *FString::Printf( TEXT("%s/%02d/%s_%s.file"), *EBuildPatchAppManifestVersion::GetFileSubdir( ManifestVersion ), FCrc::MemCrc32( &FileGUID, sizeof( FGuid ) ) % 100, *FileHash.ToString(), *FileGUID.ToString() ) );
 }
 
+const FString FBuildPatchUtils::GetFileNewFilename(const EBuildPatchAppManifestVersion::Type ManifestVersion, const FString& RootDirectory, const FGuid& FileGUID, const uint64& FileHash)
+{
+	check( FileGUID.IsValid() );
+	return FPaths::Combine( *RootDirectory, *FString::Printf( TEXT("%s/%02d/%016llX_%s.file"), *EBuildPatchAppManifestVersion::GetFileSubdir( ManifestVersion ), FCrc::MemCrc32( &FileGUID, sizeof( FGuid ) ) % 100, FileHash, *FileGUID.ToString() ) );
+}
+
 void FBuildPatchUtils::GetChunkDetailFromNewFilename( const FString& ChunkNewFilename, FGuid& ChunkGUID, uint64& ChunkHash )
 {
 	const FString ChunkFilename = FPaths::GetBaseFilename( ChunkNewFilename );
@@ -75,6 +81,42 @@ const FString FBuildPatchUtils::GetDataTypeOldFilename( const FBuildPatchData::T
 	return TEXT( "" );
 }
 
+const FString FBuildPatchUtils::GetDataFilename(const FBuildPatchAppManifestRef& Manifest, const FString& RootDirectory, const FGuid& DataGUID)
+{
+	const FBuildPatchData::Type DataType = Manifest->IsFileDataManifest() ? FBuildPatchData::FileData : FBuildPatchData::ChunkData;
+	if (Manifest->GetManifestVersion() < EBuildPatchAppManifestVersion::DataFileRenames)
+	{
+		return FBuildPatchUtils::GetDataTypeOldFilename(DataType, RootDirectory, DataGUID);
+	}
+	else if (!Manifest->IsFileDataManifest())
+	{
+		uint64 ChunkHash;
+		const bool bFound = Manifest->GetChunkHash(DataGUID, ChunkHash);
+		// Should be impossible to not exist
+		check(bFound);
+		return FBuildPatchUtils::GetChunkNewFilename(Manifest->GetManifestVersion(), RootDirectory, DataGUID, ChunkHash);
+	}
+	else if (Manifest->GetManifestVersion() <= EBuildPatchAppManifestVersion::StoredAsCompressedUClass)
+	{
+		FSHAHashData FileHash;
+		const bool bFound = Manifest->GetFileHash(DataGUID, FileHash);
+		// Should be impossible to not exist
+		check(bFound);
+		return FBuildPatchUtils::GetFileNewFilename(Manifest->GetManifestVersion(), RootDirectory, DataGUID, FileHash);
+	}
+	else
+	{
+		uint64 FileHash;
+		const bool bFound = Manifest->GetFilePartHash(DataGUID, FileHash);
+		// Should be impossible to not exist
+		check(bFound);
+		return FBuildPatchUtils::GetFileNewFilename(Manifest->GetManifestVersion(), RootDirectory, DataGUID, FileHash);
+	}
+	return TEXT("");
+}
+
+
+
 bool FBuildPatchUtils::GetGUIDFromFilename( const FString& DataFilename, FGuid& DataGUID )
 {
 	const FString DataBaseFilename = FPaths::GetBaseFilename( DataFilename );
@@ -188,6 +230,7 @@ bool FBuildPatchUtils::VerifyChunkFile( FArchive& ChunkFileData, bool bQuickChec
 			// Hashes for checking data
 			FSHA1 SHAHasher;
 			FSHAHashData SHAHash;
+			uint64 CycPoly64Hash = 0;
 			// Load the data to check
 			uint8* FileReadBuffer = new uint8[ FileBufferSize ];
 			int64 DataOffset = 0;
@@ -202,16 +245,7 @@ bool FBuildPatchUtils::VerifyChunkFile( FArchive& ChunkFileData, bool bQuickChec
 					switch ( Header.HashType )
 					{
 					case FChunkHeader::HASH_ROLLING:
-						bSuccess = bSuccess && ( ReadLen == FBuildPatchData::ChunkDataSize );
-						if ( !bSuccess )
-						{
-							GLog->Logf( TEXT( "BuildPatchServices: ERROR: VerifyChunkFile invalid chunk size %d" ), ReadLen );
-						}
-						bSuccess = bSuccess && Header.RollingHash == FRollingHash< FBuildPatchData::ChunkDataSize >::GetHashForDataSet( FileReadBuffer );
-						if ( !bSuccess )
-						{
-							GLog->Logf( TEXT( "BuildPatchServices: ERROR: VerifyChunkFile chunk hashcheck failed" ) );
-						}
+						CycPoly64Hash = FCycPoly64Hash::GetHashForDataSet(FileReadBuffer, ReadLen, CycPoly64Hash);
 						break;
 					case  FChunkHeader::HASH_SHA1:
 						SHAHasher.Update( FileReadBuffer, ReadLen );
@@ -228,16 +262,17 @@ bool FBuildPatchUtils::VerifyChunkFile( FArchive& ChunkFileData, bool bQuickChec
 					switch ( Header.HashType )
 					{
 					case FChunkHeader::HASH_ROLLING:
-						// Already checked
+						bSuccess = Header.RollingHash == CycPoly64Hash;
 						break;
 					case  FChunkHeader::HASH_SHA1:
 						SHAHasher.Final();
 						SHAHasher.GetHash( SHAHash.Hash );
 						bSuccess = SHAHash == Header.SHAHash;
-						if ( !bSuccess )
-						{
-							GLog->Logf( TEXT( "BuildPatchServices: ERROR: VerifyChunkFile file hashcheck failed" ) );
-						}
+						break;
+					}
+					if (!bSuccess)
+					{
+						GLog->Logf(TEXT("BuildPatchServices: ERROR: VerifyChunkFile file hashcheck failed"));
 					}
 				}
 				break;
@@ -300,3 +335,56 @@ bool FBuildPatchUtils::UncompressChunkFile( TArray< uint8 >& ChunkFileArray )
 	}
 	return false;
 }
+
+bool FBuildPatchUtils::UncompressFileDataFile(TArray< uint8 >& DataFileArray, FChunkHeader* OutHeader/* = nullptr*/)
+{
+	FChunkHeader Header;
+	if (OutHeader == nullptr)
+	{
+		OutHeader = &Header;
+	}
+	FMemoryReader FileArrayReader(DataFileArray);
+	// Read the header
+	FileArrayReader << *OutHeader;
+	// Check header
+	const bool bValidHeader = OutHeader->IsValidMagic();
+	const bool bSupportedFormat = !(OutHeader->StoredAs & FChunkHeader::STORED_ENCRYPTED);
+	if (bValidHeader && bSupportedFormat)
+	{
+		bool bSuccess = true;
+		// Uncompress if we need to
+		if (OutHeader->StoredAs == FChunkHeader::STORED_COMPRESSED)
+		{
+			// Load the compressed data
+			const int32 CompressedSize = DataFileArray.Num() - OutHeader->HeaderSize;
+			TArray< uint8 > CompressedData;
+			TArray< uint8 > UncompressedData;
+			CompressedData.Empty(CompressedSize);
+			CompressedData.AddUninitialized(CompressedSize);
+			UncompressedData.Empty(OutHeader->DataSize);
+			UncompressedData.AddUninitialized(OutHeader->DataSize);
+			FileArrayReader.Serialize(CompressedData.GetData(), CompressedSize);
+			FileArrayReader.Close();
+			// Uncompress
+			bSuccess = FCompression::UncompressMemory(
+				static_cast<ECompressionFlags>(COMPRESS_ZLIB | COMPRESS_BiasMemory),
+				UncompressedData.GetData(),
+				UncompressedData.Num(),
+				CompressedData.GetData(),
+				CompressedData.Num());
+			// If successful, write back over the original array
+			if (bSuccess)
+			{
+				DataFileArray.Empty();
+				FMemoryWriter FileArrayWriter(DataFileArray);
+				OutHeader->StoredAs = FChunkHeader::STORED_RAW;
+				FileArrayWriter << *OutHeader;
+				FileArrayWriter.Serialize(UncompressedData.GetData(), UncompressedData.Num());
+				FileArrayWriter.Close();
+			}
+		}
+		return bSuccess;
+	}
+	return false;
+}
+
diff --git a/Engine/Source/Runtime/Online/BuildPatchServices/Private/BuildPatchUtil.h b/Engine/Source/Runtime/Online/BuildPatchServices/Private/BuildPatchUtil.h
--- a/Engine/Source/Runtime/Online/BuildPatchServices/Private/BuildPatchUtil.h
+++ b/Engine/Source/Runtime/Online/BuildPatchServices/Private/BuildPatchUtil.h
@@ -30,6 +30,7 @@ struct FBuildPatchUtils
 	 * @return	The file chunk path
 	 */
 	static const FString GetFileNewFilename(const EBuildPatchAppManifestVersion::Type ManifestVersion, const FString& RootDirectory, const FGuid& FileGUID, const FSHAHashData& FileHash);
+	static const FString GetFileNewFilename(const EBuildPatchAppManifestVersion::Type ManifestVersion, const FString& RootDirectory, const FGuid& FileGUID, const uint64& FilePartHash);
 
 	/**
 	 * Gets the chunk GUID and Hash, from the filename, which is the new format.
@@ -72,6 +73,15 @@ struct FBuildPatchUtils
 	 */
 	static const FString GetDataTypeOldFilename(const FBuildPatchData::Type DataType, const FString& RootDirectory, const FGuid& Guid);
 
+	/**
+	 * Gets the filename for any data part. Wraps the choice between all of the above
+	 * @param Manifest			The manifest referencing this data
+	 * @param RootDirectory		The root directory
+	 * @param DataGUID			The data Guid
+	 * @return	the data part path
+	 */
+	static const FString GetDataFilename(const FBuildPatchAppManifestRef& Manifest, const FString& RootDirectory, const FGuid& DataGUID);
+
 	/**
 	 * Gets the GUID for a data file according to it's filename (new or old)
 	 * @param DataFilename		IN		The data filename, or URL
@@ -87,6 +97,14 @@ struct FBuildPatchUtils
 	 */
 	static bool UncompressChunkFile(TArray<uint8>& ChunkFileArray);
 
+	/**
+	 * Helper function to uncompress file part data. Can be called without knowing if needed and process will be just skipped.
+	 * @param DataFileArray	IN OUT		The data array, should contain full file - header plus data. Will be overwritten with uncompressed data version.
+	 * @param OutHeader		OUT			Optional ptr to a header to receive the header from the file if you need it anyway
+	 * @return		true if no errors occurred and the data is not corrupted
+	 */
+	static bool UncompressFileDataFile(TArray< uint8 >& DataFileArray, FChunkHeader* OutHeader = nullptr);
+
 	/**
 	 * Checks a file against SHA1 hashes. The function takes two so that it can return no match, match with Hash1, or match with Hash2, that way we can check the file for being the same as an old manifest or new manifest
 	 * NOTE: This function is blocking and will not return until finished. Don't run on main thread.
