Prevent crash when resetting achievements on those platforms that require a valid FUniqueNetId to do so.

#branch UE4

#change It was crashing on Steam as it can't convert a null/invalid FUniqueNetId to a FUniqueNetIdSteam. Google Play doesn't care about the FUniqueNetId atm, so should still work fine providing it still has a local user.

diff --git a/Engine/Source/Runtime/Online/OnlineSubsystem/Private/OnlineSubsystem.cpp b/Engine/Source/Runtime/Online/OnlineSubsystem/Private/OnlineSubsystem.cpp
--- a/Engine/Source/Runtime/Online/OnlineSubsystem/Private/OnlineSubsystem.cpp
+++ b/Engine/Source/Runtime/Online/OnlineSubsystem/Private/OnlineSubsystem.cpp
@@ -94,6 +94,20 @@ bool IsPlayerInSessionImpl(IOnlineSession* SessionInt, FName SessionName, const
 
 static void ResetAchievements()
 {
+	auto IdentityInterface = Online::GetIdentityInterface();
+	if (!IdentityInterface.IsValid())
+	{
+		UE_LOG_ONLINE(Warning, TEXT("ResetAchievements command: couldn't get the identity interface"));
+		return;
+	}
+	
+	TSharedPtr<FUniqueNetId> UserId = IdentityInterface->GetUniquePlayerId(0);
+	if(!UserId.IsValid())
+	{
+		UE_LOG_ONLINE(Warning, TEXT("ResetAchievements command: invalid UserId"));
+		return;
+	}
+
 	auto AchievementsInterface = Online::GetAchievementsInterface();
 	if (!AchievementsInterface.IsValid())
 	{
@@ -101,7 +115,7 @@ static void ResetAchievements()
 		return;
 	}
 
-	AchievementsInterface->ResetAchievements(FUniqueNetIdString());
+	AchievementsInterface->ResetAchievements(*UserId);
 }
 
 FAutoConsoleCommand CmdResetAchievements(
