Making some of the most frequently constructed FNames static.

diff --git a/Engine/Source/Runtime/CoreUObject/Private/UObject/GarbageCollection.cpp b/Engine/Source/Runtime/CoreUObject/Private/UObject/GarbageCollection.cpp
--- a/Engine/Source/Runtime/CoreUObject/Private/UObject/GarbageCollection.cpp
+++ b/Engine/Source/Runtime/CoreUObject/Private/UObject/GarbageCollection.cpp
@@ -1613,7 +1613,7 @@ void UClass::AssembleReferenceTokenStream()
 		}
 
 		// Emit end of stream token.
-		const FName EOSDebugName("EOS");
+		static const FName EOSDebugName("EOS");
 		EmitObjectReference(0, EOSDebugName, GCRT_EndOfStream);
 
 		// Shrink reference token stream to proper size.
diff --git a/Engine/Source/Runtime/CoreUObject/Private/UObject/UObjectBase.cpp b/Engine/Source/Runtime/CoreUObject/Private/UObject/UObjectBase.cpp
--- a/Engine/Source/Runtime/CoreUObject/Private/UObject/UObjectBase.cpp
+++ b/Engine/Source/Runtime/CoreUObject/Private/UObject/UObjectBase.cpp
@@ -339,8 +339,10 @@ bool UObjectBase::IsValidLowLevelFast(bool bRecursive /*= true*/) const
 
 void UObjectBase::EmitBaseReferences(UClass *RootClass)
 {
-	RootClass->EmitObjectReference(STRUCT_OFFSET(UObjectBase, Class), TEXT("Class"));
-	RootClass->EmitObjectReference(STRUCT_OFFSET(UObjectBase, Outer), TEXT("Outer"), GCRT_PersistentObject);
+	static const FName ClassPropertyName(TEXT("Class"));
+	static const FName OuterPropertyName(TEXT("Outer"));
+	RootClass->EmitObjectReference(STRUCT_OFFSET(UObjectBase, Class), ClassPropertyName);
+	RootClass->EmitObjectReference(STRUCT_OFFSET(UObjectBase, Outer), OuterPropertyName, GCRT_PersistentObject);
 }
 
 /** Enqueue the registration for this object. */
diff --git a/Engine/Source/Runtime/Engine/Private/Actor.cpp b/Engine/Source/Runtime/Engine/Private/Actor.cpp
--- a/Engine/Source/Runtime/Engine/Private/Actor.cpp
+++ b/Engine/Source/Runtime/Engine/Private/Actor.cpp
@@ -538,7 +538,8 @@ void AActor::PostLoadSubobjects(FObjectInstancingGraph* OuterInstanceGraph)
 void AActor::ProcessEvent(UFunction* Function, void* Parameters)
 {
 	#if WITH_EDITOR
-	const bool bAllowScriptExecution = GAllowActorScriptExecutionInEditor || Function->GetBoolMetaData( "CallInEditor" );
+	static const FName CallInEditorMeta(TEXT("CallInEditor"));
+	const bool bAllowScriptExecution = GAllowActorScriptExecutionInEditor || Function->GetBoolMetaData(CallInEditorMeta);
 	#else
 	const bool bAllowScriptExecution = GAllowActorScriptExecutionInEditor;
 	#endif
diff --git a/Engine/Source/Runtime/Engine/Private/ShapeComponent.cpp b/Engine/Source/Runtime/Engine/Private/ShapeComponent.cpp
--- a/Engine/Source/Runtime/Engine/Private/ShapeComponent.cpp
+++ b/Engine/Source/Runtime/Engine/Private/ShapeComponent.cpp
@@ -6,7 +6,8 @@
 UShapeComponent::UShapeComponent(const FObjectInitializer& ObjectInitializer)
 	: Super(ObjectInitializer)
 {
-	BodyInstance.SetCollisionProfileName(FName(TEXT("OverlapAllDynamic")));
+	static const FName CollisionProfileName(TEXT("OverlapAllDynamic"));
+	BodyInstance.SetCollisionProfileName(CollisionProfileName);
 	// when we deprecated this variable, we switched on for the shapecomponent collision profile
 	// the problem with adding collision profile later, the instanced data will be wiped
 	// since shape component is so popular for BP and so on, I'm adding manual support for compatibility
diff --git a/Engine/Source/Runtime/Online/OnlineSubsystem/Private/OnlineSubsystemImpl.cpp b/Engine/Source/Runtime/Online/OnlineSubsystem/Private/OnlineSubsystemImpl.cpp
--- a/Engine/Source/Runtime/Online/OnlineSubsystem/Private/OnlineSubsystemImpl.cpp
+++ b/Engine/Source/Runtime/Online/OnlineSubsystem/Private/OnlineSubsystemImpl.cpp
@@ -5,8 +5,10 @@
 #include "NamedInterfaces.h"
 #include "OnlineIdentityInterface.h"
 
+const FName FOnlineSubsystemImpl::DefaultInstanceName(TEXT("DefaultInstance"));
+
 FOnlineSubsystemImpl::FOnlineSubsystemImpl() :
-	InstanceName(DEFAULT_INSTANCE),
+	InstanceName(DefaultInstanceName),
 	bForceDedicated(false),
 	NamedInterfaces(NULL)
 {
@@ -68,7 +70,7 @@ void FOnlineSubsystemImpl::SetNamedInterface(FName InterfaceName, UObject* NewIn
 
 bool FOnlineSubsystemImpl::IsServer() const
 {
-	FName WorldContextHandle = (InstanceName != NAME_None && InstanceName != DEFAULT_INSTANCE) ? InstanceName : NAME_None;
+	FName WorldContextHandle = (InstanceName != NAME_None && InstanceName != DefaultInstanceName) ? InstanceName : NAME_None;
 	return IsServerForOnlineSubsystems(WorldContextHandle);
 }
 
diff --git a/Engine/Source/Runtime/Online/OnlineSubsystem/Private/OnlineSubsystemModule.cpp b/Engine/Source/Runtime/Online/OnlineSubsystem/Private/OnlineSubsystemModule.cpp
--- a/Engine/Source/Runtime/Online/OnlineSubsystem/Private/OnlineSubsystemModule.cpp
+++ b/Engine/Source/Runtime/Online/OnlineSubsystem/Private/OnlineSubsystemModule.cpp
@@ -139,7 +139,7 @@ void FOnlineSubsystemModule::UnregisterPlatformService(const FName FactoryName)
 void FOnlineSubsystemModule::ParseOnlineSubsystemName(const FName& FullName, FName& SubsystemName, FName& InstanceName) const
 {
 	SubsystemName = DefaultPlatformService;
-	InstanceName = DEFAULT_INSTANCE;
+	InstanceName = FOnlineSubsystemImpl::DefaultInstanceName;
 
 	if (FullName != NAME_None)
 	{
diff --git a/Engine/Source/Runtime/Online/OnlineSubsystem/Public/OnlineSubsystemImpl.h b/Engine/Source/Runtime/Online/OnlineSubsystem/Public/OnlineSubsystemImpl.h
--- a/Engine/Source/Runtime/Online/OnlineSubsystem/Public/OnlineSubsystemImpl.h
+++ b/Engine/Source/Runtime/Online/OnlineSubsystem/Public/OnlineSubsystemImpl.h
@@ -51,5 +51,8 @@ public:
 	 * @return the name of the online subsystem instance
 	 */
 	FName GetInstanceName() const { return InstanceName; }
+
+	/** Name given to default OSS instances (disambiguates for PIE) */
+	static const FName DefaultInstanceName;
 };
 
diff --git a/Engine/Source/Runtime/Online/OnlineSubsystem/Public/OnlineSubsystemTypes.h b/Engine/Source/Runtime/Online/OnlineSubsystem/Public/OnlineSubsystemTypes.h
--- a/Engine/Source/Runtime/Online/OnlineSubsystem/Public/OnlineSubsystemTypes.h
+++ b/Engine/Source/Runtime/Online/OnlineSubsystem/Public/OnlineSubsystemTypes.h
@@ -3,9 +3,6 @@
 #pragma once
 #include "OnlineSubsystemPackage.h"
 
-/** Name given to default OSS instances (disambiguates for PIE) */
-#define DEFAULT_INSTANCE FName(TEXT("DefaultInstance"))
-
 /** Maximum players supported on a given platform */
 #if PLATFORM_XBOXONE
 #define MAX_LOCAL_PLAYERS 4
diff --git a/Engine/Source/Runtime/Online/OnlineSubsystemUtils/Private/OnlineSubsystemUtils.cpp b/Engine/Source/Runtime/Online/OnlineSubsystemUtils/Private/OnlineSubsystemUtils.cpp
--- a/Engine/Source/Runtime/Online/OnlineSubsystemUtils/Private/OnlineSubsystemUtils.cpp
+++ b/Engine/Source/Runtime/Online/OnlineSubsystemUtils/Private/OnlineSubsystemUtils.cpp
@@ -81,7 +81,7 @@ UAudioComponent* CreateVoiceAudioComponent(uint32 SampleRate)
 UWorld* GetWorldForOnline(FName InstanceName)
 {
 	UWorld* World = NULL;
-	if (InstanceName != DEFAULT_INSTANCE && InstanceName != NAME_None)
+	if (InstanceName != FOnlineSubsystemImpl::DefaultInstanceName && InstanceName != NAME_None)
 	{
 		FWorldContext& WorldContext = GEngine->GetWorldContextFromHandleChecked(InstanceName);
 		check(WorldContext.WorldType == EWorldType::Game || WorldContext.WorldType == EWorldType::PIE);
diff --git a/Engine/Source/Runtime/Online/OnlineSubsystemUtils/Private/OnlineSubsystemUtilsPrivatePCH.h b/Engine/Source/Runtime/Online/OnlineSubsystemUtils/Private/OnlineSubsystemUtilsPrivatePCH.h
--- a/Engine/Source/Runtime/Online/OnlineSubsystemUtils/Private/OnlineSubsystemUtilsPrivatePCH.h
+++ b/Engine/Source/Runtime/Online/OnlineSubsystemUtils/Private/OnlineSubsystemUtilsPrivatePCH.h
@@ -6,6 +6,7 @@
 #include "Engine.h"
 #include "Sockets.h"
 #include "OnlineSubsystem.h"
+#include "OnlineSubsystemImpl.h"
 #include "OnlineSubsystemUtils.h"
 #include "OnlineSubsystemUtilsModule.h"
 #include "SocketSubsystem.h"
