Fix for non-unity builds after static FName cleanup.

diff --git a/Engine/Source/Runtime/Online/OnlineSubsystem/Private/OnlineSubsystemModule.cpp b/Engine/Source/Runtime/Online/OnlineSubsystem/Private/OnlineSubsystemModule.cpp
--- a/Engine/Source/Runtime/Online/OnlineSubsystem/Private/OnlineSubsystemModule.cpp
+++ b/Engine/Source/Runtime/Online/OnlineSubsystem/Private/OnlineSubsystemModule.cpp
@@ -2,6 +2,7 @@
 
 #include "OnlineSubsystemPrivatePCH.h"
 #include "ModuleManager.h"
+#include "OnlineSubsystemImpl.h"
 
 IMPLEMENT_MODULE( FOnlineSubsystemModule, OnlineSubsystem );
 
