Destroy OSS interfaces in reverse order of construction.

diff --git a/Engine/Source/Runtime/Online/OnlineSubsystemSteam/Private/OnlineSubsystemSteam.cpp b/Engine/Source/Runtime/Online/OnlineSubsystemSteam/Private/OnlineSubsystemSteam.cpp
--- a/Engine/Source/Runtime/Online/OnlineSubsystemSteam/Private/OnlineSubsystemSteam.cpp
+++ b/Engine/Source/Runtime/Online/OnlineSubsystemSteam/Private/OnlineSubsystemSteam.cpp
@@ -377,14 +377,14 @@ bool FOnlineSubsystemSteam::Shutdown()
 
 	// Destruct the interfaces
 	DESTRUCT_INTERFACE(AchievementsInterface);
-	DESTRUCT_INTERFACE(SessionInterface);
-	DESTRUCT_INTERFACE(IdentityInterface);
-	DESTRUCT_INTERFACE(FriendInterface);
+	DESTRUCT_INTERFACE(ExternalUIInterface);
+	DESTRUCT_INTERFACE(VoiceInterface);
+	DESTRUCT_INTERFACE(LeaderboardsInterface);
 	DESTRUCT_INTERFACE(SharedCloudInterface);
 	DESTRUCT_INTERFACE(UserCloudInterface);
-	DESTRUCT_INTERFACE(LeaderboardsInterface);
-	DESTRUCT_INTERFACE(VoiceInterface);
-	DESTRUCT_INTERFACE(ExternalUIInterface);
+	DESTRUCT_INTERFACE(FriendInterface);
+	DESTRUCT_INTERFACE(IdentityInterface);
+	DESTRUCT_INTERFACE(SessionInterface);
 
 	#undef DESTRUCT_INTERFACE
 
