Fix UE-2195. Moved DisplayAccountUpgradeDialog into the IOnlineExternalUIInterface and renamed function to ShowAccountUpgradeUI.

diff --git a/Engine/Source/Runtime/Online/Android/OnlineSubsystemGooglePlay/Private/OnlineExternalUIInterfaceGooglePlay.cpp b/Engine/Source/Runtime/Online/Android/OnlineSubsystemGooglePlay/Private/OnlineExternalUIInterfaceGooglePlay.cpp
--- a/Engine/Source/Runtime/Online/Android/OnlineSubsystemGooglePlay/Private/OnlineExternalUIInterfaceGooglePlay.cpp
+++ b/Engine/Source/Runtime/Online/Android/OnlineSubsystemGooglePlay/Private/OnlineExternalUIInterfaceGooglePlay.cpp
@@ -61,3 +61,8 @@ bool FOnlineExternalUIGooglePlay::ShowProfileUI( const FUniqueNetId& Requestor,
 {
 	return false;
 }
+
+bool FOnlineExternalUIGooglePlay::ShowAccountUpgradeUI(const FUniqueNetId& UniqueId)
+{
+	return false;
+}
diff --git a/Engine/Source/Runtime/Online/Android/OnlineSubsystemGooglePlay/Private/OnlineExternalUIInterfaceGooglePlay.h b/Engine/Source/Runtime/Online/Android/OnlineSubsystemGooglePlay/Private/OnlineExternalUIInterfaceGooglePlay.h
--- a/Engine/Source/Runtime/Online/Android/OnlineSubsystemGooglePlay/Private/OnlineExternalUIInterfaceGooglePlay.h
+++ b/Engine/Source/Runtime/Online/Android/OnlineSubsystemGooglePlay/Private/OnlineExternalUIInterfaceGooglePlay.h
@@ -28,6 +28,7 @@ public:
 	virtual bool ShowLeaderboardUI(const FString& LeaderboardName) override;
 	virtual bool ShowWebURL(const FString& WebURL) override;
 	virtual bool ShowProfileUI(const FUniqueNetId& Requestor, const FUniqueNetId& Requestee, const FOnProfileUIClosedDelegate& Delegate) override;
+	virtual bool ShowAccountUpgradeUI(const FUniqueNetId& UniqueId) override;
 };
 
 typedef TSharedPtr<FOnlineExternalUIGooglePlay, ESPMode::ThreadSafe> FOnlineExternalUIGooglePlayPtr;
diff --git a/Engine/Source/Runtime/Online/IOS/OnlineSubsystemIOS/Private/OnlineExternalUIInterfaceIOS.cpp b/Engine/Source/Runtime/Online/IOS/OnlineSubsystemIOS/Private/OnlineExternalUIInterfaceIOS.cpp
--- a/Engine/Source/Runtime/Online/IOS/OnlineSubsystemIOS/Private/OnlineExternalUIInterfaceIOS.cpp
+++ b/Engine/Source/Runtime/Online/IOS/OnlineSubsystemIOS/Private/OnlineExternalUIInterfaceIOS.cpp
@@ -46,3 +46,8 @@ bool FOnlineExternalUIIOS::ShowProfileUI( const FUniqueNetId& Requestor, const F
 {
 	return false;
 }
+
+bool FOnlineExternalUIIOS::ShowAccountUpgradeUI(const FUniqueNetId& UniqueId)
+{
+	return false;
+}
diff --git a/Engine/Source/Runtime/Online/IOS/OnlineSubsystemIOS/Private/OnlineExternalUIInterfaceIOS.h b/Engine/Source/Runtime/Online/IOS/OnlineSubsystemIOS/Private/OnlineExternalUIInterfaceIOS.h
--- a/Engine/Source/Runtime/Online/IOS/OnlineSubsystemIOS/Private/OnlineExternalUIInterfaceIOS.h
+++ b/Engine/Source/Runtime/Online/IOS/OnlineSubsystemIOS/Private/OnlineExternalUIInterfaceIOS.h
@@ -21,6 +21,7 @@ public:
 	virtual bool ShowLeaderboardUI(const FString& LeaderboardName) override;
 	virtual bool ShowWebURL(const FString& WebURL) override;
 	virtual bool ShowProfileUI(const FUniqueNetId& Requestor, const FUniqueNetId& Requestee, const FOnProfileUIClosedDelegate& Delegate) override;
+	virtual bool ShowAccountUpgradeUI(const FUniqueNetId& UniqueId) override;
 	// End IOnlineExternalUI interface
 };
 
diff --git a/Engine/Source/Runtime/Online/OnlineSubsystem/Public/Interfaces/OnlineExternalUIInterface.h b/Engine/Source/Runtime/Online/OnlineSubsystem/Public/Interfaces/OnlineExternalUIInterface.h
--- a/Engine/Source/Runtime/Online/OnlineSubsystem/Public/Interfaces/OnlineExternalUIInterface.h
+++ b/Engine/Source/Runtime/Online/OnlineSubsystem/Public/Interfaces/OnlineExternalUIInterface.h
@@ -103,6 +103,15 @@ public:
 	 */
 	virtual bool ShowProfileUI(const FUniqueNetId& Requestor, const FUniqueNetId& Requestee, const FOnProfileUIClosedDelegate& Delegate) = 0;
 
+	/**
+	* Displays a system dialog to purchase user account upgrades.  e.g. PlaystationPlus, XboxLive GOLD, etc.
+	*
+	* @param UniqueID of the user to show the dialog for
+	*
+	* @return true if it was able to show the UI, false if it failed
+	*/
+	virtual bool ShowAccountUpgradeUI(const FUniqueNetId& UniqueId) = 0;
+
 	/**
 	 * Delegate called when the external UI is opened or closed
 	 *
diff --git a/Engine/Source/Runtime/Online/OnlineSubsystem/Public/OnlineSubsystem.h b/Engine/Source/Runtime/Online/OnlineSubsystem/Public/OnlineSubsystem.h
--- a/Engine/Source/Runtime/Online/OnlineSubsystem/Public/OnlineSubsystem.h
+++ b/Engine/Source/Runtime/Online/OnlineSubsystem/Public/OnlineSubsystem.h
@@ -323,12 +323,6 @@ public:
 	 */
 	virtual bool Exec(class UWorld* InWorld, const TCHAR* Cmd, FOutputDevice& Ar) = 0;
 
-
-	/**
-	 * Displays a system dialog to purchase user account upgrades.  e.g. PlaystationPlus, XboxLive GOLD, etc.
-	 */
-	virtual void DisplayAccountUpgradeDialog(const FUniqueNetId& UniqueId) = 0;
-
 	/**
 	 * Some platforms must know when the game is using Multiplayer features so they can do recurring authorization checks.
 	 */
diff --git a/Engine/Source/Runtime/Online/OnlineSubsystem/Public/OnlineSubsystemImpl.h b/Engine/Source/Runtime/Online/OnlineSubsystem/Public/OnlineSubsystemImpl.h
--- a/Engine/Source/Runtime/Online/OnlineSubsystem/Public/OnlineSubsystemImpl.h
+++ b/Engine/Source/Runtime/Online/OnlineSubsystem/Public/OnlineSubsystemImpl.h
@@ -42,7 +42,6 @@ public:
 
 	virtual bool IsLocalPlayer(const FUniqueNetId& UniqueId) const override;
 
-	virtual void DisplayAccountUpgradeDialog(const FUniqueNetId& UniqueId) override {};
 	virtual void SetUsingMultiplayerFeatures(const FUniqueNetId& UniqueId, bool bUsingMP) override {};
 
 	// FOnlineSubsystemImpl
diff --git a/Engine/Source/Runtime/Online/OnlineSubsystemSteam/Private/OnlineExternalUIInterfaceSteam.cpp b/Engine/Source/Runtime/Online/OnlineSubsystemSteam/Private/OnlineExternalUIInterfaceSteam.cpp
--- a/Engine/Source/Runtime/Online/OnlineSubsystemSteam/Private/OnlineExternalUIInterfaceSteam.cpp
+++ b/Engine/Source/Runtime/Online/OnlineSubsystemSteam/Private/OnlineExternalUIInterfaceSteam.cpp
@@ -78,7 +78,8 @@ bool FOnlineExternalUISteam::ShowProfileUI( const FUniqueNetId& Requestor, const
 	return false;
 }
 
-
-
-
+bool FOnlineExternalUISteam::ShowAccountUpgradeUI(const FUniqueNetId& UniqueId)
+{
+	return false;
+}
 
diff --git a/Engine/Source/Runtime/Online/OnlineSubsystemSteam/Private/OnlineExternalUIInterfaceSteam.h b/Engine/Source/Runtime/Online/OnlineSubsystemSteam/Private/OnlineExternalUIInterfaceSteam.h
--- a/Engine/Source/Runtime/Online/OnlineSubsystemSteam/Private/OnlineExternalUIInterfaceSteam.h
+++ b/Engine/Source/Runtime/Online/OnlineSubsystemSteam/Private/OnlineExternalUIInterfaceSteam.h
@@ -82,6 +82,7 @@ public:
 	virtual bool ShowLeaderboardUI(const FString& LeaderboardName) override;
 	virtual bool ShowWebURL(const FString& WebURL) override;
 	virtual bool ShowProfileUI(const FUniqueNetId& Requestor, const FUniqueNetId& Requestee, const FOnProfileUIClosedDelegate& Delegate) override;
+	virtual bool ShowAccountUpgradeUI(const FUniqueNetId& UniqueId) override;
 };
 
 typedef TSharedPtr<FOnlineExternalUISteam, ESPMode::ThreadSafe> FOnlineExternalUISteamPtr;
