FUniqueNetIdPS4 converted to use opaque handles.
NamedSessions now have UniqueNetId for local owner. Needs implementing on other platforms.
Matching2 Contexts now handled properly.

diff --git a/Engine/Source/Runtime/Online/OnlineSubsystem/Public/OnlineSessionSettings.h b/Engine/Source/Runtime/Online/OnlineSubsystem/Public/OnlineSessionSettings.h
--- a/Engine/Source/Runtime/Online/OnlineSubsystem/Public/OnlineSessionSettings.h
+++ b/Engine/Source/Runtime/Online/OnlineSubsystem/Public/OnlineSessionSettings.h
@@ -351,9 +351,13 @@ public:
 	/** The name of the session */
 	const FName SessionName;
 	/** Index of the player who created the session [host] or joined it [client] */
-	int32 HostingPlayerNum;
+	int32 HostingPlayerNum;	
 	/** Whether or not the local player is hosting this session */
 	bool bHosting;
+
+	/** NetId of the local player that created this named session.  Could be the host, or a player joining a session. Will entirely replace HostingPlayerNum */
+	TSharedPtr<FUniqueNetId> LocalOwnerId;
+
 	/** List of players registered in the session */
 	TArray< TSharedRef<FUniqueNetId> > RegisteredPlayers;
 	/** State of the session (game thread write only) */
