Merging using UE4-To-UE4-LauncherReleases REVERSE

CL#2354432 on 10/11/2014 08:49:46
------------------------------------------
Implement file data splitting and compression support for build patch tool

diff --git a/Engine/Source/Runtime/Online/BuildPatchServices/Private/BuildPatchGeneration.cpp b/Engine/Source/Runtime/Online/BuildPatchServices/Private/BuildPatchGeneration.cpp
--- a/Engine/Source/Runtime/Online/BuildPatchServices/Private/BuildPatchGeneration.cpp
+++ b/Engine/Source/Runtime/Online/BuildPatchServices/Private/BuildPatchGeneration.cpp
@@ -3,7 +3,7 @@
 /*=============================================================================
 	BuildPatchGeneration.cpp: Implements the classes that control build
 	installation, and the generation of chunks and manifests from a build image.
-	=============================================================================*/
+=============================================================================*/
 
 #include "BuildPatchServicesPrivatePCH.h"
 
@@ -63,9 +63,9 @@ FBuildStream::FBuildStreamReader::~FBuildStreamReader()
 	}
 }
 
-bool FBuildStream::FBuildStreamReader::Init()
+bool FBuildStream::FBuildStreamReader::Init() 
 {
-	return BuildStream != NULL && IFileManager::Get().DirectoryExists(*DepotDirectory);
+	return BuildStream != NULL && IFileManager::Get().DirectoryExists( *DepotDirectory );
 }
 
 uint32 FBuildStream::FBuildStreamReader::Run()
@@ -74,42 +74,42 @@ uint32 FBuildStream::FBuildStreamReader::Run()
 	BuildStream->Clear();
 
 	TArray< FString > AllFiles;
-	IFileManager::Get().FindFilesRecursive(AllFiles, *DepotDirectory, TEXT("*.*"), true, false);
+	IFileManager::Get().FindFilesRecursive( AllFiles, *DepotDirectory, TEXT("*.*"), true, false );
 	AllFiles.Sort();
 
 	// Remove the files that appear in an ignore list
-	FBuildDataGenerator::StripIgnoredFiles(AllFiles, DepotDirectory, IgnoreListFile);
+	FBuildDataGenerator::StripIgnoredFiles( AllFiles, DepotDirectory, IgnoreListFile );
 
 	// Allocate our file read buffer
-	uint8* FileReadBuffer = new uint8[FileBufferSize];
+	uint8* FileReadBuffer = new uint8[ FileBufferSize ];
 
 	for (auto& SourceFile : AllFiles)
 	{
 		// Read the file
-		FArchive* FileReader = IFileManager::Get().CreateFileReader(*SourceFile);
+		FArchive* FileReader = IFileManager::Get().CreateFileReader( *SourceFile );
 		const bool bBuildFileOpenSuccess = FileReader != NULL;
-		if (bBuildFileOpenSuccess)
+		if( bBuildFileOpenSuccess )
 		{
 			// Make SourceFile the format we want it in and start a new file
-			FPaths::MakePathRelativeTo(SourceFile, *(DepotDirectory + TEXT("/")));
+			FPaths::MakePathRelativeTo( SourceFile, *( DepotDirectory + TEXT( "/" ) ) );
 			int64 FileSize = FileReader->TotalSize();
 			// Process files that have bytes
-			if (FileSize > 0)
+			if( FileSize > 0 )
 			{
-				BuildStream->BeginNewFile(SourceFile, FileSize);
-				while (!FileReader->AtEnd())
+				BuildStream->BeginNewFile( SourceFile, FileSize );
+				while( !FileReader->AtEnd() )
 				{
 					const int64 SizeLeft = FileSize - FileReader->Tell();
-					const uint32 ReadLen = FMath::Min< int64 >(FileBufferSize, SizeLeft);
-					FileReader->Serialize(FileReadBuffer, ReadLen);
+					const uint32 ReadLen = FMath::Min< int64 >( FileBufferSize, SizeLeft );
+					FileReader->Serialize( FileReadBuffer, ReadLen );
 					// Copy into data stream
-					BuildStream->EnqueueData(FileReadBuffer, ReadLen);
+					BuildStream->EnqueueData( FileReadBuffer, ReadLen );
 				}
 			}
 			// Special case zero byte files
-			else if (FileSize == 0)
+			else if( FileSize == 0 )
 			{
-				BuildStream->AddEmptyFile(SourceFile);
+				BuildStream->AddEmptyFile( SourceFile );
 			}
 			FileReader->Close();
 			delete FileReader;
@@ -117,10 +117,10 @@ uint32 FBuildStream::FBuildStreamReader::Run()
 		else
 		{
 			// Not being able to load a required file from the build would be fatal, hard fault.
-			GLog->Logf(TEXT("FBuildStreamReader: Could not open file from build! %s"), *SourceFile);
+			GLog->Logf( TEXT( "FBuildStreamReader: Could not open file from build! %s" ), *SourceFile );
 			GLog->PanicFlushThreadedLogs();
 			// Use bool variable for easier to understand assert message
-			check(bBuildFileOpenSuccess);
+			check( bBuildFileOpenSuccess );
 		}
 	}
 
@@ -140,7 +140,7 @@ void FBuildStream::FBuildStreamReader::StartThread()
 
 /* FBuildStream implementation
 *****************************************************************************/
-FBuildStream::FBuildStream(const FString& RootDirectory, const FString& IgnoreListFile)
+FBuildStream::FBuildStream( const FString& RootDirectory, const FString& IgnoreListFile )
 {
 	BuildStreamReader.DepotDirectory = RootDirectory;
 	BuildStreamReader.IgnoreListFile = IgnoreListFile;
@@ -152,13 +152,13 @@ FBuildStream::~FBuildStream()
 {
 }
 
-bool FBuildStream::GetFileSpan(const uint64& StartingIdx, FString& Filename, uint64& FileSize)
+bool FBuildStream::GetFileSpan( const uint64& StartingIdx, FString& Filename, uint64& FileSize )
 {
 	bool bFound = false;
 	FilesListsCS.Lock();
 	// Find the filename
-	FFileSpan* FileSpan = FilesParsed.Find(StartingIdx);
-	if (FileSpan != NULL)
+	FFileSpan* FileSpan = FilesParsed.Find( StartingIdx );
+	if( FileSpan != NULL )
 	{
 		Filename = FileSpan->Filename;
 		FileSize = FileSpan->Size;
@@ -168,20 +168,20 @@ bool FBuildStream::GetFileSpan(const uint64& StartingIdx, FString& Filename, uin
 	return bFound;
 }
 
-uint32 FBuildStream::DequeueData(uint8* Buffer, const uint32& ReqSize, const bool WaitForData)
+uint32 FBuildStream::DequeueData( uint8* Buffer, const uint32& ReqSize, const bool WaitForData )
 {
 	// Wait for data
-	if (WaitForData)
+	if( WaitForData )
 	{
-		while (DataAvailable() < ReqSize && !IsEndOfBuild())
+		while ( DataAvailable() < ReqSize && !IsEndOfBuild() )
 		{
-			FPlatformProcess::Sleep(0.01f);
+			FPlatformProcess::Sleep( 0.01f );
 		}
 	}
 
 	BuildDataStreamCS.Lock();
-	uint32 ReadLen = FMath::Min< int64 >(ReqSize, DataAvailable());
-	ReadLen = BuildDataStream.Dequeue(Buffer, ReadLen);
+	uint32 ReadLen = FMath::Min< int64 >( ReqSize, DataAvailable() );
+	ReadLen = BuildDataStream.Dequeue( Buffer, ReadLen );
 	BuildDataStreamCS.Unlock();
 
 	return ReadLen;
@@ -189,7 +189,7 @@ uint32 FBuildStream::DequeueData(uint8* Buffer, const uint32& ReqSize, const boo
 
 const TArray< FString > FBuildStream::GetEmptyFiles()
 {
-	FScopeLock ScopeLock(&FilesListsCS);
+	FScopeLock ScopeLock( &FilesListsCS );
 	return EmptyFiles;
 }
 
@@ -216,7 +216,7 @@ bool FBuildStream::IsEndOfData()
 
 void FBuildStream::Clear()
 {
-	EndOfBuild(false);
+	EndOfBuild( false );
 
 	BuildDataStreamCS.Lock();
 	BuildDataStream.Empty();
@@ -245,7 +245,7 @@ uint32 FBuildStream::DataAvailable()
 	return rtn;
 }
 
-void FBuildStream::BeginNewFile(const FString& Filename, const uint64& FileSize)
+void FBuildStream::BeginNewFile( const FString& Filename, const uint64& FileSize )
 {
 	FFileSpan FileSpan;
 	FileSpan.Filename = Filename;
@@ -254,30 +254,30 @@ void FBuildStream::BeginNewFile(const FString& Filename, const uint64& FileSize)
 	FileSpan.StartIdx = BuildDataStream.TotalDataPushed();
 	BuildDataStreamCS.Unlock();
 	FilesListsCS.Lock();
-	FilesParsed.Add(FileSpan.StartIdx, FileSpan);
+	FilesParsed.Add( FileSpan.StartIdx, FileSpan );
 	FilesListsCS.Unlock();
 }
 
-void FBuildStream::AddEmptyFile(const FString& Filename)
+void FBuildStream::AddEmptyFile( const FString& Filename )
 {
 	FilesListsCS.Lock();
-	EmptyFiles.Add(Filename);
+	EmptyFiles.Add( Filename );
 	FilesListsCS.Unlock();
 }
 
-void FBuildStream::EnqueueData(const uint8* Buffer, const uint32& Len)
+void FBuildStream::EnqueueData( const uint8* Buffer, const uint32& Len )
 {
 	// Wait for space
-	while (SpaceLeft() < Len)
+	while ( SpaceLeft() < Len )
 	{
-		FPlatformProcess::Sleep(0.01f);
+		FPlatformProcess::Sleep( 0.01f );
 	}
 	BuildDataStreamCS.Lock();
-	BuildDataStream.Enqueue(Buffer, Len);
+	BuildDataStream.Enqueue( Buffer, Len );
 	BuildDataStreamCS.Unlock();
 }
 
-void FBuildStream::EndOfBuild(bool bIsEnd)
+void FBuildStream::EndOfBuild( bool bIsEnd )
 {
 	NoMoreDataCS.Lock();
 	bNoMoreData = bIsEnd;
@@ -286,25 +286,25 @@ void FBuildStream::EndOfBuild(bool bIsEnd)
 
 /* FBuildDataProcessor implementation
 *****************************************************************************/
-FBuildDataChunkProcessor::FBuildDataChunkProcessor(FBuildPatchAppManifestRef InBuildManifest, const FString& InBuildRoot)
-	: NumNewChunks(0)
-	, NumKnownChunks(0)
-	, BuildRoot(InBuildRoot)
-	, ChunkWriter(FBuildPatchServicesModule::GetCloudDirectory())
-	, CurrentChunkBufferPos(0)
+FBuildDataChunkProcessor::FBuildDataChunkProcessor( FBuildPatchAppManifestRef InBuildManifest, const FString& InBuildRoot )
+	: NumNewChunks( 0 )
+	, NumKnownChunks( 0 )
+	, BuildRoot( InBuildRoot )
+	, ChunkWriter( FBuildPatchServicesModule::GetCloudDirectory() )
+	, CurrentChunkBufferPos( 0 )
 	, CurrentFile(nullptr)
-	, IsProcessingChunk(false)
-	, IsProcessingChunkPart(false)
-	, IsProcessingFile(false)
-	, ChunkIsPushed(false)
-	, BackupChunkBufferPos(0)
-	, BackupProcessingChunk(false)
-	, BackupProcessingChunkPart(false)
+	, IsProcessingChunk( false )
+	, IsProcessingChunkPart( false )
+	, IsProcessingFile( false )
+	, ChunkIsPushed( false )
+	, BackupChunkBufferPos( 0 )
+	, BackupProcessingChunk( false )
+	, BackupProcessingChunkPart( false )
 {
 	BuildManifest = InBuildManifest;
 	CurrentChunkGuid.Invalidate();
 	BackupChunkGuid.Invalidate();
-	CurrentChunkBuffer = new uint8[FBuildPatchData::ChunkDataSize];
+	CurrentChunkBuffer = new uint8[ FBuildPatchData::ChunkDataSize ];
 }
 
 FBuildDataChunkProcessor::~FBuildDataChunkProcessor()
@@ -312,19 +312,19 @@ FBuildDataChunkProcessor::~FBuildDataChunkProcessor()
 	delete[] CurrentChunkBuffer;
 }
 
-void FBuildDataChunkProcessor::BeginNewChunk(const bool& bZeroBuffer)
+void FBuildDataChunkProcessor::BeginNewChunk( const bool& bZeroBuffer )
 {
-	check(IsProcessingChunk == false);
-	check(IsProcessingChunkPart == false);
+	check( IsProcessingChunk == false );
+	check( IsProcessingChunkPart == false );
 	IsProcessingChunk = true;
 
 	// NB: If you change this function, make sure you make required changes to 
 	// Backup of chunk if need be at the end of a recognized chunk
 
 	// Erase the chunk buffer data ready for new chunk parts
-	if (bZeroBuffer)
+	if( bZeroBuffer )
 	{
-		FMemory::Memzero(CurrentChunkBuffer, FBuildPatchData::ChunkDataSize);
+		FMemory::Memzero( CurrentChunkBuffer, FBuildPatchData::ChunkDataSize );
 		CurrentChunkBufferPos = 0;
 	}
 
@@ -334,8 +334,8 @@ void FBuildDataChunkProcessor::BeginNewChunk(const bool& bZeroBuffer)
 
 void FBuildDataChunkProcessor::BeginNewChunkPart()
 {
-	check(IsProcessingChunkPart == false);
-	check(IsProcessingChunk == true);
+	check( IsProcessingChunkPart == false );
+	check( IsProcessingChunk == true );
 	IsProcessingChunkPart = true;
 
 	// The current file should have a new chunk part setup
@@ -350,8 +350,8 @@ void FBuildDataChunkProcessor::BeginNewChunkPart()
 
 void FBuildDataChunkProcessor::EndNewChunkPart()
 {
-	check(IsProcessingChunkPart == true);
-	check(IsProcessingChunk == true);
+	check( IsProcessingChunkPart == true );
+	check( IsProcessingChunk == true );
 	IsProcessingChunkPart = false;
 
 	// Current file should have it's last chunk part updated with the size value finalized.
@@ -359,21 +359,21 @@ void FBuildDataChunkProcessor::EndNewChunkPart()
 	{
 		FChunkPartData* ChunkPart = &CurrentFile->FileChunkParts.Last();
 		ChunkPart->Size = CurrentChunkBufferPos - ChunkPart->Offset;
-		check(ChunkPart->Size != 0);
+		check( ChunkPart->Size != 0 );
 	}
 }
 
-void FBuildDataChunkProcessor::EndNewChunk(const uint64& ChunkHash, const uint8* ChunkData, const FGuid& ChunkGuid)
+void FBuildDataChunkProcessor::EndNewChunk( const uint64& ChunkHash, const uint8* ChunkData, const FGuid& ChunkGuid )
 {
-	check(IsProcessingChunk == true);
-	check(IsProcessingChunkPart == false);
+	check( IsProcessingChunk == true );
+	check( IsProcessingChunkPart == false );
 	IsProcessingChunk = false;
 
 	// A bool that will state whether we should update and log chunking progress
 	bool bLogProgress = false;
 
 	// If the new chunk was recognized, then we will have got a different Guid through, so fix up guid for files using this chunk
-	if (CurrentChunkGuid != ChunkGuid)
+	if( CurrentChunkGuid != ChunkGuid )
 	{
 		bLogProgress = true;
 		++NumKnownChunks;
@@ -381,7 +381,7 @@ void FBuildDataChunkProcessor::EndNewChunk(const uint64& ChunkHash, const uint8*
 		{
 			for (auto& ChunkPart : FileManifest.FileChunkParts)
 			{
-				if (ChunkPart.Guid == CurrentChunkGuid)
+				if( ChunkPart.Guid == CurrentChunkGuid )
 				{
 					ChunkPart.Guid = ChunkGuid;
 				}
@@ -389,7 +389,7 @@ void FBuildDataChunkProcessor::EndNewChunk(const uint64& ChunkHash, const uint8*
 		}
 	}
 	// If the chunk is new, count it if we were passed data
-	else if (ChunkData)
+	else if( ChunkData )
 	{
 		bLogProgress = true;
 		++NumNewChunks;
@@ -397,9 +397,9 @@ void FBuildDataChunkProcessor::EndNewChunk(const uint64& ChunkHash, const uint8*
 
 	// Always queue the chunk if passed data, it will be skipped automatically if existing already, and we need to make
 	// sure the latest version is saved out when recognizing an older version
-	if (ChunkData)
+	if( ChunkData )
 	{
-		ChunkWriter.QueueChunk(ChunkData, ChunkGuid, ChunkHash);
+		ChunkWriter.QueueChunk( ChunkData, ChunkGuid, ChunkHash );
 
 		// Also add the info to the data
 		if (!BuildManifest->ChunkInfoLookup.Contains(ChunkGuid))
@@ -416,16 +416,16 @@ void FBuildDataChunkProcessor::EndNewChunk(const uint64& ChunkHash, const uint8*
 
 	}
 
-	if (bLogProgress)
+	if( bLogProgress )
 	{
 		// Output to log for builder info
-		GLog->Logf(TEXT("%s %s [%d:%d]"), *BuildManifest->GetAppName(), *BuildManifest->GetVersionString(), NumNewChunks, NumKnownChunks);
+		GLog->Logf( TEXT( "%s %s [%d:%d]" ), *BuildManifest->GetAppName(), *BuildManifest->GetVersionString(), NumNewChunks, NumKnownChunks );
 	}
 }
 
 void FBuildDataChunkProcessor::PushChunk()
 {
-	check(!ChunkIsPushed);
+	check( !ChunkIsPushed );
 	ChunkIsPushed = true;
 
 	// Then we can skip over this entire chunk using it for current files in processing
@@ -433,29 +433,29 @@ void FBuildDataChunkProcessor::PushChunk()
 	BackupProcessingChunkPart = IsProcessingChunkPart;
 	BackupChunkGuid = CurrentChunkGuid;
 	BackupChunkBufferPos = CurrentChunkBufferPos;
-	if (BackupProcessingChunkPart)
+	if( BackupProcessingChunkPart )
 	{
 		EndNewChunkPart();
 	}
-	if (BackupProcessingChunk)
+	if( BackupProcessingChunk )
 	{
 		// null ChunkData will stop the incomplete chunk being saved out
-		EndNewChunk(0, NULL, CurrentChunkGuid);
+		EndNewChunk( 0, NULL, CurrentChunkGuid );
 	}
 
 	// Start this chunk, must begin from 0
 	CurrentChunkBufferPos = 0;
-	BeginNewChunk(false);
+	BeginNewChunk( false );
 	BeginNewChunkPart();
 }
 
-void FBuildDataChunkProcessor::PopChunk(const uint64& ChunkHash, const uint8* ChunkData, const FGuid& ChunkGuid)
+void FBuildDataChunkProcessor::PopChunk( const uint64& ChunkHash, const uint8* ChunkData, const FGuid& ChunkGuid )
 {
-	check(ChunkIsPushed);
+	check( ChunkIsPushed );
 	ChunkIsPushed = false;
 
 	EndNewChunkPart();
-	EndNewChunk(ChunkHash, ChunkData, ChunkGuid);
+	EndNewChunk( ChunkHash, ChunkData, ChunkGuid );
 
 	// Backup data for previous part chunk
 	IsProcessingChunk = BackupProcessingChunk;
@@ -468,18 +468,18 @@ void FBuildDataChunkProcessor::FinalChunk()
 	// If we're processing a chunk part, then we MUST be processing a chunk too
 	check(IsProcessingChunkPart ? IsProcessingChunk : true);
 
-	if (IsProcessingChunk)
+	if( IsProcessingChunk )
 	{
 		// The there is no more data and the last file is processed.
 		// The final chunk should be finished
-		uint64 NewChunkHash = FRollingHash< FBuildPatchData::ChunkDataSize >::GetHashForDataSet(CurrentChunkBuffer);
+		uint64 NewChunkHash = FRollingHash< FBuildPatchData::ChunkDataSize >::GetHashForDataSet( CurrentChunkBuffer );
 		FGuid ChunkGuid = CurrentChunkGuid;
-		FBuildDataGenerator::FindExistingChunkData(NewChunkHash, CurrentChunkBuffer, ChunkGuid);
-		if (IsProcessingChunkPart)
+		FBuildDataGenerator::FindExistingChunkData( NewChunkHash, CurrentChunkBuffer, ChunkGuid );
+		if( IsProcessingChunkPart )
 		{
 			EndNewChunkPart();
 		}
-		EndNewChunk(NewChunkHash, CurrentChunkBuffer, ChunkGuid);
+		EndNewChunk( NewChunkHash, CurrentChunkBuffer, ChunkGuid );
 	}
 
 	// Wait for the chunk writer to finish
@@ -487,9 +487,9 @@ void FBuildDataChunkProcessor::FinalChunk()
 	ChunkWriter.WaitForThread();
 }
 
-void FBuildDataChunkProcessor::BeginFile(const FString& FileName)
+void FBuildDataChunkProcessor::BeginFile( const FString& FileName )
 {
-	check(IsProcessingFile == false);
+	check( IsProcessingFile == false );
 	IsProcessingFile = true;
 
 	// We should start a new file, which begins with the current new chunk and current chunk offset
@@ -500,7 +500,7 @@ void FBuildDataChunkProcessor::BeginFile(const FString& FileName)
 	CurrentFile->SymlinkTarget = GetSymlinkTarget(*(BuildRoot / FileName));
 
 	// Setup for current chunk part
-	if (IsProcessingChunkPart)
+	if( IsProcessingChunkPart )
 	{
 		CurrentFile->FileChunkParts.Add(FChunkPartData());
 		FChunkPartData& NewPart = CurrentFile->FileChunkParts.Last();
@@ -511,12 +511,12 @@ void FBuildDataChunkProcessor::BeginFile(const FString& FileName)
 
 void FBuildDataChunkProcessor::EndFile()
 {
-	check(IsProcessingFile == true);
+	check( IsProcessingFile == true );
 	IsProcessingFile = false;
 
 	if (CurrentFile)
 	{
-		if (IsProcessingChunkPart)
+		if( IsProcessingChunkPart )
 		{
 			// Current file should have it's last chunk part updated with the size value finalized.
 			FChunkPartData* ChunkPart = &CurrentFile->FileChunkParts.Last();
@@ -531,74 +531,74 @@ void FBuildDataChunkProcessor::EndFile()
 	CurrentFile = nullptr;
 }
 
-void FBuildDataChunkProcessor::SkipKnownByte(const uint8& NextByte, const bool& bStartOfFile, const bool& bEndOfFile, const FString& Filename)
+void FBuildDataChunkProcessor::SkipKnownByte( const uint8& NextByte, const bool& bStartOfFile, const bool& bEndOfFile, const FString& Filename )
 {
 	// Check for start of new file
-	if (bStartOfFile)
+	if( bStartOfFile )
 	{
-		BeginFile(Filename);
+		BeginFile( Filename );
 		FileHash.Reset();
 	}
 
 	// Increment position between start and end of file
 	++CurrentChunkBufferPos;
-	FileHash.Update(&NextByte, 1);
+	FileHash.Update( &NextByte, 1 );
 
 	// Check for end of file
-	if (bEndOfFile)
+	if( bEndOfFile )
 	{
 		FileHash.Final();
-		FileHash.GetHash(CurrentFile->FileHash.Hash);
+		FileHash.GetHash( CurrentFile->FileHash.Hash );
 		EndFile();
 	}
 }
 
-void FBuildDataChunkProcessor::ProcessNewByte(const uint8& NewByte, const bool& bStartOfFile, const bool& bEndOfFile, const FString& Filename)
+void FBuildDataChunkProcessor::ProcessNewByte( const uint8& NewByte, const bool& bStartOfFile, const bool& bEndOfFile, const FString& Filename )
 {
 	// If we finished a chunk, we begin a new chunk and new chunk part
-	if (!IsProcessingChunk)
+	if( !IsProcessingChunk )
 	{
 		BeginNewChunk();
 		BeginNewChunkPart();
 	}
 	// Or if we finished a chunk part (by recognizing a chunk hash), we begin a new one
-	else if (!IsProcessingChunkPart)
+	else if( !IsProcessingChunkPart )
 	{
 		BeginNewChunkPart();
 	}
 
 	// Check for start of new file
-	if (bStartOfFile)
+	if( bStartOfFile )
 	{
-		BeginFile(Filename);
+		BeginFile( Filename );
 		FileHash.Reset();
 	}
 
 	// We add the old byte to our new chunk which will be used as part of the file it belonged to
-	CurrentChunkBuffer[CurrentChunkBufferPos++] = NewByte;
-	FileHash.Update(&NewByte, 1);
+	CurrentChunkBuffer[ CurrentChunkBufferPos++ ] = NewByte;
+	FileHash.Update( &NewByte, 1 );
 
 	// Check for end of file
-	if (bEndOfFile)
+	if( bEndOfFile )
 	{
 		FileHash.Final();
-		FileHash.GetHash(CurrentFile->FileHash.Hash);
+		FileHash.GetHash( CurrentFile->FileHash.Hash );
 		EndFile();
 	}
 
 	// Do we have a full new chunk?
-	check(CurrentChunkBufferPos <= FBuildPatchData::ChunkDataSize);
-	if (CurrentChunkBufferPos == FBuildPatchData::ChunkDataSize)
+	check( CurrentChunkBufferPos <= FBuildPatchData::ChunkDataSize );
+	if( CurrentChunkBufferPos == FBuildPatchData::ChunkDataSize )
 	{
-		uint64 NewChunkHash = FRollingHash< FBuildPatchData::ChunkDataSize >::GetHashForDataSet(CurrentChunkBuffer);
+		uint64 NewChunkHash = FRollingHash< FBuildPatchData::ChunkDataSize >::GetHashForDataSet( CurrentChunkBuffer );
 		FGuid ChunkGuid = CurrentChunkGuid;
-		FBuildDataGenerator::FindExistingChunkData(NewChunkHash, CurrentChunkBuffer, ChunkGuid);
+		FBuildDataGenerator::FindExistingChunkData( NewChunkHash, CurrentChunkBuffer, ChunkGuid );
 		EndNewChunkPart();
-		EndNewChunk(NewChunkHash, CurrentChunkBuffer, ChunkGuid);
+		EndNewChunk( NewChunkHash, CurrentChunkBuffer, ChunkGuid );
 	}
 }
 
-void FBuildDataChunkProcessor::GetChunkStats(uint32& OutNewFiles, uint32& OutKnownFiles)
+void FBuildDataChunkProcessor::GetChunkStats( uint32& OutNewFiles, uint32& OutKnownFiles )
 {
 	OutNewFiles = NumNewChunks;
 	OutKnownFiles = NumKnownChunks;
@@ -612,8 +612,9 @@ const TMap<FGuid, int64>& FBuildDataChunkProcessor::GetChunkFilesizes()
 
 /* FBuildDataFileProcessor implementation
 *****************************************************************************/
-FBuildDataFileProcessor::FBuildDataFileProcessor(FBuildPatchAppManifestRef InBuildManifest, const FString& InBuildRoot, const FDateTime& InDataThresholdTime)
-	: NumNewFiles(0)
+FBuildDataFileProcessor::FBuildDataFileProcessor(FBuildPatchAppManifestRef InBuildManifest, const FString& InBuildRoot, const FDateTime& InDataThresholdTime, const uint64 MaxFilePartSize)
+	: MaxFilePartSize(MaxFilePartSize)
+	, NumNewFiles(0)
 	, NumKnownFiles(0)
 	, BuildManifest(InBuildManifest)
 	, BuildRoot(InBuildRoot)
@@ -637,12 +638,23 @@ void FBuildDataFileProcessor::BeginFile(const FString& InFileName)
 	CurrentFile->bIsUnixExecutable = IsUnixExecutable(*(BuildRoot / InFileName));
 	CurrentFile->SymlinkTarget = GetSymlinkTarget(*(BuildRoot / InFileName));
 
-	// Add the 'chunk part' which will just be the whole file
-	CurrentFile->FileChunkParts.Add(FChunkPartData());
-
 	// Reset the hash and size counter
 	FileHash.Reset();
 	FileSize = 0;
+
+	// Init part data
+	BeginFilePart();
+}
+
+void FBuildDataFileProcessor::BeginFilePart()
+{
+	CurrentFilePartData = FChunkPartData();
+	CurrentFilePartInfo = FChunkInfoData();
+	
+	CurrentFilePartData.Guid = FGuid::NewGuid();
+	CurrentFilePartInfo.Guid = CurrentFilePartData.Guid;
+
+	CurrentFilePart.Empty(MaxFilePartSize);
 }
 
 void FBuildDataFileProcessor::ProcessFileData(const uint8* Data, const uint32& DataLen)
@@ -652,66 +664,92 @@ void FBuildDataFileProcessor::ProcessFileData(const uint8* Data, const uint32& D
 	// Update the hash
 	FileHash.Update(Data, DataLen);
 
-	// Count size
-	FileSize += DataLen;
+	// Process current file part
+	const uint32 NewSize = CurrentFilePartData.Size + DataLen;
+	if (NewSize > MaxFilePartSize)
+	{
+		const uint32 RemainingSize = MaxFilePartSize - CurrentFilePartData.Size;
+		check(RemainingSize < DataLen);
+		FileSize += RemainingSize;
+		CurrentFilePartData.Size += RemainingSize;
+		CurrentFilePartInfo.Hash = FCycPoly64Hash::GetHashForDataSet(Data, RemainingSize, CurrentFilePartInfo.Hash);
+		CurrentFilePart.Append(Data, RemainingSize);
+		EndFilePart();
+		BeginFilePart();
+		const uint32 OverflowSize = DataLen - RemainingSize;
+		const uint8* OverflowData = Data + RemainingSize;
+		FileSize += OverflowSize;
+		CurrentFilePartData.Size += OverflowSize;
+		CurrentFilePartInfo.Hash = FCycPoly64Hash::GetHashForDataSet(OverflowData, OverflowSize, CurrentFilePartInfo.Hash);
+		CurrentFilePart.Append(OverflowData, OverflowSize);
+	}
+	else
+	{
+		CurrentFilePartData.Size += DataLen;
+		CurrentFilePartInfo.Hash = FCycPoly64Hash::GetHashForDataSet(Data, DataLen, CurrentFilePartInfo.Hash);
+		CurrentFilePart.Append(Data, DataLen);
+		FileSize += DataLen;
+	}
+
 }
 
-void FBuildDataFileProcessor::EndFile()
+void FBuildDataFileProcessor::EndFilePart()
 {
-	check(IsProcessingFile);
-	check(CurrentFile);
-	check(CurrentFile->FileChunkParts.Num() == 1);
-	IsProcessingFile = false;
+	const FString FullSourcePath = FPaths::Combine(*BuildRoot, *CurrentFile->Filename);
 
-	// Finalize the hash
-	FileHash.Final();
-	FileHash.GetHash(CurrentFile->FileHash.Hash);
+	bool bFoundExisting = FBuildDataGenerator::FindExistingFileData(CurrentFilePart, CurrentFilePartInfo, CurrentFilePartData, DataThresholdTime);
 
-	// Use hash and full file name to find out if we have a file match
-	FString FullPath = FPaths::Combine(*BuildRoot, *CurrentFile->Filename);
-	FGuid FileGuid = FGuid::NewGuid();
-	bool bFoundSameFile = FBuildDataGenerator::FindExistingFileData(FullPath, CurrentFile->FileHash, DataThresholdTime, FileGuid);
+	// Set group number now as GUID may have changed
+	CurrentFilePartInfo.GroupNumber = FCrc::MemCrc32(&CurrentFilePartInfo.Guid, sizeof(FGuid)) % 100;
 
-	// Fill the 'chunk part' info
-	FChunkPartData& FilePart = CurrentFile->FileChunkParts[0];
-	FilePart.Guid = FileGuid;
-	FilePart.Offset = 0;
-	FilePart.Size = FileSize;
+	// Always call save to ensure new versions exist when recognizing old ones
+	bool bFilePartSavedSuccess = FBuildDataGenerator::SaveOutFileDataPart(CurrentFilePart, CurrentFilePartInfo, CurrentFilePartData);
 
-	// Init and check size
-	CurrentFile->Init();
-	check(CurrentFile->GetFileSize() == FileSize);
+	// We must have successfully saved, otherwise fatal
+	check(bFilePartSavedSuccess);
 
-	// Always call save to ensure new versions exist when recognising old ones, the save will be skipped if required copies already exist.
-	FBuildDataGenerator::SaveOutFileData(FullPath, CurrentFile->FileHash, FileGuid);
-
-	// Also add the info to the data
-	if (!BuildManifest->ChunkInfoLookup.Contains(FileGuid))
+	// Add to arrays
+	CurrentFile->FileChunkParts.Add(CurrentFilePartData);
+	if (!BuildManifest->ChunkInfoLookup.Contains(CurrentFilePartInfo.Guid))
 	{
-		BuildManifest->Data->ChunkList.Add(FChunkInfoData());
-		FChunkInfoData& DataInfo = BuildManifest->Data->ChunkList.Last();
-		DataInfo.Guid = FileGuid;
-		DataInfo.Hash = INDEX_NONE;
-		DataInfo.FileSize = FileSize;
-		DataInfo.GroupNumber = FCrc::MemCrc32(&FileGuid, sizeof(FGuid)) % 100;
+		BuildManifest->Data->ChunkList.Add(CurrentFilePartInfo);
 		// Lookup is used just to ensure one copy of each chunk added, we must not use the ptr as array will realloc
-		BuildManifest->ChunkInfoLookup.Add(FileGuid, nullptr);
+		BuildManifest->ChunkInfoLookup.Add(CurrentFilePartInfo.Guid, nullptr);
 	}
 
 	// Count new/known stats
-	if (bFoundSameFile == false)
+	if (bFoundExisting)
 	{
-		++NumNewFiles;
+		++NumKnownFiles;
 	}
 	else
 	{
-		++NumKnownFiles;
+		++NumNewFiles;
 	}
 
 	// Output to log for builder info
 	GLog->Logf(TEXT("%s %s [%d:%d]"), *BuildManifest->GetAppName(), *BuildManifest->GetVersionString(), NumNewFiles, NumKnownFiles);
 }
 
+void FBuildDataFileProcessor::EndFile()
+{
+	check(IsProcessingFile);
+	check(CurrentFile);
+	IsProcessingFile = false;
+
+	const FString FullSourcePath = FPaths::Combine(*BuildRoot, *CurrentFile->Filename);
+
+	// Finalize the hash
+	FileHash.Final();
+	FileHash.GetHash(CurrentFile->FileHash.Hash);
+
+	EndFilePart();
+
+	// Init to check size
+	CurrentFile->Init();
+	check(CurrentFile->GetFileSize() == FileSize);
+}
+
 void FBuildDataFileProcessor::GetFileStats(uint32& OutNewFiles, uint32& OutKnownFiles)
 {
 	OutNewFiles = NumNewFiles;
@@ -720,15 +758,15 @@ void FBuildDataFileProcessor::GetFileStats(uint32& OutNewFiles, uint32& OutKnown
 
 /* FBuildSimpleChunkCache::FChunkReader implementation
 *****************************************************************************/
-FBuildGenerationChunkCache::FChunkReader::FChunkReader(const FString& InChunkFilePath, TSharedRef< FChunkFile > InChunkFile, uint32* InBytesRead, const FDateTime& InDataAgeThreshold)
-	: ChunkFilePath(InChunkFilePath)
-	, ChunkFileReader(NULL)
-	, ChunkFile(InChunkFile)
-	, FileBytesRead(InBytesRead)
-	, MemoryBytesRead(0)
-	, DataAgeThreshold(InDataAgeThreshold)
+FBuildGenerationChunkCache::FChunkReader::FChunkReader( const FString& InChunkFilePath, TSharedRef< FChunkFile > InChunkFile, uint32* InBytesRead, const FDateTime& InDataAgeThreshold )
+	: ChunkFilePath( InChunkFilePath )
+	, ChunkFileReader( NULL )
+	, ChunkFile( InChunkFile )
+	, FileBytesRead( InBytesRead )
+	, MemoryBytesRead( 0 )
+	, DataAgeThreshold( InDataAgeThreshold )
 {
-	ChunkFile->GetDataLock(&ChunkData, &ChunkHeader);
+	ChunkFile->GetDataLock( &ChunkData, &ChunkHeader );
 }
 
 FBuildGenerationChunkCache::FChunkReader::~FChunkReader()
@@ -736,7 +774,7 @@ FBuildGenerationChunkCache::FChunkReader::~FChunkReader()
 	ChunkFile->ReleaseDataLock();
 
 	// Close file handle
-	if (ChunkFileReader != NULL)
+	if( ChunkFileReader != NULL )
 	{
 		ChunkFileReader->Close();
 		delete ChunkFileReader;
@@ -747,18 +785,18 @@ FBuildGenerationChunkCache::FChunkReader::~FChunkReader()
 FArchive* FBuildGenerationChunkCache::FChunkReader::GetArchive()
 {
 	// Open file handle?
-	if (ChunkFileReader == NULL)
+	if( ChunkFileReader == NULL )
 	{
-		ChunkFileReader = IFileManager::Get().CreateFileReader(*ChunkFilePath);
-		if (ChunkFileReader == NULL)
+		ChunkFileReader = IFileManager::Get().CreateFileReader( *ChunkFilePath );
+		if( ChunkFileReader == NULL )
 		{
 			// Break the magic to mark as invalid
 			ChunkHeader->Magic = 0;
-			GLog->Logf(TEXT("WARNING: Skipped missing chunk file %s"), *ChunkFilePath);
+			GLog->Logf( TEXT( "WARNING: Skipped missing chunk file %s" ), *ChunkFilePath );
 			return NULL;
 		}
 		// Read Header?
-		if (ChunkHeader->Guid.IsValid() == false)
+		if( ChunkHeader->Guid.IsValid() == false )
 		{
 			*ChunkFileReader << *ChunkHeader;
 		}
@@ -773,13 +811,13 @@ FArchive* FBuildGenerationChunkCache::FChunkReader::GetArchive()
 		const int64 ExpectedFileSize = ChunkHeader->DataSize + ChunkHeader->HeaderSize;
 		const int64 ChunkFileSize = ChunkFileReader->TotalSize();
 		const int64 NextByte = ChunkHeader->HeaderSize + *FileBytesRead;
-		if (ChunkFileSize == ExpectedFileSize
-			&& NextByte < ChunkFileSize)
+		if( ChunkFileSize == ExpectedFileSize
+		 && NextByte < ChunkFileSize )
 		{
 			// Seek to next byte
-			ChunkFileReader->Seek(NextByte);
+			ChunkFileReader->Seek( NextByte );
 			// Break the magic to mark as invalid if archive errored, this chunk will get ignored
-			if (ChunkFileReader->GetError())
+			if( ChunkFileReader->GetError() )
 			{
 				ChunkHeader->Magic = 0;
 			}
@@ -791,25 +829,25 @@ FArchive* FBuildGenerationChunkCache::FChunkReader::GetArchive()
 		}
 		// If this chunk is valid and compressed, we must read the entire file and decompress to memory now if we have not already
 		// as we cannot compare to compressed data
-		if (*FileBytesRead == 0 && ChunkHeader->StoredAs & FChunkHeader::STORED_COMPRESSED)
+		if( *FileBytesRead == 0 && ChunkHeader->StoredAs & FChunkHeader::STORED_COMPRESSED )
 		{
 			// Load the compressed chunk data
 			TArray< uint8 > CompressedData;
-			CompressedData.Empty(ChunkHeader->DataSize);
-			CompressedData.AddUninitialized(ChunkHeader->DataSize);
-			ChunkFileReader->Serialize(CompressedData.GetData(), ChunkHeader->DataSize);
+			CompressedData.Empty( ChunkHeader->DataSize );
+			CompressedData.AddUninitialized( ChunkHeader->DataSize );
+			ChunkFileReader->Serialize( CompressedData.GetData(), ChunkHeader->DataSize );
 			// Uncompress
 			bool bSuceess = FCompression::UncompressMemory(
-				static_cast<ECompressionFlags>(COMPRESS_ZLIB | COMPRESS_BiasMemory),
+				static_cast< ECompressionFlags >( COMPRESS_ZLIB | COMPRESS_BiasMemory ),
 				ChunkData,
 				FBuildPatchData::ChunkDataSize,
 				CompressedData.GetData(),
-				ChunkHeader->DataSize);
+				ChunkHeader->DataSize );
 			// Mark that we have fully read decompressed data and update the chunkfile's data size as we are expanding it
 			*FileBytesRead = FBuildPatchData::ChunkDataSize;
 			ChunkHeader->DataSize = FBuildPatchData::ChunkDataSize;
 			// Check uncompression was OK
-			if (!bSuceess)
+			if( !bSuceess )
 			{
 				ChunkHeader->Magic = 0;
 			}
@@ -820,7 +858,7 @@ FArchive* FBuildGenerationChunkCache::FChunkReader::GetArchive()
 
 const bool FBuildGenerationChunkCache::FChunkReader::IsValidChunk()
 {
-	if (ChunkHeader->Guid.IsValid() == false)
+	if( ChunkHeader->Guid.IsValid() == false )
 	{
 		GetArchive();
 	}
@@ -833,38 +871,38 @@ const bool FBuildGenerationChunkCache::FChunkReader::IsValidChunk()
 
 const FGuid& FBuildGenerationChunkCache::FChunkReader::GetChunkGuid()
 {
-	if (ChunkHeader->Guid.IsValid() == false)
+	if( ChunkHeader->Guid.IsValid() == false )
 	{
 		GetArchive();
 	}
 	return ChunkHeader->Guid;
 }
 
-void FBuildGenerationChunkCache::FChunkReader::ReadNextBytes(uint8** OutDataBuffer, const uint32& ReadLength)
+void FBuildGenerationChunkCache::FChunkReader::ReadNextBytes( uint8** OutDataBuffer, const uint32& ReadLength )
 {
-	uint8* BufferNextByte = &ChunkData[MemoryBytesRead];
+	uint8* BufferNextByte = &ChunkData[ MemoryBytesRead ];
 	// Do we need to load from disk?
-	if ((MemoryBytesRead + ReadLength) > *FileBytesRead)
+	if( ( MemoryBytesRead + ReadLength ) > *FileBytesRead )
 	{
 		FArchive* Reader = GetArchive();
 		// Do not allow incorrect usage
-		check(Reader);
-		check(ReadLength <= BytesLeft());
+		check( Reader );
+		check( ReadLength <= BytesLeft() );
 		// Read the number of bytes extra we need
 		const int32 NumFileBytesRead = *FileBytesRead;
 		const int32 NextMemoryBytesRead = MemoryBytesRead + ReadLength;
-		const uint32 FileReadLen = FMath::Max<int32>(0, NextMemoryBytesRead - NumFileBytesRead);
+		const uint32 FileReadLen = FMath::Max<int32>( 0, NextMemoryBytesRead - NumFileBytesRead );
 		*FileBytesRead += FileReadLen;
-		Reader->Serialize(BufferNextByte, FileReadLen);
+		Reader->Serialize( BufferNextByte, FileReadLen );
 		// Assert if read error, if theres some problem accessing chunks then continuing would cause bad patch
 		// ratios, so it's better to hard fault.
 		const bool bChunkReadOK = !Reader->GetError();
-		if (!bChunkReadOK)
+		if( !bChunkReadOK )
 		{
 			// Print something helpful
-			GLog->Logf(TEXT("FATAL ERROR: Could not read from chunk FArchive %s"), *ChunkFilePath);
+			GLog->Logf( TEXT( "FATAL ERROR: Could not read from chunk FArchive %s" ), *ChunkFilePath );
 			// Check with bool variable so that output will be readable
-			check(bChunkReadOK);
+			check( bChunkReadOK );
 		}
 	}
 	MemoryBytesRead += ReadLength;
@@ -873,7 +911,7 @@ void FBuildGenerationChunkCache::FChunkReader::ReadNextBytes(uint8** OutDataBuff
 
 const uint32 FBuildGenerationChunkCache::FChunkReader::BytesLeft()
 {
-	if (ChunkHeader->Guid.IsValid() == false)
+	if( ChunkHeader->Guid.IsValid() == false )
 	{
 		GetArchive();
 	}
@@ -887,44 +925,44 @@ FBuildGenerationChunkCache::FBuildGenerationChunkCache(const FDateTime& DataAgeT
 {
 }
 
-TSharedRef< FBuildGenerationChunkCache::FChunkReader > FBuildGenerationChunkCache::GetChunkReader(const FString& ChunkFilePath)
+TSharedRef< FBuildGenerationChunkCache::FChunkReader > FBuildGenerationChunkCache::GetChunkReader( const FString& ChunkFilePath )
 {
-	if (ChunkCache.Contains(ChunkFilePath) == false)
+	if( ChunkCache.Contains( ChunkFilePath ) == false )
 	{
 		// Remove oldest access from cache?
-		if (ChunkCache.Num() >= NumChunksToCache)
+		if( ChunkCache.Num() >= NumChunksToCache )
 		{
-			FString OldestAccessChunk = TEXT("");
+			FString OldestAccessChunk = TEXT( "" );
 			double OldestAccessTime = FPlatformTime::Seconds();
-			for (auto ChunkCacheIt = ChunkCache.CreateConstIterator(); ChunkCacheIt; ++ChunkCacheIt)
+			for( auto ChunkCacheIt = ChunkCache.CreateConstIterator(); ChunkCacheIt; ++ChunkCacheIt )
 			{
 				const FString& ChunkFilePath = ChunkCacheIt.Key();
 				const FChunkFile& ChunkFile = ChunkCacheIt.Value().Get();
-				if (ChunkFile.GetLastAccessTime() < OldestAccessTime)
+				if( ChunkFile.GetLastAccessTime() < OldestAccessTime )
 				{
 					OldestAccessTime = ChunkFile.GetLastAccessTime();
 					OldestAccessChunk = ChunkFilePath;
 				}
 			}
-			ChunkCache.Remove(OldestAccessChunk);
-			delete BytesReadPerChunk[OldestAccessChunk];
-			BytesReadPerChunk.Remove(OldestAccessChunk);
+			ChunkCache.Remove( OldestAccessChunk );
+			delete BytesReadPerChunk[ OldestAccessChunk ];
+			BytesReadPerChunk.Remove( OldestAccessChunk );
 		}
 		// Add the chunk to cache
-		ChunkCache.Add(ChunkFilePath, MakeShareable(new FChunkFile(1, true)));
-		BytesReadPerChunk.Add(ChunkFilePath, new uint32(0));
+		ChunkCache.Add( ChunkFilePath, MakeShareable( new FChunkFile( 1, true ) ) );
+		BytesReadPerChunk.Add( ChunkFilePath, new uint32( 0 ) );
 	}
-	return MakeShareable(new FChunkReader(ChunkFilePath, ChunkCache[ChunkFilePath], BytesReadPerChunk[ChunkFilePath], DataAgeThreshold));
+	return MakeShareable( new FChunkReader( ChunkFilePath, ChunkCache[ ChunkFilePath ], BytesReadPerChunk[ ChunkFilePath ], DataAgeThreshold ) );
 }
 
 void FBuildGenerationChunkCache::Cleanup()
 {
-	ChunkCache.Empty(0);
-	for (auto BytesReadPerChunkIt = BytesReadPerChunk.CreateConstIterator(); BytesReadPerChunkIt; ++BytesReadPerChunkIt)
+	ChunkCache.Empty( 0 );
+	for( auto BytesReadPerChunkIt = BytesReadPerChunk.CreateConstIterator(); BytesReadPerChunkIt; ++BytesReadPerChunkIt )
 	{
 		delete BytesReadPerChunkIt.Value();
 	}
-	BytesReadPerChunk.Empty(0);
+	BytesReadPerChunk.Empty( 0 );
 }
 
 /* FBuildGenerationChunkCache system singleton setup
@@ -934,21 +972,21 @@ TSharedPtr< FBuildGenerationChunkCache > FBuildGenerationChunkCache::SingletonIn
 void FBuildGenerationChunkCache::Init(const FDateTime& DataAgeThreshold)
 {
 	// We won't allow misuse of these functions
-	check(!SingletonInstance.IsValid());
+	check( !SingletonInstance.IsValid() );
 	SingletonInstance = MakeShareable(new FBuildGenerationChunkCache(DataAgeThreshold));
 }
 
 FBuildGenerationChunkCache& FBuildGenerationChunkCache::Get()
 {
 	// We won't allow misuse of these functions
-	check(SingletonInstance.IsValid());
+	check( SingletonInstance.IsValid() );
 	return *SingletonInstance.Get();
 }
 
 void FBuildGenerationChunkCache::Shutdown()
 {
 	// We won't allow misuse of these functions
-	check(SingletonInstance.IsValid());
+	check( SingletonInstance.IsValid() );
 	SingletonInstance->Cleanup();
 	SingletonInstance.Reset();
 }
@@ -958,7 +996,7 @@ void FBuildGenerationChunkCache::Shutdown()
 TMap< FGuid, FString > FBuildDataGenerator::ExistingChunkGuidInventory;
 TMap< uint64, TArray< FGuid > > FBuildDataGenerator::ExistingChunkHashInventory;
 bool FBuildDataGenerator::ExistingChunksEnumerated = false;
-TMap< FSHAHashData, TArray< FString > > FBuildDataGenerator::ExistingFileInventory;
+TMap< uint64, TArray< FString > > FBuildDataGenerator::ExistingFileInventory;
 bool FBuildDataGenerator::ExistingFilesEnumerated = false;
 FCriticalSection FBuildDataGenerator::SingleConcurrentBuildCS;
 
@@ -985,13 +1023,13 @@ static void AddCustomFieldsToBuildManifest(const TMap<FString, FVariant>& Custom
 
 /* FBuildDataGenerator implementation
 *****************************************************************************/
-bool FBuildDataGenerator::GenerateChunksManifestFromDirectory(const FBuildPatchSettings& Settings)
+bool FBuildDataGenerator::GenerateChunksManifestFromDirectory( const FBuildPatchSettings& Settings )
 {
 	// Output to log for builder info
 	GLog->Logf(TEXT("Running Chunks Patch Generation for: %u:%s %s"), Settings.AppID, *Settings.AppName, *Settings.BuildVersion);
 
 	// Take the build CS
-	FScopeLock SingleConcurrentBuild(&SingleConcurrentBuildCS);
+	FScopeLock SingleConcurrentBuild( &SingleConcurrentBuildCS );
 
 	// Create our chunk cache
 	const FDateTime Cutoff = Settings.bShouldHonorReuseThreshold ? FDateTime::UtcNow() - FTimespan::FromDays(Settings.DataAgeThreshold) : FDateTime::MinValue();
@@ -999,6 +1037,7 @@ bool FBuildDataGenerator::GenerateChunksManifestFromDirectory(const FBuildPatchS
 
 	// Create a manifest
 	FBuildPatchAppManifestRef BuildManifest = MakeShareable(new FBuildPatchAppManifest());
+	BuildManifest->Data->ManifestFileVersion = EBuildPatchAppManifestVersion::GetLatestVersion();
 
 	// Setup custom fields
 	AddCustomFieldsToBuildManifest(Settings.CustomFields, BuildManifest);
@@ -1027,11 +1066,11 @@ bool FBuildDataGenerator::GenerateChunksManifestFromDirectory(const FBuildPatchS
 
 	// Create a data buffer
 	const uint32 DataBufferSize = FBuildPatchData::ChunkDataSize;
-	uint8* DataBuffer = new uint8[DataBufferSize];
+	uint8* DataBuffer = new uint8[ DataBufferSize ];
 
 	// We'll need a rolling hash for chunking
 	FRollingHash< FBuildPatchData::ChunkDataSize >* RollingHash = new FRollingHash< FBuildPatchData::ChunkDataSize >();
-
+	
 	// Refers to how much data has been processed (into the FBuildDataProcessor)
 	uint64 ProcessPos = 0;
 
@@ -1049,10 +1088,10 @@ bool FBuildDataGenerator::GenerateChunksManifestFromDirectory(const FBuildPatchS
 	const double TimeGenStarted = LastProgressLog;
 
 	// Loop through all data
-	while (!BuildStream->IsEndOfData())
+	while ( !BuildStream->IsEndOfData() )
 	{
 		// Grab some data from the build stream
-		ReadLen = BuildStream->DequeueData(DataBuffer, DataBufferSize);
+		ReadLen = BuildStream->DequeueData( DataBuffer, DataBufferSize );
 
 		// A bool says if there's no more data to come from the Build Stream
 		const bool bNoMoreData = BuildStream->IsEndOfData();
@@ -1064,18 +1103,18 @@ bool FBuildDataGenerator::GenerateChunksManifestFromDirectory(const FBuildPatchS
 		uint32 PaddedZeros = 0;
 
 		// Process data while we have more
-		while ((DataBufferPos < ReadLen) || (bNoMoreData && PaddedZeros < RollingHash->GetWindowSize()))
+		while ( ( DataBufferPos < ReadLen ) || ( bNoMoreData && PaddedZeros < RollingHash->GetWindowSize() ) )
 		{
 			// Prime the rolling hash
-			if (RollingHash->GetNumDataNeeded() > 0)
+			if( RollingHash->GetNumDataNeeded() > 0 )
 			{
-				if (DataBufferPos < ReadLen)
+				if( DataBufferPos < ReadLen )
 				{
-					RollingHash->ConsumeByte(DataBuffer[DataBufferPos++]);
+					RollingHash->ConsumeByte( DataBuffer[ DataBufferPos++ ] );
 				}
 				else
 				{
-					RollingHash->ConsumeByte(0);
+					RollingHash->ConsumeByte( 0 );
 					++PaddedZeros;
 				}
 				// Keep looping until primed
@@ -1086,24 +1125,24 @@ bool FBuildDataGenerator::GenerateChunksManifestFromDirectory(const FBuildPatchS
 			FGuid ChunkGuid;
 			const uint64 WindowHash = RollingHash->GetWindowHash();
 			const TRingBuffer< uint8, FBuildPatchData::ChunkDataSize >& WindowData = RollingHash->GetWindowData();
-			bool ChunkRecognised = FindExistingChunkData(WindowHash, WindowData, ChunkGuid);
-			if (ChunkRecognised)
+			bool ChunkRecognised = FindExistingChunkData( WindowHash, WindowData, ChunkGuid );
+			if( ChunkRecognised )
 			{
 				// Process all bytes
 				DataProcessor.PushChunk();
 				const uint32 WindowDataSize = RollingHash->GetWindowSize() - PaddedZeros;
-				for (uint32 i = 0; i < WindowDataSize; ++i)
+				for( uint32 i = 0; i < WindowDataSize; ++i )
 				{
-					const bool bStartOfFile = BuildStream->GetFileSpan(ProcessPos, FileName, FileDataCount);
+					const bool bStartOfFile = BuildStream->GetFileSpan( ProcessPos, FileName, FileDataCount );
 					const bool bEndOfFile = FileDataCount <= 1;
-					check(FileDataCount > 0);// If FileDataCount is ever 0, it means this piece of data belongs to no file, so something is wrong
-					DataProcessor.SkipKnownByte(WindowData[i], bStartOfFile, bEndOfFile, FileName);
+					check( FileDataCount > 0 );// If FileDataCount is ever 0, it means this piece of data belongs to no file, so something is wrong
+					DataProcessor.SkipKnownByte( WindowData[i], bStartOfFile, bEndOfFile, FileName );
 					++ProcessPos;
 					--FileDataCount;
 				}
-				uint8* SerialWindowData = new uint8[FBuildPatchData::ChunkDataSize];
-				WindowData.Serialize(SerialWindowData);
-				DataProcessor.PopChunk(WindowHash, SerialWindowData, ChunkGuid);
+				uint8* SerialWindowData = new uint8[ FBuildPatchData::ChunkDataSize ];
+				WindowData.Serialize( SerialWindowData );
+				DataProcessor.PopChunk( WindowHash, SerialWindowData, ChunkGuid );
 				delete[] SerialWindowData;
 
 				// Clear
@@ -1112,29 +1151,29 @@ bool FBuildDataGenerator::GenerateChunksManifestFromDirectory(const FBuildPatchS
 			else
 			{
 				// Process one byte
-				const bool bStartOfFile = BuildStream->GetFileSpan(ProcessPos, FileName, FileDataCount);
+				const bool bStartOfFile = BuildStream->GetFileSpan( ProcessPos, FileName, FileDataCount );
 				const bool bEndOfFile = FileDataCount <= 1;
-				DataProcessor.ProcessNewByte(WindowData.Bottom(), bStartOfFile, bEndOfFile, FileName);
+				DataProcessor.ProcessNewByte( WindowData.Bottom(), bStartOfFile, bEndOfFile, FileName );
 				++ProcessPos;
 				--FileDataCount;
 
 				// Roll
-				if (DataBufferPos < ReadLen)
+				if( DataBufferPos < ReadLen )
 				{
-					RollingHash->RollForward(DataBuffer[DataBufferPos++]);
+					RollingHash->RollForward( DataBuffer[ DataBufferPos++ ] );
 				}
-				else if (bNoMoreData)
+				else if( bNoMoreData )
 				{
-					RollingHash->RollForward(0);
+					RollingHash->RollForward( 0 );
 					++PaddedZeros;
 				}
 			}
 
 			// Log processed data
-			if ((FPlatformTime::Seconds() - LastProgressLog) >= 10.0)
+			if( ( FPlatformTime::Seconds() - LastProgressLog ) >= 10.0 )
 			{
 				LastProgressLog = FPlatformTime::Seconds();
-				GLog->Logf(TEXT("Processed %lld bytes."), ProcessPos);
+				GLog->Logf( TEXT( "Processed %lld bytes." ), ProcessPos );
 			}
 		}
 	}
@@ -1169,15 +1208,11 @@ bool FBuildDataGenerator::GenerateChunksManifestFromDirectory(const FBuildPatchS
 	BuildManifest->InitLookups();
 
 	// Save manifest into the cloud directory
-	FString BaseFilename = FBuildPatchServicesModule::GetCloudDirectory() / FDefaultValueHelper::RemoveWhitespaces(BuildManifest->Data->AppName + BuildManifest->Data->BuildVersion);
-	FString JsonFilename = BaseFilename + TEXT(".manifest");
-	FString BinaryFilename = BaseFilename + TEXT(".binary.manifest");
-	BuildManifest->SaveToFile(JsonFilename, false);
-	BuildManifest->SaveToFile(BinaryFilename, true);
+	const FString OutFilename = FBuildPatchServicesModule::GetCloudDirectory() / FDefaultValueHelper::RemoveWhitespaces(BuildManifest->Data->AppName + BuildManifest->Data->BuildVersion) + TEXT(".manifest");
+	BuildManifest->SaveToFile(OutFilename, true);
 
 	// Output to log for builder info
-	GLog->Logf(TEXT("Saved manifest to %s"), *JsonFilename);
-	GLog->Logf(TEXT("Saved manifest to %s"), *BinaryFilename);
+	GLog->Logf(TEXT("Saved manifest to %s"), *OutFilename);
 
 	// Clean up memory
 	delete[] DataBuffer;
@@ -1189,16 +1224,17 @@ bool FBuildDataGenerator::GenerateChunksManifestFromDirectory(const FBuildPatchS
 	return true;
 }
 
-bool FBuildDataGenerator::GenerateFilesManifestFromDirectory(const FBuildPatchSettings& Settings)
+bool FBuildDataGenerator::GenerateFilesManifestFromDirectory( const FBuildPatchSettings& Settings )
 {
 	// Output to log for builder info
 	GLog->Logf(TEXT("Running Files Patch Generation for: %u:%s %s"), Settings.AppID, *Settings.AppName, *Settings.BuildVersion);
 
 	// Take the build CS
-	FScopeLock SingleConcurrentBuild(&SingleConcurrentBuildCS);
+	FScopeLock SingleConcurrentBuild( &SingleConcurrentBuildCS );
 
 	// Create a manifest
-	FBuildPatchAppManifestRef BuildManifest = MakeShareable(new FBuildPatchAppManifest());
+	FBuildPatchAppManifestRef BuildManifest = MakeShareable( new FBuildPatchAppManifest() );
+	BuildManifest->Data->ManifestFileVersion = EBuildPatchAppManifestVersion::GetLatestVersion();
 
 	// Setup custom fields
 	AddCustomFieldsToBuildManifest(Settings.CustomFields, BuildManifest);
@@ -1209,7 +1245,7 @@ bool FBuildDataGenerator::GenerateFilesManifestFromDirectory(const FBuildPatchSe
 
 	// Declare a build processor
 	const FDateTime Cutoff = Settings.bShouldHonorReuseThreshold ? FDateTime::UtcNow() - FTimespan::FromDays(Settings.DataAgeThreshold) : FDateTime::MinValue();
-	FBuildDataFileProcessor DataProcessor(BuildManifest, Settings.RootDirectory, Cutoff);
+	FBuildDataFileProcessor DataProcessor(BuildManifest, Settings.RootDirectory, Cutoff, 1024*1024*5);
 
 	// Set the basic details
 	BuildManifest->Data->bIsFileData = true;
@@ -1223,7 +1259,7 @@ bool FBuildDataGenerator::GenerateFilesManifestFromDirectory(const FBuildPatchSe
 	BuildManifest->Data->PrereqArgs = Settings.PrereqArgs;
 
 	// Create a data buffer
-	uint8* FileReadBuffer = new uint8[FileBufferSize];
+	uint8* FileReadBuffer = new uint8[ FileBufferSize ];
 
 	// Refers to how much data has been processed (into the FBuildDataFileProcessor)
 	uint64 ProcessPos = 0;
@@ -1237,31 +1273,31 @@ bool FBuildDataGenerator::GenerateFilesManifestFromDirectory(const FBuildPatchSe
 	FBuildDataGenerator::StripIgnoredFiles(AllFiles, Settings.RootDirectory, Settings.IgnoreListFile);
 
 	// Loop through all files
-	for (auto FileIt = AllFiles.CreateConstIterator(); FileIt; ++FileIt)
+	for( auto FileIt = AllFiles.CreateConstIterator(); FileIt; ++FileIt )
 	{
 		const FString& FileName = *FileIt;
 		// Read the file
-		FArchive* FileReader = IFileManager::Get().CreateFileReader(*FileName);
-		if (FileReader != NULL)
+		FArchive* FileReader = IFileManager::Get().CreateFileReader( *FileName );
+		if( FileReader != NULL )
 		{
 			// Make SourceFile the format we want it in and start a new file
 			FString SourceFile = FileName;
 			FPaths::MakePathRelativeTo(SourceFile, *(Settings.RootDirectory + TEXT("/")));
 			int64 FileSize = FileReader->TotalSize();
-			if (FileSize < 0)
+			if( FileSize < 0 )
 			{
 				// Skip potential error ( INDEX_NONE == -1 )
 				continue;
 			}
-			DataProcessor.BeginFile(SourceFile);
-			while (!FileReader->AtEnd())
+			DataProcessor.BeginFile( SourceFile );
+			while( !FileReader->AtEnd() )
 			{
 				const int64 SizeLeft = FileSize - FileReader->Tell();
-				const uint32 ReadLen = FMath::Min< int64 >(FileBufferSize, SizeLeft);
+				const uint32 ReadLen = FMath::Min< int64 >( FileBufferSize, SizeLeft );
 				ProcessPos += ReadLen;
-				FileReader->Serialize(FileReadBuffer, ReadLen);
+				FileReader->Serialize( FileReadBuffer, ReadLen );
 				// Copy into data stream
-				DataProcessor.ProcessFileData(FileReadBuffer, ReadLen);
+				DataProcessor.ProcessFileData( FileReadBuffer, ReadLen );
 			}
 			FileReader->Close();
 			delete FileReader;
@@ -1277,15 +1313,11 @@ bool FBuildDataGenerator::GenerateFilesManifestFromDirectory(const FBuildPatchSe
 	BuildManifest->InitLookups();
 
 	// Save manifest into the cloud directory
-	FString BaseFilename = FBuildPatchServicesModule::GetCloudDirectory() / FDefaultValueHelper::RemoveWhitespaces(BuildManifest->Data->AppName + BuildManifest->Data->BuildVersion);
-	FString JsonFilename = BaseFilename + TEXT(".manifest");
-	FString BinaryFilename = BaseFilename + TEXT(".binary.manifest");
-	BuildManifest->SaveToFile(JsonFilename, false);
-	BuildManifest->SaveToFile(BinaryFilename, true);
+	const FString OutFilename = FBuildPatchServicesModule::GetCloudDirectory() / FDefaultValueHelper::RemoveWhitespaces(BuildManifest->Data->AppName + BuildManifest->Data->BuildVersion) + TEXT(".manifest");
+	BuildManifest->SaveToFile(OutFilename, true);
 
 	// Output to log for builder info
-	GLog->Logf(TEXT("Saved manifest to %s"), *JsonFilename);
-	GLog->Logf(TEXT("Saved manifest to %s"), *BinaryFilename);
+	GLog->Logf(TEXT("Saved manifest to %s"), *OutFilename);
 
 	// Clean up memory
 	delete[] FileReadBuffer;
@@ -1294,13 +1326,13 @@ bool FBuildDataGenerator::GenerateFilesManifestFromDirectory(const FBuildPatchSe
 	return true;
 }
 
-bool FBuildDataGenerator::FindExistingChunkData(const uint64& ChunkHash, const uint8* ChunkData, FGuid& ChunkGuid)
+bool FBuildDataGenerator::FindExistingChunkData( const uint64& ChunkHash, const uint8* ChunkData, FGuid& ChunkGuid )
 {
 	// Quick code hack
 	TRingBuffer< uint8, FBuildPatchData::ChunkDataSize > ChunkDataRing;
 	const uint32 ChunkDataLen = FBuildPatchData::ChunkDataSize;
-	ChunkDataRing.Enqueue(ChunkData, ChunkDataLen);
-	return FindExistingChunkData(ChunkHash, ChunkDataRing, ChunkGuid);
+	ChunkDataRing.Enqueue( ChunkData, ChunkDataLen );
+	return FindExistingChunkData( ChunkHash, ChunkDataRing, ChunkGuid );
 }
 
 bool FBuildDataGenerator::FindExistingChunkData(const uint64& ChunkHash, const TRingBuffer< uint8, FBuildPatchData::ChunkDataSize >& ChunkData, FGuid& ChunkGuid)
@@ -1311,9 +1343,6 @@ bool FBuildDataGenerator::FindExistingChunkData(const uint64& ChunkHash, const T
 	if (ExistingChunksEnumerated == false)
 	{
 		IFileManager& FileManager = IFileManager::Get();
-		FString JSONOutput;
-		TSharedRef< TJsonWriter< TCHAR, TPrettyJsonPrintPolicy< TCHAR > > > DebugWriter = TJsonWriterFactory< TCHAR, TPrettyJsonPrintPolicy< TCHAR > >::Create(&JSONOutput);
-		DebugWriter->WriteObjectStart();
 
 		// Find all manifest files
 		const FString CloudDir = FBuildPatchServicesModule::GetCloudDirectory();
@@ -1350,11 +1379,13 @@ bool FBuildDataGenerator::FindExistingChunkData(const uint64& ChunkHash, const T
 							{
 								if (ChunkHash != 0)
 								{
+									const FString SourceFile = FBuildPatchUtils::GetFileNewFilename(BuildManifest->GetManifestVersion(), CloudDir, ChunkGuid, ChunkHash);
 									TArray< FGuid >& HashChunkList = ExistingChunkHashInventory.FindOrAdd(ChunkHash);
 									if (!HashChunkList.Contains(ChunkGuid))
 									{
 										++NumChunksFound;
 										HashChunkList.Add(ChunkGuid);
+										ExistingChunkGuidInventory.Add(ChunkGuid, SourceFile);
 									}
 								}
 								else
@@ -1399,20 +1430,20 @@ bool FBuildDataGenerator::FindExistingChunkData(const uint64& ChunkHash, const T
 			for (auto ChunkIt = ChunkList->CreateConstIterator(); ChunkIt && !bFoundMatchingChunk; ++ChunkIt)
 			{
 				FGuid Guid = *ChunkIt;
-				if (!ExistingChunkGuidInventory.Contains(Guid))
-				{
-					ExistingChunkGuidInventory.Add(Guid, DiscoverChunkFilename(Guid, ChunkHash));
-				}
-				const FString& SourceFile = ExistingChunkGuidInventory[Guid];
-				// Read the file
-				uint8* TempChunkData = new uint8[FBuildPatchData::ChunkDataSize];
-				ChunkData.Serialize(TempChunkData);
-				bool bChunkIsUsable = true;
-				if (CompareDataToChunk(SourceFile, TempChunkData, Guid, bChunkIsUsable))
+				bool bChunkIsUsable = false;
+				if (ExistingChunkGuidInventory.Contains(Guid))
 				{
-					// We have a chunk match!!
-					bFoundMatchingChunk = true;
-					ChunkGuid = Guid;
+					const FString& SourceFile = ExistingChunkGuidInventory[Guid];
+					// Read the file
+					uint8* TempChunkData = new uint8[FBuildPatchData::ChunkDataSize];
+					ChunkData.Serialize(TempChunkData);
+					if (CompareDataToChunk(SourceFile, TempChunkData, Guid, bChunkIsUsable))
+					{
+						// We have a chunk match!!
+						bFoundMatchingChunk = true;
+						ChunkGuid = Guid;
+					}
+					delete[] TempChunkData;
 				}
 				// Check if this chunk should be dumped
 				if (!bChunkIsUsable)
@@ -1421,7 +1452,6 @@ bool FBuildDataGenerator::FindExistingChunkData(const uint64& ChunkHash, const T
 					ChunkList->Remove(Guid);
 					--ChunkIt;
 				}
-				delete[] TempChunkData;
 			}
 		}
 	}
@@ -1429,159 +1459,84 @@ bool FBuildDataGenerator::FindExistingChunkData(const uint64& ChunkHash, const T
 	return bFoundMatchingChunk;
 }
 
-FString FBuildDataGenerator::DiscoverChunkFilename(const FGuid& ChunkGuid, const uint64& ChunkHash)
+bool FBuildDataGenerator::FindExistingFileData(const TArray<uint8>& Data, FChunkInfoData& FilePartInfo, FChunkPartData& FilePartData, const FDateTime& DataThresholdTime)
 {
-	static double AccumTime = 0.0;
-	const double StartDiscovery = FPlatformTime::Seconds();
-	const FString CloudDir = FBuildPatchServicesModule::GetCloudDirectory();
-	FString ChunkFilenameChecked;
-	for (EBuildPatchAppManifestVersion::Type VersionCounter = static_cast<EBuildPatchAppManifestVersion::Type>(EBuildPatchAppManifestVersion::LatestPlusOne - 1);
-		VersionCounter >= EBuildPatchAppManifestVersion::Original;
-		VersionCounter = static_cast<EBuildPatchAppManifestVersion::Type>(VersionCounter - 1))
-	{
-		const FString ChunkFilename = VersionCounter < EBuildPatchAppManifestVersion::DataFileRenames ?
-			FBuildPatchUtils::GetChunkOldFilename(CloudDir, ChunkGuid) :
-			FBuildPatchUtils::GetChunkNewFilename(VersionCounter, CloudDir, ChunkGuid, ChunkHash);
-		if (ChunkFilenameChecked != ChunkFilename)
-		{
-			ChunkFilenameChecked = ChunkFilename;
-			if (FPaths::FileExists(ChunkFilename))
-			{
-				const double DiscoveryTime = FPlatformTime::Seconds() - StartDiscovery;
-				AccumTime += DiscoveryTime;
-				GLog->Logf(TEXT("BuildDataGenerator: DiscoverChunkFilename: Chunk %s found in %f secs. Accum time %.2f secs"), *ChunkGuid.ToString(), DiscoveryTime, AccumTime);
-				return ChunkFilename;
-			}
-		}
-	}
-	GLog->Logf(TEXT("BuildDataGenerator: DiscoverChunkFilename: Chunk %s not found"), *ChunkGuid.ToString());
-	return TEXT("");
-}
+	check(Data.Num() == FilePartData.Size);
 
-bool FBuildDataGenerator::FindExistingFileData(const FString& InSourceFile, const FSHAHashData& InFileHash, const FDateTime& DataThresholdTime, FGuid& OutFileGuid)
-{
-	bool bFoundMatchingFile = false;
+	bool bFoundMatchingData = false;
 
 	// Perform an inventory on Cloud files if not already done
 	if (ExistingFilesEnumerated == false)
 	{
 		IFileManager& FileManager = IFileManager::Get();
-		FString JSONOutput;
-		TSharedRef< TJsonWriter< TCHAR, TPrettyJsonPrintPolicy< TCHAR > > > DebugWriter = TJsonWriterFactory< TCHAR, TPrettyJsonPrintPolicy< TCHAR > >::Create(&JSONOutput);
-		DebugWriter->WriteObjectStart();
-
-		TArray< FGuid > FoundFiles;
-
-		// The directory containing old filename version files
-		const FString CloudFileDir = FBuildPatchServicesModule::GetCloudDirectory() / TEXT("Files");
-		// The directory containing new filename version files
-		const FString CloudFile2Dir = FBuildPatchServicesModule::GetCloudDirectory() / TEXT("FilesV2");
-
-		GLog->Logf(TEXT("BuildDataGenerator: Enumerating Files from %s"), *CloudFile2Dir);
-		if (FileManager.DirectoryExists(*CloudFile2Dir))
+		// Find all manifest files
+		const FString CloudDir = FBuildPatchServicesModule::GetCloudDirectory();
+		if (FileManager.DirectoryExists(*CloudDir))
 		{
 			const double StartEnumerate = FPlatformTime::Seconds();
-
-			// Find all files
-			TArray<FString> AllFiles;
-			FileManager.FindFilesRecursive(AllFiles, *CloudFile2Dir, TEXT("*.file"), true, false);
-
-			for (auto FileIt = AllFiles.CreateConstIterator(); FileIt; ++FileIt)
-			{
-				FString SourceFile = *FileIt;
-				FSHAHashData FoundHash;
-				FGuid FoundGuid;
-				FBuildPatchUtils::GetFileDetailFromNewFilename(SourceFile, FoundGuid, FoundHash);
-
-				// Add to inventory
-				FoundFiles.Add(FoundGuid);
-				TArray< FString >& FileList = ExistingFileInventory.FindOrAdd(FoundHash);
-				FileList.Add(SourceFile);
-
-				const void* HashBuffer = &FoundHash;
-				DebugWriter->WriteValue(FoundGuid.ToString(), FString::FromBlob(static_cast<const uint8*>(HashBuffer), sizeof(FoundHash)));
-			}
-
+			TArray<FString> AllManifests;
+			GLog->Logf(TEXT("BuildDataGenerator: Enumerating Manifests from %s"), *CloudDir);
+			FileManager.FindFiles(AllManifests, *(CloudDir / TEXT("*.manifest")), true, false);
 			const double EnumerateTime = FPlatformTime::Seconds() - StartEnumerate;
-			GLog->Logf(TEXT("BuildDataGenerator: Found %d new name files in %.1f seconds"), AllFiles.Num(), EnumerateTime);
-		}
-		else
-		{
-			GLog->Logf(TEXT("BuildDataGenerator: Cloud directory does not exist: %s"), *CloudFile2Dir);
-		}
-
-		GLog->Logf(TEXT("BuildDataGenerator: Enumerating Files From %s"), *CloudFileDir);
-		if (IFileManager::Get().DirectoryExists(*CloudFileDir))
-		{
-			const double StartEnumerate = FPlatformTime::Seconds();
-			const int32 PreviousNumFiles = FoundFiles.Num();
-
-			// Find all files
-			TArray<FString> AllFiles;
-			IFileManager::Get().FindFilesRecursive(AllFiles, *CloudFileDir, TEXT("*.file"), true, false);
+			GLog->Logf(TEXT("BuildDataGenerator: Found %d manifests in %.1f seconds"), AllManifests.Num(), EnumerateTime);
 
-			for (auto FileIt = AllFiles.CreateConstIterator(); FileIt; ++FileIt)
+			// Load all manifest files
+			uint64 NumFilePartsFound = 0;
+			const double StartLoadAllManifest = FPlatformTime::Seconds();
+			for (const auto& ManifestFile : AllManifests)
 			{
-				FString SourceFile = *FileIt;
-
-				// Skip any for GUID we already found
-				FGuid FoundGuid;
-				FBuildPatchUtils::GetGUIDFromFilename(SourceFile, FoundGuid);
-				if (FoundFiles.Contains(FoundGuid))
-				{
-					continue;
-				}
-
-				// Read the file
-				FArchive* FileReader = IFileManager::Get().CreateFileReader(*SourceFile);
-				if (FileReader != NULL)
+				// Determine file parts from manifest file
+				const FString ManifestFilename = CloudDir / ManifestFile;
+				FBuildPatchAppManifestRef BuildManifest = MakeShareable(new FBuildPatchAppManifest());
+				const double StartLoadManifest = FPlatformTime::Seconds();
+				if (BuildManifest->LoadFromFile(ManifestFilename))
 				{
-					// Read the header
-					FChunkHeader Header;
-					*FileReader << Header;
-
-					// Check magic
-					if (Header.IsValidMagic() && Header.HashType == FChunkHeader::HASH_SHA1)
+					const double LoadManifestTime = FPlatformTime::Seconds() - StartLoadManifest;
+					GLog->Logf(TEXT("BuildDataGenerator: Loaded %s in %.1f seconds"), *ManifestFile, LoadManifestTime);
+					if (BuildManifest->GetManifestVersion() >= EBuildPatchAppManifestVersion::FileSplittingSupport)
 					{
-						TArray< FString >& FileList = ExistingFileInventory.FindOrAdd(Header.SHAHash);
-						FileList.Add(SourceFile);
-						FoundFiles.Add(Header.Guid);
-
-						const void* HashBuffer = &Header.SHAHash;
-						DebugWriter->WriteValue(Header.Guid.ToString(), FString::FromBlob(static_cast<const uint8*>(HashBuffer), sizeof(Header.SHAHash)));
+						if (BuildManifest->IsFileDataManifest())
+						{
+							TArray<FGuid> FilePartsReferenced;
+							BuildManifest->GetDataList(FilePartsReferenced);
+							for (const auto& FilePartGuid : FilePartsReferenced)
+							{
+								uint64 FilePartHash;
+								if (BuildManifest->GetFilePartHash(FilePartGuid, FilePartHash))
+								{
+									// Add to inventory
+									const FString SourceFile = FBuildPatchUtils::GetFileNewFilename(BuildManifest->GetManifestVersion(), CloudDir, FilePartGuid, FilePartHash);
+									TArray< FString >& FileList = ExistingFileInventory.FindOrAdd(FilePartHash);
+									FileList.Add(SourceFile);
+									++NumFilePartsFound;
+								}
+								else
+								{
+									GLog->Logf(TEXT("BuildDataGenerator: WARNING: Missing part hash for %s in manifest %s"), *FilePartGuid.ToString(), *ManifestFile);
+								}
+							}
+						}
+						else
+						{
+							GLog->Logf(TEXT("BuildDataGenerator: INFO: Ignoring chunked manifest %s"), *ManifestFilename);
+						}
 					}
 					else
 					{
-						GLog->Logf(TEXT("BuildDataGenerator: Failed magic/hashtype check on file [%d:%d] %s"), Header.Magic, Header.HashType, *SourceFile);
+						GLog->Logf(TEXT("BuildDataGenerator: INFO: Not recognising from old file manifests %s"), *ManifestFilename);
 					}
-
-					FileReader->Close();
-					delete FileReader;
 				}
 				else
 				{
-					GLog->Logf(TEXT("BuildDataGenerator: Failed to read chunk %s"), *SourceFile);
+					GLog->Logf(TEXT("BuildDataGenerator: WARNING: Could not read Manifest file. Data recognition will suffer (%s)"), *ManifestFilename);
 				}
 			}
-
-			const double EnumerateTime = FPlatformTime::Seconds() - StartEnumerate;
-			const int32 NewNumFiles = FoundFiles.Num();
-			GLog->Logf(TEXT("BuildDataGenerator: Found %d extra old files in %.1f seconds"), NewNumFiles - PreviousNumFiles, EnumerateTime);
+			const double LoadAllManifestTime = FPlatformTime::Seconds() - StartLoadAllManifest;
+			GLog->Logf(TEXT("BuildDataGenerator: Used %d manifests to enumerate %llu file parts in %.1f seconds"), AllManifests.Num(), NumFilePartsFound, LoadAllManifestTime);
 		}
 		else
 		{
-			GLog->Logf(TEXT("BuildDataGenerator: Cloud directory does not exist: %s"), *CloudFileDir);
-		}
-
-		DebugWriter->WriteObjectEnd();
-		DebugWriter->Close();
-
-		FArchive* FileOut = IFileManager::Get().CreateFileWriter(*(CloudFileDir + TEXT("DebugFileList.txt")));
-		if (FileOut != NULL)
-		{
-			FileOut->Serialize(TCHAR_TO_ANSI(*JSONOutput), JSONOutput.Len());
-			FileOut->Close();
-			delete FileOut;
+			GLog->Logf(TEXT("BuildDataGenerator: Cloud directory does not exist: %s"), *CloudDir);
 		}
 
 		ExistingFilesEnumerated = true;
@@ -1590,13 +1545,12 @@ bool FBuildDataGenerator::FindExistingFileData(const FString& InSourceFile, cons
 	// Do we have a file matching this data?
 	if (ExistingFileInventory.Num() > 0)
 	{
-		TArray< FString >* FileList = ExistingFileInventory.Find(InFileHash);
+		TArray< FString >* FileList = ExistingFileInventory.Find(FilePartInfo.Hash);
 		if (FileList != NULL)
 		{
-			// We need to load each chunk in this list and compare data
-			for (auto FileIt = FileList->CreateConstIterator(); FileIt && !bFoundMatchingFile; ++FileIt)
+			// We need to load each file part in this list and compare data
+			for (const FString& CloudFilename : *FileList)
 			{
-				FString CloudFilename = *FileIt;
 				// Check the file date
 				FDateTime ModifiedDate = IFileManager::Get().GetTimeStamp(*CloudFilename);
 				if (ModifiedDate < DataThresholdTime)
@@ -1605,161 +1559,125 @@ bool FBuildDataGenerator::FindExistingFileData(const FString& InSourceFile, cons
 					continue;
 				}
 				// Compare the files
-				FArchive* SourceFile = IFileManager::Get().CreateFileReader(*InSourceFile);
-				FArchive* FoundFile = IFileManager::Get().CreateFileReader(*CloudFilename);
-				if (SourceFile != NULL && FoundFile != NULL)
+				TArray<uint8> FoundFileData;
+				FChunkHeader FoundHeader;
+				if (FFileHelper::LoadFileToArray(FoundFileData, *CloudFilename))
 				{
-					FChunkHeader FoundHeader;
-					*FoundFile << FoundHeader;
-					const int64 SourceFileSize = SourceFile->TotalSize();
-					const int64 FoundFileSize = FoundFile->TotalSize();
-					if (SourceFileSize == FoundHeader.DataSize)
+					const int32 OriginalFileSize = FoundFileData.Num();
+					if (FBuildPatchUtils::UncompressFileDataFile(FoundFileData, &FoundHeader))
 					{
-						// Currently only support stored raw!
 						check(FoundHeader.StoredAs == FChunkHeader::STORED_RAW);
-						// Makes no sense here for the sizes to not add up
-						check(FoundFileSize == (FoundHeader.DataSize + FoundHeader.HeaderSize));
-						// Move FoundFile to start of file data
-						FoundFile->Seek(FoundHeader.HeaderSize);
-						// Compare
-						bool bSameData = true;
-						uint8* TempSourceBuffer = new uint8[FileBufferSize];
-						uint8* TempFoundBuffer = new uint8[FileBufferSize];
-						while (bSameData && ((SourceFile->AtEnd() || FoundFile->AtEnd()) == false))
+						const int64 FoundFileSize = FoundFileData.Num();
+						if (FoundHeader.RollingHash == FilePartInfo.Hash && Data.Num() == FoundHeader.DataSize)
 						{
-							const int64 SizeLeft = SourceFileSize - SourceFile->Tell();
-							const uint32 ReadLen = FMath::Min< int64 >(FileBufferSize, SizeLeft);
-							SourceFile->Serialize(TempSourceBuffer, ReadLen);
-							FoundFile->Serialize(TempFoundBuffer, ReadLen);
-							bSameData = FMemory::Memcmp(TempSourceBuffer, TempFoundBuffer, ReadLen) == 0;
-						}
-						delete[] TempSourceBuffer;
-						delete[] TempFoundBuffer;
-						// Did we match?
-						if (bSameData && SourceFile->AtEnd() && FoundFile->AtEnd())
-						{
-							// Yes we did!
-							bFoundMatchingFile = true;
-							OutFileGuid = FoundHeader.Guid;
+							// Compare
+							const bool bSameData = FMemory::Memcmp(Data.GetData(), FoundFileData.GetData() + FoundHeader.HeaderSize, FoundHeader.DataSize) == 0;
+							// Did we match?
+							if (bSameData)
+							{
+								// Yes we did!
+								bFoundMatchingData = true;
+								FilePartInfo.Guid = FoundHeader.Guid;
+								FilePartInfo.FileSize = OriginalFileSize;
+								FilePartData.Guid = FoundHeader.Guid;
+								break;
+							}
 						}
 					}
 				}
-				if (SourceFile != NULL)
-				{
-					SourceFile->Close();
-					delete SourceFile;
-				}
-				if (FoundFile != NULL)
-				{
-					FoundFile->Close();
-					delete FoundFile;
-				}
 			}
 		}
 	}
 
-	return bFoundMatchingFile;
+	return bFoundMatchingData;
 }
 
-bool FBuildDataGenerator::SaveOutFileData(const FString& SourceFile, const FSHAHashData& FileHash, const FGuid& FileGuid)
+bool FBuildDataGenerator::SaveOutFileDataPart(TArray<uint8>& Data, FChunkInfoData& FilePartInfo, FChunkPartData& FilePartData)
 {
 	bool bAlreadySaved = false;
 	bool bSuccess = false;
 	IFileManager& FileManager = IFileManager::Get();
 
-	const FString NewFilename = FBuildPatchUtils::GetFileNewFilename( EBuildPatchAppManifestVersion::GetLatestVersion(), FBuildPatchServicesModule::GetCloudDirectory(), FileGuid, FileHash );
+	const FString NewFilename = FBuildPatchUtils::GetFileNewFilename(EBuildPatchAppManifestVersion::GetLatestVersion(), FBuildPatchServicesModule::GetCloudDirectory(), FilePartInfo.Guid, FilePartInfo.Hash);
 	bAlreadySaved = FPaths::FileExists( NewFilename );
 
-#if SAVE_OLD_FILEDATA_FILENAMES
-	const FString OldFilename = FBuildPatchUtils::GetFileOldFilename( FBuildPatchServicesModule::GetCloudDirectory(), FileGuid );
-	bAlreadySaved = bAlreadySaved && FPaths::FileExists( OldFilename );
-#endif
-
 	if( !bAlreadySaved )
 	{
 		FArchive* FileOut = FileManager.CreateFileWriter( *NewFilename );
-		FArchive* FileIn = FileManager.CreateFileReader( *SourceFile );
-
-		bSuccess = FileOut != NULL && FileIn != NULL;
-
-#if SAVE_OLD_FILEDATA_FILENAMES
-		FArchive* OldFileOut = FileManager.CreateFileWriter( *OldFilename );
-		bSuccess = bSuccess && OldFileOut != NULL;
-#endif
+		bSuccess = FileOut != nullptr;
 
 		if( bSuccess )
 		{
-			const int64 FileSize = FileIn->TotalSize();
-
-			// LSwift: No support for too large files currently
-			check( FileSize <= 0xFFFFFFFF );
+			// LSwift: No support for too large file parts
+			check(Data.Num() <= 0xFFFFFFFF);
 
 			// Setup Header
 			FChunkHeader Header;
 			*FileOut << Header;
 			Header.HeaderSize = FileOut->Tell();
-			Header.DataSize = FileSize;
-			Header.Guid = FileGuid;
-			Header.HashType = FChunkHeader::HASH_SHA1;
-			Header.SHAHash = FileHash;
-			Header.StoredAs = FChunkHeader::STORED_RAW;
-
-			// Write out file
-			FileOut->Seek( 0 );
-			*FileOut << Header;
-
-#if SAVE_OLD_FILEDATA_FILENAMES
-			*OldFileOut << Header;
-#endif
+			Header.Guid = FilePartInfo.Guid;
+			Header.HashType = FChunkHeader::HASH_ROLLING;
+			Header.RollingHash = FilePartInfo.Hash;
+
+			// DataSize must be uncompressed size!
+			Header.DataSize = Data.Num();
+
+			// Setup to handle compression
+			bool bDataIsCompressed = true;
+			uint8* DataSource = Data.GetData();
+			int32 DataSourceSize = Data.Num();
+			TArray< uint8 > TempCompressedData;
+			TempCompressedData.Empty(DataSourceSize);
+			TempCompressedData.AddUninitialized(DataSourceSize);
+			int32 CompressedSize = DataSourceSize;
+
+			// Compressed can increase in size, but the function will return as failure in that case
+			// we can allow that to happen since we would not keep larger compressed data anyway.
+			bDataIsCompressed = FCompression::CompressMemory(
+				static_cast<ECompressionFlags>(COMPRESS_ZLIB | COMPRESS_BiasMemory),
+				TempCompressedData.GetData(),
+				CompressedSize,
+				DataSource,
+				DataSourceSize);
 
-			uint8* FileDataBuffer = new uint8[ FileBufferSize ];
-			while( !FileIn->AtEnd() )
+			// If compression succeeded, set data vars
+			if (bDataIsCompressed)
 			{
-				const int64 SizeLeft = FileSize - FileIn->Tell();
-				const uint32 ReadLen = FMath::Min< int64 >( FileBufferSize, SizeLeft );
-				FileIn->Serialize( FileDataBuffer, ReadLen );
-				FileOut->Serialize( FileDataBuffer, ReadLen );
+				DataSource = TempCompressedData.GetData();
+				DataSourceSize = CompressedSize;
+			}
 
-#if SAVE_OLD_FILEDATA_FILENAMES
-				OldFileOut->Serialize( FileDataBuffer, ReadLen );
-#endif
+			Header.StoredAs = bDataIsCompressed ? FChunkHeader::STORED_COMPRESSED : FChunkHeader::STORED_RAW;
+			Header.HashType = FChunkHeader::HASH_ROLLING;
 
-			}
-			delete[] FileDataBuffer;
+			// Write out files
+			FileOut->Seek(0);
+			*FileOut << Header;
+			FileOut->Serialize(DataSource, DataSourceSize);
+			
+			FilePartInfo.FileSize = FileOut->TotalSize();
+
+			bSuccess = !FileOut->IsError();
 		}
 
 		// Close files
-		if( FileIn != NULL )
-		{
-			FileIn->Close();
-			delete FileIn;
-		}
 		if( FileOut != NULL )
 		{
 			FileOut->Close();
 			delete FileOut;
 		}
-
-#if SAVE_OLD_FILEDATA_FILENAMES
-		if( OldFileOut != NULL )
-		{
-			OldFileOut->Close();
-			delete OldFileOut;
-		}
-#endif
-
 	}
 
-	return bSuccess;
+	return bSuccess || bAlreadySaved;
 }
 
-bool FBuildDataGenerator::CompareDataToChunk(const FString& ChunkFilePath, uint8* ChunkData, FGuid& ChunkGuid, bool& OutSourceChunkIsValid)
+bool FBuildDataGenerator::CompareDataToChunk( const FString& ChunkFilePath, uint8* ChunkData, FGuid& ChunkGuid, bool& OutSourceChunkIsValid )
 {
 	bool bMatching = false;
 
 	// Read the file
-	TSharedRef< FBuildGenerationChunkCache::FChunkReader > ChunkReader = FBuildGenerationChunkCache::Get().GetChunkReader(ChunkFilePath);
-	if (ChunkReader->IsValidChunk())
+	TSharedRef< FBuildGenerationChunkCache::FChunkReader > ChunkReader = FBuildGenerationChunkCache::Get().GetChunkReader( ChunkFilePath );
+	if( ChunkReader->IsValidChunk() )
 	{
 		ChunkGuid = ChunkReader->GetChunkGuid();
 		// Default true
@@ -1768,11 +1686,11 @@ bool FBuildDataGenerator::CompareDataToChunk(const FString& ChunkFilePath, uint8
 		const uint32 CompareSize = 64;
 		uint8* ReadBuffer;
 		uint32 NumCompared = 0;
-		while (bMatching && ChunkReader->BytesLeft() > 0 && NumCompared < FBuildPatchData::ChunkDataSize)
+		while( bMatching && ChunkReader->BytesLeft() > 0 && NumCompared < FBuildPatchData::ChunkDataSize )
 		{
-			const uint32 ReadLen = FMath::Min< uint32 >(CompareSize, ChunkReader->BytesLeft());
-			ChunkReader->ReadNextBytes(&ReadBuffer, ReadLen);
-			bMatching = FMemory::Memcmp(&ChunkData[NumCompared], ReadBuffer, ReadLen) == 0;
+			const uint32 ReadLen = FMath::Min< uint32 >( CompareSize, ChunkReader->BytesLeft() );
+			ChunkReader->ReadNextBytes( &ReadBuffer, ReadLen );
+			bMatching = FMemory::Memcmp( &ChunkData[ NumCompared ], ReadBuffer, ReadLen ) == 0;
 			NumCompared += ReadLen;
 		}
 	}
@@ -1783,37 +1701,37 @@ bool FBuildDataGenerator::CompareDataToChunk(const FString& ChunkFilePath, uint8
 	return bMatching;
 }
 
-void FBuildDataGenerator::StripIgnoredFiles(TArray< FString >& AllFiles, const FString& DepotDirectory, const FString& IgnoreListFile)
+void FBuildDataGenerator::StripIgnoredFiles( TArray< FString >& AllFiles, const FString& DepotDirectory, const FString& IgnoreListFile )
 {
 	const int32 OriginalNumFiles = AllFiles.Num();
-	FString IgnoreFileList = TEXT("");
-	FFileHelper::LoadFileToString(IgnoreFileList, *IgnoreListFile);
+	FString IgnoreFileList = TEXT( "" );
+	FFileHelper::LoadFileToString( IgnoreFileList, *IgnoreListFile );
 	TArray< FString > IgnoreFiles;
-	IgnoreFileList.ParseIntoArray(&IgnoreFiles, TEXT("\r\n"), true);
+	IgnoreFileList.ParseIntoArray( &IgnoreFiles, TEXT( "\r\n" ), true );
 	struct FRemoveMatchingStrings
-	{
+	{ 
 		const FString* MatchingString;
-		FRemoveMatchingStrings(const FString* InMatch)
+		FRemoveMatchingStrings( const FString* InMatch )
 			: MatchingString(InMatch) {}
 
 		bool operator()(const FString& RemovalCandidate) const
 		{
 			FString PathA = RemovalCandidate;
 			FString PathB = *MatchingString;
-			FPaths::NormalizeFilename(PathA);
-			FPaths::NormalizeFilename(PathB);
+			FPaths::NormalizeFilename( PathA );
+			FPaths::NormalizeFilename( PathB );
 			return PathA == PathB;
 		}
 	};
 
-	for (int32 IgnoreIdx = 0; IgnoreIdx < IgnoreFiles.Num(); ++IgnoreIdx)
+	for( int32 IgnoreIdx = 0; IgnoreIdx < IgnoreFiles.Num(); ++IgnoreIdx )
 	{
 		const FString& IgnoreFile = IgnoreFiles[IgnoreIdx];
 		const FString FullIgnorePath = DepotDirectory / IgnoreFile;
-		AllFiles.RemoveAll(FRemoveMatchingStrings(&FullIgnorePath));
+		AllFiles.RemoveAll( FRemoveMatchingStrings( &FullIgnorePath ) );
 	}
 	const int32 NewNumFiles = AllFiles.Num();
-	GLog->Logf(TEXT("Stripped %d ignorable file(s)"), (OriginalNumFiles - NewNumFiles));
+	GLog->Logf( TEXT( "Stripped %d ignorable file(s)" ), ( OriginalNumFiles - NewNumFiles ) );
 }
 
 #endif //WITH_BUILDPATCHGENERATION
diff --git a/Engine/Source/Runtime/Online/BuildPatchServices/Private/BuildPatchGeneration.h b/Engine/Source/Runtime/Online/BuildPatchServices/Private/BuildPatchGeneration.h
--- a/Engine/Source/Runtime/Online/BuildPatchServices/Private/BuildPatchGeneration.h
+++ b/Engine/Source/Runtime/Online/BuildPatchServices/Private/BuildPatchGeneration.h
@@ -365,7 +365,7 @@ private:
 	/**
 	 * Make the default constructor private so that we must have the required params passed in.
 	 */
-	FBuildDataFileProcessor() { check(false); }
+	FBuildDataFileProcessor() : MaxFilePartSize(0) { check(false); }
 
 public:
 
@@ -375,13 +375,18 @@ public:
 	 * @param	InBuildRoot				The root directory of the build
 	 * @param	InDataThresholdTime		The cutoff time used for file reuse. Files older than this will not be reused
 	 */
-	FBuildDataFileProcessor( FBuildPatchAppManifestRef InBuildManifest, const FString& InBuildRoot, const FDateTime& InDataThresholdTime );
+	FBuildDataFileProcessor(FBuildPatchAppManifestRef InBuildManifest, const FString& InBuildRoot, const FDateTime& InDataThresholdTime, const uint64 MaxFilePartSize);
 
 	/**
 	 * Mark the beginning of a new file
 	 * @param	InFileName		The file's local filename
 	 */
-	void BeginFile( const FString& InFileName );
+	void BeginFile(const FString& InFileName);
+
+	/**
+	 * Starts a new file part for the current file
+	 */
+	void BeginFilePart();
 
 	/**
 	 * Process a set of data for a file
@@ -390,6 +395,11 @@ public:
 	 */
 	void ProcessFileData( const uint8* Data, const uint32& DataLen );
 
+	/**
+	 * Ends the current file part of this file
+	 */
+	void EndFilePart();
+
 	/**
 	 * Mark the end of the current file
 	 */
@@ -404,6 +414,9 @@ public:
 
 private:
 
+	// The maximum filepart allowed
+	const uint64 MaxFilePartSize;
+
 	// Keep track of the number of new files
 	uint32 NumNewFiles;
 
@@ -419,6 +432,13 @@ private:
 	// Holds the current file being processed. Receives the chunk parts that we pass over.
 	FFileManifestData* CurrentFile;
 
+	// Holds current file part data and info
+	FChunkPartData CurrentFilePartData;
+	FChunkInfoData CurrentFilePartInfo;
+
+	// Holds the current part data for saving out
+	TArray<uint8> CurrentFilePart;
+
 	// Holds an SHA hash calculator
 	FSHA1 FileHash;
 
@@ -610,35 +630,24 @@ public:
 	static bool FindExistingChunkData( const uint64& ChunkHash, const uint8* ChunkData, FGuid& ChunkGuid );
 
 	/**
-	 * Given a Guid and Hash for a chunk file, try to find it from a cloud directory, starting from the latest
-	 * version of the naming convention.
-	 * NOTE: This function is blocking and will not return until finished. Don't run on main thread.
-	 * @param ChunkGuid		The guid of the chunk data
-	 * @param ChunkHash		The chunk's hash value
-	 * @return		The full patch to the found chunk, otherwise empty string.
-	 */
-	static FString DiscoverChunkFilename(const FGuid& ChunkGuid, const uint64& ChunkHash);
-
-	/**
-	 * Checks to see if a given file already exists in the database of known files
+	 * Checks to see if a given file data part already exists in the database of known files.
 	 * NOTE: This function is blocking and will not return until finished. Don't run on main thread.
-	 * @param InSourceFile			IN		The filename of the file on disk.
-	 * @param InFileHash			IN		The hash for the file.
+	 * @param Data					IN		The data for the file part
+	 * @param FilePartInfo			INOUT	The part info, the guid and hash will be setup by this function
+	 * @param FilePartData			INOUT	The part data, the guid will be set correctly if we found a match
 	 * @param DataThresholdTime		IN		The date/time before which files will not be reused
-	 * @param OutFileGuid			OUT		Will be set to the existing file guid if found.
 	 * @return		Whether this file is new or existing.
 	 */
-	static bool FindExistingFileData(const FString& InSourceFile, const FSHAHashData& InFileHash, const FDateTime& DataThresholdTime, FGuid& OutFileGuid);
+	static bool FindExistingFileData(const TArray<uint8>& Data, FChunkInfoData& FilePartInfo, FChunkPartData& FilePartData, const FDateTime& DataThresholdTime);
 
 	/**
 	 * Saves out the file data for use with file based patch manifests
 	 * NOTE: This function is blocking and will not return until finished. Don't run on main thread.
 	 * @param SourceFile		The filename of the source file on disk.
-	 * @param FileHash			The hash for the file.
-	 * @param FileGuid			The GUID for the file.
+	 * @param InFileManifest	The file manifest for saving out parts.
 	 * @return		Whether this file is new or existing.
 	 */
-	static bool SaveOutFileData(const FString& SourceFile, const FSHAHashData& FileHash, const FGuid& FileGuid);
+	static bool SaveOutFileDataPart(TArray<uint8>& Data, FChunkInfoData& FilePartInfo, FChunkPartData& FilePartData);
 
 	/**
 	 * Strips out files from the provided array that appear in the ignore list file
@@ -671,7 +680,7 @@ private:
 	static bool ExistingChunksEnumerated;
 
 	// For File Manifest generation, the existing files
-	static TMap< FSHAHashData, TArray< FString > > ExistingFileInventory;
+	static TMap< uint64, TArray< FString > > ExistingFileInventory;
 
 	// For File Manifest generation, have we enumerated the existing files yet
 	static bool ExistingFilesEnumerated;
diff --git a/Engine/Source/Runtime/Online/BuildPatchServices/Private/BuildPatchManifest.cpp b/Engine/Source/Runtime/Online/BuildPatchServices/Private/BuildPatchManifest.cpp
--- a/Engine/Source/Runtime/Online/BuildPatchServices/Private/BuildPatchManifest.cpp
+++ b/Engine/Source/Runtime/Online/BuildPatchServices/Private/BuildPatchManifest.cpp
@@ -8,9 +8,6 @@
 
 #define LOCTEXT_NAMESPACE "BuildPatchManifest"
 
-// Remove this when we are to enable creating compressed file data and related manifests
-#define ENABLE_NOCHUNKS_COMPRESSION 0
-
 // The manifest header magic codeword, for quick checking that the opened file is probably a manifest file.
 #define MANIFEST_HEADER_MAGIC		0x44BEC00C
 
@@ -57,11 +54,7 @@ bool BufferIsJsonManifest(const TArray<uint8>& DataInput)
 *****************************************************************************/
 const EBuildPatchAppManifestVersion::Type EBuildPatchAppManifestVersion::GetLatestVersion()
 {
-#if ENABLE_NOCHUNKS_COMPRESSION
 	return static_cast<EBuildPatchAppManifestVersion::Type>(LatestPlusOne - 1);
-#else
-	return EBuildPatchAppManifestVersion::StoredAsCompressedUClass;
-#endif
 }
 
 const EBuildPatchAppManifestVersion::Type EBuildPatchAppManifestVersion::GetLatestJsonVersion()
