Fix Google Play crash, UEPLAT-237
Uncomment the code that logs in to Google Play
Add some null checks for sanity

diff --git a/Engine/Source/Runtime/Online/Android/OnlineSubsystemGooglePlay/Private/OnlineAchievementsInterfaceGooglePlay.cpp b/Engine/Source/Runtime/Online/Android/OnlineSubsystemGooglePlay/Private/OnlineAchievementsInterfaceGooglePlay.cpp
--- a/Engine/Source/Runtime/Online/Android/OnlineSubsystemGooglePlay/Private/OnlineAchievementsInterfaceGooglePlay.cpp
+++ b/Engine/Source/Runtime/Online/Android/OnlineSubsystemGooglePlay/Private/OnlineAchievementsInterfaceGooglePlay.cpp
@@ -17,6 +17,12 @@ FOnlineAchievementsGooglePlay::FOnlineAchievementsGooglePlay( FOnlineSubsystemGo
 
 void FOnlineAchievementsGooglePlay::QueryAchievements(const FUniqueNetId& PlayerId, const FOnQueryAchievementsCompleteDelegate& Delegate)
 {
+	if (AndroidSubsystem->GetGameServices() == nullptr)
+	{
+		Delegate.ExecuteIfBound(PlayerId, false);
+		return;
+	}
+
 	auto QueryTask = new FOnlineAsyncTaskGooglePlayQueryAchievements(AndroidSubsystem, FUniqueNetIdString(PlayerId), Delegate);
 	AndroidSubsystem->QueueAsyncTask(QueryTask);
 }
@@ -129,6 +135,12 @@ void FOnlineAchievementsGooglePlay::WriteAchievements( const FUniqueNetId& Playe
 {
 	// @todo: verify the PlayerId passed in is the currently signed in player
 		
+	if (AndroidSubsystem->GetGameServices() == nullptr)
+	{
+		Delegate.ExecuteIfBound(PlayerId, false);
+		return;
+	}
+
 	// We need valid data from Google first because we use the achievement type to convert Unreal progress
 	// (percentage) to Google progress (binary locked/unlocked or integer number of arbitrary steps).
 	// Kick off a query if we don't have valid data.
diff --git a/Engine/Source/Runtime/Online/Android/OnlineSubsystemGooglePlay/Private/OnlineExternalUIInterfaceGooglePlay.cpp b/Engine/Source/Runtime/Online/Android/OnlineSubsystemGooglePlay/Private/OnlineExternalUIInterfaceGooglePlay.cpp
--- a/Engine/Source/Runtime/Online/Android/OnlineSubsystemGooglePlay/Private/OnlineExternalUIInterfaceGooglePlay.cpp
+++ b/Engine/Source/Runtime/Online/Android/OnlineSubsystemGooglePlay/Private/OnlineExternalUIInterfaceGooglePlay.cpp
@@ -30,12 +30,22 @@ bool FOnlineExternalUIGooglePlay::ShowInviteUI(int32 LocalUserNum)
 
 bool FOnlineExternalUIGooglePlay::ShowAchievementsUI(int32 LocalUserNum) 
 {
+	if (Subsystem->GetGameServices() == nullptr)
+	{
+		return false;
+	}
+
 	Subsystem->GetGameServices()->Achievements().ShowAllUI();
 	return true;
 }
 
 bool FOnlineExternalUIGooglePlay::ShowLeaderboardUI(const FString& LeaderboardName)
 {
+	if (Subsystem->GetGameServices() == nullptr)
+	{
+		return false;
+	}
+
 	auto DefaultSettings = GetDefault<UAndroidRuntimeSettings>();
 
 	for(const auto& Mapping : DefaultSettings->LeaderboardMap)
diff --git a/Engine/Source/Runtime/Online/Android/OnlineSubsystemGooglePlay/Private/OnlineLeaderboardInterfaceGooglePlay.cpp b/Engine/Source/Runtime/Online/Android/OnlineSubsystemGooglePlay/Private/OnlineLeaderboardInterfaceGooglePlay.cpp
--- a/Engine/Source/Runtime/Online/Android/OnlineSubsystemGooglePlay/Private/OnlineLeaderboardInterfaceGooglePlay.cpp
+++ b/Engine/Source/Runtime/Online/Android/OnlineSubsystemGooglePlay/Private/OnlineLeaderboardInterfaceGooglePlay.cpp
@@ -15,9 +15,17 @@ FOnlineLeaderboardsGooglePlay::FOnlineLeaderboardsGooglePlay(FOnlineSubsystemGoo
 
 bool FOnlineLeaderboardsGooglePlay::ReadLeaderboards(const TArray< TSharedRef<FUniqueNetId> >& Players, FOnlineLeaderboardReadRef& ReadObject)
 {
-	ReadObject->ReadState = EOnlineAsyncTaskState::InProgress;
 	ReadObject->Rows.Empty();
 
+	if (Subsystem->GetGameServices() == nullptr)
+	{
+		ReadObject->ReadState = EOnlineAsyncTaskState::Failed;
+		Subsystem->GetLeaderboardsInterface()->TriggerOnLeaderboardReadCompleteDelegates(false);
+		return false;
+	}
+
+	ReadObject->ReadState = EOnlineAsyncTaskState::InProgress;
+
 	auto ReadTask = new FOnlineAsyncTaskGooglePlayReadLeaderboard(
 		Subsystem,
 		ReadObject,
@@ -99,6 +107,12 @@ bool FOnlineLeaderboardsGooglePlay::FlushLeaderboards(const FName& SessionName)
 {
 	UE_LOG_ONLINE(Display, TEXT("flush leaderboards session name :%s"), *SessionName.ToString());
 
+	if (Subsystem->GetGameServices() == nullptr)
+	{
+		Subsystem->GetLeaderboardsInterface()->TriggerOnLeaderboardFlushCompleteDelegates(SessionName, false);
+		return false;
+	}
+
 	bool Success = true;
 
 	for(int32 Index = 0; Index < UnreportedScores.Num(); ++Index)
diff --git a/Engine/Source/Runtime/Online/Android/OnlineSubsystemGooglePlay/Private/OnlineSubsystemGooglePlay.cpp b/Engine/Source/Runtime/Online/Android/OnlineSubsystemGooglePlay/Private/OnlineSubsystemGooglePlay.cpp
--- a/Engine/Source/Runtime/Online/Android/OnlineSubsystemGooglePlay/Private/OnlineSubsystemGooglePlay.cpp
+++ b/Engine/Source/Runtime/Online/Android/OnlineSubsystemGooglePlay/Private/OnlineSubsystemGooglePlay.cpp
@@ -120,7 +120,7 @@ bool FOnlineSubsystemGooglePlay::Init()
 	{
 		StoreInterface = MakeShareable(new FOnlineStoreGooglePlay(this));
 	}
-	/*
+	
 	extern struct android_app* GNativeAndroidApp;
 	check(GNativeAndroidApp != nullptr);
 	AndroidInitialization::android_main(GNativeAndroidApp);
@@ -144,7 +144,7 @@ bool FOnlineSubsystemGooglePlay::Init()
 			OnAuthActionFinished(Op, Status);
 		})
 		.Create(PlatformConfiguration).release() );
-*/
+
 	return true;
 }
 
