Merging using UE4-To-UE4-LauncherReleases

#http
- expose LogHttp category

diff --git a/Engine/Source/Runtime/Online/HTTP/Private/HttpPrivatePCH.h b/Engine/Source/Runtime/Online/HTTP/Private/HttpPrivatePCH.h
--- a/Engine/Source/Runtime/Online/HTTP/Private/HttpPrivatePCH.h
+++ b/Engine/Source/Runtime/Online/HTTP/Private/HttpPrivatePCH.h
@@ -5,4 +5,3 @@
 #include "Core.h"
 #include "Http.h"
 
-DECLARE_LOG_CATEGORY_EXTERN(LogHttp, Warning, All);
\ No newline at end of file
diff --git a/Engine/Source/Runtime/Online/HTTP/Public/Http.h b/Engine/Source/Runtime/Online/HTTP/Public/Http.h
--- a/Engine/Source/Runtime/Online/HTTP/Public/Http.h
+++ b/Engine/Source/Runtime/Online/HTTP/Public/Http.h
@@ -9,3 +9,5 @@
 
 // Module includes
 #include "HttpModule.h"
+
+HTTP_API DECLARE_LOG_CATEGORY_EXTERN(LogHttp, Warning, All);
\ No newline at end of file
