Fix for shutdown crash when BuildPatchServices module is released after UObject has already shutdown

diff --git a/Engine/Source/Runtime/Online/BuildPatchServices/Private/BuildPatchServicesModule.cpp b/Engine/Source/Runtime/Online/BuildPatchServices/Private/BuildPatchServicesModule.cpp
--- a/Engine/Source/Runtime/Online/BuildPatchServices/Private/BuildPatchServicesModule.cpp
+++ b/Engine/Source/Runtime/Online/BuildPatchServices/Private/BuildPatchServicesModule.cpp
@@ -65,6 +65,9 @@ void FBuildPatchServicesModule::StartupModule()
 	// Add our ticker
 	FTicker::GetCoreTicker().AddTicker( FTickerDelegate::CreateRaw( this, &FBuildPatchServicesModule::Tick ) );
 
+	// Register core PreExit
+	FCoreDelegates::OnPreExit.AddRaw(this, &FBuildPatchServicesModule::PreExit);
+
 	// Test the rolling hash algorithm
 	check( CheckRollingHashAlgorithm() );
 }
@@ -72,22 +75,8 @@ void FBuildPatchServicesModule::StartupModule()
 void FBuildPatchServicesModule::ShutdownModule()
 {
 	GWarn->Logf( TEXT( "BuildPatchServicesModule: Shutting Down" ) );
-	// We need to stop and wait for any threads to exit.
-	FBuildPatchInstallError::SetFatalError( EBuildPatchInstallError::ApplicationClosing );
-
-	// Release our ptr to analytics
-	FBuildPatchAnalytics::SetAnalyticsProvider( NULL );
 
-	// Reset all installer ptrs (this will cause us to wait for each thread to complete)
-	for (auto& BuildPatchInstaller : BuildPatchInstallers)
-	{
-		// Make sure it is not paused, this function only un-pauses when in error state
-		BuildPatchInstaller->TogglePauseInstall();
-		// We still have to manually wait for the thread as another system could hold a shared ptr
-		// thus we would not be calling the destructor here
-		BuildPatchInstaller->WaitForThread();
-	}
-	BuildPatchInstallers.Empty();
+	checkf(BuildPatchInstallers.Num() == 0, TEXT("BuildPatchServicesModule: FATAL ERROR: Core PreExit not called, or installer created during shutdown!"));
 
 	// Remove our ticker
 	FTicker::GetCoreTicker().RemoveTicker( FTickerDelegate::CreateRaw( this, &FBuildPatchServicesModule::Tick ) );
@@ -233,6 +222,26 @@ void FBuildPatchServicesModule::RegisterAppInstallation(IBuildManifestRef AppMan
 	InstallationInfo.RegisterAppInstallation(AppManifest, AppInstallDirectory);
 }
 
+void FBuildPatchServicesModule::PreExit()
+{
+	// Set shutdown error so any running threads know to exit.
+	FBuildPatchInstallError::SetFatalError(EBuildPatchInstallError::ApplicationClosing);
+
+	// Release our ptr to analytics
+	FBuildPatchAnalytics::SetAnalyticsProvider(NULL);
+
+	// Cleanup installers
+	for (auto& BuildPatchInstaller : BuildPatchInstallers)
+	{
+		// Make sure it is not paused, this function only un-pauses when in error state
+		BuildPatchInstaller->TogglePauseInstall();
+		// We still have to manually wait for the thread as another system could hold a shared ptr
+		// thus we would not be calling the destructor here
+		BuildPatchInstaller->WaitForThread();
+	}
+	BuildPatchInstallers.Empty();
+}
+
 const FString& FBuildPatchServicesModule::GetStagingDirectory()
 {
 	// Default staging directory
diff --git a/Engine/Source/Runtime/Online/BuildPatchServices/Private/BuildPatchServicesModule.h b/Engine/Source/Runtime/Online/BuildPatchServices/Private/BuildPatchServicesModule.h
--- a/Engine/Source/Runtime/Online/BuildPatchServices/Private/BuildPatchServicesModule.h
+++ b/Engine/Source/Runtime/Online/BuildPatchServices/Private/BuildPatchServicesModule.h
@@ -115,6 +115,11 @@ private:
 	 */
 	bool Tick( float Delta );
 
+	/**
+	 * This will get called when core PreExits. Make sure any running installers are canceled out.
+	 */
+	void PreExit();
+
 private:
 	// Holds the cloud directory where chunks should belong
 	static FString CloudDirectory;
