Merging using UE4-To-UE4-LauncherReleases

CL# 2364809 on 19/11/2014 08:11:51
--------------------------------------------
Fixing resume problem when build label changes or update is released

diff --git a/Engine/Source/Runtime/Online/BuildPatchServices/Private/BuildPatchManifest.cpp b/Engine/Source/Runtime/Online/BuildPatchServices/Private/BuildPatchManifest.cpp
--- a/Engine/Source/Runtime/Online/BuildPatchServices/Private/BuildPatchManifest.cpp
+++ b/Engine/Source/Runtime/Online/BuildPatchServices/Private/BuildPatchManifest.cpp
@@ -1260,6 +1260,39 @@ IBuildManifestRef FBuildPatchAppManifest::Duplicate() const
 	return MakeShareable(new FBuildPatchAppManifest(*this));
 }
 
+void FBuildPatchAppManifest::CopyCustomFields(IBuildManifestRef InOther, bool bClobber)
+{
+	// Cast manifest parameters
+	FBuildPatchAppManifestRef Other = StaticCastSharedRef< FBuildPatchAppManifest >(InOther);
+
+	// Use lookup to overwrite existing and list additional fields
+	TArray<FCustomFieldData> Extras;
+	for (const auto& CustomField : Other->Data->CustomFields)
+	{
+		if (CustomFieldLookup.Contains(CustomField.Key))
+		{
+			if (bClobber)
+			{
+				CustomFieldLookup[CustomField.Key]->Value = CustomField.Value;
+			}
+		}
+		else
+		{
+			Extras.Add(CustomField);
+		}
+	}
+
+	// Add the extra fields
+	Data->CustomFields.Append(Extras);
+
+	// Reset the loookup
+	CustomFieldLookup.Empty(Data->CustomFields.Num());
+	for (auto& CustomField : Data->CustomFields)
+	{
+		CustomFieldLookup.Add(CustomField.Key, &CustomField);
+	}
+}
+
 const IManifestFieldPtr FBuildPatchAppManifest::GetCustomField(const FString& FieldName) const
 {
 	IManifestFieldPtr Return;
diff --git a/Engine/Source/Runtime/Online/BuildPatchServices/Private/BuildPatchManifest.h b/Engine/Source/Runtime/Online/BuildPatchServices/Private/BuildPatchManifest.h
--- a/Engine/Source/Runtime/Online/BuildPatchServices/Private/BuildPatchManifest.h
+++ b/Engine/Source/Runtime/Online/BuildPatchServices/Private/BuildPatchManifest.h
@@ -357,11 +357,12 @@ public:
 	virtual void GetRemovableFiles(IBuildManifestRef OldManifest, TArray< FString >& RemovableFiles) const override;
 	virtual void GetRemovableFiles(const TCHAR* InstallPath, TArray< FString >& RemovableFiles) const override;
 	virtual bool NeedsResaving() const override;
-	virtual IBuildManifestRef Duplicate() const override;
+	virtual void CopyCustomFields(IBuildManifestRef Other, bool bClobber) override;
 	virtual const IManifestFieldPtr GetCustomField(const FString& FieldName) const override;
 	virtual const IManifestFieldPtr SetCustomField(const FString& FieldName, const FString& Value) override;
 	virtual const IManifestFieldPtr SetCustomField(const FString& FieldName, const double& Value) override;
 	virtual const IManifestFieldPtr SetCustomField(const FString& FieldName, const int64& Value) override;
+	virtual IBuildManifestRef Duplicate() const override;
 	// END IBuildManifest Interface
 
 	/**
diff --git a/Engine/Source/Runtime/Online/BuildPatchServices/Public/Interfaces/IBuildManifest.h b/Engine/Source/Runtime/Online/BuildPatchServices/Public/Interfaces/IBuildManifest.h
--- a/Engine/Source/Runtime/Online/BuildPatchServices/Public/Interfaces/IBuildManifest.h
+++ b/Engine/Source/Runtime/Online/BuildPatchServices/Public/Interfaces/IBuildManifest.h
@@ -136,13 +136,13 @@ public:
 	 * @return True if the manifest was created from the latest format
 	 */
 	virtual bool NeedsResaving() const = 0;
-	
+
 	/**
-	 * Duplicated this manifest to create a copy. Should be used if storing a received manifest as an installed
-	 * manifest which would then be unique
-	 * @return A shared ref to the new manifest
+	 * Copy the custom fields from another manifest into this one. If this manifest has custom fields, matching keys will be overwritten but extras will remain
+	 * @param Other		The manifest to copy from
+	 * @param bClobber	Whether to overwrite any already existing fields
 	 */
-	virtual IBuildManifestRef Duplicate() const = 0;
+	virtual void CopyCustomFields(IBuildManifestRef Other, bool bClobber) = 0;
 
 	/**
 	 * Get a custom field from the manifest
@@ -160,6 +160,13 @@ public:
 	virtual const IManifestFieldPtr SetCustomField(const FString& FieldName, const FString& Value) = 0;
 	virtual const IManifestFieldPtr SetCustomField(const FString& FieldName, const double& Value) = 0;
 	virtual const IManifestFieldPtr SetCustomField(const FString& FieldName, const int64& Value) = 0;
+	
+	/**
+	 * Duplicated this manifest to create a copy. Should be used if storing a received manifest as an installed
+	 * manifest which would then be unique
+	 * @return A shared ref to the new manifest
+	 */
+	virtual IBuildManifestRef Duplicate() const = 0;
 };
 
 #endif // __IBuildManifest_h__
