Refactor of url handlers to be navigators for managing each of the types

#branch UE4

#change Considated http and marketplace\\ handling into a single place and added a placeholder on for # (anchor) handler (still todo)

#change Added code allowing the ability to specify anchor tags on any slate launcher widget (unused atm)

#change Fixed issue with pre-existing non unity compile by including missing headers

diff --git a/Engine/Source/Runtime/Online/BuildPatchServices/Private/BuildPatchHash.h b/Engine/Source/Runtime/Online/BuildPatchServices/Private/BuildPatchHash.h
--- a/Engine/Source/Runtime/Online/BuildPatchServices/Private/BuildPatchHash.h
+++ b/Engine/Source/Runtime/Online/BuildPatchServices/Private/BuildPatchHash.h
@@ -265,13 +265,13 @@ private:
 namespace FCycPoly64Hash
 {
 	/**
-	 * Static function to simply return the hash for a given data range.
+	 * Calculate the hash for a given data range.
 	 * @param DataSet	The buffer to the data. This must be a buffer of length == DataSize
 	 * @param DataSize	The size of the buffer
 	 * @param State		The current hash state, if continuing hash from previous buffer. 0 for new data.
 	 * @return			The hash state for the provided data
 	 */
-	static uint64 GetHashForDataSet(const uint8* DataSet, const uint32 DataSize, const uint64 State = 0);
+	uint64 GetHashForDataSet(const uint8* DataSet, const uint32 DataSize, const uint64 State = 0);
 }
 
 /* TRingBuffer implementation
