Fix for manifest duplication not creating new UObject data

diff --git a/Engine/Source/Runtime/Online/BuildPatchServices/Private/BuildPatchManifest.cpp b/Engine/Source/Runtime/Online/BuildPatchServices/Private/BuildPatchManifest.cpp
--- a/Engine/Source/Runtime/Online/BuildPatchServices/Private/BuildPatchManifest.cpp
+++ b/Engine/Source/Runtime/Online/BuildPatchServices/Private/BuildPatchManifest.cpp
@@ -475,6 +475,13 @@ FBuildPatchAppManifest::FBuildPatchAppManifest(const uint32& InAppID, const FStr
 	Data->AppName = AppName;
 }
 
+FBuildPatchAppManifest::FBuildPatchAppManifest(const FBuildPatchAppManifest& Other)
+{
+	Data = DuplicateObject<UBuildPatchManifest>(Other.Data, Other.Data->GetOuter());
+	InitLookups();
+	bNeedsResaving = Other.bNeedsResaving;
+}
+
 FBuildPatchAppManifest::~FBuildPatchAppManifest()
 {
 //	Data->MarkPendingKill(); ??? Mark for destory and run GC?
diff --git a/Engine/Source/Runtime/Online/BuildPatchServices/Private/BuildPatchManifest.h b/Engine/Source/Runtime/Online/BuildPatchServices/Private/BuildPatchManifest.h
--- a/Engine/Source/Runtime/Online/BuildPatchServices/Private/BuildPatchManifest.h
+++ b/Engine/Source/Runtime/Online/BuildPatchServices/Private/BuildPatchManifest.h
@@ -338,6 +338,11 @@ public:
 	 */
 	FBuildPatchAppManifest(const uint32& InAppID, const FString& AppName);
 
+	/**
+	 * Copy constructor
+	 */
+	FBuildPatchAppManifest(const FBuildPatchAppManifest& Other);
+
 	/**
 	 * Default destructor
 	 */
