Rename -touchonly parameter in BuildPatchTool.exe to -nopatchdelete
Improve logging in BuildPatchTool.exe when -nopatchdelete is specified (was previously logging files as being deleted)

diff --git a/Engine/Source/Programs/BuildPatchTool/Private/BuildPatchToolMain.cpp b/Engine/Source/Programs/BuildPatchTool/Private/BuildPatchToolMain.cpp
--- a/Engine/Source/Programs/BuildPatchTool/Private/BuildPatchToolMain.cpp
+++ b/Engine/Source/Programs/BuildPatchTool/Private/BuildPatchToolMain.cpp
@@ -55,7 +55,7 @@
 	Optional arguments:
 	-stdout				Adds stdout logging to the app.
 	-preview			Log all the actions it will take to update internal structures, but don't actually execute them.
-	-touchonly			When specified, will only set the modified date of referenced files to the current time, but will NOT delete any patch data.
+	-nopatchdelete		When specified, will only set the modified date of referenced files to the current time, but will NOT delete any patch data.
 	-ManifestsList=""			Specifies in quotes, the list of manifest filenames to keep following the operation. If omitted, all manifests are kept.
 	-ManifestsFile=""			Specifies in quotes, the name of the file (relative to CloudDir) which contains a list of manifests to keep, one manifest filename per line
 	-DataAgeThreshold=14.25		The maximum age in days of chunk files that will be retained. All older chunks will be deleted.
@@ -156,7 +156,7 @@ int32 BuildPatchToolMain( const TCHAR* CommandLine )
 	bool bCompactify = false;
 	bool bPatchGeneration = true;
 	bool bPreview = false;
-	bool bTouchOnly = false;
+	bool bNoPatchDelete = false;
 	bool bPatchWithReuseAgeThreshold = true;
 
 	// Collect all the info from the CommandLine
@@ -189,8 +189,8 @@ int32 BuildPatchToolMain( const TCHAR* CommandLine )
 		Matcher.Command = TEXT("preview");
 		bPreview = bCompactify && Switches.IndexOfByPredicate(Matcher) != INDEX_NONE;
 
-		Matcher.Command = TEXT("touchonly");
-		bTouchOnly = bCompactify && Switches.IndexOfByPredicate(Matcher) != INDEX_NONE;
+		Matcher.Command = TEXT("nopatchdelete");
+		bNoPatchDelete = bCompactify && Switches.IndexOfByPredicate(Matcher) != INDEX_NONE;
 		
 		Matcher.Command = TEXT( "BuildRoot" );
 		BuildRootIdx = Switches.IndexOfByPredicate(Matcher);
@@ -383,9 +383,9 @@ int32 BuildPatchToolMain( const TCHAR* CommandLine )
 		return 1;
 	}
 
-	if (bCompactify && bPreview && bTouchOnly)
+	if (bCompactify && bPreview && bNoPatchDelete)
 	{
-		GLog->Log(ELogVerbosity::Error, TEXT("Only one of -preview and -touchonly can be specified"));
+		GLog->Log(ELogVerbosity::Error, TEXT("Only one of -preview and -nopatchdelete can be specified"));
 		return 5;
 	}
 
@@ -435,9 +435,9 @@ int32 BuildPatchToolMain( const TCHAR* CommandLine )
 		{
 			CompactifyMode = ECompactifyMode::Preview;
 		}
-		else if (bTouchOnly)
+		else if (bNoPatchDelete)
 		{
-			CompactifyMode = ECompactifyMode::TouchOnly;
+			CompactifyMode = ECompactifyMode::NoPatchDelete;
 		}
 
 		// Run the compactify routine
diff --git a/Engine/Source/Runtime/Online/BuildPatchServices/Private/BuildPatchCompactifier.cpp b/Engine/Source/Runtime/Online/BuildPatchServices/Private/BuildPatchCompactifier.cpp
--- a/Engine/Source/Runtime/Online/BuildPatchServices/Private/BuildPatchCompactifier.cpp
+++ b/Engine/Source/Runtime/Online/BuildPatchServices/Private/BuildPatchCompactifier.cpp
@@ -18,10 +18,10 @@ namespace BuildDataCompactifierDefs
 /* Constructors
 *****************************************************************************/
 
-FBuildDataCompactifier::FBuildDataCompactifier(const FString& CloudDir, const bool bPreview, const bool bTouchOnly)
+FBuildDataCompactifier::FBuildDataCompactifier(const FString& CloudDir, const bool bPreview, const bool bNoPatchDelete)
 	: CloudDir(CloudDir)
 	, bPreview(bPreview)
-	, bTouchOnly(bTouchOnly)
+	, bNoPatchDelete(bNoPatchDelete)
 {
 
 }
@@ -29,15 +29,15 @@ FBuildDataCompactifier::FBuildDataCompactifier(const FString& CloudDir, const bo
 /* Public static methods
 *****************************************************************************/
 
-bool FBuildDataCompactifier::CompactifyCloudDirectory(const FString& CloudDir, const TArray<FString>& ManifestsToKeep, const float DataAgeThreshold, const bool bPreview, const bool bTouchOnly)
+bool FBuildDataCompactifier::CompactifyCloudDirectory(const FString& CloudDir, const TArray<FString>& ManifestsToKeep, const float DataAgeThreshold, const bool bPreview, const bool bNoPatchDelete)
 {
-	FBuildDataCompactifier Compactifier(CloudDir, bPreview, bTouchOnly);
+	FBuildDataCompactifier Compactifier(CloudDir, bPreview, bNoPatchDelete);
 	return Compactifier.Compactify(ManifestsToKeep, DataAgeThreshold);
 }
 
-bool FBuildDataCompactifier::CompactifyCloudDirectory(const TArray<FString>& ManifestsToKeep, const float DataAgeThreshold, const bool bPreview, const bool bTouchOnly)
+bool FBuildDataCompactifier::CompactifyCloudDirectory(const TArray<FString>& ManifestsToKeep, const float DataAgeThreshold, const bool bPreview, const bool bNoPatchDelete)
 {
-	return CompactifyCloudDirectory(FBuildPatchServicesModule::GetCloudDirectory(), ManifestsToKeep, DataAgeThreshold, bPreview, bTouchOnly);
+	return CompactifyCloudDirectory(FBuildPatchServicesModule::GetCloudDirectory(), ManifestsToKeep, DataAgeThreshold, bPreview, bNoPatchDelete);
 }
 
 /* Private methods
@@ -45,7 +45,7 @@ bool FBuildDataCompactifier::CompactifyCloudDirectory(const TArray<FString>& Man
 
 bool FBuildDataCompactifier::Compactify(const TArray<FString>& ManifestsToKeep, const float DataAgeThreshold) const
 {
-	GLog->Logf(TEXT("Running Compactify on %s%s"), *CloudDir, bPreview ? TEXT(". Preview mode. NO action will be taken.") : bTouchOnly ? TEXT(". TouchOnly mode. NO patch data will be deleted.") : TEXT(""));
+	GLog->Logf(TEXT("Running Compactify on %s%s"), *CloudDir, bPreview ? TEXT(". Preview mode. NO action will be taken.") : bNoPatchDelete ? TEXT(". NoPatchDelete mode. NO patch data will be deleted.") : TEXT(""));
 	if (ManifestsToKeep.Num() > 0)
 	{
 		GLog->Logf(TEXT("Preserving manifest files: %s"), *FString::Join(ManifestsToKeep, TEXT(", ")));
@@ -211,7 +211,14 @@ bool FBuildDataCompactifier::Compactify(const TArray<FString>& ManifestsToKeep,
 	{
 		GLog->Logf(TEXT("Of these, %u (totalling %s) were not chunk/manifest files."), NonPatchFilesProcessed, *HumanReadableSize(NonPatchBytesProcessed));
 	}
-	GLog->Logf(TEXT("Compactify deleted %u chunk/manifest files totalling %s."), FilesDeleted, *HumanReadableSize(BytesDeleted));
+	if (bNoPatchDelete)
+	{
+		GLog->Logf(TEXT("Compactify skipped deleting %u unreferenced chunk/manifest files (totalling %s) due to -nopatchdelete being specified."), FilesDeleted, *HumanReadableSize(BytesDeleted));
+	}
+	else
+	{
+		GLog->Logf(TEXT("Compactify deleted %u chunk/manifest files totalling %s."), FilesDeleted, *HumanReadableSize(BytesDeleted));
+	}
 	GLog->Logf(TEXT("Compactify touched %u chunk files totalling %s."), FilesTouched, *HumanReadableSize(BytesTouched));
 	if (TouchFailureCount > 0)
 	{
@@ -223,11 +230,16 @@ bool FBuildDataCompactifier::Compactify(const TArray<FString>& ManifestsToKeep,
 
 void FBuildDataCompactifier::DeleteFile(const FString& FilePath) const
 {
-	GLog->Logf(TEXT("Deleting file %s"), *FilePath);
-	if (!bPreview && !bTouchOnly)
+	FString LogMsg = FString::Printf(TEXT("Deprecated data %s"), *FilePath);
+	if (!bNoPatchDelete)
 	{
-		IFileManager::Get().Delete(*FilePath);
+		LogMsg = LogMsg.Append(TEXT(" ... deleted"));
+		if (!bPreview)
+		{
+			IFileManager::Get().Delete(*FilePath);
+		}
 	}
+	GLog->Logf(*LogMsg);
 }
 
 void FBuildDataCompactifier::EnumerateManifests(TArray<FString>& OutManifests) const
@@ -268,7 +280,7 @@ bool FBuildDataCompactifier::DeleteNonReferencedManifests(TArray<FString>& AllMa
 			DeletedManifests.Add(Manifest);
 			BytesDeleted += IFileManager::Get().FileSize(*ManifestPath);
 
-			if (!bPreview && !bTouchOnly)
+			if (!bPreview && !bNoPatchDelete)
 			{
 				IFileManager::Get().Delete(*ManifestPath);
 				if (FPaths::FileExists(ManifestPath))
diff --git a/Engine/Source/Runtime/Online/BuildPatchServices/Private/BuildPatchCompactifier.h b/Engine/Source/Runtime/Online/BuildPatchServices/Private/BuildPatchCompactifier.h
--- a/Engine/Source/Runtime/Online/BuildPatchServices/Private/BuildPatchCompactifier.h
+++ b/Engine/Source/Runtime/Online/BuildPatchServices/Private/BuildPatchCompactifier.h
@@ -22,11 +22,11 @@ public:
 	 * NOTE: THIS function is blocking and will not return until finished. Don't run on main thread.
 	 * @param ManifestsToKeep   If specified, these manifests will be retained, and all others will be deleted
 	 * @param DataAgeThreshold  Chunks which are not referenced by a valid manifest, and which are older than this age (in days), will be deleted
-	 * @param bPreview      If true, then no actual work will be done, but all operations which would be carried out will be logged.
-	 * @param bTouchOnly    If true, then no patch data will be deleted, but referenced data will have its timestamp updated to the current time.
-	 * @return              true if no file errors occurred
+	 * @param bPreview          If true, then no actual work will be done, but all operations which would be carried out will be logged.
+	 * @param bNoPatchDelete    If true, then no patch data will be deleted, but referenced data will have its timestamp updated to the current time.
+	 * @return                  true if no file errors occurred
 	 */
-	static bool CompactifyCloudDirectory(const TArray<FString>& ManifestsToKeep, const float DataAgeThreshold, const bool bPreview, const bool bTouchOnly);
+	static bool CompactifyCloudDirectory(const TArray<FString>& ManifestsToKeep, const float DataAgeThreshold, const bool bPreview, const bool bNoPatchDelete);
 
 	/**
 	 * Processes the specified Cloud Directory to identify and delete any orphaned chunks or files.
@@ -34,24 +34,24 @@ public:
 	 * @param CloudDir      The path to the Cloud Dir to compactify
 	 * @param ManifestsToKeep   If specified, these manifests will be retained, and all others will be deleted
 	 * @param DataAgeThreshold  Chunks which are not referenced by a valid manifest, and which are older than this age (in days), will be deleted
-	 * @param bPreview      If true, then no actual work will be done, but all operations which would be carried out will be logged.
-	 * @param bTouchOnly    If true, then no patch data will be deleted, but referenced data will have its timestamp updated to the current time.
-	 * @return              true if no file errors occurred
+	 * @param bPreview          If true, then no actual work will be done, but all operations which would be carried out will be logged.
+	 * @param bNoPatchDelete    If true, then no patch data will be deleted, but referenced data will have its timestamp updated to the current time.
+	 * @return                  true if no file errors occurred
 	 */
-	static bool CompactifyCloudDirectory(const FString& CloudDir, const TArray<FString>& ManifestsToKeep, const float DataAgeThreshold, const bool bPreview, const bool bTouchOnly);
+	static bool CompactifyCloudDirectory(const FString& CloudDir, const TArray<FString>& ManifestsToKeep, const float DataAgeThreshold, const bool bPreview, const bool bNoPatchDelete);
 
 private:
 	const FString CloudDir;
 	const bool bPreview;
-	const bool bTouchOnly;
+	const bool bNoPatchDelete;
 
 	/**
 	 * Constructor.
-	 * @param CloudDir      The directory that this compactifier should operate on.
-	 * @param bPreview      If true, then no actual work will be done, but all operations which would be carried out will be logged.
-	 * @param bTouchOnly    If true, then no patch data will be deleted, but referenced data will have its timestamp updated to the current time.
+	 * @param CloudDir         The directory that this compactifier should operate on.
+	 * @param bPreview         If true, then no actual work will be done, but all operations which would be carried out will be logged.
+	 * @param bNoPatchDelete   If true, then no patch data will be deleted, but referenced data will have its timestamp updated to the current time.
 	 */
-	FBuildDataCompactifier(const FString& CloudDir, const bool bPreview, const bool bTouchOnly);
+	FBuildDataCompactifier(const FString& CloudDir, const bool bPreview, const bool bNoPatchDelete);
 
 	/**
 	 * Processes the Cloud Directory to identify and delete any orphaned chunks or files.
diff --git a/Engine/Source/Runtime/Online/BuildPatchServices/Private/BuildPatchServicesModule.cpp b/Engine/Source/Runtime/Online/BuildPatchServices/Private/BuildPatchServicesModule.cpp
--- a/Engine/Source/Runtime/Online/BuildPatchServices/Private/BuildPatchServicesModule.cpp
+++ b/Engine/Source/Runtime/Online/BuildPatchServices/Private/BuildPatchServicesModule.cpp
@@ -191,8 +191,8 @@ bool FBuildPatchServicesModule::GenerateFilesManifestFromDirectory( const FBuild
 bool FBuildPatchServicesModule::CompactifyCloudDirectory(const TArray<FString>& ManifestsToKeep, const float DataAgeThreshold, const ECompactifyMode::Type Mode)
 {
 	const bool bPreview = Mode == ECompactifyMode::Preview;
-	const bool bTouchOnly = Mode == ECompactifyMode::TouchOnly;
-	return FBuildDataCompactifier::CompactifyCloudDirectory(ManifestsToKeep, DataAgeThreshold, bPreview, bTouchOnly);
+	const bool bNoPatchDelete = Mode == ECompactifyMode::NoPatchDelete;
+	return FBuildDataCompactifier::CompactifyCloudDirectory(ManifestsToKeep, DataAgeThreshold, bPreview, bNoPatchDelete);
 }
 
 #endif //WITH_BUILDPATCHGENERATION
diff --git a/Engine/Source/Runtime/Online/BuildPatchServices/Public/Interfaces/IBuildPatchServicesModule.h b/Engine/Source/Runtime/Online/BuildPatchServices/Public/Interfaces/IBuildPatchServicesModule.h
--- a/Engine/Source/Runtime/Online/BuildPatchServices/Public/Interfaces/IBuildPatchServicesModule.h
+++ b/Engine/Source/Runtime/Online/BuildPatchServices/Public/Interfaces/IBuildPatchServicesModule.h
@@ -21,7 +21,7 @@ namespace ECompactifyMode
 	enum Type
 	{
 		Preview,
-		TouchOnly,
+		NoPatchDelete,
 		Full
 	};
 }
@@ -124,7 +124,7 @@ public:
 	 * NOTE: THIS function is blocking and will not return until finished. Don't run on main thread.
 	 * @param ManifestsToKeep      If specified, these manifests will be retained, and all others will be deleted
 	 * @param DataAgeThreshold     Chunks which are not referenced by a valid manifest, and which are older than this age (in days), will be deleted
-	 * @param Mode      The mode that compactify will run in. If Preview, then no work will be carried out, if TouchOnly, then no patch-data will be deleted
+	 * @param Mode      The mode that compactify will run in. If Preview, then no work will be carried out, if NoPatchData, then no patch-data will be deleted
 	 * @return              true if no file errors occurred
 	 */
 	virtual bool CompactifyCloudDirectory(const TArray<FString>& ManifestsToKeep, const float DataAgeThreshold, const ECompactifyMode::Type Mode) = 0;
