WinINet: Moving GetLastError call so that it is not always returning 0 when http logging is switched on. Also adding it to the basic InternetReadFile log output for extra info.

diff --git a/Engine/Source/Runtime/Online/HTTP/Private/Windows/HttpWinInet.cpp b/Engine/Source/Runtime/Online/HTTP/Private/Windows/HttpWinInet.cpp
--- a/Engine/Source/Runtime/Online/HTTP/Private/Windows/HttpWinInet.cpp
+++ b/Engine/Source/Runtime/Online/HTTP/Private/Windows/HttpWinInet.cpp
@@ -920,13 +920,13 @@ void FHttpResponseWinInet::ProcessResponse()
 			&ResponsePayload[TotalBytesRead], 
 			NumBytesToRead, 
 			(LPDWORD)&AsyncBytesRead);
+		const int32 ErrorCode = GetLastError();
 
-		DEBUG_LOG_HTTP(bDebugLog, VeryVerbose, TEXT("InternetReadFile result=%d (%u bytes read) (%u bytes total read) (%u bytes to read). LoopCount=%d %p"), 
-			bReadFile, AsyncBytesRead, TotalBytesRead, NumBytesToRead, LoopCount, &Request);
+		DEBUG_LOG_HTTP(bDebugLog, VeryVerbose, TEXT("InternetReadFile result=%d error=%d (%u bytes read) (%u bytes total read) (%u bytes to read). LoopCount=%d %p"), 
+			bReadFile, ErrorCode, AsyncBytesRead, TotalBytesRead, NumBytesToRead, LoopCount, &Request);
 
 		if (!bReadFile)
 		{
-			int32 ErrorCode = GetLastError();
 			if (ErrorCode == ERROR_IO_PENDING)
 			{
 				// Chunked responses could cause InternetReadFile to return IO_PENDING, in which
