Aborting chunk downloads that appear to be suffering from crippling download speed

diff --git a/Engine/Source/Runtime/Online/BuildPatchServices/Private/BuildPatchAnalytics.cpp b/Engine/Source/Runtime/Online/BuildPatchServices/Private/BuildPatchAnalytics.cpp
--- a/Engine/Source/Runtime/Online/BuildPatchServices/Private/BuildPatchAnalytics.cpp
+++ b/Engine/Source/Runtime/Online/BuildPatchServices/Private/BuildPatchAnalytics.cpp
@@ -37,10 +37,21 @@ void FBuildPatchAnalytics::RecordChunkDownloadError( const FString& ChunkUrl, co
 		Attributes.Add( FAnalyticsEventAttribute( TEXT( "ChunkURL" ), ChunkUrl ) );
 		Attributes.Add( FAnalyticsEventAttribute( TEXT( "ResponseCode" ), ResponseCode ) );
 		Attributes.Add( FAnalyticsEventAttribute( TEXT( "ErrorString" ), ErrorString ) );
-		FBuildPatchHTTP::QueueAnalyticsEvent( TEXT( "Patcher.Error.Download" ) , Attributes );
+		FBuildPatchHTTP::QueueAnalyticsEvent( TEXT( "Patcher.Error.Download" ), Attributes );
 	}
 }
 
+void FBuildPatchAnalytics::RecordChunkDownloadAborted( const FString& ChunkUrl, const double& ChunkTime, const double& ChunkMean, const double& ChunkStd, const double& BreakingPoint )
+{
+	TArray<FAnalyticsEventAttribute> Attributes;
+	Attributes.Add( FAnalyticsEventAttribute( TEXT( "ChunkURL" ), ChunkUrl ) );
+	Attributes.Add( FAnalyticsEventAttribute( TEXT( "ChunkTime" ), ChunkTime ) );
+	Attributes.Add( FAnalyticsEventAttribute( TEXT( "ChunkMean" ), ChunkMean ) );
+	Attributes.Add( FAnalyticsEventAttribute( TEXT( "ChunkStd" ), ChunkStd ) );
+	Attributes.Add( FAnalyticsEventAttribute( TEXT( "BreakingPoint" ), BreakingPoint ) );
+	FBuildPatchHTTP::QueueAnalyticsEvent( TEXT( "Patcher.Warning.ChunkAborted" ), Attributes );
+}
+
 void FBuildPatchAnalytics::RecordChunkCacheError( const FGuid& ChunkGuid, const FString& Filename, const int32 LastError, const FString& SystemName, const FString& ErrorString )
 {
 	if ( CacheErrors.Increment() <= ERROR_EVENT_SEND_LIMIT )
@@ -51,7 +62,7 @@ void FBuildPatchAnalytics::RecordChunkCacheError( const FGuid& ChunkGuid, const
 		Attributes.Add( FAnalyticsEventAttribute( TEXT( "LastError" ), LastError ) );
 		Attributes.Add( FAnalyticsEventAttribute( TEXT( "SystemName" ), SystemName ) );
 		Attributes.Add( FAnalyticsEventAttribute( TEXT( "ErrorString" ), ErrorString ) );
-		FBuildPatchHTTP::QueueAnalyticsEvent( TEXT( "Patcher.Error.Cache" ) , Attributes );
+		FBuildPatchHTTP::QueueAnalyticsEvent( TEXT( "Patcher.Error.Cache" ), Attributes );
 	}
 }
 
@@ -63,7 +74,7 @@ void FBuildPatchAnalytics::RecordConstructionError( const FString& Filename, con
 		Attributes.Add( FAnalyticsEventAttribute( TEXT( "Filename" ), Filename ) );
 		Attributes.Add( FAnalyticsEventAttribute( TEXT( "LastError" ), LastError ) );
 		Attributes.Add( FAnalyticsEventAttribute( TEXT( "ErrorString" ), ErrorString ) );
-		FBuildPatchHTTP::QueueAnalyticsEvent( TEXT( "Patcher.Error.Construction" ) , Attributes );
+		FBuildPatchHTTP::QueueAnalyticsEvent( TEXT( "Patcher.Error.Construction" ), Attributes );
 	}
 }
 
@@ -74,7 +85,7 @@ void FBuildPatchAnalytics::RecordPrereqInstallnError( const FString& Filename, c
 	Attributes.Add( FAnalyticsEventAttribute( TEXT( "CommandLine" ), CommandLine ) );
 	Attributes.Add( FAnalyticsEventAttribute( TEXT( "ErrorCode" ), ErrorCode ) );
 	Attributes.Add( FAnalyticsEventAttribute( TEXT( "ErrorString" ), ErrorString ) );
-	FBuildPatchHTTP::QueueAnalyticsEvent( TEXT( "Patcher.Error.Prerequisites" ) , Attributes );
+	FBuildPatchHTTP::QueueAnalyticsEvent( TEXT( "Patcher.Error.Prerequisites" ), Attributes );
 
 }
 
diff --git a/Engine/Source/Runtime/Online/BuildPatchServices/Private/BuildPatchAnalytics.h b/Engine/Source/Runtime/Online/BuildPatchServices/Private/BuildPatchAnalytics.h
--- a/Engine/Source/Runtime/Online/BuildPatchServices/Private/BuildPatchAnalytics.h
+++ b/Engine/Source/Runtime/Online/BuildPatchServices/Private/BuildPatchAnalytics.h
@@ -53,6 +53,16 @@ public:
 	 */
 	static void RecordChunkDownloadError( const FString& ChunkUrl, const int32& ResponseCode, const FString& ErrorString );
 
+	/**
+	 * Record a chunk download aborted event
+	 * @param ChunkUrl		The url for the chunk being downloaded
+	 * @param ChunkTime		The current running time for this chunk
+	 * @param ChunkMean		The recorded average chunk download time
+	 * @param ChunkStd		The recorded standard deviation for chunk download time
+	 * @param BreakingPoint	The value used as the breaking point for chunk download time
+	 */
+	static void RecordChunkDownloadAborted( const FString& ChunkUrl, const double& ChunkTime, const double& ChunkMean, const double& ChunkStd, const double& BreakingPoint );
+
 	/**
 	 * Record a chunk cache error event
 	 * @param ChunkGuid		The chunk guid that the error occurred for
diff --git a/Engine/Source/Runtime/Online/BuildPatchServices/Private/BuildPatchDownloader.cpp b/Engine/Source/Runtime/Online/BuildPatchServices/Private/BuildPatchDownloader.cpp
--- a/Engine/Source/Runtime/Online/BuildPatchServices/Private/BuildPatchDownloader.cpp
+++ b/Engine/Source/Runtime/Online/BuildPatchServices/Private/BuildPatchDownloader.cpp
@@ -15,6 +15,59 @@ int32 CHUNK_DOWNLOAD_RETRIES = 20;
 
 float CHUNK_RETRY_TIME = 1.0;
 
+/* A class used to monitor the average chunk download time and standard deviation
+*****************************************************************************/
+class FMeanChunkTime
+{
+public:
+	FMeanChunkTime()
+		: Count(0)
+		, Total(0)
+		, TotalSqs(0)
+	{}
+
+	void Reset()
+	{
+		Count = 0;
+		Total = 0;
+		TotalSqs = 0;
+	}
+
+	bool IsReliable() const
+	{
+		return Count > 10;
+	}
+
+	double GetMean() const
+	{
+		return Total / Count;
+	}
+
+	double GetStd() const
+	{
+		const double Mean = GetMean();
+		return FMath::Sqrt((TotalSqs / Count) - (Mean*Mean));
+	}
+
+	void GetValues(double& Mean, double& Std) const
+	{
+		Mean = GetMean();
+		Std = FMath::Sqrt((TotalSqs / Count) - (Mean*Mean));
+	}
+
+	void AddSample(double Sample)
+	{
+		Total += Sample;
+		TotalSqs += Sample * Sample;
+		++Count;
+	}
+
+private:
+	int32 Count;
+	double Total;
+	double TotalSqs;
+};
+
 /* FBuildPatchDownloader implementation
 *****************************************************************************/
 FBuildPatchDownloader::FBuildPatchDownloader( const FString& InSaveDirectory, const FBuildPatchAppManifestRef& InInstallManifest, FBuildPatchProgress* InBuildProgress )
@@ -67,6 +120,9 @@ uint32 FBuildPatchDownloader::Run()
 	const bool bIsChunkData = !InstallManifest->IsFileDataManifest();
 	const FBuildPatchData::Type DataType = bIsChunkData ? FBuildPatchData::ChunkData : FBuildPatchData::FileData;
 
+	// The Average chunk download time
+	FMeanChunkTime MeanChunkTime;
+
 	// While there is the possibility of more chunks to download
 	while ( ShouldBeRunning() && !FBuildPatchInstallError::HasFatalError() )
 	{
@@ -87,6 +143,8 @@ uint32 FBuildPatchDownloader::Run()
 				bool bSuccess = InFlightJob.StateFlag.GetValue() == EDownloadState::DownloadSuccess;
 				if( bSuccess )
 				{
+					const double ChunkTime = InFlightJob.DownloadRecord.EndTime - InFlightJob.DownloadRecord.StartTime;
+					MeanChunkTime.AddSample(ChunkTime);
 					// Uncompress if needed
 					if (bIsChunkData)
 					{
@@ -172,6 +230,26 @@ uint32 FBuildPatchDownloader::Run()
 					InFlightDownloads.Remove( InFlightKey );
 				}
 			}
+			else if (MeanChunkTime.IsReliable() && InFlightJob.RetryCount.GetValue() == 0)
+			{
+				// Cancel any chunk taking 4x longer if still on first try
+				// This attempts to catch out issues with trickling requests
+				const double ChunkTime = (FPlatformTime::Seconds() - InFlightJob.DownloadRecord.StartTime);
+				double ChunkMean, ChunkStd;
+				MeanChunkTime.GetValues(ChunkMean, ChunkStd);
+				// The point at which we decide the chunk is delayed, with a sane minimum
+				const double BreakingPoint = FMath::Max<double>(20.0, ChunkMean + (ChunkStd * 4.0));
+				if (ChunkTime > BreakingPoint)
+				{
+					GWarn->Logf(TEXT("BuildPatchServices: WARNING: Cancelling download %s. T:%.2f Av:%.2f Std:%.2f Bp:%.2f"), *InFlightJob.Guid.ToString(), ChunkTime, ChunkMean, ChunkStd, BreakingPoint);
+					MeanChunkTime.Reset();
+					InFlightDownloadsLock.Unlock();
+					FBuildPatchAnalytics::RecordChunkDownloadAborted(InFlightJob.DownloadUrl, ChunkTime, ChunkMean, ChunkStd, BreakingPoint);
+					FBuildPatchHTTP::CancelHttpRequest(InFlightJob.HttpRequestId);
+					InFlightDownloadsLock.Lock();
+				}
+			}
+
 		}
 		InFlightDownloadsLock.Unlock();
 
@@ -199,14 +277,19 @@ uint32 FBuildPatchDownloader::Run()
 		const bool bIsHTTPRequest = DownloadUrl.Contains( TEXT( "http" ), ESearchCase::IgnoreCase );
 
 		// Start the download
-		int32 HttpRequestID = INDEX_NONE;
+		int32 HttpRequestId = INDEX_NONE;
 		if( bIsHTTPRequest )
 		{
 			// Kick off the download
 #if EXTRA_DOWNLOAD_LOGGING
 			GWarn->Logf(TEXT("BuildPatchDownloader: Queuing chunk HTTP download %s"), *NextGuid.ToString());
 #endif
-			HttpRequestID = FBuildPatchHTTP::QueueHttpRequest( DownloadUrl, FHttpRequestCompleteDelegate::CreateThreadSafeSP( this, &FBuildPatchDownloader::HttpRequestComplete ), FHttpRequestProgressDelegate::CreateThreadSafeSP( this, &FBuildPatchDownloader::HttpRequestProgress ) );
+			HttpRequestId = FBuildPatchHTTP::QueueHttpRequest(DownloadUrl, FHttpRequestCompleteDelegate::CreateThreadSafeSP(this, &FBuildPatchDownloader::HttpRequestComplete), FHttpRequestProgressDelegate::CreateThreadSafeSP(this, &FBuildPatchDownloader::HttpRequestProgress));
+
+			InFlightDownloadsLock.Lock();
+			InFlightDownloads[NextGuid].HttpRequestId = HttpRequestId;
+			InFlightDownloads[NextGuid].DownloadUrl = DownloadUrl;
+			InFlightDownloadsLock.Unlock();
 		}
 		else
 		{
@@ -383,6 +466,9 @@ bool FBuildPatchDownloader::ShouldBeRunning()
 
 void FBuildPatchDownloader::HttpRequestProgress( FHttpRequestPtr Request, int32 BytesSoFar )
 {
+#if EXTRA_DOWNLOAD_LOGGING
+	GWarn->Logf(TEXT("BuildPatchDownloader: Request %p Bytes: %d"), Request.Get(), BytesSoFar);
+#endif
 	if ( !FBuildPatchInstallError::HasFatalError() )
 	{
 		FGuid DataGUID;
@@ -418,17 +504,13 @@ void FBuildPatchDownloader::OnDownloadProgress( const FGuid& Guid, const int32&
 		CurrentJob.DownloadRecord.DownloadSize = BytesSoFar;
 		IncrementByteDownloadCount( DeltaBytes );
 #if EXTRA_DOWNLOAD_LOGGING
-		// Log slow downloads
+		// Log download progress
 		const double tNow = FPlatformTime::Seconds();
 		const double tDelta = tNow - CurrentJob.DownloadRecord.EndTime;
 		const double tRunning = tNow - CurrentJob.DownloadRecord.StartTime;
 		const double currSpeed = DeltaBytes / tDelta;
 		const double avSpeed = BytesSoFar / tRunning;
-		if (tRunning > 60 && (tNow - CurrentJob.TimeLastLogged > 5))
-		{
-			CurrentJob.TimeLastLogged = tNow;
-			GWarn->Logf(TEXT("BuildPatchDownloader: WARNING: Slow download %s Running:%.2fsec\tDownloaded:%d\tSpeed:%.0f\tAverageSpeed:%.0f (B/s)"), *Guid.ToString(), tRunning, BytesSoFar, currSpeed, avSpeed);
-		}
+		GWarn->Logf(TEXT("BuildPatchDownloader: Download %s Running:%.2fsec\tDownloaded:%d\tSpeed:%.0f\tAverageSpeed:%.0f (B/s)"), *Guid.ToString(), tRunning, BytesSoFar, currSpeed, avSpeed);
 		CurrentJob.DownloadRecord.EndTime = tNow;
 #endif
 	}
diff --git a/Engine/Source/Runtime/Online/BuildPatchServices/Private/BuildPatchDownloader.h b/Engine/Source/Runtime/Online/BuildPatchServices/Private/BuildPatchDownloader.h
--- a/Engine/Source/Runtime/Online/BuildPatchServices/Private/BuildPatchDownloader.h
+++ b/Engine/Source/Runtime/Online/BuildPatchServices/Private/BuildPatchDownloader.h
@@ -77,8 +77,8 @@ private:
 		FBuildPatchDownloadRecord DownloadRecord;
 		// The response code, for HTTP downloads
 		int32 ResponseCode;
-		// A time we last logged info for this download, to help not overly spam
-		double TimeLastLogged;
+		// The request ID, for HTTP downloads
+		int32 HttpRequestId;
 
 		/**
 		 * Default constructor
@@ -88,7 +88,7 @@ private:
 			, RetryCount( 0 )
 			, StateFlag( EDownloadState::DownloadInit )
 			, ResponseCode( 0 )
-			, TimeLastLogged( 0 )
+			, HttpRequestId( INDEX_NONE )
 		{}
 	};
 
