Clean up

diff --git a/Engine/Source/Runtime/Online/BuildPatchServices/Private/BuildPatchDownloader.cpp b/Engine/Source/Runtime/Online/BuildPatchServices/Private/BuildPatchDownloader.cpp
--- a/Engine/Source/Runtime/Online/BuildPatchServices/Private/BuildPatchDownloader.cpp
+++ b/Engine/Source/Runtime/Online/BuildPatchServices/Private/BuildPatchDownloader.cpp
@@ -46,13 +46,13 @@ public:
 	double GetStd() const
 	{
 		const double Mean = GetMean();
-		return FMath::Sqrt((TotalSqs / Count) - (Mean*Mean));
+		return FMath::Sqrt((TotalSqs / Count) - (Mean * Mean));
 	}
 
 	void GetValues(double& Mean, double& Std) const
 	{
 		Mean = GetMean();
-		Std = FMath::Sqrt((TotalSqs / Count) - (Mean*Mean));
+		Std = FMath::Sqrt((TotalSqs / Count) - (Mean * Mean));
 	}
 
 	void AddSample(double Sample)
@@ -232,8 +232,8 @@ uint32 FBuildPatchDownloader::Run()
 			}
 			else if (MeanChunkTime.IsReliable() && InFlightJob.RetryCount.GetValue() == 0)
 			{
-				// Cancel any chunk taking 4x longer if still on first try
-				// This attempts to catch out issues with trickling requests
+				// If still on first try, cancel any chunk taking longer than the mean time plus 4x standard deviation. In statistical terms, that's 1 in 15,787 chance of being
+				// a download time that appears in the normal distribution. So we are guessing that this chunk would be an abnormally delayed one.
 				const double ChunkTime = (FPlatformTime::Seconds() - InFlightJob.DownloadRecord.StartTime);
 				double ChunkMean, ChunkStd;
 				MeanChunkTime.GetValues(ChunkMean, ChunkStd);
@@ -241,7 +241,7 @@ uint32 FBuildPatchDownloader::Run()
 				const double BreakingPoint = FMath::Max<double>(20.0, ChunkMean + (ChunkStd * 4.0));
 				if (ChunkTime > BreakingPoint)
 				{
-					GWarn->Logf(TEXT("BuildPatchServices: WARNING: Cancelling download %s. T:%.2f Av:%.2f Std:%.2f Bp:%.2f"), *InFlightJob.Guid.ToString(), ChunkTime, ChunkMean, ChunkStd, BreakingPoint);
+					GWarn->Logf(TEXT("BuildPatchServices: WARNING: Canceling download %s. T:%.2f Av:%.2f Std:%.2f Bp:%.2f"), *InFlightJob.Guid.ToString(), ChunkTime, ChunkMean, ChunkStd, BreakingPoint);
 					MeanChunkTime.Reset();
 					InFlightDownloadsLock.Unlock();
 					FBuildPatchAnalytics::RecordChunkDownloadAborted(InFlightJob.DownloadUrl, ChunkTime, ChunkMean, ChunkStd, BreakingPoint);
