Further improve logging of BuildPatchTool.exe deletions when in preview or nonpatchdelete modes.

diff --git a/Engine/Source/Runtime/Online/BuildPatchServices/Private/BuildPatchCompactifier.cpp b/Engine/Source/Runtime/Online/BuildPatchServices/Private/BuildPatchCompactifier.cpp
--- a/Engine/Source/Runtime/Online/BuildPatchServices/Private/BuildPatchCompactifier.cpp
+++ b/Engine/Source/Runtime/Online/BuildPatchServices/Private/BuildPatchCompactifier.cpp
@@ -231,13 +231,10 @@ bool FBuildDataCompactifier::Compactify(const TArray<FString>& ManifestsToKeep,
 void FBuildDataCompactifier::DeleteFile(const FString& FilePath) const
 {
 	FString LogMsg = FString::Printf(TEXT("Deprecated data %s"), *FilePath);
-	if (!bNoPatchDelete)
+	if (!bNoPatchDelete && !bPreview)
 	{
 		LogMsg = LogMsg.Append(TEXT(" ... deleted"));
-		if (!bPreview)
-		{
-			IFileManager::Get().Delete(*FilePath);
-		}
+		IFileManager::Get().Delete(*FilePath);
 	}
 	GLog->Logf(*LogMsg);
 }
@@ -274,7 +271,7 @@ bool FBuildDataCompactifier::DeleteNonReferencedManifests(TArray<FString>& AllMa
 	{
 		if (!ManifestsToKeep.Contains(Manifest))
 		{
-			GLog->Logf(TEXT("Deleting manifest file %s"), *Manifest);
+			FString LogMsg = FString::Printf(TEXT("Found deleteable manifest file %s"), *Manifest);
 			FString ManifestPath = CloudDir / Manifest;
 
 			DeletedManifests.Add(Manifest);
@@ -289,7 +286,9 @@ bool FBuildDataCompactifier::DeleteNonReferencedManifests(TArray<FString>& AllMa
 					GLog->Logf(ELogVerbosity::Error, TEXT("Compactify could not delete manifest file %s"), *Manifest);
 					return false;
 				}
+				LogMsg.Append(TEXT(" ... deleted"));
 			}
+			GLog->Logf(*LogMsg);
 		}
 	}
 
