More online delegate cleanup int he friends interface (changing multicast delegates to delegates that are passed in to friends interface functions and called when requests are completed)

diff --git a/Engine/Source/Developer/FriendsAndChat/Private/FriendsAndChatManager.cpp b/Engine/Source/Developer/FriendsAndChat/Private/FriendsAndChatManager.cpp
--- a/Engine/Source/Developer/FriendsAndChat/Private/FriendsAndChatManager.cpp
+++ b/Engine/Source/Developer/FriendsAndChat/Private/FriendsAndChatManager.cpp
@@ -64,8 +64,6 @@ void FFriendsAndChatManager::Login()
 			// Create delegates for list refreshes
 			OnQueryRecentPlayersCompleteDelegate = FOnQueryRecentPlayersCompleteDelegate::CreateRaw(this, &FFriendsAndChatManager::OnQueryRecentPlayersComplete);
 			OnFriendsListChangedDelegate = FOnFriendsChangeDelegate::CreateSP(this, &FFriendsAndChatManager::OnFriendsListChanged);
-			OnAcceptInviteCompleteDelegate = FOnAcceptInviteCompleteDelegate::CreateSP( this, &FFriendsAndChatManager::OnAcceptInviteComplete );
-			OnSendInviteCompleteDelegate = FOnSendInviteCompleteDelegate::CreateSP( this, &FFriendsAndChatManager::OnSendInviteComplete );
 			OnDeleteFriendCompleteDelegate = FOnDeleteFriendCompleteDelegate::CreateSP( this, &FFriendsAndChatManager::OnDeleteFriendComplete );
 			OnQueryUserIdMappingCompleteDelegate = FOnQueryUserIdMappingCompleteDelegate::CreateSP( this, &FFriendsAndChatManager::OnQueryUserIdMappingComplete );
 			OnQueryUserInfoCompleteDelegate = FOnQueryUserInfoCompleteDelegate::CreateSP( this, &FFriendsAndChatManager::OnQueryUserInfoComplete );
@@ -84,9 +82,7 @@ void FFriendsAndChatManager::Login()
 			FriendsInterface->AddOnFriendRemovedDelegate(OnFriendRemovedDelegate);
 			FriendsInterface->AddOnInviteRejectedDelegate(OnFriendInviteRejected);
 			FriendsInterface->AddOnInviteAcceptedDelegate(OnFriendInviteAccepted);
-			FriendsInterface->AddOnAcceptInviteCompleteDelegate( 0, OnAcceptInviteCompleteDelegate );
 			FriendsInterface->AddOnDeleteFriendCompleteDelegate( 0, OnDeleteFriendCompleteDelegate );
-			FriendsInterface->AddOnSendInviteCompleteDelegate( 0, OnSendInviteCompleteDelegate );
 			UserInterface->AddOnQueryUserInfoCompleteDelegate(0, OnQueryUserInfoCompleteDelegate);
 			OnlineSubMcp->GetPresenceInterface()->AddOnPresenceReceivedDelegate(OnPresenceReceivedCompleteDelegate);
 			OnlineSubMcp->GetSessionInterface()->AddOnSessionInviteReceivedDelegate(OnGameInviteReceivedDelegate);
@@ -135,9 +131,7 @@ void FFriendsAndChatManager::Logout()
 			OnlineSubMcp->GetFriendsInterface()->ClearOnFriendRemovedDelegate(OnFriendRemovedDelegate);
 			OnlineSubMcp->GetFriendsInterface()->ClearOnInviteRejectedDelegate(OnFriendInviteRejected);
 			OnlineSubMcp->GetFriendsInterface()->ClearOnInviteAcceptedDelegate(OnFriendInviteAccepted);
-			OnlineSubMcp->GetFriendsInterface()->ClearOnAcceptInviteCompleteDelegate(0, OnAcceptInviteCompleteDelegate);
 			OnlineSubMcp->GetFriendsInterface()->ClearOnDeleteFriendCompleteDelegate(0, OnDeleteFriendCompleteDelegate);
-			OnlineSubMcp->GetFriendsInterface()->ClearOnSendInviteCompleteDelegate(0, OnSendInviteCompleteDelegate);
 		}
 		if (OnlineSubMcp->GetPresenceInterface().IsValid())
 		{
@@ -785,7 +779,8 @@ void FFriendsAndChatManager::SetState( EFriendsAndManagerState::Type NewState )
 		break;
 	case EFriendsAndManagerState::AcceptingFriendRequest:
 		{
-			FriendsInterface->AcceptInvite( 0, PendingOutgoingAcceptFriendRequests[0], EFriendsLists::ToString( EFriendsLists::Default ) );
+			FOnAcceptInviteComplete Delegate = FOnAcceptInviteComplete::CreateSP(this, &FFriendsAndChatManager::OnAcceptInviteComplete);
+			FriendsInterface->AcceptInvite( 0, PendingOutgoingAcceptFriendRequests[0], EFriendsLists::ToString( EFriendsLists::Default ), Delegate );
 		}
 		break;
 	case EFriendsAndManagerState::RequestGameInviteRefresh:
@@ -1201,7 +1196,8 @@ void FFriendsAndChatManager::OnQueryUserIdMappingComplete(bool bWasSuccessful, c
 		{
 			for ( int32 Index = 0; Index < PendingOutgoingFriendRequests.Num(); Index++ )
 			{
-				FriendsInterface->SendInvite( 0, PendingOutgoingFriendRequests[Index].Get(), EFriendsLists::ToString( EFriendsLists::Default ) );
+				FOnSendInviteComplete Delegate = FOnSendInviteComplete::CreateSP(this, &FFriendsAndChatManager::OnSendInviteComplete);
+				FriendsInterface->SendInvite(0, PendingOutgoingFriendRequests[Index].Get(), EFriendsLists::ToString( EFriendsLists::Default ), Delegate);
 				AddFriendsToast(FText::FromString("Request Sent"));
 			}
 		}
diff --git a/Engine/Source/Developer/FriendsAndChat/Private/FriendsAndChatManager.h b/Engine/Source/Developer/FriendsAndChat/Private/FriendsAndChatManager.h
--- a/Engine/Source/Developer/FriendsAndChat/Private/FriendsAndChatManager.h
+++ b/Engine/Source/Developer/FriendsAndChat/Private/FriendsAndChatManager.h
@@ -667,10 +667,6 @@ private:
 	FTickerDelegate UpdateFriendsTickerDelegate;
 	// Delegate to use for querying for recent players 
 	FOnQueryRecentPlayersCompleteDelegate OnQueryRecentPlayersCompleteDelegate;
-	// Delegate when invites accepted
-	FOnAcceptInviteCompleteDelegate OnAcceptInviteCompleteDelegate;
-	// Delegate when invites have been sent
-	FOnSendInviteCompleteDelegate OnSendInviteCompleteDelegate;
 	// Delegate to use for deleting a friend
 	FOnDeleteFriendCompleteDelegate OnDeleteFriendCompleteDelegate;
 	// Delegate for querying user id from a name string
diff --git a/Engine/Source/Runtime/Online/IOS/OnlineSubsystemIOS/Private/OnlineFriendsInterfaceIOS.cpp b/Engine/Source/Runtime/Online/IOS/OnlineSubsystemIOS/Private/OnlineFriendsInterfaceIOS.cpp
--- a/Engine/Source/Runtime/Online/IOS/OnlineSubsystemIOS/Private/OnlineFriendsInterfaceIOS.cpp
+++ b/Engine/Source/Runtime/Online/IOS/OnlineSubsystemIOS/Private/OnlineFriendsInterfaceIOS.cpp
@@ -131,21 +131,21 @@ bool FOnlineFriendsIOS::ReadFriendsList(int32 LocalUserNum, const FString& ListN
 	return bSuccessfullyBeganReadFriends;
 }
 
-bool FOnlineFriendsIOS::DeleteFriendsList(int32 LocalUserNum, const FString& ListName)
+bool FOnlineFriendsIOS::DeleteFriendsList(int32 LocalUserNum, const FString& ListName, const FOnDeleteFriendsListComplete& Delegate /*= FOnDeleteFriendsListComplete()*/)
 {
-	TriggerOnDeleteFriendsListCompleteDelegates(LocalUserNum, false, ListName, FString(TEXT("DeleteFriendsList() is not supported")));
+	Delegate.ExecuteIfBound(LocalUserNum, false, ListName, FString(TEXT("DeleteFriendsList() is not supported")));
 	return false;
 }
 
-bool FOnlineFriendsIOS::SendInvite(int32 LocalUserNum, const FUniqueNetId& FriendId, const FString& ListName)
+bool FOnlineFriendsIOS::SendInvite(int32 LocalUserNum, const FUniqueNetId& FriendId, const FString& ListName, const FOnSendInviteComplete& Delegate /*= FOnSendInviteComplete()*/)
 {
-	TriggerOnSendInviteCompleteDelegates(LocalUserNum, false, FriendId, ListName, FString(TEXT("SendInvite() is not supported")));
+	Delegate.ExecuteIfBound(LocalUserNum, false, FriendId, ListName, FString(TEXT("SendInvite() is not supported")));
 	return false;
 }
 
-bool FOnlineFriendsIOS::AcceptInvite(int32 LocalUserNum, const FUniqueNetId& FriendId, const FString& ListName)
+bool FOnlineFriendsIOS::AcceptInvite(int32 LocalUserNum, const FUniqueNetId& FriendId, const FString& ListName, const FOnAcceptInviteComplete& Delegate /*= FOnAcceptInviteComplete()*/)
 {
-	TriggerOnAcceptInviteCompleteDelegates(LocalUserNum, false, FriendId, ListName, FString(TEXT("AcceptInvite() is not supported")));
+	Delegate.ExecuteIfBound(LocalUserNum, false, FriendId, ListName, FString(TEXT("AcceptInvite() is not supported")));
 	return false;
 }
 
diff --git a/Engine/Source/Runtime/Online/IOS/OnlineSubsystemIOS/Private/OnlineFriendsInterfaceIOS.h b/Engine/Source/Runtime/Online/IOS/OnlineSubsystemIOS/Private/OnlineFriendsInterfaceIOS.h
--- a/Engine/Source/Runtime/Online/IOS/OnlineSubsystemIOS/Private/OnlineFriendsInterfaceIOS.h
+++ b/Engine/Source/Runtime/Online/IOS/OnlineSubsystemIOS/Private/OnlineFriendsInterfaceIOS.h
@@ -89,9 +89,9 @@ public:
 	// IOnlineFriends
 
 	virtual bool ReadFriendsList(int32 LocalUserNum, const FString& ListName, const FOnReadFriendsListComplete& Delegate = FOnReadFriendsListComplete()) override;
-	virtual bool DeleteFriendsList(int32 LocalUserNum, const FString& ListName) override;
-	virtual bool SendInvite(int32 LocalUserNum, const FUniqueNetId& FriendId, const FString& ListName) override;
-	virtual bool AcceptInvite(int32 LocalUserNum, const FUniqueNetId& FriendId, const FString& ListName) override;
+	virtual bool DeleteFriendsList(int32 LocalUserNum, const FString& ListName, const FOnDeleteFriendsListComplete& Delegate = FOnDeleteFriendsListComplete()) override;
+	virtual bool SendInvite(int32 LocalUserNum, const FUniqueNetId& FriendId, const FString& ListName,  const FOnSendInviteComplete& Delegate = FOnSendInviteComplete()) override;
+	virtual bool AcceptInvite(int32 LocalUserNum, const FUniqueNetId& FriendId, const FString& ListName, const FOnAcceptInviteComplete& Delegate = FOnAcceptInviteComplete()) override;
  	virtual bool RejectInvite(int32 LocalUserNum, const FUniqueNetId& FriendId, const FString& ListName) override;
  	virtual bool DeleteFriend(int32 LocalUserNum, const FUniqueNetId& FriendId, const FString& ListName) override;
 	virtual bool GetFriendsList(int32 LocalUserNum, const FString& ListName, TArray< TSharedRef<FOnlineFriend> >& OutFriends) override;
diff --git a/Engine/Source/Runtime/Online/OnlineSubsystem/Public/Interfaces/OnlineFriendsInterface.h b/Engine/Source/Runtime/Online/OnlineSubsystem/Public/Interfaces/OnlineFriendsInterface.h
--- a/Engine/Source/Runtime/Online/OnlineSubsystem/Public/Interfaces/OnlineFriendsInterface.h
+++ b/Engine/Source/Runtime/Online/OnlineSubsystem/Public/Interfaces/OnlineFriendsInterface.h
@@ -57,8 +57,7 @@ DECLARE_DELEGATE_FourParams(FOnReadFriendsListComplete, int32, bool, const FStri
  * @param ListName name of the friends list that was operated on
  * @param ErrorStr string representing the error condition
  */
-DECLARE_MULTICAST_DELEGATE_FourParams(FOnDeleteFriendsListComplete, int32, bool, const FString&, const FString&);
-typedef FOnDeleteFriendsListComplete::FDelegate FOnDeleteFriendsListCompleteDelegate;
+DECLARE_DELEGATE_FourParams(FOnDeleteFriendsListComplete, int32, bool, const FString&, const FString&);
 
 /**
  * Delegate used when an invite send request has completed
@@ -69,8 +68,7 @@ typedef FOnDeleteFriendsListComplete::FDelegate FOnDeleteFriendsListCompleteDele
  * @param ListName name of the friends list that was operated on
  * @param ErrorStr string representing the error condition
  */
-DECLARE_MULTICAST_DELEGATE_FiveParams(FOnSendInviteComplete, int32, bool, const FUniqueNetId&, const FString&, const FString&);
-typedef FOnSendInviteComplete::FDelegate FOnSendInviteCompleteDelegate;
+DECLARE_DELEGATE_FiveParams(FOnSendInviteComplete, int32, bool, const FUniqueNetId&, const FString&, const FString&);
 
 /**
  * Delegate used when an invite accept request has completed
@@ -81,8 +79,7 @@ typedef FOnSendInviteComplete::FDelegate FOnSendInviteCompleteDelegate;
  * @param ListName name of the friends list that was operated on
  * @param ErrorStr string representing the error condition
  */
-DECLARE_MULTICAST_DELEGATE_FiveParams(FOnAcceptInviteComplete, int32, bool, const FUniqueNetId&, const FString&, const FString&);
-typedef FOnAcceptInviteComplete::FDelegate FOnAcceptInviteCompleteDelegate;
+DECLARE_DELEGATE_FiveParams(FOnAcceptInviteComplete, int32, bool, const FUniqueNetId&, const FString&, const FString&);
 
 /**
  * Delegate used when an invite reject request has completed
@@ -221,17 +218,7 @@ public:
 	 *
 	 * @return true if the delete request was started successfully, false otherwise
 	 */
-	virtual bool DeleteFriendsList(int32 LocalUserNum, const FString& ListName) = 0;
-
-	/**
-     * Delegate used when the friends list delete request has completed
-     *
-	 * @param LocalUserNum the controller number of the associated user that made the request
-     * @param bWasSuccessful true if the async action completed without error, false if there was an error
-	 * @param ListName name of the friends list that was operated on
-	 * @param ErrorStr string representing the error condition
-     */
-	DEFINE_ONLINE_PLAYER_DELEGATE_THREE_PARAM(MAX_LOCAL_PLAYERS, OnDeleteFriendsListComplete, bool, const FString&, const FString&);
+	virtual bool DeleteFriendsList(int32 LocalUserNum, const FString& ListName, const FOnDeleteFriendsListComplete& Delegate = FOnDeleteFriendsListComplete()) = 0;
 
 	/**
 	 * Starts an async task that sends an invite to another player. 
@@ -242,18 +229,7 @@ public:
 	 *
 	 * @return true if the request was started successfully, false otherwise
 	 */
-	virtual bool SendInvite(int32 LocalUserNum, const FUniqueNetId& FriendId, const FString& ListName) = 0;
-
-	/**
-	 * Delegate used when an invite send request has completed
-	 *
-	 * @param LocalUserNum the controller number of the associated user that made the request
-	 * @param bWasSuccessful true if the async action completed without error, false if there was an error
-	 * @param FriendId player that was invited
-	 * @param ListName name of the friends list that was operated on
-	 * @param ErrorStr string representing the error condition
-	 */
-	DEFINE_ONLINE_PLAYER_DELEGATE_FOUR_PARAM(MAX_LOCAL_PLAYERS, OnSendInviteComplete, bool, const FUniqueNetId&, const FString&, const FString&);
+	virtual bool SendInvite(int32 LocalUserNum, const FUniqueNetId& FriendId, const FString& ListName, const FOnSendInviteComplete& Delegate = FOnSendInviteComplete()) = 0;
 
 	/**
 	 * Starts an async task that accepts an invite from another player. 
@@ -264,18 +240,7 @@ public:
 	 *
 	 * @return true if the request was started successfully, false otherwise
 	 */
-	virtual bool AcceptInvite(int32 LocalUserNum, const FUniqueNetId& FriendId, const FString& ListName) = 0;
-
-	/**
-	 * Delegate used when an invite accept request has completed
-	 *
-	 * @param LocalUserNum the controller number of the associated user that made the request
-	 * @param bWasSuccessful true if the async action completed without error, false if there was an error
-	 * @param FriendId player that invited us
-	 * @param ListName name of the friends list that was operated on
-	 * @param ErrorStr string representing the error condition
-	 */
-	DEFINE_ONLINE_PLAYER_DELEGATE_FOUR_PARAM(MAX_LOCAL_PLAYERS, OnAcceptInviteComplete, bool, const FUniqueNetId&, const FString&, const FString&);
+	virtual bool AcceptInvite(int32 LocalUserNum, const FUniqueNetId& FriendId, const FString& ListName, const FOnAcceptInviteComplete& Delegate = FOnAcceptInviteComplete()) = 0;
 
 	/**
 	 * Starts an async task that rejects an invite from another player. 
diff --git a/Engine/Source/Runtime/Online/OnlineSubsystemFacebook/Private/IOS/OnlineFriendsFacebook.cpp b/Engine/Source/Runtime/Online/OnlineSubsystemFacebook/Private/IOS/OnlineFriendsFacebook.cpp
--- a/Engine/Source/Runtime/Online/OnlineSubsystemFacebook/Private/IOS/OnlineFriendsFacebook.cpp
+++ b/Engine/Source/Runtime/Online/OnlineSubsystemFacebook/Private/IOS/OnlineFriendsFacebook.cpp
@@ -95,21 +95,21 @@ bool FOnlineFriendsFacebook::ReadFriendsList(int32 LocalUserNum, const FString&
 	return bRequestTriggered;
 }
 
-bool FOnlineFriendsFacebook::DeleteFriendsList(int32 LocalUserNum, const FString& ListName)
+bool FOnlineFriendsFacebook::DeleteFriendsList(int32 LocalUserNum, const FString& ListName, const FOnDeleteFriendsListComplete& Delegate /*= FOnDeleteFriendsListComplete()*/)
 {
-	TriggerOnDeleteFriendsListCompleteDelegates(LocalUserNum, false, ListName, FString(TEXT("DeleteFriendsList() is not supported")));
+	Delegate.ExecuteIfBound(LocalUserNum, false, ListName, FString(TEXT("DeleteFriendsList() is not supported")));
 	return false;
 }
 
-bool FOnlineFriendsFacebook::SendInvite(int32 LocalUserNum, const FUniqueNetId& FriendId, const FString& ListName)
+bool FOnlineFriendsFacebook::SendInvite(int32 LocalUserNum, const FUniqueNetId& FriendId, const FString& ListName, const FOnSendInviteComplete& Delegate /*= FOnSendInviteComplete()*/)
 {
-	TriggerOnSendInviteCompleteDelegates(LocalUserNum, false, FriendId, ListName, FString(TEXT("SendInvite() is not supported")));
+	Delegate.ExecuteIfBound(LocalUserNum, false, FriendId, ListName, FString(TEXT("SendInvite() is not supported")));
 	return false;
 }
 
-bool FOnlineFriendsFacebook::AcceptInvite(int32 LocalUserNum, const FUniqueNetId& FriendId, const FString& ListName)
+bool FOnlineFriendsFacebook::AcceptInvite(int32 LocalUserNum, const FUniqueNetId& FriendId, const FString& ListName, const FOnAcceptInviteComplete& Delegate /*= FOnAcceptInviteComplete()*/)
 {
-	TriggerOnAcceptInviteCompleteDelegates(LocalUserNum, false, FriendId, ListName, FString(TEXT("AcceptInvite() is not supported")));
+	Delegate.ExecuteIfBound(LocalUserNum, false, FriendId, ListName, FString(TEXT("AcceptInvite() is not supported")));
 	return false;
 }
 
diff --git a/Engine/Source/Runtime/Online/OnlineSubsystemFacebook/Private/IOS/OnlineFriendsFacebook.h b/Engine/Source/Runtime/Online/OnlineSubsystemFacebook/Private/IOS/OnlineFriendsFacebook.h
--- a/Engine/Source/Runtime/Online/OnlineSubsystemFacebook/Private/IOS/OnlineFriendsFacebook.h
+++ b/Engine/Source/Runtime/Online/OnlineSubsystemFacebook/Private/IOS/OnlineFriendsFacebook.h
@@ -84,9 +84,9 @@ public:
 	// IOnlineFriends
 
 	virtual bool ReadFriendsList(int32 LocalUserNum, const FString& ListName, const FOnReadFriendsListComplete& Delegate = FOnReadFriendsListComplete()) override;
-	virtual bool DeleteFriendsList(int32 LocalUserNum, const FString& ListName) override;
-	virtual bool SendInvite(int32 LocalUserNum, const FUniqueNetId& FriendId, const FString& ListName) override;
-	virtual bool AcceptInvite(int32 LocalUserNum, const FUniqueNetId& FriendId, const FString& ListName) override;
+	virtual bool DeleteFriendsList(int32 LocalUserNum, const FString& ListName, const FOnDeleteFriendsListComplete& Delegate = FOnDeleteFriendsListComplete()) override;
+	virtual bool SendInvite(int32 LocalUserNum, const FUniqueNetId& FriendId, const FString& ListName,  const FOnSendInviteComplete& Delegate = FOnSendInviteComplete()) override;
+	virtual bool AcceptInvite(int32 LocalUserNum, const FUniqueNetId& FriendId, const FString& ListName, const FOnAcceptInviteComplete& Delegate = FOnAcceptInviteComplete()) override;
  	virtual bool RejectInvite(int32 LocalUserNum, const FUniqueNetId& FriendId, const FString& ListName) override;
  	virtual bool DeleteFriend(int32 LocalUserNum, const FUniqueNetId& FriendId, const FString& ListName) override;
 	virtual bool GetFriendsList(int32 LocalUserNum, const FString& ListName, TArray< TSharedRef<FOnlineFriend> >& OutFriends) override;
diff --git a/Engine/Source/Runtime/Online/OnlineSubsystemFacebook/Private/Windows/OnlineFriendsFacebook.cpp b/Engine/Source/Runtime/Online/OnlineSubsystemFacebook/Private/Windows/OnlineFriendsFacebook.cpp
--- a/Engine/Source/Runtime/Online/OnlineSubsystemFacebook/Private/Windows/OnlineFriendsFacebook.cpp
+++ b/Engine/Source/Runtime/Online/OnlineSubsystemFacebook/Private/Windows/OnlineFriendsFacebook.cpp
@@ -126,21 +126,21 @@ bool FOnlineFriendsFacebook::ReadFriendsList(int32 LocalUserNum, const FString&
 	return HttpRequest->ProcessRequest();
 }
 
-bool FOnlineFriendsFacebook::DeleteFriendsList(int32 LocalUserNum, const FString& ListName)
+bool FOnlineFriendsFacebook::DeleteFriendsList(int32 LocalUserNum, const FString& ListName, const FOnDeleteFriendsListComplete& Delegate /*= FOnDeleteFriendsListComplete()*/)
 {
-	TriggerOnDeleteFriendsListCompleteDelegates(LocalUserNum, false, ListName, FString(TEXT("DeleteFriendsList() is not supported")));
+	Delegate.ExecuteIfBound(LocalUserNum, false, ListName, FString(TEXT("DeleteFriendsList() is not supported")));
 	return false;
 }
 
-bool FOnlineFriendsFacebook::SendInvite(int32 LocalUserNum, const FUniqueNetId& FriendId, const FString& ListName)
+bool FOnlineFriendsFacebook::SendInvite(int32 LocalUserNum, const FUniqueNetId& FriendId, const FString& ListName, const FOnSendInviteComplete& Delegate /*= FOnSendInviteComplete()*/)
 {
-	TriggerOnSendInviteCompleteDelegates(LocalUserNum, false, FriendId, ListName, FString(TEXT("SendInvite() is not supported")));
+	Delegate.ExecuteIfBound(LocalUserNum, false, FriendId, ListName, FString(TEXT("SendInvite() is not supported")));
 	return false;
 }
 
-bool FOnlineFriendsFacebook::AcceptInvite(int32 LocalUserNum, const FUniqueNetId& FriendId, const FString& ListName)
+bool FOnlineFriendsFacebook::AcceptInvite(int32 LocalUserNum, const FUniqueNetId& FriendId, const FString& ListName, const FOnAcceptInviteComplete& Delegate /*= FOnAcceptInviteComplete()*/)
 {
-	TriggerOnAcceptInviteCompleteDelegates(LocalUserNum, false, FriendId, ListName, FString(TEXT("AcceptInvite() is not supported")));
+	Delegate.ExecuteIfBound(LocalUserNum, false, FriendId, ListName, FString(TEXT("AcceptInvite() is not supported")));
 	return false;
 }
 
diff --git a/Engine/Source/Runtime/Online/OnlineSubsystemFacebook/Private/Windows/OnlineFriendsFacebook.h b/Engine/Source/Runtime/Online/OnlineSubsystemFacebook/Private/Windows/OnlineFriendsFacebook.h
--- a/Engine/Source/Runtime/Online/OnlineSubsystemFacebook/Private/Windows/OnlineFriendsFacebook.h
+++ b/Engine/Source/Runtime/Online/OnlineSubsystemFacebook/Private/Windows/OnlineFriendsFacebook.h
@@ -82,9 +82,9 @@ public:
 	// IOnlineFriends
 
 	virtual bool ReadFriendsList(int32 LocalUserNum, const FString& ListName, const FOnReadFriendsListComplete& Delegate = FOnReadFriendsListComplete()) override;
-	virtual bool DeleteFriendsList(int32 LocalUserNum, const FString& ListName) override;
-	virtual bool SendInvite(int32 LocalUserNum, const FUniqueNetId& FriendId, const FString& ListName) override;
-	virtual bool AcceptInvite(int32 LocalUserNum, const FUniqueNetId& FriendId, const FString& ListName) override;
+	virtual bool DeleteFriendsList(int32 LocalUserNum, const FString& ListName, const FOnDeleteFriendsListComplete& Delegate = FOnDeleteFriendsListComplete()) override;
+	virtual bool SendInvite(int32 LocalUserNum, const FUniqueNetId& FriendId, const FString& ListName,  const FOnSendInviteComplete& Delegate = FOnSendInviteComplete()) override;
+	virtual bool AcceptInvite(int32 LocalUserNum, const FUniqueNetId& FriendId, const FString& ListName, const FOnAcceptInviteComplete& Delegate = FOnAcceptInviteComplete()) override;
  	virtual bool RejectInvite(int32 LocalUserNum, const FUniqueNetId& FriendId, const FString& ListName) override;
  	virtual bool DeleteFriend(int32 LocalUserNum, const FUniqueNetId& FriendId, const FString& ListName) override;
 	virtual bool GetFriendsList(int32 LocalUserNum, const FString& ListName, TArray< TSharedRef<FOnlineFriend> >& OutFriends) override;
diff --git a/Engine/Source/Runtime/Online/OnlineSubsystemSteam/Private/OnlineFriendsInterfaceSteam.cpp b/Engine/Source/Runtime/Online/OnlineSubsystemSteam/Private/OnlineFriendsInterfaceSteam.cpp
--- a/Engine/Source/Runtime/Online/OnlineSubsystemSteam/Private/OnlineFriendsInterfaceSteam.cpp
+++ b/Engine/Source/Runtime/Online/OnlineSubsystemSteam/Private/OnlineFriendsInterfaceSteam.cpp
@@ -85,21 +85,21 @@ bool FOnlineFriendsSteam::ReadFriendsList(int32 LocalUserNum, const FString& Lis
 	return true;
 }
 
-bool FOnlineFriendsSteam::DeleteFriendsList(int32 LocalUserNum, const FString& ListName)
+bool FOnlineFriendsSteam::DeleteFriendsList(int32 LocalUserNum, const FString& ListName, const FOnDeleteFriendsListComplete& Delegate /*= FOnDeleteFriendsListComplete()*/)
 {
-	TriggerOnDeleteFriendsListCompleteDelegates(LocalUserNum, false, ListName, FString(TEXT("DeleteFriendsList() is not supported")));
+	Delegate.ExecuteIfBound(LocalUserNum, false, ListName, FString(TEXT("DeleteFriendsList() is not supported")));
 	return false;
 }
 
-bool FOnlineFriendsSteam::SendInvite(int32 LocalUserNum, const FUniqueNetId& FriendId, const FString& ListName)
+bool FOnlineFriendsSteam::SendInvite(int32 LocalUserNum, const FUniqueNetId& FriendId, const FString& ListName, const FOnSendInviteComplete& Delegate /*= FOnSendInviteComplete()*/)
 {
-	TriggerOnSendInviteCompleteDelegates(LocalUserNum, false, FriendId, ListName, FString(TEXT("SendInvite() is not supported")));
+	Delegate.ExecuteIfBound(LocalUserNum, false, FriendId, ListName, FString(TEXT("SendInvite() is not supported")));
 	return false;
 }
 
-bool FOnlineFriendsSteam::AcceptInvite(int32 LocalUserNum, const FUniqueNetId& FriendId, const FString& ListName)
+bool FOnlineFriendsSteam::AcceptInvite(int32 LocalUserNum, const FUniqueNetId& FriendId, const FString& ListName, const FOnAcceptInviteComplete& Delegate /*= FOnAcceptInviteComplete()*/)
 {
-	TriggerOnAcceptInviteCompleteDelegates(LocalUserNum, false, FriendId, ListName, FString(TEXT("AcceptInvite() is not supported")));
+	Delegate.ExecuteIfBound(LocalUserNum, false, FriendId, ListName, FString(TEXT("AcceptInvite() is not supported")));
 	return false;
 }
 
diff --git a/Engine/Source/Runtime/Online/OnlineSubsystemSteam/Private/OnlineFriendsInterfaceSteam.h b/Engine/Source/Runtime/Online/OnlineSubsystemSteam/Private/OnlineFriendsInterfaceSteam.h
--- a/Engine/Source/Runtime/Online/OnlineSubsystemSteam/Private/OnlineFriendsInterfaceSteam.h
+++ b/Engine/Source/Runtime/Online/OnlineSubsystemSteam/Private/OnlineFriendsInterfaceSteam.h
@@ -104,9 +104,9 @@ public:
 	// IOnlineFriends
 
 	virtual bool ReadFriendsList(int32 LocalUserNum, const FString& ListName, const FOnReadFriendsListComplete& Delegate = FOnReadFriendsListComplete()) override;
-	virtual bool DeleteFriendsList(int32 LocalUserNum, const FString& ListName) override;
-	virtual bool SendInvite(int32 LocalUserNum, const FUniqueNetId& FriendId, const FString& ListName) override;
-	virtual bool AcceptInvite(int32 LocalUserNum, const FUniqueNetId& FriendId, const FString& ListName) override;
+	virtual bool DeleteFriendsList(int32 LocalUserNum, const FString& ListName, const FOnDeleteFriendsListComplete& Delegate = FOnDeleteFriendsListComplete()) override;
+	virtual bool SendInvite(int32 LocalUserNum, const FUniqueNetId& FriendId, const FString& ListName,  const FOnSendInviteComplete& Delegate = FOnSendInviteComplete()) override;
+	virtual bool AcceptInvite(int32 LocalUserNum, const FUniqueNetId& FriendId, const FString& ListName, const FOnAcceptInviteComplete& Delegate = FOnAcceptInviteComplete()) override;
  	virtual bool RejectInvite(int32 LocalUserNum, const FUniqueNetId& FriendId, const FString& ListName) override;
  	virtual bool DeleteFriend(int32 LocalUserNum, const FUniqueNetId& FriendId, const FString& ListName) override;
 	virtual bool GetFriendsList(int32 LocalUserNum, const FString& ListName, TArray< TSharedRef<FOnlineFriend> >& OutFriends) override;
diff --git a/Engine/Source/Runtime/Online/OnlineSubsystemUtils/Private/Tests/TestFriendsInterface.cpp b/Engine/Source/Runtime/Online/OnlineSubsystemUtils/Private/Tests/TestFriendsInterface.cpp
--- a/Engine/Source/Runtime/Online/OnlineSubsystemUtils/Private/Tests/TestFriendsInterface.cpp
+++ b/Engine/Source/Runtime/Online/OnlineSubsystemUtils/Private/Tests/TestFriendsInterface.cpp
@@ -19,17 +19,11 @@ void FTestFriendsInterface::Test(UWorld* InWorld, const TArray<FString>& Invites
 		}
 
 		// Add our delegate for the async call
-		OnAcceptInviteCompleteDelegate = FOnAcceptInviteCompleteDelegate::CreateRaw(this, &FTestFriendsInterface::OnAcceptInviteComplete);
-		OnSendInviteCompleteDelegate = FOnSendInviteCompleteDelegate::CreateRaw(this, &FTestFriendsInterface::OnSendInviteComplete);
 		OnDeleteFriendCompleteDelegate = FOnDeleteFriendCompleteDelegate::CreateRaw(this, &FTestFriendsInterface::OnDeleteFriendComplete);
-		OnDeleteFriendsListCompleteDelegate = FOnDeleteFriendsListCompleteDelegate::CreateRaw(this, &FTestFriendsInterface::OnDeleteFriendsListComplete);
 		OnQueryRecentPlayersCompleteDelegate = FOnQueryRecentPlayersCompleteDelegate::CreateRaw(this, &FTestFriendsInterface::OnQueryRecentPlayersComplete);
 		
 
-		OnlineSub->GetFriendsInterface()->AddOnAcceptInviteCompleteDelegate(0, OnAcceptInviteCompleteDelegate);
-		OnlineSub->GetFriendsInterface()->AddOnSendInviteCompleteDelegate(0, OnSendInviteCompleteDelegate);
 		OnlineSub->GetFriendsInterface()->AddOnDeleteFriendCompleteDelegate(0, OnDeleteFriendCompleteDelegate);
-		OnlineSub->GetFriendsInterface()->AddOnDeleteFriendsListCompleteDelegate(0, OnDeleteFriendsListCompleteDelegate);
 		OnlineSub->GetFriendsInterface()->AddOnQueryRecentPlayersCompleteDelegate(OnQueryRecentPlayersCompleteDelegate);
 
 		// list of pending users to send invites to
@@ -72,11 +66,13 @@ void FTestFriendsInterface::StartNextTest()
 	}
 	else if (bAcceptInvites && InvitesToAccept.Num() > 0)
 	{
-		OnlineSub->GetFriendsInterface()->AcceptInvite(0, *InvitesToAccept[0], FriendsListName);
+		FOnAcceptInviteComplete Delegate = FOnAcceptInviteComplete::CreateRaw(this, &FTestFriendsInterface::OnAcceptInviteComplete);
+		OnlineSub->GetFriendsInterface()->AcceptInvite(0, *InvitesToAccept[0], FriendsListName, Delegate);
 	}
 	else if (bSendInvites && InvitesToSend.Num() > 0)
 	{
-		OnlineSub->GetFriendsInterface()->SendInvite(0, *InvitesToSend[0], FriendsListName);
+		FOnSendInviteComplete OnSendInviteCompleteDelegate = FOnSendInviteComplete::CreateRaw(this, &FTestFriendsInterface::OnSendInviteComplete);
+		OnlineSub->GetFriendsInterface()->SendInvite(0, *InvitesToSend[0], FriendsListName, OnSendInviteCompleteDelegate);
 	}
 	else if (bDeleteFriends && FriendsToDelete.Num() > 0)
 	{
@@ -84,7 +80,8 @@ void FTestFriendsInterface::StartNextTest()
 	}
 	else if (bDeleteFriendsList)
 	{
-		OnlineSub->GetFriendsInterface()->DeleteFriendsList(0, FriendsListName);
+		FOnDeleteFriendsListComplete Delegate = FOnDeleteFriendsListComplete::CreateRaw(this, &FTestFriendsInterface::OnDeleteFriendsListComplete);
+		OnlineSub->GetFriendsInterface()->DeleteFriendsList(0, FriendsListName, Delegate);
 	}
 	else
 	{
@@ -99,10 +96,7 @@ void FTestFriendsInterface::FinishTest()
 		OnlineSub->GetFriendsInterface().IsValid())
 	{
 		// Clear delegates for the various async calls
-		OnlineSub->GetFriendsInterface()->ClearOnAcceptInviteCompleteDelegate(0, OnAcceptInviteCompleteDelegate);
-		OnlineSub->GetFriendsInterface()->ClearOnSendInviteCompleteDelegate(0, OnSendInviteCompleteDelegate);
 		OnlineSub->GetFriendsInterface()->ClearOnDeleteFriendCompleteDelegate(0, OnDeleteFriendCompleteDelegate);
-		OnlineSub->GetFriendsInterface()->ClearOnDeleteFriendsListCompleteDelegate(0, OnDeleteFriendsListCompleteDelegate);
 		OnlineSub->GetFriendsInterface()->ClearOnQueryRecentPlayersCompleteDelegate(OnQueryRecentPlayersCompleteDelegate);
 	}
 	delete this;
diff --git a/Engine/Source/Runtime/Online/OnlineSubsystemUtils/Private/Tests/TestFriendsInterface.h b/Engine/Source/Runtime/Online/OnlineSubsystemUtils/Private/Tests/TestFriendsInterface.h
--- a/Engine/Source/Runtime/Online/OnlineSubsystemUtils/Private/Tests/TestFriendsInterface.h
+++ b/Engine/Source/Runtime/Online/OnlineSubsystemUtils/Private/Tests/TestFriendsInterface.h
@@ -14,14 +14,8 @@
 	const FString SubsystemName;
 	/** The online interface to use for testing */
 	IOnlineSubsystem* OnlineSub;
-	/** Delegate when invites accepted */
-	FOnAcceptInviteCompleteDelegate OnAcceptInviteCompleteDelegate;
-	/** Delegate when invites have been sent */
-	FOnSendInviteCompleteDelegate OnSendInviteCompleteDelegate;
 	/** Delegate to use for deleting a friend entry */
 	FOnDeleteFriendCompleteDelegate OnDeleteFriendCompleteDelegate;
-	/** Delegate to use for deleting the friends list */
-	FOnDeleteFriendsListCompleteDelegate OnDeleteFriendsListCompleteDelegate;
 	/** Delegate to use for querying for recent players */
 	FOnQueryRecentPlayersCompleteDelegate OnQueryRecentPlayersCompleteDelegate;
 
