Update OSS Google Play function that uses deprecated GetTypedData

diff --git a/Engine/Source/Runtime/Online/Android/OnlineSubsystemGooglePlay/Private/OnlineSubsystemGooglePlay.cpp b/Engine/Source/Runtime/Online/Android/OnlineSubsystemGooglePlay/Private/OnlineSubsystemGooglePlay.cpp
--- a/Engine/Source/Runtime/Online/Android/OnlineSubsystemGooglePlay/Private/OnlineSubsystemGooglePlay.cpp
+++ b/Engine/Source/Runtime/Online/Android/OnlineSubsystemGooglePlay/Private/OnlineSubsystemGooglePlay.cpp
@@ -248,7 +248,7 @@ std::string FOnlineSubsystemGooglePlay::ConvertFStringToStdString(const FString&
 	TArray<ANSICHAR> Converted;
 	Converted.AddUninitialized(DestLen);
 	
-	FPlatformString::Convert(Converted.GetTypedData(), DestLen, *InString, SrcLen);
+	FPlatformString::Convert(Converted.GetData(), DestLen, *InString, SrcLen);
 
-	return std::string(Converted.GetTypedData());
+	return std::string(Converted.GetData());
 }
\ No newline at end of file
