fix deprecated api which was generating a warning for html5.

diff --git a/Engine/Source/Runtime/Online/HTTP/Private/HTML5/HTML5HTTP.cpp b/Engine/Source/Runtime/Online/HTTP/Private/HTML5/HTML5HTTP.cpp
--- a/Engine/Source/Runtime/Online/HTTP/Private/HTML5/HTML5HTTP.cpp
+++ b/Engine/Source/Runtime/Online/HTTP/Private/HTML5/HTML5HTTP.cpp
@@ -155,7 +155,7 @@ void FHTML5HttpRequest::SetContentAsString(const FString& ContentString)
 
 	FTCHARToUTF8 Converter(*ContentString);
 	RequestPayload.SetNum(Converter.Length());
-	FMemory::Memcpy(RequestPayload.GetTypedData(), (uint8*)(ANSICHAR*)Converter.Get(), RequestPayload.Num());
+	FMemory::Memcpy(RequestPayload.GetData(), (uint8*)(ANSICHAR*)Converter.Get(), RequestPayload.Num());
 }
 
 
