Rename "room" to "session" in cross-platform enum and add comments
Fixes UE-6094

diff --git a/Engine/Source/Runtime/Online/OnlineSubsystem/Public/Interfaces/OnlineSessionInterface.h b/Engine/Source/Runtime/Online/OnlineSubsystem/Public/Interfaces/OnlineSessionInterface.h
--- a/Engine/Source/Runtime/Online/OnlineSubsystem/Public/Interfaces/OnlineSessionInterface.h
+++ b/Engine/Source/Runtime/Online/OnlineSubsystem/Public/Interfaces/OnlineSessionInterface.h
@@ -90,17 +90,22 @@ typedef FOnFindSessionsComplete::FDelegate FOnFindSessionsCompleteDelegate;
 DECLARE_MULTICAST_DELEGATE_OneParam(FOnCancelFindSessionsComplete, bool);
 typedef FOnCancelFindSessionsComplete::FDelegate FOnCancelFindSessionsCompleteDelegate;
 
-
-
+/** Possible results of a JoinSession attempt */
 namespace EOnJoinSessionCompleteResult
 {
 	enum Type
 	{
+		/** The join worked as expected */
 		Success,
-		RoomIsFull,
-		RoomDoesNotExist,
+		/** There are no open slots to join */
+		SessionIsFull,
+		/** The session couldn't be found on the service */
+		SessionDoesNotExist,
+		/** There was an error getting the session server's address */
 		CouldNotRetrieveAddress,
+		/** The user attempting to join is already a member of the session */
 		AlreadyInSession,
+		/** An error not covered above occurred */
 		UnknownError
 	};
 }
