Remove dynamic allocation of FLANSession in OnlineSessionInterfaceNull
Simplifies code by removing manual memory management and null checks
Fixes memory stomp reported on UDN: FinalizeLANSearch was deleting the FLANSession during it's own tick, inside a delegate call

diff --git a/Engine/Source/Runtime/Online/OnlineSubsystem/Private/LANBeacon.cpp b/Engine/Source/Runtime/Online/OnlineSubsystem/Private/LANBeacon.cpp
--- a/Engine/Source/Runtime/Online/OnlineSubsystem/Private/LANBeacon.cpp
+++ b/Engine/Source/Runtime/Online/OnlineSubsystem/Private/LANBeacon.cpp
@@ -212,6 +212,11 @@ void FLANSession::StopLANSession()
 
 void FLANSession::Tick(float DeltaTime)
 {
+	if (LanBeaconState == ELanBeaconState::NotUsingLanBeacon)
+	{
+		return;
+	}
+
 	uint8 PacketData[LAN_BEACON_MAX_PACKET_SIZE];
 	bool bShouldRead = true;
 	// Read each pending packet and pass it out for processing
diff --git a/Engine/Source/Runtime/Online/OnlineSubsystemNull/Private/OnlineSessionInterfaceNull.cpp b/Engine/Source/Runtime/Online/OnlineSubsystemNull/Private/OnlineSessionInterfaceNull.cpp
--- a/Engine/Source/Runtime/Online/OnlineSubsystemNull/Private/OnlineSessionInterfaceNull.cpp
+++ b/Engine/Source/Runtime/Online/OnlineSubsystemNull/Private/OnlineSessionInterfaceNull.cpp
@@ -300,31 +300,23 @@ uint32 FOnlineSessionNull::UpdateLANStatus()
 	if ( NeedsToAdvertise() )
 	{
 		// set up LAN session
-		if (NULL == LANSessionManager)
+		if (LANSessionManager.GetBeaconState() == ELanBeaconState::NotUsingLanBeacon)
 		{
-			LANSessionManager = new FLANSession();
-
 			FOnValidQueryPacketDelegate QueryPacketDelegate = FOnValidQueryPacketDelegate::CreateRaw(this, &FOnlineSessionNull::OnValidQueryPacketReceived);
-			if (!LANSessionManager->Host(QueryPacketDelegate))
+			if (!LANSessionManager.Host(QueryPacketDelegate))
 			{
 				Result = E_FAIL;
 
-				LANSessionManager->StopLANSession();
-				delete LANSessionManager;
-
-				LANSessionManager = NULL;
+				LANSessionManager.StopLANSession();
 			}
 		}
 	}
 	else
 	{
-		if (LANSessionManager && LANSessionManager->GetBeaconState() != ELanBeaconState::Searching)
+		if (LANSessionManager.GetBeaconState() != ELanBeaconState::Searching)
 		{
 			// Tear down the LAN beacon
-			LANSessionManager->StopLANSession();
-			delete LANSessionManager;
-
-			LANSessionManager = NULL;
+			LANSessionManager.StopLANSession();
 		}
 	}
 
@@ -526,20 +518,15 @@ uint32 FOnlineSessionNull::FindLANSession()
 {
 	uint32 Return = ERROR_IO_PENDING;
 
-	if (!LANSessionManager)
-	{
-		LANSessionManager = new FLANSession();
-	}
-
 	// Recreate the unique identifier for this client
-	GenerateNonce((uint8*)&LANSessionManager->LanNonce, 8);
+	GenerateNonce((uint8*)&LANSessionManager.LanNonce, 8);
 
 	FOnValidResponsePacketDelegate ResponseDelegate = FOnValidResponsePacketDelegate::CreateRaw(this, &FOnlineSessionNull::OnValidResponsePacketReceived);
 	FOnSearchingTimeoutDelegate TimeoutDelegate = FOnSearchingTimeoutDelegate::CreateRaw(this, &FOnlineSessionNull::OnLANSearchTimeout);
 
 	FNboSerializeToBufferNull Packet(LAN_BEACON_MAX_PACKET_SIZE);
-	LANSessionManager->CreateClientQueryPacket(Packet, LANSessionManager->LanNonce);
-	if (LANSessionManager->Search(Packet, ResponseDelegate, TimeoutDelegate) == false)
+	LANSessionManager.CreateClientQueryPacket(Packet, LANSessionManager.LanNonce);
+	if (LANSessionManager.Search(Packet, ResponseDelegate, TimeoutDelegate) == false)
 	{
 		Return = E_FAIL;
 
@@ -943,10 +930,7 @@ void FOnlineSessionNull::Tick(float DeltaTime)
 
 void FOnlineSessionNull::TickLanTasks(float DeltaTime)
 {
-	if (LANSessionManager != NULL && LANSessionManager->GetBeaconState() > ELanBeaconState::NotUsingLanBeacon)
-	{
-		LANSessionManager->Tick(DeltaTime);
-	}
+	LANSessionManager.Tick(DeltaTime);
 }
 
 void FOnlineSessionNull::AppendSessionToPacket(FNboSerializeToBufferNull& Packet, FOnlineSession* Session)
@@ -1033,7 +1017,7 @@ void FOnlineSessionNull::OnValidQueryPacketReceived(uint8* PacketData, int32 Pac
 			{
 				FNboSerializeToBufferNull Packet(LAN_BEACON_MAX_PACKET_SIZE);
 				// Create the basic header before appending additional information
-				LANSessionManager->CreateHostResponsePacket(Packet, ClientNonce);
+				LANSessionManager.CreateHostResponsePacket(Packet, ClientNonce);
 
 				// Add all the session details
 				AppendSessionToPacket(Packet, Session);
@@ -1041,7 +1025,7 @@ void FOnlineSessionNull::OnValidQueryPacketReceived(uint8* PacketData, int32 Pac
 				// Broadcast this response so the client can see us
 				if (!Packet.HasOverflow())
 				{
-					LANSessionManager->BroadcastPacket(Packet, Packet.GetByteCount());
+					LANSessionManager.BroadcastPacket(Packet, Packet.GetByteCount());
 				}
 				else
 				{
@@ -1174,12 +1158,9 @@ void FOnlineSessionNull::OnValidResponsePacketReceived(uint8* PacketData, int32
 
 uint32 FOnlineSessionNull::FinalizeLANSearch()
 {
-	if (LANSessionManager)
+	if (LANSessionManager.GetBeaconState() == ELanBeaconState::Searching)
 	{
-		check(LANSessionManager->GetBeaconState() == ELanBeaconState::Searching);
-		LANSessionManager->StopLANSession();
-		delete LANSessionManager;
-		LANSessionManager = NULL;
+		LANSessionManager.StopLANSession();
 	}
 
 	return UpdateLANStatus();
diff --git a/Engine/Source/Runtime/Online/OnlineSubsystemNull/Private/OnlineSessionInterfaceNull.h b/Engine/Source/Runtime/Online/OnlineSubsystemNull/Private/OnlineSessionInterfaceNull.h
--- a/Engine/Source/Runtime/Online/OnlineSubsystemNull/Private/OnlineSessionInterfaceNull.h
+++ b/Engine/Source/Runtime/Online/OnlineSubsystemNull/Private/OnlineSessionInterfaceNull.h
@@ -5,6 +5,7 @@
 #include "OnlineSessionInterface.h"
 #include "OnlineSubsystemNullTypes.h"
 #include "OnlineSubsystemNullPackage.h"
+#include "LANBeacon.h"
 
 /**
  * Interface definition for the online services session services 
@@ -19,12 +20,11 @@ private:
 	class FOnlineSubsystemNull* NullSubsystem;
 
 	/** Handles advertising sessions over LAN and client searches */
-	class FLANSession* LANSessionManager;
+	FLANSession LANSessionManager;
 
 	/** Hidden on purpose */
 	FOnlineSessionNull() :
 		NullSubsystem(NULL),
-		LANSessionManager(NULL),
 		CurrentSessionSearch(NULL)
 	{}
 
@@ -165,7 +165,6 @@ PACKAGE_SCOPE:
 
 	FOnlineSessionNull(class FOnlineSubsystemNull* InSubsystem) :
 		NullSubsystem(InSubsystem),
-		LANSessionManager(NULL),
 		CurrentSessionSearch(NULL),
 		SessionSearchStartInSeconds(0)
 	{}
