Remove special-cased OSS module shutdown.
OSS module is now shut down in the same way as other modules, in reverse load order, but the service-specific OSS modules are still shut down before the base OSS module.

diff --git a/Engine/Source/Runtime/Launch/Private/LaunchEngineLoop.cpp b/Engine/Source/Runtime/Launch/Private/LaunchEngineLoop.cpp
--- a/Engine/Source/Runtime/Launch/Private/LaunchEngineLoop.cpp
+++ b/Engine/Source/Runtime/Launch/Private/LaunchEngineLoop.cpp
@@ -2052,9 +2052,6 @@ void FEngineLoop::Exit()
 
 	AppPreExit();
 
-	// Shutdown and unregister all online subsystems
-	Online::ShutdownOnlineSubsystem();
-
 	TermGamePhys();
 	ParticleVertexFactoryPool_FreePool();
 	MotionBlur_Free();
diff --git a/Engine/Source/Runtime/Online/OnlineSubsystem/Public/Online.h b/Engine/Source/Runtime/Online/OnlineSubsystem/Public/Online.h
--- a/Engine/Source/Runtime/Online/OnlineSubsystem/Public/Online.h
+++ b/Engine/Source/Runtime/Online/OnlineSubsystem/Public/Online.h
@@ -31,21 +31,6 @@
 /** Helpers for accessing all the online features available in the online subsystem */
 namespace Online
 {
-	/** 
-	 * Shutdown all online services
-	 */
-	static void ShutdownOnlineSubsystem() 
-	{ 
-		// This will potentially be called before the online subsystem was loaded...
-		FModuleManager& ModuleManager = FModuleManager::Get();
-		if (ModuleManager.IsModuleLoaded(TEXT("OnlineSubsystem")) == true)
-		{
-			// Unloading the module will call FOnlineSubsystemModule::ShutdownOnlineSubsystem()
-			const bool bIsShutdown = true;
-			ModuleManager.UnloadModule(TEXT("OnlineSubsystem"), bIsShutdown);
-		}
-	}
-
 	/** 
 	 * Get the interface for accessing the session services
 	 * @param SubsystemName - Name of the requested online service
diff --git a/Engine/Source/Runtime/Online/OnlineSubsystem/Public/OnlineSubsystemModule.h b/Engine/Source/Runtime/Online/OnlineSubsystem/Public/OnlineSubsystemModule.h
--- a/Engine/Source/Runtime/Online/OnlineSubsystem/Public/OnlineSubsystemModule.h
+++ b/Engine/Source/Runtime/Online/OnlineSubsystem/Public/OnlineSubsystemModule.h
@@ -138,7 +138,7 @@ public:
 	 */
 	virtual bool SupportsAutomaticShutdown() override
 	{
-		return false;
+		return true;
 	}
 };
 
