Simpify network version handshake + allow game to override with delegate

diff --git a/Engine/Source/Runtime/Engine/Classes/Engine/Channel.h b/Engine/Source/Runtime/Engine/Classes/Engine/Channel.h
--- a/Engine/Source/Runtime/Engine/Classes/Engine/Channel.h
+++ b/Engine/Source/Runtime/Engine/Classes/Engine/Channel.h
@@ -46,7 +46,6 @@ class ENGINE_API UChannel : public UObject
 	EChannelType		ChType;				// Type of this channel.
 	int32					NumInRec;			// Number of packets in InRec.
 	int32					NumOutRec;			// Number of packets in OutRec.
-	int32					NegotiatedVer;		// Negotiated version of engine = Min(client version, server version).
 	class FInBunch*		InRec;				// Incoming data with queued dependencies.
 	class FOutBunch*	OutRec;				// Outgoing reliable unacked data.
 	class FInBunch*		InPartialBunch;		// Partial bunch we are receiving (incoming partial bunches are appended to this)
diff --git a/Engine/Source/Runtime/Engine/Classes/Engine/NetConnection.h b/Engine/Source/Runtime/Engine/Classes/Engine/NetConnection.h
--- a/Engine/Source/Runtime/Engine/Classes/Engine/NetConnection.h
+++ b/Engine/Source/Runtime/Engine/Classes/Engine/NetConnection.h
@@ -160,7 +160,6 @@ public:
 	FString			Challenge;				// Server-generated challenge.
 	FString			ClientResponse;			// Client-generated response.
 	int32			ResponseId;				// Id assigned by the server for linking responses to connections upon authentication
-	int32			NegotiatedVer;			// Negotiated version for new channels.
 	FString			RequestURL;				// URL requested by client
 
 	// Login state tracking
diff --git a/Engine/Source/Runtime/Engine/Classes/Engine/NetDriver.h b/Engine/Source/Runtime/Engine/Classes/Engine/NetDriver.h
--- a/Engine/Source/Runtime/Engine/Classes/Engine/NetDriver.h
+++ b/Engine/Source/Runtime/Engine/Classes/Engine/NetDriver.h
@@ -151,10 +151,6 @@ public:
 	UPROPERTY(Config)
 	float ConnectionTimeout;
 
-	/** Requires engine version to match exactly in order to connect. Else fall back to GEngineMinNetVersion check */
-	UPROPERTY(Config)
-	bool RequireEngineVersionMatch;
-
 	/** Connection to the server (this net driver is a client) */
 	UPROPERTY()
 	class UNetConnection* ServerConnection;
diff --git a/Engine/Source/Runtime/Engine/Private/DataChannel.cpp b/Engine/Source/Runtime/Engine/Private/DataChannel.cpp
--- a/Engine/Source/Runtime/Engine/Private/DataChannel.cpp
+++ b/Engine/Source/Runtime/Engine/Private/DataChannel.cpp
@@ -53,7 +53,6 @@ void UChannel::Init( UNetConnection* InConnection, int32 InChIndex, bool InOpene
 	ChIndex			= InChIndex;
 	OpenedLocally	= InOpenedLocally;
 	OpenPacketId	= FPacketIdRange();
-	NegotiatedVer	= InConnection->NegotiatedVer;
 }
 
 
diff --git a/Engine/Source/Runtime/Engine/Private/NetConnection.cpp b/Engine/Source/Runtime/Engine/Private/NetConnection.cpp
--- a/Engine/Source/Runtime/Engine/Private/NetConnection.cpp
+++ b/Engine/Source/Runtime/Engine/Private/NetConnection.cpp
@@ -31,7 +31,6 @@ UNetConnection::UNetConnection(const FObjectInitializer& ObjectInitializer)
 ,	ProtocolVersion		( MIN_PROTOCOL_VERSION )
 ,	PacketOverhead		( 0 )
 ,	ResponseId			( 0 )
-,	NegotiatedVer		( GEngineNegotiationVersion )
 
 ,	QueuedBytes			( 0 )
 ,	TickCount			( 0 )
diff --git a/Engine/Source/Runtime/Engine/Private/NetworkDriver.cpp b/Engine/Source/Runtime/Engine/Private/NetworkDriver.cpp
--- a/Engine/Source/Runtime/Engine/Private/NetworkDriver.cpp
+++ b/Engine/Source/Runtime/Engine/Private/NetworkDriver.cpp
@@ -101,7 +101,6 @@ UNetDriver::UNetDriver(const FObjectInitializer& ObjectInitializer)
 :	UObject(ObjectInitializer)
 ,	MaxInternetClientRate(10000)
 , 	MaxClientRate(15000)
-,	RequireEngineVersionMatch(true)
 ,	ClientConnections()
 ,	Time( 0.f )
 ,	InBytes(0)
diff --git a/Engine/Source/Runtime/Engine/Private/PendingNetGame.cpp b/Engine/Source/Runtime/Engine/Private/PendingNetGame.cpp
--- a/Engine/Source/Runtime/Engine/Private/PendingNetGame.cpp
+++ b/Engine/Source/Runtime/Engine/Private/PendingNetGame.cpp
@@ -42,7 +42,10 @@ void UPendingNetGame::InitNetDriver()
 			// Send initial message.
 			uint8 IsLittleEndian = uint8(PLATFORM_LITTLE_ENDIAN);
 			check(IsLittleEndian == !!IsLittleEndian); // should only be one or zero
-			FNetControlMessage<NMT_Hello>::Send(NetDriver->ServerConnection, IsLittleEndian, GEngineMinNetVersion, GEngineNetVersion, Cast<UGeneralProjectSettings>(UGeneralProjectSettings::StaticClass()->GetDefaultObject())->ProjectID);
+			
+			uint32 LocalNetworkVersion = FNetworkVersion::GetLocalNetworkVersion();
+
+			FNetControlMessage<NMT_Hello>::Send( NetDriver->ServerConnection, IsLittleEndian, LocalNetworkVersion );
 
 			NetDriver->ServerConnection->FlushNet();
 		}
@@ -114,20 +117,11 @@ void UPendingNetGame::NotifyControlMessage(UNetConnection* Connection, uint8 Mes
 	{
 		case NMT_Upgrade:
 			// Report mismatch.
-			int32 RemoteMinVer, RemoteVer;
-			FNetControlMessage<NMT_Upgrade>::Receive(Bunch, RemoteMinVer, RemoteVer);
-			if (GEngineNetVersion < RemoteMinVer)
-			{
-				// Upgrade
-				ConnectionError = NSLOCTEXT("Engine", "ClientOutdated", "The match you are trying to join is running an incompatible version of the game.  Please try upgrading your game version.").ToString();
-				GEngine->BroadcastNetworkFailure(NULL, NetDriver, ENetworkFailure::OutdatedClient, ConnectionError);
-			}
-			else
-			{
-				// Downgrade
-				ConnectionError = NSLOCTEXT("Engine", "ServerOutdated", "Server's version is outdated").ToString();
-				GEngine->BroadcastNetworkFailure(NULL, NetDriver, ENetworkFailure::OutdatedServer, ConnectionError);
-			}
+			uint32 RemoteNetworkVersion;
+			FNetControlMessage<NMT_Upgrade>::Receive(Bunch, RemoteNetworkVersion);
+			// Upgrade
+			ConnectionError = NSLOCTEXT("Engine", "ClientOutdated", "The match you are trying to join is running an incompatible version of the game.  Please try upgrading your game version.").ToString();
+			GEngine->BroadcastNetworkFailure(NULL, NetDriver, ENetworkFailure::OutdatedClient, ConnectionError);
 			break;
 
 		case NMT_Failure:
@@ -158,7 +152,7 @@ void UPendingNetGame::NotifyControlMessage(UNetConnection* Connection, uint8 Mes
 		case NMT_Challenge:
 		{
 			// Challenged by server.
-			FNetControlMessage<NMT_Challenge>::Receive(Bunch, Connection->NegotiatedVer, Connection->Challenge);
+			FNetControlMessage<NMT_Challenge>::Receive(Bunch, Connection->Challenge);
 
 			FURL PartialURL(URL);
 			PartialURL.Host = TEXT("");
diff --git a/Engine/Source/Runtime/Engine/Private/UnrealNetwork.cpp b/Engine/Source/Runtime/Engine/Private/UnrealNetwork.cpp
new file mode 100644
--- /dev/null
+++ b/Engine/Source/Runtime/Engine/Private/UnrealNetwork.cpp
@@ -0,0 +1,45 @@
+// Copyright 1998-2015 Epic Games, Inc. All Rights Reserved.
+
+#include "EnginePrivate.h"
+#include "Net/UnrealNetwork.h"
+#include "GeneralProjectSettings.h"
+
+FNetworkVersion::FGetLocalNetworkVersionOverride FNetworkVersion::GetLocalNetworkVersionOverride;
+FNetworkVersion::FIsNetworkCompatibleOverride FNetworkVersion::IsNetworkCompatibleOverride;
+
+uint32 FNetworkVersion::GetLocalNetworkVersion()
+{
+	if ( GetLocalNetworkVersionOverride.IsBound() )
+	{
+		const uint32 LocalNetworkVersion = GetLocalNetworkVersionOverride.Execute();
+
+		UE_LOG( LogNet, Log, TEXT( "GetLocalNetworkVersionOverride: LocalNetworkVersion: %i" ), LocalNetworkVersion );
+
+		return LocalNetworkVersion;
+	}
+
+	// Get the project name (NOT case sensitive)
+	const FString ProjectName( FString( FApp::GetGameName() ).ToLower() );
+
+	// Get the project version string (IS case sensitive!)
+	const FString& ProjectVersion = Cast<UGeneralProjectSettings>(UGeneralProjectSettings::StaticClass()->GetDefaultObject())->ProjectVersion;
+
+	// Start with engine version as seed, and then hash with project name + project version
+	const uint32 LocalNetworkVersion = FCrc::StrCrc32( *ProjectVersion, FCrc::StrCrc32( *ProjectName, GEngineNetVersion ) );
+
+	UE_LOG( LogNet, Log, TEXT( "GetLocalNetworkVersion: GEngineNetVersion: %i, ProjectName: %s, ProjectVersion: %s, LocalNetworkVersion: %i" ), GEngineNetVersion, *ProjectName, *ProjectVersion, LocalNetworkVersion );
+
+	return LocalNetworkVersion;
+}
+
+bool FNetworkVersion::IsNetworkCompatible( const uint32 LocalNetworkVersion, const uint32 RemoteNetworkVersion )
+{
+	if ( IsNetworkCompatibleOverride.IsBound() )
+	{
+		return IsNetworkCompatibleOverride.Execute( LocalNetworkVersion, RemoteNetworkVersion );
+	}
+
+	return LocalNetworkVersion == RemoteNetworkVersion;
+}
+
+// ----------------------------------------------------------------
diff --git a/Engine/Source/Runtime/Engine/Private/World.cpp b/Engine/Source/Runtime/Engine/Private/World.cpp
--- a/Engine/Source/Runtime/Engine/Private/World.cpp
+++ b/Engine/Source/Runtime/Engine/Private/World.cpp
@@ -3585,35 +3585,24 @@ void UWorld::NotifyControlMessage(UNetConnection* Connection, uint8 MessageType,
 		{
 			case NMT_Hello:
 			{
-				uint8 IsLittleEndian;
-				int32 RemoteMinVer, RemoteVer;
-				FGuid RemoteGameGUID;
-				FNetControlMessage<NMT_Hello>::Receive(Bunch, IsLittleEndian, RemoteMinVer, RemoteVer, RemoteGameGUID);
+				uint8 IsLittleEndian = 0;
+				uint32 RemoteNetworkVersion = 0;
+				uint32 LocalNetworkVersion = FNetworkVersion::GetLocalNetworkVersion();
 
-				if (!IsNetworkCompatible(Connection->Driver->RequireEngineVersionMatch, RemoteVer, RemoteMinVer))
+				FNetControlMessage<NMT_Hello>::Receive(Bunch, IsLittleEndian, RemoteNetworkVersion);
+
+				if (!FNetworkVersion::IsNetworkCompatible(LocalNetworkVersion, RemoteNetworkVersion))
 				{
-					UE_LOG(LogNet, Log, TEXT("Requesting client to upgrade or downgrade to GEngineMinNetVersion=%d, GEngineNetVersion=%d"), GEngineMinNetVersion, GEngineNetVersion);
-					FNetControlMessage<NMT_Upgrade>::Send(Connection, GEngineMinNetVersion, GEngineNetVersion);
+					UE_LOG(LogNet, Log, TEXT("NotifyControlMessage: Client connecting with invalid version. LocalNetworkVersion: %i, RemoteNetworkVersion: %i"), LocalNetworkVersion, RemoteNetworkVersion);
+					FNetControlMessage<NMT_Upgrade>::Send(Connection, LocalNetworkVersion);
 					Connection->FlushNet(true);
 					Connection->Close();
 				}
 				else
 				{
-					Connection->NegotiatedVer = FMath::Min(RemoteVer, GEngineNetVersion);
-
-					// Make sure the server has the same GameGUID as we do
-					if( RemoteGameGUID != GetDefault<UGeneralProjectSettings>()->ProjectID )
-					{
-						FString ErrorMsg = NSLOCTEXT("NetworkErrors", "ServerHostingDifferentGame", "Incompatible game connection.").ToString();
-						FNetControlMessage<NMT_Failure>::Send(Connection, ErrorMsg);
-						Connection->FlushNet(true);
-						Connection->Close();
-						break;
-					}
-
 					Connection->Challenge = FString::Printf(TEXT("%08X"), FPlatformTime::Cycles());
 					Connection->SetExpectedClientLoginMsgType( NMT_Login );
-					FNetControlMessage<NMT_Challenge>::Send(Connection, Connection->NegotiatedVer, Connection->Challenge);
+					FNetControlMessage<NMT_Challenge>::Send(Connection, Connection->Challenge);
 					Connection->FlushNet();
 				}
 				break;
diff --git a/Engine/Source/Runtime/Engine/Public/Net/DataChannel.h b/Engine/Source/Runtime/Engine/Public/Net/DataChannel.h
--- a/Engine/Source/Runtime/Engine/Public/Net/DataChannel.h
+++ b/Engine/Source/Runtime/Engine/Public/Net/DataChannel.h
@@ -370,10 +370,10 @@ private:
 #define IMPLEMENT_CONTROL_CHANNEL_MESSAGE(Name) static uint8 Dummy##_FNetControlMessage_##Name = FNetControlMessage<NMT_##Name>::Initialize();
 
 // message type definitions
-DEFINE_CONTROL_CHANNEL_MESSAGE_FOURPARAM(Hello, 0, uint8, int32, int32, FGuid); // initial client connection message
+DEFINE_CONTROL_CHANNEL_MESSAGE_TWOPARAM(Hello, 0, uint8, uint32); // initial client connection message
 DEFINE_CONTROL_CHANNEL_MESSAGE_THREEPARAM(Welcome, 1, FString, FString, FString); // server tells client they're ok'ed to load the server's level
-DEFINE_CONTROL_CHANNEL_MESSAGE_TWOPARAM(Upgrade, 2, int32, int32); // server tells client their version is incompatible
-DEFINE_CONTROL_CHANNEL_MESSAGE_TWOPARAM(Challenge, 3, int32, FString); // server sends client challenge string to verify integrity
+DEFINE_CONTROL_CHANNEL_MESSAGE_ONEPARAM(Upgrade, 2, uint32); // server tells client their version is incompatible
+DEFINE_CONTROL_CHANNEL_MESSAGE_ONEPARAM(Challenge, 3, FString); // server sends client challenge string to verify integrity
 DEFINE_CONTROL_CHANNEL_MESSAGE_ONEPARAM(Netspeed, 4, int32); // client sends requested transfer rate
 DEFINE_CONTROL_CHANNEL_MESSAGE_THREEPARAM(Login, 5, FString, FString, FUniqueNetIdRepl); // client requests to be admitted to the game
 DEFINE_CONTROL_CHANNEL_MESSAGE_ONEPARAM(Failure, 6, FString); // indicates connection failure
@@ -407,4 +407,4 @@ DEFINE_CONTROL_CHANNEL_MESSAGE_TWOPARAM(NetGUIDAssign, 18, FNetworkGUID, FString
 DEFINE_CONTROL_CHANNEL_MESSAGE_ZEROPARAM(BeaconWelcome, 25); // server tells client they're ok to attempt to join (client sends netspeed/beacontype)
 DEFINE_CONTROL_CHANNEL_MESSAGE_ONEPARAM(BeaconJoin, 26, FString);  // server tries to create beacon type requested by client, sends NetGUID for actor sync
 DEFINE_CONTROL_CHANNEL_MESSAGE_ONEPARAM(BeaconAssignGUID, 27, FNetworkGUID); // client assigns NetGUID from server to beacon actor, sends NetGUIDAck
-DEFINE_CONTROL_CHANNEL_MESSAGE_ONEPARAM(BeaconNetGUIDAck, 28, FString); // server received NetGUIDAck from client, connection established successfully
+DEFINE_CONTROL_CHANNEL_MESSAGE_ONEPARAM(BeaconNetGUIDAck, 28, FString); // server received NetGUIDAck from client, connection established successfully
\ No newline at end of file
diff --git a/Engine/Source/Runtime/Engine/Public/Net/UnrealNetwork.h b/Engine/Source/Runtime/Engine/Public/Net/UnrealNetwork.h
--- a/Engine/Source/Runtime/Engine/Public/Net/UnrealNetwork.h
+++ b/Engine/Source/Runtime/Engine/Public/Net/UnrealNetwork.h
@@ -37,21 +37,33 @@ inline int32 MakeRelative( int32 Value, int32 Reference, int32 Max )
 	return Reference + BestSignedDifference(Value,Reference,Max);
 }
 
-/**
- * Determine if a connection is compatible with this instance
- *
- * @param bRequireEngineVersionMatch should the engine versions match exactly
- * @param RemoteVer current version of the remote party
- * @param RemoteMinVer min net compatible version of the remote party
- *
- * @return true if the two instances can communicate, false otherwise
- */
-inline bool IsNetworkCompatible(bool bRequireEngineVersionMatch, int32 RemoteVer, int32 RemoteMinVer)
+struct ENGINE_API FNetworkVersion
 {
-	bool bEngineVerMatch = !bRequireEngineVersionMatch || RemoteVer == GEngineNetVersion || GEngineNetVersion == 0 || RemoteVer == 0;
-	bool bMinEngineVerMatch = (RemoteVer > GEngineMinNetVersion || RemoteVer == 0) && (RemoteMinVer < GEngineNetVersion || GEngineNetVersion == 0);
-	return bEngineVerMatch && bMinEngineVerMatch;
-}
+	/** Called in GetLocalNetworkVersion if bound */
+	DECLARE_DELEGATE_RetVal( uint32, FGetLocalNetworkVersionOverride );
+	static FGetLocalNetworkVersionOverride GetLocalNetworkVersionOverride;
+
+	/** Called in IsNetworkCompatible if bound */
+	DECLARE_DELEGATE_RetVal_TwoParams( bool, FIsNetworkCompatibleOverride, uint32, uint32 );
+	static FIsNetworkCompatibleOverride IsNetworkCompatibleOverride;
+
+	/**
+	 * Generates a version number, that by default, is based on a checksum of the engine version + project name + project version string
+	 * Game/project code can completely override what this value returns through the GetLocalNetworkVersionOverride delegate
+	 */
+	static uint32 GetLocalNetworkVersion();
+
+	/**
+	 * Determine if a connection is compatible with this instance
+	 *
+	 * @param bRequireEngineVersionMatch should the engine versions match exactly
+	 * @param LocalNetworkVersion current version of the local machine
+	 * @param RemoteNetworkVersion current version of the remote machine
+	 *
+	 * @return true if the two instances can communicate, false otherwise
+	 */
+	static bool IsNetworkCompatible( const uint32 LocalNetworkVersion, const uint32 RemoteNetworkVersion );
+};
 
 /*-----------------------------------------------------------------------------
 	Replication.
diff --git a/Engine/Source/Runtime/Online/OnlineSubsystemUtils/Private/OnlineBeaconClient.cpp b/Engine/Source/Runtime/Online/OnlineSubsystemUtils/Private/OnlineBeaconClient.cpp
--- a/Engine/Source/Runtime/Online/OnlineSubsystemUtils/Private/OnlineBeaconClient.cpp
+++ b/Engine/Source/Runtime/Online/OnlineSubsystemUtils/Private/OnlineBeaconClient.cpp
@@ -2,6 +2,7 @@
 
 #include "OnlineSubsystemUtilsPrivatePCH.h"
 #include "Net/DataChannel.h"
+#include "Net/UnrealNetwork.h"
 
 #define BEACON_RPC_TIMEOUT 15.0f
 
@@ -43,7 +44,10 @@ bool AOnlineBeaconClient::InitClient(FURL& URL)
 				uint8 IsLittleEndian = uint8(PLATFORM_LITTLE_ENDIAN);
 				check(IsLittleEndian == !!IsLittleEndian); // should only be one or zero
 				BeaconConnection = NetDriver->ServerConnection;
-				FNetControlMessage<NMT_Hello>::Send(NetDriver->ServerConnection, IsLittleEndian, GEngineMinNetVersion, GEngineNetVersion, Cast<UGeneralProjectSettings>(UGeneralProjectSettings::StaticClass()->GetDefaultObject())->ProjectID);
+
+				uint32 LocalNetworkVersion = FNetworkVersion::GetLocalNetworkVersion();
+				
+				FNetControlMessage<NMT_Hello>::Send(NetDriver->ServerConnection, IsLittleEndian, LocalNetworkVersion);
 				NetDriver->ServerConnection->FlushNet();
 
 				bSuccess = true;
@@ -153,20 +157,11 @@ void AOnlineBeaconClient::NotifyControlMessage(UNetConnection* Connection, uint8
 		case NMT_Upgrade:
 			{
 				// Report mismatch.
-				int32 RemoteMinVer, RemoteVer;
-				FNetControlMessage<NMT_Upgrade>::Receive(Bunch, RemoteMinVer, RemoteVer);
-				if (GEngineNetVersion < RemoteMinVer)
-				{
-					// Upgrade
-					const FString ConnectionError = NSLOCTEXT("Engine", "ClientOutdated", "The match you are trying to join is running an incompatible version of the game.  Please try upgrading your game version.").ToString();
-					GEngine->BroadcastNetworkFailure(GetWorld(), NetDriver, ENetworkFailure::OutdatedClient, ConnectionError);
-				}
-				else
-				{
-					// Downgrade
-					const FString ConnectionError = NSLOCTEXT("Engine", "ServerOutdated", "Server's version is outdated").ToString();
-					GEngine->BroadcastNetworkFailure(GetWorld(), NetDriver, ENetworkFailure::OutdatedServer, ConnectionError);
-				}
+				uint32 RemoteNetworkVersion;
+				FNetControlMessage<NMT_Upgrade>::Receive(Bunch, RemoteNetworkVersion);
+				// Upgrade
+				const FString ConnectionError = NSLOCTEXT("Engine", "ClientOutdated", "The match you are trying to join is running an incompatible version of the game.  Please try upgrading your game version.").ToString();
+				GEngine->BroadcastNetworkFailure(GetWorld(), NetDriver, ENetworkFailure::OutdatedClient, ConnectionError);
 				break;
 			}
 		case NMT_Failure:
diff --git a/Engine/Source/Runtime/Online/OnlineSubsystemUtils/Private/OnlineBeaconHost.cpp b/Engine/Source/Runtime/Online/OnlineSubsystemUtils/Private/OnlineBeaconHost.cpp
--- a/Engine/Source/Runtime/Online/OnlineSubsystemUtils/Private/OnlineBeaconHost.cpp
+++ b/Engine/Source/Runtime/Online/OnlineSubsystemUtils/Private/OnlineBeaconHost.cpp
@@ -81,28 +81,19 @@ void AOnlineBeaconHost::NotifyControlMessage(UNetConnection* Connection, uint8 M
 			{
 				UE_LOG(LogNet, Log, TEXT("Beacon Hello"));
 				uint8 IsLittleEndian;
-				int32 RemoteMinVer, RemoteVer;
-				FGuid RemoteGameGUID;
-				FNetControlMessage<NMT_Hello>::Receive(Bunch, IsLittleEndian, RemoteMinVer, RemoteVer, RemoteGameGUID);
 
-				if (!IsNetworkCompatible(Connection->Driver->RequireEngineVersionMatch, RemoteVer, RemoteMinVer))
+				uint32 RemoteNetworkVersion = 0;
+				uint32 LocalNetworkVersion = FNetworkVersion::GetLocalNetworkVersion();
+
+				FNetControlMessage<NMT_Hello>::Receive(Bunch, IsLittleEndian, RemoteNetworkVersion);
+
+				if (!FNetworkVersion::IsNetworkCompatible(LocalNetworkVersion, RemoteNetworkVersion))
 				{
-					FNetControlMessage<NMT_Upgrade>::Send(Connection, GEngineMinNetVersion, GEngineNetVersion);
+					FNetControlMessage<NMT_Upgrade>::Send(Connection, LocalNetworkVersion);
 					bCloseConnection = true;
 				}
 				else
 				{
-					Connection->NegotiatedVer = FMath::Min(RemoteVer, GEngineNetVersion);
-
-					// Make sure the server has the same GameGUID as we do
-					if( RemoteGameGUID != GetDefault<UGeneralProjectSettings>()->ProjectID )
-					{
-						FString ErrorMsg = NSLOCTEXT("NetworkErrors", "ServerHostingDifferentGame", "Incompatible game connection.").ToString();
-						FNetControlMessage<NMT_Failure>::Send(Connection, ErrorMsg);
-						bCloseConnection = true;
-						break;
-					}
-
 					Connection->Challenge = FString::Printf(TEXT("%08X"), FPlatformTime::Cycles());
 					FNetControlMessage<NMT_BeaconWelcome>::Send(Connection);
 					Connection->FlushNet();
