Fix for jira OPP-1885, crash when calling FinishedRequest() on a request that's already been cleared.

diff --git a/Engine/Source/Runtime/Online/HTTP/Private/Mac/MacHTTP.cpp b/Engine/Source/Runtime/Online/HTTP/Private/Mac/MacHTTP.cpp
--- a/Engine/Source/Runtime/Online/HTTP/Private/Mac/MacHTTP.cpp
+++ b/Engine/Source/Runtime/Online/HTTP/Private/Mac/MacHTTP.cpp
@@ -274,7 +274,10 @@ void FMacHttpRequest::FinishedRequest()
 	CleanupRequest();
 
 	// Remove from global list since processing is now complete
-	FHttpModule::Get().GetHttpManager().RemoveRequest(SharedThis(this));
+	if (FHttpModule::Get().GetHttpManager().IsValidRequest(this))
+	{
+		FHttpModule::Get().GetHttpManager().RemoveRequest(SharedThis(this));
+	}
 }
 
 
@@ -417,7 +420,10 @@ FMacHttpResponse::~FMacHttpResponse()
 {
 	UE_LOG(LogHttp, Verbose, TEXT("FMacHttpResponse::~FMacHttpResponse()"));
 	[[ResponseWrapper Payload] release];
+	ResponseWrapper.Payload = nil;
+
 	[ResponseWrapper release];
+	ResponseWrapper = nil;
 }
 
 
@@ -497,6 +503,7 @@ const TArray<uint8>& FMacHttpResponse::GetContent()
 	}
 	else
 	{
+		// todo: Find out why we're empty and doing a memcopy every time instead of just once when we receive the payload
 		Payload.Empty();
 		Payload.AddZeroed( [[ResponseWrapper Payload] length] );
 	
