[INTEGRATE] UE4-LauncherReleases->UE4 integrate all up to cl 2392917

diff --git a/Engine/Source/Runtime/Online/HTTP/HTTP.Build.cs b/Engine/Source/Runtime/Online/HTTP/HTTP.Build.cs
--- a/Engine/Source/Runtime/Online/HTTP/HTTP.Build.cs
+++ b/Engine/Source/Runtime/Online/HTTP/HTTP.Build.cs
@@ -14,7 +14,12 @@ public class HTTP : ModuleRules
 			}
             );
 
-        PrivateDependencyModuleNames.AddRange(new string[] { "Core" });
+        PrivateDependencyModuleNames.AddRange(
+			new string[] { 
+				"Core",
+				"Analytics"
+			}
+			);
 
         if (Target.Platform == UnrealTargetPlatform.Win32 ||
             Target.Platform == UnrealTargetPlatform.Win64)
diff --git a/Engine/Source/Runtime/Online/HTTP/Private/Curl/CurlHttpManager.cpp b/Engine/Source/Runtime/Online/HTTP/Private/Curl/CurlHttpManager.cpp
--- a/Engine/Source/Runtime/Online/HTTP/Private/Curl/CurlHttpManager.cpp
+++ b/Engine/Source/Runtime/Online/HTTP/Private/Curl/CurlHttpManager.cpp
@@ -182,18 +182,20 @@ FCurlHttpManager::FCurlHttpManager()
 }
 
 // note that we cannot call parent implementation because lock might be possible non-multiple
-void FCurlHttpManager::AddRequest(TSharedRef<class IHttpRequest> Request)
+void FCurlHttpManager::AddRequest(const TSharedRef<class IHttpRequest>& Request)
 {
 	FScopeLock ScopeLock(&RequestLock);
 
-	Requests.AddUnique(Request);
+	Requests.Add(&Request.Get(), FActiveHttpRequest(FPlatformTime::Seconds(), Request));
 
 	FCurlHttpRequest* CurlRequest = static_cast< FCurlHttpRequest* >( &Request.Get() );
 	HandlesToRequests.Add(CurlRequest->GetEasyHandle(), Request);
+
+	
 }
 
 // note that we cannot call parent implementation because lock might be possible non-multiple
-void FCurlHttpManager::RemoveRequest(TSharedRef<class IHttpRequest> Request)
+void FCurlHttpManager::RemoveRequest(const TSharedRef<class IHttpRequest>& Request)
 {
 	FScopeLock ScopeLock(&RequestLock);
 
@@ -203,7 +205,13 @@ void FCurlHttpManager::RemoveRequest(TSharedRef<class IHttpRequest> Request)
 	FCurlHttpRequest* CurlRequest = static_cast< FCurlHttpRequest* >( &Request.Get() );
 	HandlesToRequests.Remove(CurlRequest->GetEasyHandle());
 
-	Requests.RemoveSingle(Request);
+	// add analytics info before removing the event
+	const FActiveHttpRequest* Found = Requests.Find(&Request.Get());
+	if (Found != NULL)
+	{
+		AddAnalytics(FPlatformTime::Seconds() - Found->StartTime, Request);
+	}
+	Requests.Remove(&Request.Get());
 }
 
 bool FCurlHttpManager::Tick(float DeltaSeconds)
diff --git a/Engine/Source/Runtime/Online/HTTP/Private/Curl/CurlHttpManager.h b/Engine/Source/Runtime/Online/HTTP/Private/Curl/CurlHttpManager.h
--- a/Engine/Source/Runtime/Online/HTTP/Private/Curl/CurlHttpManager.h
+++ b/Engine/Source/Runtime/Online/HTTP/Private/Curl/CurlHttpManager.h
@@ -28,8 +28,8 @@ protected:
 public:
 
 	// Begin HttpManager interface
-	virtual void AddRequest(TSharedRef<class IHttpRequest> Request) override;
-	virtual void RemoveRequest(TSharedRef<class IHttpRequest> Request) override;
+	virtual void AddRequest(const TSharedRef<class IHttpRequest>& Request) override;
+	virtual void RemoveRequest(const TSharedRef<class IHttpRequest>& Request) override;
 	virtual bool Tick(float DeltaSeconds) override;
 	// End HttpManager interface
 
diff --git a/Engine/Source/Runtime/Online/HTTP/Private/HttpManager.cpp b/Engine/Source/Runtime/Online/HTTP/Private/HttpManager.cpp
new file mode 100644
--- /dev/null
+++ b/Engine/Source/Runtime/Online/HTTP/Private/HttpManager.cpp
@@ -0,0 +1,161 @@
+// Copyright 1998-2014 Epic Games, Inc. All Rights Reserved.
+
+#include "HttpPrivatePCH.h"
+#include "HttpManager.h"
+#include "IAnalyticsProvider.h"
+
+// FHttpManager
+
+FCriticalSection FHttpManager::RequestLock;
+
+FHttpManager::FHttpManager()
+	:	FTickerObjectBase(0.0f)
+	,	DeferredDestroyDelay(10.0f)
+{
+	
+}
+
+FHttpManager::~FHttpManager()
+{
+	
+}
+
+bool FHttpManager::Tick(float DeltaSeconds)
+{
+	FScopeLock ScopeLock(&RequestLock);
+
+	// Tick each active request
+	for (auto It = Requests.CreateIterator(); It; ++It)
+	{
+		const FActiveHttpRequest& ActiveRequest = It.Value();
+		ActiveRequest.HttpRequest->Tick(DeltaSeconds);
+	}
+	// Tick any pending destroy objects
+	for (int Idx=0; Idx < PendingDestroyRequests.Num(); Idx++)
+	{
+		FRequestPendingDestroy& Request = PendingDestroyRequests[Idx];
+		Request.TimeLeft -= DeltaSeconds;
+		if (Request.TimeLeft <= 0)
+		{	
+			PendingDestroyRequests.RemoveAt(Idx--);
+		}		
+	}
+	// keep ticking
+	return true;
+}
+
+void FHttpManager::AddRequest(const TSharedRef<IHttpRequest>& Request)
+{
+	FScopeLock ScopeLock(&RequestLock);
+
+	Requests.Add(&Request.Get(), FActiveHttpRequest(FPlatformTime::Seconds(), Request));
+}
+
+void FHttpManager::RemoveRequest(const TSharedRef<IHttpRequest>& Request)
+{
+	FScopeLock ScopeLock(&RequestLock);
+
+	// Keep track of requests that have been removed to be destroyed later
+	PendingDestroyRequests.AddUnique(FRequestPendingDestroy(DeferredDestroyDelay,Request));
+
+	// add analytics info before removing the event
+	const FActiveHttpRequest* Found = Requests.Find(&Request.Get());
+	if (Found != NULL)
+	{
+		AddAnalytics(FPlatformTime::Seconds() - Found->StartTime, Request);
+	}
+	Requests.Remove(&Request.Get());
+}
+
+bool FHttpManager::IsValidRequest(const IHttpRequest* RequestPtr) const
+{
+	FScopeLock ScopeLock(&RequestLock);
+
+	return Requests.Contains(RequestPtr);
+}
+
+void FHttpManager::DumpRequests(FOutputDevice& Ar) const
+{
+	FScopeLock ScopeLock(&RequestLock);
+	
+	Ar.Logf(TEXT("------- (%d) Http Requests"), Requests.Num());
+	for (auto It = Requests.CreateConstIterator(); It; ++It)
+	{
+		const FActiveHttpRequest& ActiveRequest = It.Value();
+		Ar.Logf(TEXT("	verb=[%s] url=[%s] status=%s"),
+			*ActiveRequest.HttpRequest->GetVerb(), *ActiveRequest.HttpRequest->GetURL(), EHttpRequestStatus::ToString(ActiveRequest.HttpRequest->GetStatus()));
+	}
+}
+
+void FHttpManager::FAnalyticsHttpRequest::RecordRequest(double ElapsedTime, const TSharedRef<IHttpRequest>& HttpRequest)
+{
+	// keep a histogram of response codes
+	int32 ResponseCode = HttpRequest->GetResponse()->GetResponseCode();
+	int32* ResponseCodeEntry = ResponseCodes.Find(ResponseCode);
+	if (ResponseCodeEntry == NULL)
+	{
+		ResponseCodeEntry = &ResponseCodes.Add(ResponseCode, 0);
+	}
+	(*ResponseCodeEntry)++;
+	// sum elapsed time for average calc
+	ElapsedTimeTotal += ElapsedTime;
+	ElapsedTimeMax = ElapsedTimeMax > 0 ? FMath::Max<double>(ElapsedTimeMax, ElapsedTime) : ElapsedTime;
+	ElapsedTimeMin = ElapsedTimeMin > 0 ? FMath::Min<double>(ElapsedTimeMin, ElapsedTime) : ElapsedTime;
+	// sum download rate for average calc
+	const double DlSize = HttpRequest->GetResponse()->GetContent().Num() / 1024.f;
+	DownloadRateTotal += DlSize / ElapsedTime;
+	// keep track of total # of requests
+	Count++;
+}
+
+static FString GetUrlDomain(const FString& Url)
+{
+	FString Protocol;
+	FString Domain = Url;
+	// split the http protocol portion from domain
+	Url.Split(TEXT("://"), &Protocol, &Domain);
+	// strip off everything but the domain portion
+	int32 Idx = Domain.Find(TEXT("/"));
+	int32 IdxOpt = Domain.Find(TEXT("?"));
+	Idx = IdxOpt >= 0 && IdxOpt < Idx ? IdxOpt : Idx;
+	if (Idx > 0)
+	{
+		Domain = Domain.Left(Idx);
+	}
+	return Domain;
+}
+
+void FHttpManager::AddAnalytics(double ElapsedTime, const TSharedRef<IHttpRequest>& HttpRequest)
+{
+	// only record events that actually received a response
+	if (HttpRequest->GetResponse().IsValid())
+	{
+		// add/update analytics entry for this domain
+		FString Domain = GetUrlDomain(HttpRequest->GetURL());
+		FAnalyticsHttpRequest& AnalyticsEntry = RequestAnalytics.FindOrAdd(Domain);
+		AnalyticsEntry.RecordRequest(ElapsedTime, HttpRequest);
+	}
+}
+
+void FHttpManager::FlushAnalytics(class IAnalyticsProvider& Analytics)
+{
+	for (auto It = RequestAnalytics.CreateConstIterator(); It; ++It)
+	{
+		const FString& Domain = It.Key();
+		const FAnalyticsHttpRequest& AnalyticsEntry = It.Value();
+
+		TArray<FAnalyticsEventAttribute> Attributes;
+		Attributes.Add(FAnalyticsEventAttribute(TEXT("Domain"), Domain));
+		Attributes.Add(FAnalyticsEventAttribute(TEXT("AvgTime"), AnalyticsEntry.ElapsedTimeTotal / AnalyticsEntry.Count));
+		Attributes.Add(FAnalyticsEventAttribute(TEXT("MinTime"), AnalyticsEntry.ElapsedTimeMin));
+		Attributes.Add(FAnalyticsEventAttribute(TEXT("MaxTime"), AnalyticsEntry.ElapsedTimeMax));
+		Attributes.Add(FAnalyticsEventAttribute(TEXT("AvgRate"), AnalyticsEntry.DownloadRateTotal / AnalyticsEntry.Count));
+		Attributes.Add(FAnalyticsEventAttribute(TEXT("Hits"), AnalyticsEntry.Count));
+		for (auto CodeIt = AnalyticsEntry.ResponseCodes.CreateConstIterator(); CodeIt; ++CodeIt)
+		{
+			Attributes.Add(FAnalyticsEventAttribute(FString::Printf(TEXT("Code-%d"), CodeIt.Key()), CodeIt.Value()));
+		}
+		Analytics.RecordEvent(TEXT("Http.Services"), Attributes);
+	}
+	RequestAnalytics.Empty();
+}
diff --git a/Engine/Source/Runtime/Online/HTTP/Private/HttpModule.cpp b/Engine/Source/Runtime/Online/HTTP/Private/HttpModule.cpp
--- a/Engine/Source/Runtime/Online/HTTP/Private/HttpModule.cpp
+++ b/Engine/Source/Runtime/Online/HTTP/Private/HttpModule.cpp
@@ -11,8 +11,6 @@ IMPLEMENT_MODULE(FHttpModule, HTTP);
 
 FHttpModule* FHttpModule::Singleton = NULL;
 
-FCriticalSection FHttpManager::RequestLock;
-
 void FHttpModule::StartupModule()
 {	
 	Singleton = this;
@@ -73,7 +71,7 @@ void FHttpModule::ShutdownModule()
 	Singleton = nullptr;
 }
 
-bool FHttpModule::HandleHTTPCommand( const TCHAR* Cmd, FOutputDevice& Ar )
+bool FHttpModule::HandleHTTPCommand(const TCHAR* Cmd, FOutputDevice& Ar)
 {
 	if (FParse::Command(&Cmd, TEXT("TEST")))
 	{
@@ -127,88 +125,3 @@ TSharedRef<IHttpRequest> FHttpModule::CreateRequest()
 	return TSharedRef<IHttpRequest>( FPlatformHttp::ConstructRequest() );
 }
 
-// FHttpManager
-
-FHttpManager::FHttpManager()
-	:	FTickerObjectBase(0.0f)
-	,	DeferredDestroyDelay(10.0f)
-{
-	
-}
-
-FHttpManager::~FHttpManager()
-{
-	
-}
-
-bool FHttpManager::Tick(float DeltaSeconds)
-{
-	FScopeLock ScopeLock(&RequestLock);
-
-	// Tick each active request
-	for (TArray<TSharedRef<class IHttpRequest> >::TIterator It(Requests); It; ++It)
-	{
-		TSharedRef<class IHttpRequest> Request = *It;
-		Request->Tick(DeltaSeconds);
-	}
-	// Tick any pending destroy objects
-	for (int Idx=0; Idx < PendingDestroyRequests.Num(); Idx++)
-	{
-		FRequestPendingDestroy& Request = PendingDestroyRequests[Idx];
-		Request.TimeLeft -= DeltaSeconds;
-		if (Request.TimeLeft <= 0)
-		{	
-			PendingDestroyRequests.RemoveAt(Idx--);
-		}		
-	}
-	// keep ticking
-	return true;
-}
-
-void FHttpManager::AddRequest(TSharedRef<class IHttpRequest> Request)
-{
-	FScopeLock ScopeLock(&RequestLock);
-
-	Requests.AddUnique(Request);
-}
-
-void FHttpManager::RemoveRequest(TSharedRef<class IHttpRequest> Request)
-{
-	FScopeLock ScopeLock(&RequestLock);
-
-	// Keep track of requests that have been removed to be destroyed later
-	PendingDestroyRequests.AddUnique(FRequestPendingDestroy(DeferredDestroyDelay,Request));
-
-	Requests.RemoveSingle(Request);
-}
-
-bool FHttpManager::IsValidRequest(class IHttpRequest* RequestPtr)
-{
-	FScopeLock ScopeLock(&RequestLock);
-
-	bool bResult = false;
-	for (TArray<TSharedRef<class IHttpRequest> >::TConstIterator It(Requests); It; ++It)
-	{
-		const TSharedRef<class IHttpRequest>& Request = *It;
-
-		if (&Request.Get() == RequestPtr)
-		{
-			bResult = true;
-			break;
-		}
-	}
-	return bResult;
-}
-
-void FHttpManager::DumpRequests(FOutputDevice& Ar)
-{
-	FScopeLock ScopeLock(&RequestLock);
-	
-	Ar.Logf(TEXT("------- (%d) Http Requests"), Requests.Num());
-	for (TArray<TSharedRef<class IHttpRequest> >::TIterator It(Requests); It; ++It)
-	{
-		TSharedRef<class IHttpRequest> Request = *It;
-		Ar.Logf(TEXT("	verb=[%s] url=[%s] status=%s"),
-			*Request->GetVerb(), *Request->GetURL(), EHttpRequestStatus::ToString(Request->GetStatus()));
-	}
-}
diff --git a/Engine/Source/Runtime/Online/HTTP/Public/HttpManager.h b/Engine/Source/Runtime/Online/HTTP/Public/HttpManager.h
--- a/Engine/Source/Runtime/Online/HTTP/Public/HttpManager.h
+++ b/Engine/Source/Runtime/Online/HTTP/Public/HttpManager.h
@@ -30,15 +30,15 @@ public:
 	 *
 	 * @param Request - the request object to add
 	 */
-	virtual void AddRequest(TSharedRef<class IHttpRequest> Request);
+	virtual void AddRequest(const TSharedRef<IHttpRequest>& Request);
 
 	/**
-	 * Removes an Http request intance from the manager
+	 * Removes an Http request instance from the manager
 	 * Presumably it is done being processed
 	 *
 	 * @param Request - the request object to remove
 	 */
-	virtual void RemoveRequest(TSharedRef<class IHttpRequest> Request);
+	virtual void RemoveRequest(const TSharedRef<IHttpRequest>& Request);
 
 	/**
 	* Find an Http request in the lists of current valid requests
@@ -47,7 +47,7 @@ public:
 	*
 	* @return shared ptr to the request or invalid shared ptr
 	*/
-	virtual bool IsValidRequest(class IHttpRequest* RequestPtr);
+	virtual bool IsValidRequest(const IHttpRequest* RequestPtr) const;
 
 	/**
 	 * FTicker callback
@@ -63,35 +63,79 @@ public:
 	 *
 	 * @param Ar - output device to log with
 	 */
-	virtual void DumpRequests(FOutputDevice& Ar);
+	virtual void DumpRequests(FOutputDevice& Ar) const;
+
+	/**
+	 * Write http analytics events and clear them out
+	 *
+	 * @param Analytics - provider to use for writing events
+	 */
+	virtual void FlushAnalytics(class IAnalyticsProvider& Analytics);
 
 protected:
+	
+	/** Keep track of an http request while it is being processed */
+	class FActiveHttpRequest
+	{
+	public:
+		FActiveHttpRequest(double InStartTime, const TSharedRef<IHttpRequest>& InHttpRequest)
+			: StartTime(InStartTime)
+			, HttpRequest(InHttpRequest)
+		{}
+
+		double StartTime;
+		TSharedRef<IHttpRequest> HttpRequest;
+	};
 
 	/** List of Http requests that are actively being processed */
-	TArray<TSharedRef<class IHttpRequest> > Requests;
+	TMap<IHttpRequest*, FActiveHttpRequest> Requests;
 
 	/** Keep track of a request that should be deleted later */
-	struct FRequestPendingDestroy
+	class FRequestPendingDestroy
 	{
-		FRequestPendingDestroy(float InTimeLeft, TSharedPtr<class IHttpRequest> InHttpRequest)
+	public:
+		FRequestPendingDestroy(float InTimeLeft, const TSharedPtr<IHttpRequest>& InHttpRequest)
 			: TimeLeft(InTimeLeft)
 			, HttpRequest(InHttpRequest)
-		{
-
-		}
+		{}
 
-		FORCEINLINE bool operator==( const FRequestPendingDestroy& Other ) const
+		FORCEINLINE bool operator==(const FRequestPendingDestroy& Other) const
 		{
 			return Other.HttpRequest == HttpRequest;
 		}
 
 		float TimeLeft;
-		TSharedPtr<class IHttpRequest> HttpRequest;
+		TSharedPtr<IHttpRequest> HttpRequest;
 	};
 
 	/** Dead requests that need to be destroyed */
 	TArray<FRequestPendingDestroy> PendingDestroyRequests;
 
+	/** Analytics data to keep track of for completed http requests */
+	class FAnalyticsHttpRequest
+	{
+	public:
+		FAnalyticsHttpRequest()
+			: ElapsedTimeTotal(0)
+			, ElapsedTimeMin(0)
+			, ElapsedTimeMax(0)
+			, DownloadRateTotal(0)
+			, Count(0)
+		{}
+
+		void RecordRequest(double ElapsedTime, const TSharedRef<IHttpRequest>& HttpRequest);
+
+		double ElapsedTimeTotal;
+		double ElapsedTimeMin;
+		double ElapsedTimeMax;
+		double DownloadRateTotal;
+		int32 Count;
+		TMap<int32, int32> ResponseCodes;
+	};
+	TMap<FString, FAnalyticsHttpRequest> RequestAnalytics;
+
+	void AddAnalytics(double ElapsedTime, const TSharedRef<IHttpRequest>& HttpRequest);
+
 PACKAGE_SCOPE:
 
 	/** Used to lock access to add/remove/find requests */
