Xbox Live 2015 multiplayer implementation, contributed from Microsoft
Removed IOnlineParty since Xbox was the only platform implementing it and the 2015 APIs don't support it
IOnlineSession::StartMatchmaking now accepts an array of user IDs instead of just one
Xbox Live invitations are now implemented using protocol activation
OnlineSessionInterfaceLive now uses session subscriptions to receive session change notifications
And many other changes to support the 2015 APIs

diff --git a/Engine/Source/Runtime/Online/Android/OnlineSubsystemGooglePlay/Private/OnlineSubsystemGooglePlay.cpp b/Engine/Source/Runtime/Online/Android/OnlineSubsystemGooglePlay/Private/OnlineSubsystemGooglePlay.cpp
--- a/Engine/Source/Runtime/Online/Android/OnlineSubsystemGooglePlay/Private/OnlineSubsystemGooglePlay.cpp
+++ b/Engine/Source/Runtime/Online/Android/OnlineSubsystemGooglePlay/Private/OnlineSubsystemGooglePlay.cpp
@@ -84,11 +84,6 @@ IOnlineTimePtr FOnlineSubsystemGooglePlay::GetTimeInterface() const
 	return nullptr;
 }
 
-IOnlinePartyPtr FOnlineSubsystemGooglePlay::GetPartyInterface() const
-{
-	return nullptr;
-}
-
 IOnlineTitleFilePtr FOnlineSubsystemGooglePlay::GetTitleFileInterface() const
 {
 	return nullptr;
diff --git a/Engine/Source/Runtime/Online/Android/OnlineSubsystemGooglePlay/Public/OnlineSubsystemGooglePlay.h b/Engine/Source/Runtime/Online/Android/OnlineSubsystemGooglePlay/Public/OnlineSubsystemGooglePlay.h
--- a/Engine/Source/Runtime/Online/Android/OnlineSubsystemGooglePlay/Public/OnlineSubsystemGooglePlay.h
+++ b/Engine/Source/Runtime/Online/Android/OnlineSubsystemGooglePlay/Public/OnlineSubsystemGooglePlay.h
@@ -36,7 +36,6 @@ public:
 	virtual IOnlineExternalUIPtr GetExternalUIInterface() const override;
 	virtual IOnlineTimePtr GetTimeInterface() const override;
 	virtual IOnlineIdentityPtr GetIdentityInterface() const override;
-	virtual IOnlinePartyPtr GetPartyInterface() const override;
 	virtual IOnlineTitleFilePtr GetTitleFileInterface() const override;
 	virtual IOnlineEntitlementsPtr GetEntitlementsInterface() const override;
 	virtual IOnlineStorePtr GetStoreInterface() const override;
diff --git a/Engine/Source/Runtime/Online/IOS/OnlineSubsystemIOS/Private/OnlineSessionInterfaceIOS.cpp b/Engine/Source/Runtime/Online/IOS/OnlineSubsystemIOS/Private/OnlineSessionInterfaceIOS.cpp
--- a/Engine/Source/Runtime/Online/IOS/OnlineSubsystemIOS/Private/OnlineSessionInterfaceIOS.cpp
+++ b/Engine/Source/Runtime/Online/IOS/OnlineSubsystemIOS/Private/OnlineSessionInterfaceIOS.cpp
@@ -319,15 +319,7 @@ bool FOnlineSessionIOS::IsPlayerInSession(FName SessionName, const FUniqueNetId&
 	return IsPlayerInSessionImpl(this, SessionName, UniqueId);
 }
 
-bool FOnlineSessionIOS::StartMatchmaking(int32 SearchingPlayerNum, FName SessionName, const FOnlineSessionSettings& NewSessionSettings, TSharedRef<FOnlineSessionSearch>& SearchSettings)
-{
-	UE_LOG(LogOnline, Warning, TEXT("Matchmaking is not supported on this platform."));
-	TriggerOnMatchmakingCompleteDelegates(SessionName, false);
-	return false;
-}
-
-
-bool FOnlineSessionIOS::StartMatchmaking(const FUniqueNetId& SearchingPlayerId, FName SessionName, const FOnlineSessionSettings& NewSessionSettings, TSharedRef<FOnlineSessionSearch>& SearchSettings)
+bool FOnlineSessionIOS::StartMatchmaking(const TArray< TSharedRef<FUniqueNetId> >& LocalPlayers, FName SessionName, const FOnlineSessionSettings& NewSessionSettings, TSharedRef<FOnlineSessionSearch>& SearchSettings)
 {
 	UE_LOG(LogOnline, Warning, TEXT("Matchmaking is not supported on this platform."));
 	TriggerOnMatchmakingCompleteDelegates(SessionName, false);
diff --git a/Engine/Source/Runtime/Online/IOS/OnlineSubsystemIOS/Private/OnlineSessionInterfaceIOS.h b/Engine/Source/Runtime/Online/IOS/OnlineSubsystemIOS/Private/OnlineSessionInterfaceIOS.h
--- a/Engine/Source/Runtime/Online/IOS/OnlineSubsystemIOS/Private/OnlineSessionInterfaceIOS.h
+++ b/Engine/Source/Runtime/Online/IOS/OnlineSubsystemIOS/Private/OnlineSessionInterfaceIOS.h
@@ -92,9 +92,7 @@ public:
 
 	virtual bool IsPlayerInSession(FName SessionName, const FUniqueNetId& UniqueId) override;
 
-	virtual bool StartMatchmaking(int32 SearchingPlayerNum, FName SessionName, const FOnlineSessionSettings& NewSessionSettings, TSharedRef<FOnlineSessionSearch>& SearchSettings) override;
-
-	virtual bool StartMatchmaking(const FUniqueNetId& SearchingPlayerId, FName SessionName, const FOnlineSessionSettings& NewSessionSettings, TSharedRef<FOnlineSessionSearch>& SearchSettings) override;
+	virtual bool StartMatchmaking(const TArray< TSharedRef<FUniqueNetId> >& LocalPlayers, FName SessionName, const FOnlineSessionSettings& NewSessionSettings, TSharedRef<FOnlineSessionSearch>& SearchSettings) override;
 
 	virtual bool CancelMatchmaking(int32 SearchingPlayerNum, FName SessionName) override;
 
diff --git a/Engine/Source/Runtime/Online/IOS/OnlineSubsystemIOS/Private/OnlineSubsystemIOS.cpp b/Engine/Source/Runtime/Online/IOS/OnlineSubsystemIOS/Private/OnlineSubsystemIOS.cpp
--- a/Engine/Source/Runtime/Online/IOS/OnlineSubsystemIOS/Private/OnlineSubsystemIOS.cpp
+++ b/Engine/Source/Runtime/Online/IOS/OnlineSubsystemIOS/Private/OnlineSubsystemIOS.cpp
@@ -62,12 +62,6 @@ IOnlineIdentityPtr FOnlineSubsystemIOS::GetIdentityInterface() const
 }
 
 
-IOnlinePartyPtr FOnlineSubsystemIOS::GetPartyInterface() const
-{
-	return NULL;
-}
-
-
 IOnlineTitleFilePtr FOnlineSubsystemIOS::GetTitleFileInterface() const
 {
 	return NULL;
diff --git a/Engine/Source/Runtime/Online/IOS/OnlineSubsystemIOS/Public/OnlineSubsystemIOS.h b/Engine/Source/Runtime/Online/IOS/OnlineSubsystemIOS/Public/OnlineSubsystemIOS.h
--- a/Engine/Source/Runtime/Online/IOS/OnlineSubsystemIOS/Public/OnlineSubsystemIOS.h
+++ b/Engine/Source/Runtime/Online/IOS/OnlineSubsystemIOS/Public/OnlineSubsystemIOS.h
@@ -29,7 +29,6 @@ public:
 	virtual IOnlineExternalUIPtr GetExternalUIInterface() const override;
 	virtual IOnlineTimePtr GetTimeInterface() const override;
 	virtual IOnlineIdentityPtr GetIdentityInterface() const override;
-	virtual IOnlinePartyPtr GetPartyInterface() const override;
 	virtual IOnlineTitleFilePtr GetTitleFileInterface() const override;
 	virtual IOnlineEntitlementsPtr GetEntitlementsInterface() const override;
 	virtual IOnlineStorePtr GetStoreInterface() const override;
diff --git a/Engine/Source/Runtime/Online/OnlineSubsystem/Public/Interfaces/OnlinePartyInterface.h b/Engine/Source/Runtime/Online/OnlineSubsystem/Public/Interfaces/OnlinePartyInterface.h
deleted file mode 100644
diff --git a/Engine/Source/Runtime/Online/OnlineSubsystem/Public/Interfaces/OnlineSessionInterface.h b/Engine/Source/Runtime/Online/OnlineSubsystem/Public/Interfaces/OnlineSessionInterface.h
--- a/Engine/Source/Runtime/Online/OnlineSubsystem/Public/Interfaces/OnlineSessionInterface.h
+++ b/Engine/Source/Runtime/Online/OnlineSubsystem/Public/Interfaces/OnlineSessionInterface.h
@@ -214,6 +214,25 @@ DECLARE_DELEGATE_TwoParams(FOnRegisterLocalPlayerCompleteDelegate, const FUnique
  */
 DECLARE_DELEGATE_TwoParams(FOnUnregisterLocalPlayerCompleteDelegate, const FUniqueNetId&, const bool);	
 
+/** Possible reasons for the service to cause a session failure */
+namespace ESessionFailure
+{
+	enum Type
+	{
+		/** General loss of connection */
+		ServiceConnectionLost
+	};
+}
+
+/**
+ * Delegate fired when an unexpected error occurs that impacts session connectivity or use
+ *
+ * @param PlayerId The player impacted by the failure (may be empty if unknown, or if all players are affected)
+ * @param FailureType What kind of failure occurred 
+ */
+DECLARE_MULTICAST_DELEGATE_TwoParams(FOnSessionFailure, const FUniqueNetId&, ESessionFailure::Type);
+typedef FOnSessionFailure::FDelegate FOnSessionFailureDelegate;
+
 /**
  * Interface definition for the online services session services 
  * Session services are defined as anything related managing a session 
@@ -396,28 +415,15 @@ public:
 	/**
 	 * Begins cloud based matchmaking for a session
 	 *
-	 * @param SearchingPlayerNum the index of the player searching for a match
-	 * @param SessionName the name of the session to use, usually will be GAME_SESSION_NAME
-	 * @param NewSessionSettings the desired settings to match against or create with when forming new sessions
-	 * @param SearchSettings the desired settings that the matched session will have
-	 *
-	 * @return true if successful searching for sessions, false otherwise
-	 */
-	virtual bool StartMatchmaking(int32 SearchingPlayerNum, FName SessionName, const FOnlineSessionSettings& NewSessionSettings, TSharedRef<FOnlineSessionSearch>& SearchSettings) = 0;
-
-	/**
-	 * Begins cloud based matchmaking for a session
-	 *
-	 * @param SearchingPlayerId the id of the player searching for a match
+	 * @param LocalPlayers the ids of all local players that will participate in the match
 	 * @param SessionName the name of the session to use, usually will be GAME_SESSION_NAME
 	 * @param NewSessionSettings the desired settings to match against or create with when forming new sessions
 	 * @param SearchSettings the desired settings that the matched session will have
 	 *
 	 * @return true if successful searching for sessions, false otherwise
 	 */
-	virtual bool StartMatchmaking(const FUniqueNetId& SearchingPlayerId, FName SessionName, const FOnlineSessionSettings& NewSessionSettings, TSharedRef<FOnlineSessionSearch>& SearchSettings) = 0;
-
-
+	virtual bool StartMatchmaking(const TArray< TSharedRef<FUniqueNetId> >& LocalPlayers, FName SessionName, const FOnlineSessionSettings& NewSessionSettings, TSharedRef<FOnlineSessionSearch>& SearchSettings) = 0;
+	
 	/**
 	 * Delegate fired when the cloud matchmaking has completed
 	 *
@@ -767,6 +773,14 @@ public:
 	 */
 	virtual void UnregisterLocalPlayer(const FUniqueNetId& PlayerId, FName SessionName, const FOnUnregisterLocalPlayerCompleteDelegate& Delegate) = 0;
 
+	/**
+	 * Delegate fired when an unexpected error occurs that impacts session connectivity or use
+	 *
+	 * @param PlayerId The player impacted by the failure (may be empty if unknown, or if all players are affected)
+	 * @param FailureType What kind of failure occured 
+	 */
+	DEFINE_ONLINE_DELEGATE_TWO_PARAM(OnSessionFailure, const FUniqueNetId&, ESessionFailure::Type);
+
 	/**
 	 * Gets the number of known sessions registered with the interface
 	 */
diff --git a/Engine/Source/Runtime/Online/OnlineSubsystem/Public/Online.h b/Engine/Source/Runtime/Online/OnlineSubsystem/Public/Online.h
--- a/Engine/Source/Runtime/Online/OnlineSubsystem/Public/Online.h
+++ b/Engine/Source/Runtime/Online/OnlineSubsystem/Public/Online.h
@@ -17,7 +17,6 @@
 #include "OnlinePresenceInterface.h"
 #include "OnlineUserCloudInterface.h"
 #include "OnlineUserInterface.h"
-#include "OnlinePartyInterface.h"
 #include "OnlineFriendsInterface.h" //Needed for online friend functionality
 
 /** Macro to handle the boilerplate of accessing the proper online subsystem and getting the requested interface */
@@ -91,14 +90,7 @@ namespace Online
 	 * @return Interface pointer for the appropriate identity service
 	 */
 	IMPLEMENT_GET_INTERFACE(Identity);
-
-	/** 
-	 * Get the interface for accessing party services
-	 * @param SubsystemName - Name of the requested online service
-	 * @return Interface pointer for the appropriate party service
-	 */
-	IMPLEMENT_GET_INTERFACE(Party);
-
+	
 	/** 
 	 * Get the interface for accessing title file online services
 	 * @param SubsystemName - Name of the requested online service
diff --git a/Engine/Source/Runtime/Online/OnlineSubsystem/Public/OnlineSubsystem.h b/Engine/Source/Runtime/Online/OnlineSubsystem/Public/OnlineSubsystem.h
--- a/Engine/Source/Runtime/Online/OnlineSubsystem/Public/OnlineSubsystem.h
+++ b/Engine/Source/Runtime/Online/OnlineSubsystem/Public/OnlineSubsystem.h
@@ -49,7 +49,6 @@ typedef TSharedPtr<class IOnlineSharing, ESPMode::ThreadSafe> IOnlineSharingPtr;
 typedef TSharedPtr<class IOnlineUser, ESPMode::ThreadSafe> IOnlineUserPtr;
 typedef TSharedPtr<class IOnlineMessage, ESPMode::ThreadSafe> IOnlineMessagePtr;
 typedef TSharedPtr<class IOnlinePresence, ESPMode::ThreadSafe> IOnlinePresencePtr;
-typedef TSharedPtr<class IOnlineParty, ESPMode::ThreadSafe> IOnlinePartyPtr;
 typedef TSharedPtr<class IOnlineChat, ESPMode::ThreadSafe> IOnlineChatPtr;
 typedef TSharedPtr<class FOnlineNotificationHandler, ESPMode::ThreadSafe> FOnlineNotificationHandlerPtr;
 typedef TSharedPtr<class FOnlineNotificationTransportManager, ESPMode::ThreadSafe> FOnlineNotificationTransportManagerPtr;
@@ -191,12 +190,6 @@ public:
 	 */
 	virtual IOnlineIdentityPtr GetIdentityInterface() const = 0;
 
-	/** 
-	 * Get the interface for accessing party online services
-	 * @return Interface pointer for the appropriate party service
-	 */
-	virtual IOnlinePartyPtr GetPartyInterface() const = 0;
-
 	/** 
 	 * Get the interface for accessing title file online services
 	 * @return Interface pointer for the appropriate title file service
diff --git a/Engine/Source/Runtime/Online/OnlineSubsystem/Public/OnlineSubsystemTypes.h b/Engine/Source/Runtime/Online/OnlineSubsystem/Public/OnlineSubsystemTypes.h
--- a/Engine/Source/Runtime/Online/OnlineSubsystem/Public/OnlineSubsystemTypes.h
+++ b/Engine/Source/Runtime/Online/OnlineSubsystem/Public/OnlineSubsystemTypes.h
@@ -1138,37 +1138,6 @@ public:
 	virtual FDateTime GetLastSeen() const = 0;
 };
 
-/**
- * Party user info returned via IOnlineParty interface
- */
-struct FOnlinePartyMember
-{
-	/** The unique ID of this party member */
-	TSharedPtr<FUniqueNetId>	UniqueId;
-
-	/** True if this party member is using the local console */
-	bool						IsLocal;
-
-	/** The time at which the player joined the party */
-	FDateTime					JoinTime;
-
-	/** If the platform supports grouping party members by device, this will be set */
-	int32						DeviceGroup;
-
-	/** Constructor */
-	FOnlinePartyMember(
-		const TSharedPtr<FUniqueNetId>& PartyMemberId,
-		const bool IsLocalPartyMember,
-		const FDateTime& PartyJoinTime,
-		const int32 InDeviceGroup
-		)
-		: UniqueId(PartyMemberId)
-		, IsLocal(IsLocalPartyMember)
-		, JoinTime(PartyJoinTime)
-		, DeviceGroup(InDeviceGroup)
-	{}
-};
-
 /** The possible permission categories we can choose from to read from the server */
 namespace EOnlineSharingReadCategory
 {
diff --git a/Engine/Source/Runtime/Online/OnlineSubsystemAmazon/Public/OnlineSubsystemAmazon.h b/Engine/Source/Runtime/Online/OnlineSubsystemAmazon/Public/OnlineSubsystemAmazon.h
--- a/Engine/Source/Runtime/Online/OnlineSubsystemAmazon/Public/OnlineSubsystemAmazon.h
+++ b/Engine/Source/Runtime/Online/OnlineSubsystemAmazon/Public/OnlineSubsystemAmazon.h
@@ -66,10 +66,6 @@ public:
 	{
 		return NULL;
 	}
-	virtual IOnlinePartyPtr GetPartyInterface() const override
-	{
-		return NULL;
-	}
 	virtual IOnlineTitleFilePtr GetTitleFileInterface() const override
 	{
 		return NULL;
diff --git a/Engine/Source/Runtime/Online/OnlineSubsystemFacebook/Private/IOS/OnlineSubsystemFacebook.cpp b/Engine/Source/Runtime/Online/OnlineSubsystemFacebook/Private/IOS/OnlineSubsystemFacebook.cpp
--- a/Engine/Source/Runtime/Online/OnlineSubsystemFacebook/Private/IOS/OnlineSubsystemFacebook.cpp
+++ b/Engine/Source/Runtime/Online/OnlineSubsystemFacebook/Private/IOS/OnlineSubsystemFacebook.cpp
@@ -111,11 +111,6 @@ IOnlinePresencePtr FOnlineSubsystemFacebook::GetPresenceInterface() const
 	return NULL;
 }
 
-IOnlinePartyPtr FOnlineSubsystemFacebook::GetPartyInterface() const
-{
-	return NULL;
-}
-
 IOnlineChatPtr FOnlineSubsystemFacebook::GetChatInterface() const
 {
 	return NULL;
diff --git a/Engine/Source/Runtime/Online/OnlineSubsystemFacebook/Private/Windows/OnlineSubsystemFacebook.cpp b/Engine/Source/Runtime/Online/OnlineSubsystemFacebook/Private/Windows/OnlineSubsystemFacebook.cpp
--- a/Engine/Source/Runtime/Online/OnlineSubsystemFacebook/Private/Windows/OnlineSubsystemFacebook.cpp
+++ b/Engine/Source/Runtime/Online/OnlineSubsystemFacebook/Private/Windows/OnlineSubsystemFacebook.cpp
@@ -109,11 +109,6 @@ IOnlinePresencePtr FOnlineSubsystemFacebook::GetPresenceInterface() const
 	return NULL;
 }
 
-IOnlinePartyPtr FOnlineSubsystemFacebook::GetPartyInterface() const
-{
-	return NULL;
-}
-
 IOnlineChatPtr FOnlineSubsystemFacebook::GetChatInterface() const
 {
 	return NULL;
diff --git a/Engine/Source/Runtime/Online/OnlineSubsystemFacebook/Public/OnlineSubsystemFacebook.h b/Engine/Source/Runtime/Online/OnlineSubsystemFacebook/Public/OnlineSubsystemFacebook.h
--- a/Engine/Source/Runtime/Online/OnlineSubsystemFacebook/Public/OnlineSubsystemFacebook.h
+++ b/Engine/Source/Runtime/Online/OnlineSubsystemFacebook/Public/OnlineSubsystemFacebook.h
@@ -31,7 +31,6 @@ public:
 	virtual IOnlineExternalUIPtr GetExternalUIInterface() const override;
 	virtual IOnlineTimePtr GetTimeInterface() const override;
 	virtual IOnlineIdentityPtr GetIdentityInterface() const override;
-	virtual IOnlinePartyPtr GetPartyInterface() const override;
 	virtual IOnlineTitleFilePtr GetTitleFileInterface() const override;
 	virtual IOnlineEntitlementsPtr GetEntitlementsInterface() const override;
 	virtual IOnlineStorePtr GetStoreInterface() const override;
diff --git a/Engine/Source/Runtime/Online/OnlineSubsystemNull/Private/OnlineSessionInterfaceNull.cpp b/Engine/Source/Runtime/Online/OnlineSubsystemNull/Private/OnlineSessionInterfaceNull.cpp
--- a/Engine/Source/Runtime/Online/OnlineSubsystemNull/Private/OnlineSessionInterfaceNull.cpp
+++ b/Engine/Source/Runtime/Online/OnlineSubsystemNull/Private/OnlineSessionInterfaceNull.cpp
@@ -447,14 +447,7 @@ bool FOnlineSessionNull::IsPlayerInSession(FName SessionName, const FUniqueNetId
 	return IsPlayerInSessionImpl(this, SessionName, UniqueId);
 }
 
-bool FOnlineSessionNull::StartMatchmaking(int32 SearchingPlayerNum, FName SessionName, const FOnlineSessionSettings& NewSessionSettings, TSharedRef<FOnlineSessionSearch>& SearchSettings)
-{
-	UE_LOG(LogOnline, Warning, TEXT("Matchmaking is not supported on this platform."));
-	TriggerOnMatchmakingCompleteDelegates(SessionName, false);
-	return false;
-}
-
-bool FOnlineSessionNull::StartMatchmaking(const FUniqueNetId& SearchingPlayerId, FName SessionName, const FOnlineSessionSettings& NewSessionSettings, TSharedRef<FOnlineSessionSearch>& SearchSettings)
+bool FOnlineSessionNull::StartMatchmaking(const TArray< TSharedRef<FUniqueNetId> >& LocalPlayers, FName SessionName, const FOnlineSessionSettings& NewSessionSettings, TSharedRef<FOnlineSessionSearch>& SearchSettings)
 {
 	UE_LOG(LogOnline, Warning, TEXT("Matchmaking is not supported on this platform."));
 	TriggerOnMatchmakingCompleteDelegates(SessionName, false);
diff --git a/Engine/Source/Runtime/Online/OnlineSubsystemNull/Private/OnlineSessionInterfaceNull.h b/Engine/Source/Runtime/Online/OnlineSubsystemNull/Private/OnlineSessionInterfaceNull.h
--- a/Engine/Source/Runtime/Online/OnlineSubsystemNull/Private/OnlineSessionInterfaceNull.h
+++ b/Engine/Source/Runtime/Online/OnlineSubsystemNull/Private/OnlineSessionInterfaceNull.h
@@ -279,8 +279,7 @@ public:
 	virtual bool EndSession(FName SessionName) override;
 	virtual bool DestroySession(FName SessionName) override;
 	virtual bool IsPlayerInSession(FName SessionName, const FUniqueNetId& UniqueId) override;
-	virtual bool StartMatchmaking(int32 SearchingPlayerNum, FName SessionName, const FOnlineSessionSettings& NewSessionSettings, TSharedRef<FOnlineSessionSearch>& SearchSettings) override;
-	virtual bool StartMatchmaking(const FUniqueNetId& SearchingPlayerId, FName SessionName, const FOnlineSessionSettings& NewSessionSettings, TSharedRef<FOnlineSessionSearch>& SearchSettings) override;
+	virtual bool StartMatchmaking(const TArray< TSharedRef<FUniqueNetId> >& LocalPlayers, FName SessionName, const FOnlineSessionSettings& NewSessionSettings, TSharedRef<FOnlineSessionSearch>& SearchSettings) override;
 	virtual bool CancelMatchmaking(int32 SearchingPlayerNum, FName SessionName) override;
 	virtual bool CancelMatchmaking(const FUniqueNetId& SearchingPlayerId, FName SessionName) override;
 	virtual bool FindSessions(int32 SearchingPlayerNum, const TSharedRef<FOnlineSessionSearch>& SearchSettings) override;
diff --git a/Engine/Source/Runtime/Online/OnlineSubsystemNull/Private/OnlineSubsystemNull.cpp b/Engine/Source/Runtime/Online/OnlineSubsystemNull/Private/OnlineSubsystemNull.cpp
--- a/Engine/Source/Runtime/Online/OnlineSubsystemNull/Private/OnlineSubsystemNull.cpp
+++ b/Engine/Source/Runtime/Online/OnlineSubsystemNull/Private/OnlineSubsystemNull.cpp
@@ -63,11 +63,6 @@ IOnlineIdentityPtr FOnlineSubsystemNull::GetIdentityInterface() const
 	return IdentityInterface;
 }
 
-IOnlinePartyPtr FOnlineSubsystemNull::GetPartyInterface() const
-{
-	return NULL;
-}
-
 IOnlineTitleFilePtr FOnlineSubsystemNull::GetTitleFileInterface() const
 {
 	return NULL;
diff --git a/Engine/Source/Runtime/Online/OnlineSubsystemNull/Public/OnlineSubsystemNull.h b/Engine/Source/Runtime/Online/OnlineSubsystemNull/Public/OnlineSubsystemNull.h
--- a/Engine/Source/Runtime/Online/OnlineSubsystemNull/Public/OnlineSubsystemNull.h
+++ b/Engine/Source/Runtime/Online/OnlineSubsystemNull/Public/OnlineSubsystemNull.h
@@ -43,7 +43,6 @@ public:
 	virtual IOnlineExternalUIPtr GetExternalUIInterface() const override;	
 	virtual IOnlineTimePtr GetTimeInterface() const override;
 	virtual IOnlineIdentityPtr GetIdentityInterface() const override;
-	virtual IOnlinePartyPtr GetPartyInterface() const override;
 	virtual IOnlineTitleFilePtr GetTitleFileInterface() const override;
 	virtual IOnlineStorePtr GetStoreInterface() const override;
 	virtual IOnlineEventsPtr GetEventsInterface() const override;
diff --git a/Engine/Source/Runtime/Online/OnlineSubsystemSteam/Private/OnlineSessionInterfaceSteam.cpp b/Engine/Source/Runtime/Online/OnlineSubsystemSteam/Private/OnlineSessionInterfaceSteam.cpp
--- a/Engine/Source/Runtime/Online/OnlineSubsystemSteam/Private/OnlineSessionInterfaceSteam.cpp
+++ b/Engine/Source/Runtime/Online/OnlineSubsystemSteam/Private/OnlineSessionInterfaceSteam.cpp
@@ -640,14 +640,7 @@ bool FOnlineSessionSteam::IsPlayerInSession(FName SessionName, const FUniqueNetI
 	return IsPlayerInSessionImpl(this, SessionName, UniqueId);
 }
 
-bool FOnlineSessionSteam::StartMatchmaking(int32 SearchingPlayerNum, FName SessionName, const FOnlineSessionSettings& NewSessionSettings, TSharedRef<FOnlineSessionSearch>& SearchSettings)
-{
-	UE_LOG(LogOnline, Warning, TEXT("Matchmaking is not supported on this platform."));
-	TriggerOnMatchmakingCompleteDelegates(SessionName, false);
-	return false;
-}
-
-bool FOnlineSessionSteam::StartMatchmaking(const FUniqueNetId& SearchingPlayerId, FName SessionName, const FOnlineSessionSettings& NewSessionSettings, TSharedRef<FOnlineSessionSearch>& SearchSettings)
+bool FOnlineSessionSteam::StartMatchmaking(const TArray< TSharedRef<FUniqueNetId> >& LocalPlayers, FName SessionName, const FOnlineSessionSettings& NewSessionSettings, TSharedRef<FOnlineSessionSearch>& SearchSettings)
 {
 	UE_LOG(LogOnline, Warning, TEXT("Matchmaking is not supported on this platform."));
 	TriggerOnMatchmakingCompleteDelegates(SessionName, false);
diff --git a/Engine/Source/Runtime/Online/OnlineSubsystemSteam/Private/OnlineSessionInterfaceSteam.h b/Engine/Source/Runtime/Online/OnlineSubsystemSteam/Private/OnlineSessionInterfaceSteam.h
--- a/Engine/Source/Runtime/Online/OnlineSubsystemSteam/Private/OnlineSessionInterfaceSteam.h
+++ b/Engine/Source/Runtime/Online/OnlineSubsystemSteam/Private/OnlineSessionInterfaceSteam.h
@@ -496,8 +496,7 @@ public:
 	virtual bool EndSession(FName SessionName) override;
 	virtual bool DestroySession(FName SessionName) override;
 	virtual bool IsPlayerInSession(FName SessionName, const FUniqueNetId& UniqueId) override;
-	virtual bool StartMatchmaking(int32 SearchingPlayerNum, FName SessionName, const FOnlineSessionSettings& NewSessionSettings, TSharedRef<FOnlineSessionSearch>& SearchSettings) override;
-	virtual bool StartMatchmaking(const FUniqueNetId& SearchingPlayerId, FName SessionName, const FOnlineSessionSettings& NewSessionSettings, TSharedRef<FOnlineSessionSearch>& SearchSettings) override;
+	virtual bool StartMatchmaking(const TArray< TSharedRef<FUniqueNetId> >& LocalPlayers, FName SessionName, const FOnlineSessionSettings& NewSessionSettings, TSharedRef<FOnlineSessionSearch>& SearchSettings) override;
 	virtual bool CancelMatchmaking(int32 SearchingPlayerNum, FName SessionName) override;
 	virtual bool CancelMatchmaking(const FUniqueNetId& SearchingPlayerId, FName SessionName) override;
 	virtual bool FindSessions(int32 SearchingPlayerNum, const TSharedRef<FOnlineSessionSearch>& SearchSettings) override;
diff --git a/Engine/Source/Runtime/Online/OnlineSubsystemSteam/Private/OnlineSubsystemSteam.cpp b/Engine/Source/Runtime/Online/OnlineSubsystemSteam/Private/OnlineSubsystemSteam.cpp
--- a/Engine/Source/Runtime/Online/OnlineSubsystemSteam/Private/OnlineSubsystemSteam.cpp
+++ b/Engine/Source/Runtime/Online/OnlineSubsystemSteam/Private/OnlineSubsystemSteam.cpp
@@ -199,11 +199,6 @@ IOnlineIdentityPtr FOnlineSubsystemSteam::GetIdentityInterface() const
 	return IdentityInterface;
 }
 
-IOnlinePartyPtr FOnlineSubsystemSteam::GetPartyInterface() const
-{
-	return NULL;
-}
-
 IOnlineTitleFilePtr FOnlineSubsystemSteam::GetTitleFileInterface() const
 {
 	return NULL;
diff --git a/Engine/Source/Runtime/Online/OnlineSubsystemSteam/Public/OnlineSubsystemSteam.h b/Engine/Source/Runtime/Online/OnlineSubsystemSteam/Public/OnlineSubsystemSteam.h
--- a/Engine/Source/Runtime/Online/OnlineSubsystemSteam/Public/OnlineSubsystemSteam.h
+++ b/Engine/Source/Runtime/Online/OnlineSubsystemSteam/Public/OnlineSubsystemSteam.h
@@ -214,7 +214,6 @@ public:
 	virtual IOnlineExternalUIPtr GetExternalUIInterface() const override;
 	virtual IOnlineTimePtr GetTimeInterface() const override;
 	virtual IOnlineIdentityPtr GetIdentityInterface() const override;
-	virtual IOnlinePartyPtr GetPartyInterface() const override;
 	virtual IOnlineTitleFilePtr GetTitleFileInterface() const override;
 	virtual IOnlineEntitlementsPtr GetEntitlementsInterface() const override;
 	virtual IOnlineStorePtr GetStoreInterface() const override;
