Delegate comparisons deprecated, lots of other associated code deprecated, and lots of warning fixups:

diff --git a/Engine/Source/Runtime/Online/OnlineSubsystemUtils/Private/InAppPurchaseQueryCallbackProxy.cpp b/Engine/Source/Runtime/Online/OnlineSubsystemUtils/Private/InAppPurchaseQueryCallbackProxy.cpp
--- a/Engine/Source/Runtime/Online/OnlineSubsystemUtils/Private/InAppPurchaseQueryCallbackProxy.cpp
+++ b/Engine/Source/Runtime/Online/OnlineSubsystemUtils/Private/InAppPurchaseQueryCallbackProxy.cpp
@@ -25,8 +25,8 @@ void UInAppPurchaseQueryCallbackProxy::TriggerQuery(APlayerController* PlayerCon
 				bFailedToEvenSubmit = false;
 
 				// Register the completion callback
-				InAppPurchaseReadCompleteDelegate = FOnQueryForAvailablePurchasesCompleteDelegate::CreateUObject(this, &UInAppPurchaseQueryCallbackProxy::OnInAppPurchaseRead);
-				StoreInterface->AddOnQueryForAvailablePurchasesCompleteDelegate(InAppPurchaseReadCompleteDelegate);
+				InAppPurchaseReadCompleteDelegate       = FOnQueryForAvailablePurchasesCompleteDelegate::CreateUObject(this, &UInAppPurchaseQueryCallbackProxy::OnInAppPurchaseRead);
+				InAppPurchaseReadCompleteDelegateHandle = StoreInterface->AddOnQueryForAvailablePurchasesCompleteDelegate_Handle(InAppPurchaseReadCompleteDelegate);
 
 
 				ReadObject = MakeShareable(new FOnlineProductInformationRead());
@@ -72,7 +72,7 @@ void UInAppPurchaseQueryCallbackProxy::OnInAppPurchaseRead(bool bWasSuccessful)
 	{
 		// Use a local timer handle as we don't need to store it for later but we don't need to look for something to clear
 		FTimerHandle TimerHandle;
-		World->GetTimerManager().SetTimer(this, &UInAppPurchaseQueryCallbackProxy::OnInAppPurchaseRead_Delayed, 0.001f, false);
+		World->GetTimerManager().SetTimer(OnInAppPurchaseRead_DelayedTimerHandle, this, &UInAppPurchaseQueryCallbackProxy::OnInAppPurchaseRead_Delayed, 0.001f, false);
 	}
 
 	ReadObject = NULL;
@@ -99,7 +99,7 @@ void UInAppPurchaseQueryCallbackProxy::RemoveDelegate()
 			IOnlineStorePtr InAppPurchases = OnlineSub->GetStoreInterface();
 			if (InAppPurchases.IsValid())
 			{
-				InAppPurchases->ClearOnQueryForAvailablePurchasesCompleteDelegate(InAppPurchaseReadCompleteDelegate);
+				InAppPurchases->ClearOnQueryForAvailablePurchasesCompleteDelegate_Handle(InAppPurchaseReadCompleteDelegateHandle);
 			}
 		}
 	}
diff --git a/Engine/Source/Runtime/Online/OnlineSubsystemUtils/Private/JoinSessionCallbackProxy.cpp b/Engine/Source/Runtime/Online/OnlineSubsystemUtils/Private/JoinSessionCallbackProxy.cpp
--- a/Engine/Source/Runtime/Online/OnlineSubsystemUtils/Private/JoinSessionCallbackProxy.cpp
+++ b/Engine/Source/Runtime/Online/OnlineSubsystemUtils/Private/JoinSessionCallbackProxy.cpp
@@ -31,7 +31,7 @@ void UJoinSessionCallbackProxy::Activate()
 		auto Sessions = Helper.OnlineSub->GetSessionInterface();
 		if (Sessions.IsValid())
 		{
-			Sessions->AddOnJoinSessionCompleteDelegate(Delegate);
+			DelegateHandle = Sessions->AddOnJoinSessionCompleteDelegate_Handle(Delegate);
 			Sessions->JoinSession(*Helper.UserID, GameSessionName, OnlineSearchResult);
 
 			// OnCompleted will get called, nothing more to do now
@@ -57,7 +57,7 @@ void UJoinSessionCallbackProxy::OnCompleted(FName SessionName, EOnJoinSessionCom
 		auto Sessions = Helper.OnlineSub->GetSessionInterface();
 		if (Sessions.IsValid())
 		{
-			Sessions->ClearOnJoinSessionCompleteDelegate(Delegate);
+			Sessions->ClearOnJoinSessionCompleteDelegate_Handle(DelegateHandle);
 
 			if (Result == EOnJoinSessionCompleteResult::Success)
 			{
diff --git a/Engine/Source/Runtime/Online/OnlineSubsystemUtils/Private/OnlineBeacon.cpp b/Engine/Source/Runtime/Online/OnlineSubsystemUtils/Private/OnlineBeacon.cpp
--- a/Engine/Source/Runtime/Online/OnlineSubsystemUtils/Private/OnlineBeacon.cpp
+++ b/Engine/Source/Runtime/Online/OnlineSubsystemUtils/Private/OnlineBeacon.cpp
@@ -15,7 +15,7 @@ bool AOnlineBeacon::InitBase()
 {
 	if (BeaconNetDriverName != NAME_None && GEngine->CreateNamedNetDriver(GetWorld(), BeaconNetDriverName, NAME_BeaconNetDriver))
 	{
-		GEngine->OnNetworkFailure().AddUObject(this, &AOnlineBeacon::HandleNetworkFailure);
+		HandleNetworkFailureDelegateHandle = GEngine->OnNetworkFailure().AddUObject(this, &AOnlineBeacon::HandleNetworkFailure);
 		NetDriver = GEngine->FindNamedNetDriver(GetWorld(), BeaconNetDriverName);
 		return true;
 	}
@@ -32,7 +32,7 @@ bool AOnlineBeacon::HasNetOwner() const
 void AOnlineBeacon::DestroyBeacon()
 {
 	UE_LOG(LogBeacon, Verbose, TEXT("Destroying beacon %s, netdriver %s"), *GetName(), NetDriver ? *NetDriver->GetDescription() : TEXT("NULL"));
-	GEngine->OnNetworkFailure().RemoveUObject(this, &AOnlineBeacon::HandleNetworkFailure);
+	GEngine->OnNetworkFailure().Remove(HandleNetworkFailureDelegateHandle);
 	GEngine->DestroyNamedNetDriver(GetWorld(), BeaconNetDriverName);
 	NetDriver = NULL;
 
@@ -49,7 +49,7 @@ void AOnlineBeacon::HandleNetworkFailure(UWorld *World, UNetDriver *InNetDriver,
 
 void AOnlineBeacon::OnFailure()
 {
-	GEngine->OnNetworkFailure().RemoveUObject(this, &AOnlineBeacon::HandleNetworkFailure);
+	GEngine->OnNetworkFailure().Remove(HandleNetworkFailureDelegateHandle);
 	GEngine->DestroyNamedNetDriver(GetWorld(), BeaconNetDriverName);
 	NetDriver = NULL;
 }
diff --git a/Engine/Source/Runtime/Online/OnlineSubsystemUtils/Private/OnlineSessionClient.cpp b/Engine/Source/Runtime/Online/OnlineSubsystemUtils/Private/OnlineSessionClient.cpp
--- a/Engine/Source/Runtime/Online/OnlineSubsystemUtils/Private/OnlineSessionClient.cpp
+++ b/Engine/Source/Runtime/Online/OnlineSubsystemUtils/Private/OnlineSessionClient.cpp
@@ -60,13 +60,13 @@ void UOnlineSessionClient::RegisterOnlineDelegates(UWorld* InWorld)
 			if (ControllerId != INVALID_CONTROLLERID)
 			{
 				// Always on the lookout for invite acceptance (via actual invite or join from external ui)
-				OnSessionInviteAcceptedDelegate = FOnSessionInviteAcceptedDelegate::CreateUObject(this, &UOnlineSessionClient::OnSessionInviteAccepted);
-				SessionInt->AddOnSessionInviteAcceptedDelegate(ControllerId, OnSessionInviteAcceptedDelegate);
+				OnSessionInviteAcceptedDelegate       = FOnSessionInviteAcceptedDelegate::CreateUObject(this, &UOnlineSessionClient::OnSessionInviteAccepted);
+				OnSessionInviteAcceptedDelegateHandle = SessionInt->AddOnSessionInviteAcceptedDelegate_Handle(ControllerId, OnSessionInviteAcceptedDelegate);
 			}
 		}
 
-		OnJoinSessionCompleteDelegate = FOnJoinSessionCompleteDelegate::CreateUObject(this, &UOnlineSessionClient::OnJoinSessionComplete);
-		OnEndForJoinSessionCompleteDelegate = FOnEndSessionCompleteDelegate::CreateUObject(this, &UOnlineSessionClient::OnEndForJoinSessionComplete);
+		OnJoinSessionCompleteDelegate           = FOnJoinSessionCompleteDelegate   ::CreateUObject(this, &UOnlineSessionClient::OnJoinSessionComplete);
+		OnEndForJoinSessionCompleteDelegate     = FOnEndSessionCompleteDelegate    ::CreateUObject(this, &UOnlineSessionClient::OnEndForJoinSessionComplete);
 		OnDestroyForJoinSessionCompleteDelegate = FOnDestroySessionCompleteDelegate::CreateUObject(this, &UOnlineSessionClient::OnDestroyForJoinSessionComplete);
 	}
 
@@ -87,7 +87,7 @@ void UOnlineSessionClient::ClearOnlineDelegates(UWorld* InWorld)
 			int32 ControllerId = GetControllerId();
 			if (ControllerId != INVALID_CONTROLLERID)
 			{
-				SessionInt->ClearOnSessionInviteAcceptedDelegate(ControllerId, OnSessionInviteAcceptedDelegate);
+				SessionInt->ClearOnSessionInviteAcceptedDelegate_Handle(ControllerId, OnSessionInviteAcceptedDelegateHandle);
 			}
 		}
 	}
@@ -131,8 +131,8 @@ void UOnlineSessionClient::OnSessionInviteAccepted(int32 LocalUserNum, bool bWas
 void UOnlineSessionClient::OnEndForJoinSessionComplete(FName SessionName, bool bWasSuccessful)
 {
 	UE_LOG(LogOnline, Verbose, TEXT("OnEndForJoinSessionComplete %s bSuccess: %d"), *SessionName.ToString(), bWasSuccessful);
-	SessionInt->ClearOnEndSessionCompleteDelegate(OnEndForJoinSessionCompleteDelegate);
-	DestroyExistingSession(SessionName, OnDestroyForJoinSessionCompleteDelegate);
+	SessionInt->ClearOnEndSessionCompleteDelegate_Handle(OnEndForJoinSessionCompleteDelegateHandle);
+	OnDestroyForJoinSessionCompleteDelegateHandle = DestroyExistingSession_Impl(SessionName, OnDestroyForJoinSessionCompleteDelegate);
 }
 
 /**
@@ -143,15 +143,31 @@ void UOnlineSessionClient::OnEndForJoinSessionComplete(FName SessionName, bool b
  */
 void UOnlineSessionClient::EndExistingSession(FName SessionName, FOnEndSessionCompleteDelegate& Delegate)
 {
+	EndExistingSession_Impl(SessionName, Delegate);
+}
+
+/**
+ * Implementation of EndExistingSession
+ *
+ * @param SessionName name of session to end
+ * @param Delegate delegate to call at session end
+ * @return Handle to the added delegate.
+ */
+FDelegateHandle UOnlineSessionClient::EndExistingSession_Impl(FName SessionName, FOnEndSessionCompleteDelegate& Delegate)
+{
+	FDelegateHandle Result;
+
 	if (SessionInt.IsValid())
 	{
-		SessionInt->AddOnEndSessionCompleteDelegate(Delegate);
+		Result = SessionInt->AddOnEndSessionCompleteDelegate_Handle(Delegate);
 		SessionInt->EndSession(SessionName);
 	}
 	else
 	{
 		Delegate.ExecuteIfBound(SessionName, true);
 	}
+
+	return Result;
 }
 
 /**
@@ -165,7 +181,7 @@ void UOnlineSessionClient::OnDestroyForJoinSessionComplete(FName SessionName, bo
 	UE_LOG(LogOnline, Verbose, TEXT("OnDestroyForJoinSessionComplete %s bSuccess: %d"), *SessionName.ToString(), bWasSuccessful);
 	if (SessionInt.IsValid())
 	{
-		SessionInt->ClearOnDestroySessionCompleteDelegate(OnDestroyForJoinSessionCompleteDelegate);
+		SessionInt->ClearOnDestroySessionCompleteDelegate_Handle(OnDestroyForJoinSessionCompleteDelegateHandle);
 	}
 
 	if (bWasSuccessful)
@@ -191,7 +207,7 @@ void UOnlineSessionClient::OnDestroyForMainMenuComplete(FName SessionName, bool
 	UE_LOG(LogOnline, Verbose, TEXT("OnDestroyForMainMenuComplete %s bSuccess: %d"), *SessionName.ToString(), bWasSuccessful);
 	if (SessionInt.IsValid())
 	{
-		SessionInt->ClearOnDestroySessionCompleteDelegate(OnDestroyForMainMenuCompleteDelegate);
+		SessionInt->ClearOnDestroySessionCompleteDelegate_Handle(OnDestroyForMainMenuCompleteDelegateHandle);
 	}	
 
 	APlayerController* PC = GetPlayerController();
@@ -212,15 +228,31 @@ void UOnlineSessionClient::OnDestroyForMainMenuComplete(FName SessionName, bool
  */
 void UOnlineSessionClient::DestroyExistingSession(FName SessionName, FOnDestroySessionCompleteDelegate& Delegate)
 {
+	DestroyExistingSession_Impl(SessionName, Delegate);
+}
+
+/**
+ * Implementation of DestroyExistingSession
+ *
+ * @param SessionName name of session to destroy
+ * @param Delegate delegate to call at session destruction
+ * @return Handle to the added delegate.
+ */
+FDelegateHandle UOnlineSessionClient::DestroyExistingSession_Impl(FName SessionName, FOnDestroySessionCompleteDelegate& Delegate)
+{
+	FDelegateHandle Result;
+
 	if (SessionInt.IsValid())
 	{
-		SessionInt->AddOnDestroySessionCompleteDelegate(Delegate);
+		Result = SessionInt->AddOnDestroySessionCompleteDelegate_Handle(Delegate);
 		SessionInt->DestroySession(SessionName);
 	}
 	else
 	{
 		Delegate.ExecuteIfBound(SessionName, true);
 	}
+
+	return Result;
 }
 
 /**
@@ -232,7 +264,7 @@ void UOnlineSessionClient::DestroyExistingSession(FName SessionName, FOnDestroyS
 void UOnlineSessionClient::OnJoinSessionComplete(FName SessionName, EOnJoinSessionCompleteResult::Type Result)
 {
 	UE_LOG(LogOnline, Verbose, TEXT("OnJoinSessionComplete %s bSuccess: %d"), *SessionName.ToString(), static_cast<int32>(Result));
-	SessionInt->ClearOnJoinSessionCompleteDelegate(OnJoinSessionCompleteDelegate);
+	SessionInt->ClearOnJoinSessionCompleteDelegate_Handle(OnJoinSessionCompleteDelegateHandle);
 
 	if (Result == EOnJoinSessionCompleteResult::Success)
 	{
@@ -271,13 +303,13 @@ void UOnlineSessionClient::JoinSession(int32 LocalUserNum, FName SessionName, co
 	if (SessionState != EOnlineSessionState::NoSession)
 	{
 		CachedSessionResult = SearchResult;
-		EndExistingSession(SessionName, OnEndForJoinSessionCompleteDelegate);
+		OnEndForJoinSessionCompleteDelegateHandle = EndExistingSession_Impl(SessionName, OnEndForJoinSessionCompleteDelegate);
 	}
 	else
 	{
 		UGameInstance * const GameInstance = GetPlayerController()->GetGameInstance();
 		GameInstance->JoinSession(static_cast<ULocalPlayer*>(GetPlayerController()->Player), SearchResult);
-		/*SessionInt->AddOnJoinSessionCompleteDelegate(OnJoinSessionCompleteDelegate);
+		/*OnJoinSessionCompleteDelegateHandle = SessionInt->AddOnJoinSessionCompleteDelegate_Handle(OnJoinSessionCompleteDelegate);
 		SessionInt->JoinSession(LocalUserNum, SessionName, SearchResult);*/
 	}
 }
@@ -309,7 +341,7 @@ bool UOnlineSessionClient::HandleDisconnectInternal(UWorld* World, UNetDriver* N
 			if (!bHandlingDisconnect)
 			{
 				bHandlingDisconnect = true;
-				DestroyExistingSession(GameSessionName, OnDestroyForMainMenuCompleteDelegate);
+				OnDestroyForMainMenuCompleteDelegateHandle = DestroyExistingSession_Impl(GameSessionName, OnDestroyForMainMenuCompleteDelegate);
 			}
 
 			return true;
diff --git a/Engine/Source/Runtime/Online/OnlineSubsystemUtils/Private/Tests/TestAchievementsInterface.cpp b/Engine/Source/Runtime/Online/OnlineSubsystemUtils/Private/Tests/TestAchievementsInterface.cpp
--- a/Engine/Source/Runtime/Online/OnlineSubsystemUtils/Private/Tests/TestAchievementsInterface.cpp
+++ b/Engine/Source/Runtime/Online/OnlineSubsystemUtils/Private/Tests/TestAchievementsInterface.cpp
@@ -159,7 +159,7 @@ void FTestAchievementsInterface::WriteAchievements()
 	}
 
 	// Add our delegate and write some achievements to the server
-	OnlineAchievements->AddOnAchievementUnlockedDelegate( OnAchievementUnlockedDelegate );
+	OnAchievementUnlockedDelegateHandle = OnlineAchievements->AddOnAchievementUnlockedDelegate_Handle( OnAchievementUnlockedDelegate );
 	OnlineAchievements->WriteAchievements( *UserId.Get(), WriteObjectRef, OnAchievementsWrittenDelegate );
 }
 
@@ -169,7 +169,7 @@ void FTestAchievementsInterface::OnAchievementsWritten(const FUniqueNetId& Playe
 	UE_LOG(LogOnline, Display, TEXT("FTestAchievementsInterface::OnAchievementsWritten( bWasSuccessful = %s )"), bWasSuccessful ? TEXT("true") : TEXT("false"));
 
 	// Clear our delegate and delete this test.
-	OnlineAchievements->ClearOnAchievementUnlockedDelegate(OnAchievementUnlockedDelegate);
+	OnlineAchievements->ClearOnAchievementUnlockedDelegate_Handle(OnAchievementUnlockedDelegateHandle);
 
 	if (bWasSuccessful && WriteObject->WriteState == EOnlineAsyncTaskState::Done)
 	{
diff --git a/Engine/Source/Runtime/Online/OnlineSubsystemUtils/Private/Tests/TestAchievementsInterface.h b/Engine/Source/Runtime/Online/OnlineSubsystemUtils/Private/Tests/TestAchievementsInterface.h
--- a/Engine/Source/Runtime/Online/OnlineSubsystemUtils/Private/Tests/TestAchievementsInterface.h
+++ b/Engine/Source/Runtime/Online/OnlineSubsystemUtils/Private/Tests/TestAchievementsInterface.h
@@ -31,6 +31,9 @@
 	/** Leaderboard read object */
 	FOnlineAchievementsWritePtr WriteObject;
 
+	/** OnAchievementUnlocked delegate handle */
+	FDelegateHandle OnAchievementUnlockedDelegateHandle;
+
 	/** Hidden on purpose */
 	FTestAchievementsInterface()
 		: SubsystemName()
diff --git a/Engine/Source/Runtime/Online/OnlineSubsystemUtils/Private/Tests/TestCloudInterface.cpp b/Engine/Source/Runtime/Online/OnlineSubsystemUtils/Private/Tests/TestCloudInterface.cpp
--- a/Engine/Source/Runtime/Online/OnlineSubsystemUtils/Private/Tests/TestCloudInterface.cpp
+++ b/Engine/Source/Runtime/Online/OnlineSubsystemUtils/Private/Tests/TestCloudInterface.cpp
@@ -74,39 +74,39 @@ bool FTestCloudInterface::Tick( float DeltaTime )
 		case 1:
 			// Write out N files for the user of various sizes/names and enumerate
 			// @todo: make the file count configurable on a per-platform basis, since different platforms have different file count limits
-			WriteNUserCloudFiles(*UserId, FString(TEXT("UserCloud.bin")), 15, OnWriteUserCloudFileCompleteDelegate);
+			OnWriteUserCloudFileCompleteDelegateHandle = WriteNUserCloudFiles(*UserId, FString(TEXT("UserCloud.bin")), 15, OnWriteUserCloudFileCompleteDelegate);
 			break;
 		case 2:
 			// Read in N files for the user and enumerate
-			ReadEnumeratedUserFiles(OnReadEnumeratedUserFilesCompleteDelegate);
+			OnReadEnumeratedUserFilesCompleteDelegateHandle = ReadEnumeratedUserFiles(OnReadEnumeratedUserFilesCompleteDelegate);
 			break;
 		case 3: 
 			// Forget the N files for the user and enumerate (should still remain locally)
-			DeleteEnumeratedUserFiles(true, false, OnDeleteEnumeratedUserFilesCompleteDelegate);
+			OnDeleteEnumeratedUserFilesCompleteDelegateHandle = DeleteEnumeratedUserFiles(true, false, OnDeleteEnumeratedUserFilesCompleteDelegate);
 			break;
 		case 4:
 			// Delete the N files for the user and enumerate (should be removed permanently)
-			DeleteEnumeratedUserFiles(false, true, OnDeleteEnumeratedUserFilesCompleteDelegate);
+			OnDeleteEnumeratedUserFilesCompleteDelegateHandle = DeleteEnumeratedUserFiles(false, true, OnDeleteEnumeratedUserFilesCompleteDelegate);
 			break;
 		case 5:
 			// Write out N files for the user and share with the cloud and enumerate
-			WriteNSharedCloudFiles(*UserId, FString(TEXT("SharedCloud.bin")), 15, OnWriteSharedCloudFileCompleteDelegate);
+			OnWriteSharedCloudFileCompleteDelegateHandle = WriteNSharedCloudFiles(*UserId, FString(TEXT("SharedCloud.bin")), 15, OnWriteSharedCloudFileCompleteDelegate);
 			break;
 		case 6:
 			// Read in N files given the shared handles that this user generated
-			ReadEnumeratedSharedFiles(false, OnReadEnumerateSharedFileCompleteDelegate);
+			OnReadEnumerateSharedFileCompleteDelegateHandle = ReadEnumeratedSharedFiles(false, OnReadEnumerateSharedFileCompleteDelegate);
 			break;
 		case 7:
 			// Delete, marking all the files as only forgotten
-			DeleteEnumeratedUserFiles(true, false, OnDeleteEnumeratedUserFilesCompleteDelegate);
+			OnDeleteEnumeratedUserFilesCompleteDelegateHandle = DeleteEnumeratedUserFiles(true, false, OnDeleteEnumeratedUserFilesCompleteDelegate);
 			break;
 		case 8:
 			// Try to read in the N random files from their shared handles
-			ReadEnumeratedSharedFiles(true, OnReadEnumerateSharedFileCompleteDelegate);
+			OnReadEnumerateSharedFileCompleteDelegateHandle = ReadEnumeratedSharedFiles(true, OnReadEnumerateSharedFileCompleteDelegate);
 			break;
 		case 9:
 			// Delete, marking all the files for permanent removal
-			DeleteEnumeratedUserFiles(true, true, OnDeleteEnumeratedUserFilesCompleteDelegate);
+			OnDeleteEnumeratedUserFilesCompleteDelegateHandle = DeleteEnumeratedUserFiles(true, true, OnDeleteEnumeratedUserFilesCompleteDelegate);
 			break;	
 		case 10:
 			bOverallSuccess = bOverallSuccess && Cleanup();
@@ -140,7 +140,7 @@ bool FTestCloudInterface::Cleanup()
 
 void FTestCloudInterface::EnumerateUserFiles()
 {
-	UserCloud->AddOnEnumerateUserFilesCompleteDelegate(EnumerationDelegate);
+	EnumerationDelegateHandle = UserCloud->AddOnEnumerateUserFilesCompleteDelegate_Handle(EnumerationDelegate);
 	UserCloud->EnumerateUserFiles(*UserId);
 }
 
@@ -152,7 +152,7 @@ void FTestCloudInterface::OnEnumerateUserFilesComplete(bool bWasSuccessful, cons
 
 	TArray<FCloudFileHeader> UserFiles;
 	UserCloud->GetUserFileList(InUserId, UserFiles);
-	UserCloud->ClearOnEnumerateUserFilesCompleteDelegate(EnumerationDelegate);
+	UserCloud->ClearOnEnumerateUserFilesCompleteDelegate_Handle(EnumerationDelegateHandle);
 
 	for (int32 Idx=0; Idx<UserFiles.Num(); Idx++)
 	{
@@ -174,7 +174,7 @@ void FTestCloudInterface::WriteRandomFile(TArray<uint8>& Buffer, int32 Size)
 	}
 }
 
-void FTestCloudInterface::WriteNUserCloudFiles(const FUniqueNetId& InUserId, const FString& FileNameBase, int32 FileCount, FOnWriteUserFileCompleteDelegate& Delegate)
+FDelegateHandle FTestCloudInterface::WriteNUserCloudFiles(const FUniqueNetId& InUserId, const FString& FileNameBase, int32 FileCount, FOnWriteUserFileCompleteDelegate& Delegate)
 {
 	UE_LOG(LogOnline, Log, TEXT("------------------------------------------------"));
 	UE_LOG(LogOnline, Log, TEXT("Writing %d files to the cloud for user %s"), FileCount, *InUserId.ToDebugString());
@@ -182,7 +182,7 @@ void FTestCloudInterface::WriteNUserCloudFiles(const FUniqueNetId& InUserId, con
 	FString FileName(FileNameBase);
 
 	WriteUserCloudFileCount = FileCount;
-	UserCloud->AddOnWriteUserFileCompleteDelegate(Delegate);
+	FDelegateHandle Result = UserCloud->AddOnWriteUserFileCompleteDelegate_Handle(Delegate);
 
 	TArray<uint8> DummyData;
 	for (int32 FileIdx=0; FileIdx<FileCount; FileIdx++)
@@ -192,9 +192,11 @@ void FTestCloudInterface::WriteNUserCloudFiles(const FUniqueNetId& InUserId, con
 		WriteRandomFile(DummyData, FMath::TruncToInt(FMath::FRandRange(256, 1024)));
 		UserCloud->WriteUserFile(InUserId, FString::Printf(TEXT("%s%d.%s"), *FPaths::GetBaseFilename(FileName), FileIdx, *FPaths::GetExtension(FileName)), DummyData);
 	}
+
+	return Result;
 }
 
-void FTestCloudInterface::WriteNSharedCloudFiles(const FUniqueNetId& InUserId, const FString& FileNameBase, int32 FileCount, FOnWriteSharedFileCompleteDelegate& Delegate)
+FDelegateHandle FTestCloudInterface::WriteNSharedCloudFiles(const FUniqueNetId& InUserId, const FString& FileNameBase, int32 FileCount, FOnWriteSharedFileCompleteDelegate& Delegate)
 {
 	UE_LOG(LogOnline, Log, TEXT("------------------------------------------------"));
 	UE_LOG(LogOnline, Log, TEXT("Writing %d files to the cloud and sharing for user %s"), FileCount, *InUserId.ToDebugString());
@@ -202,7 +204,7 @@ void FTestCloudInterface::WriteNSharedCloudFiles(const FUniqueNetId& InUserId, c
 	FString FileName(FileNameBase);
 
 	WriteSharedCloudFileCount = FileCount;
-	SharedCloud->AddOnWriteSharedFileCompleteDelegate(Delegate);
+	FDelegateHandle Result = SharedCloud->AddOnWriteSharedFileCompleteDelegate_Handle(Delegate);
 
 	TArray<uint8> DummyData;
 	for (int32 FileIdx=0; FileIdx<FileCount; FileIdx++)
@@ -210,6 +212,8 @@ void FTestCloudInterface::WriteNSharedCloudFiles(const FUniqueNetId& InUserId, c
 		WriteRandomFile(DummyData, FMath::TruncToInt(FMath::FRandRange(1024, 100*1024)));
 		SharedCloud->WriteSharedFile(InUserId, FString::Printf(TEXT("%s%d.%s"), *FPaths::GetBaseFilename(FileName), FileIdx, *FPaths::GetExtension(FileName)), DummyData);
 	}
+
+	return Result;
 }
 
 void FTestCloudInterface::OnWriteUserCloudFileComplete(bool bWasSuccessful, const FUniqueNetId& InUserId, const FString& FileName)
@@ -224,13 +228,13 @@ void FTestCloudInterface::OnWriteUserCloudFileComplete(bool bWasSuccessful, cons
 	{
 		UE_LOG(LogOnline, Log, TEXT("Write %d User Files Complete!"), NumWrittenCount);
 		UE_LOG(LogOnline, Log, TEXT("------------------------------------------------"));
-		UserCloud->ClearOnWriteUserFileCompleteDelegate(OnWriteUserCloudFileCompleteDelegate);
+		UserCloud->ClearOnWriteUserFileCompleteDelegate_Handle(OnWriteUserCloudFileCompleteDelegateHandle);
 		NumWrittenCount = 0;
 		EnumerateUserFiles();
 	}
 }
 
-void FTestCloudInterface::ReadEnumeratedUserFiles(FOnReadUserFileCompleteDelegate& Delegate)
+FDelegateHandle FTestCloudInterface::ReadEnumeratedUserFiles(FOnReadUserFileCompleteDelegate& Delegate)
 {
 	TArray<FCloudFileHeader> UserFiles;
 	UserCloud->GetUserFileList(*UserId, UserFiles);
@@ -239,9 +243,11 @@ void FTestCloudInterface::ReadEnumeratedUserFiles(FOnReadUserFileCompleteDelegat
 
 	UE_LOG(LogOnline, Log, TEXT("------------------------------------------------"));
 	UE_LOG(LogOnline, Log, TEXT("Reading %d enumerated files for user %s"), ReadUserFileCount, *UserId->ToDebugString());
+
+	FDelegateHandle Result;
 	if (ReadUserFileCount > 0)
 	{
-		UserCloud->AddOnReadUserFileCompleteDelegate(Delegate);
+		Result = UserCloud->AddOnReadUserFileCompleteDelegate_Handle(Delegate);
 		for (int32 Idx=0; Idx < ReadUserFileCount; Idx++)
 		{
 			UE_LOG(LogOnline, Log, TEXT("\tFileName:%s Size:%d"), *UserFiles[Idx].FileName, UserFiles[Idx].FileSize);
@@ -252,6 +258,8 @@ void FTestCloudInterface::ReadEnumeratedUserFiles(FOnReadUserFileCompleteDelegat
 	{
 		EnumerateUserFiles();
 	}
+
+	return Result;
 }
 
 void FTestCloudInterface::OnReadEnumeratedUserFilesComplete(bool bWasSuccessful, const FUniqueNetId& InUserId, const FString& FileName)
@@ -275,13 +283,13 @@ void FTestCloudInterface::OnReadEnumeratedUserFilesComplete(bool bWasSuccessful,
 	{
 		UE_LOG(LogOnline, Log, TEXT("Read %d User Files Complete!"), NumReadCount);
 		UE_LOG(LogOnline, Log, TEXT("------------------------------------------------"));
-		UserCloud->ClearOnReadUserFileCompleteDelegate(OnReadEnumeratedUserFilesCompleteDelegate);
+		UserCloud->ClearOnReadUserFileCompleteDelegate_Handle(OnReadEnumeratedUserFilesCompleteDelegateHandle);
 		NumReadCount = 0;
 		EnumerateUserFiles();
 	}
 }
 
-void FTestCloudInterface::DeleteEnumeratedUserFiles(bool bCloudDelete, bool bLocalDelete, FOnDeleteUserFileCompleteDelegate& Delegate)
+FDelegateHandle FTestCloudInterface::DeleteEnumeratedUserFiles(bool bCloudDelete, bool bLocalDelete, FOnDeleteUserFileCompleteDelegate& Delegate)
 {
 	TArray<FCloudFileHeader> UserFiles;
 	UserCloud->GetUserFileList(*UserId, UserFiles);
@@ -290,9 +298,10 @@ void FTestCloudInterface::DeleteEnumeratedUserFiles(bool bCloudDelete, bool bLoc
 
 	UE_LOG(LogOnline, Log, TEXT("------------------------------------------------"));
 	UE_LOG(LogOnline, Log, TEXT("Deleting %d files from the cloud for user %s CLOUD: %s LOCAL: %s"), DeleteUserFileCount, *UserId->ToDebugString(), bCloudDelete ? TEXT("true") : TEXT("false"), bLocalDelete ? TEXT("true") : TEXT("false"));
+	FDelegateHandle Result;
 	if (DeleteUserFileCount > 0)
 	{
-		UserCloud->AddOnDeleteUserFileCompleteDelegate(Delegate);
+		Result = UserCloud->AddOnDeleteUserFileCompleteDelegate_Handle(Delegate);
 		for (int32 Idx=0; Idx < DeleteUserFileCount; Idx++)
 		{
 			UE_LOG(LogOnline, Log, TEXT("\tFileName:%s Size:%d"), *UserFiles[Idx].FileName, UserFiles[Idx].FileSize);
@@ -303,6 +312,8 @@ void FTestCloudInterface::DeleteEnumeratedUserFiles(bool bCloudDelete, bool bLoc
 	{
 		EnumerateUserFiles();
 	}
+
+	return Result;
 }
 
 void FTestCloudInterface::OnDeleteEnumeratedUserFilesComplete(bool bWasSuccessful, const FUniqueNetId& InUserId, const FString& FileName)
@@ -316,7 +327,7 @@ void FTestCloudInterface::OnDeleteEnumeratedUserFilesComplete(bool bWasSuccessfu
 	{
 		UE_LOG(LogOnline, Log, TEXT("Delete %d User Files Complete!"), NumDeletedCount);
 		UE_LOG(LogOnline, Log, TEXT("------------------------------------------------"));
-		UserCloud->ClearOnDeleteUserFileCompleteDelegate(OnDeleteEnumeratedUserFilesCompleteDelegate);
+		UserCloud->ClearOnDeleteUserFileCompleteDelegate_Handle(OnDeleteEnumeratedUserFilesCompleteDelegateHandle);
 		NumDeletedCount = 0;
 		EnumerateUserFiles();
 	}
@@ -335,14 +346,14 @@ void FTestCloudInterface::OnWriteSharedCloudFileComplete(bool bWasSuccessful, co
 	{
 		UE_LOG(LogOnline, Log, TEXT("Write %d Shared Files Complete!"), NumWrittenCount);
 		UE_LOG(LogOnline, Log, TEXT("------------------------------------------------"));
-		SharedCloud->ClearOnWriteSharedFileCompleteDelegate(OnWriteSharedCloudFileCompleteDelegate);
+		SharedCloud->ClearOnWriteSharedFileCompleteDelegate_Handle(OnWriteSharedCloudFileCompleteDelegateHandle);
 		NumWrittenCount = 0;
 		// Enumeration will put the newly shared files into this user's list
 		EnumerateUserFiles();
 	}
 }
 
-void FTestCloudInterface::ReadEnumeratedSharedFiles(bool bUseRandom, FOnReadSharedFileCompleteDelegate& Delegate)
+FDelegateHandle FTestCloudInterface::ReadEnumeratedSharedFiles(bool bUseRandom, FOnReadSharedFileCompleteDelegate& Delegate)
 {
 	if (bUseRandom)
 	{
@@ -355,9 +366,11 @@ void FTestCloudInterface::ReadEnumeratedSharedFiles(bool bUseRandom, FOnReadShar
 	
 	UE_LOG(LogOnline, Log, TEXT("------------------------------------------------"));
 	UE_LOG(LogOnline, Log, TEXT("Reading %d enumerated shared files"), ReadSharedFileCount);
+
+	FDelegateHandle Result;
 	if (ReadSharedFileCount > 0)
 	{
-		SharedCloud->AddOnReadSharedFileCompleteDelegate(Delegate);
+		Result = SharedCloud->AddOnReadSharedFileCompleteDelegate_Handle(Delegate);
 		for (int32 Idx=0; Idx < ReadSharedFileCount; Idx++)
 		{
 			if (bUseRandom)
@@ -376,6 +389,8 @@ void FTestCloudInterface::ReadEnumeratedSharedFiles(bool bUseRandom, FOnReadShar
 	{
 		EnumerateUserFiles();
 	}
+
+	return Result;
 }
 
 void FTestCloudInterface::OnReadEnumeratedSharedFileCompleteDelegate(bool bWasSuccessful, const FSharedContentHandle& SharedHandle)
@@ -389,7 +404,7 @@ void FTestCloudInterface::OnReadEnumeratedSharedFileCompleteDelegate(bool bWasSu
 	{
 		UE_LOG(LogOnline, Log, TEXT("Read %d Shared Files Complete!"), NumReadCount);
 		UE_LOG(LogOnline, Log, TEXT("------------------------------------------------"));
-		SharedCloud->ClearOnReadSharedFileCompleteDelegate(OnReadEnumerateSharedFileCompleteDelegate);
+		SharedCloud->ClearOnReadSharedFileCompleteDelegate_Handle(OnReadEnumerateSharedFileCompleteDelegateHandle);
 		NumReadCount = 0;
 		EnumerateUserFiles();
 	}
diff --git a/Engine/Source/Runtime/Online/OnlineSubsystemUtils/Private/Tests/TestCloudInterface.h b/Engine/Source/Runtime/Online/OnlineSubsystemUtils/Private/Tests/TestCloudInterface.h
--- a/Engine/Source/Runtime/Online/OnlineSubsystemUtils/Private/Tests/TestCloudInterface.h
+++ b/Engine/Source/Runtime/Online/OnlineSubsystemUtils/Private/Tests/TestCloudInterface.h
@@ -29,6 +29,14 @@
 	FOnWriteSharedFileCompleteDelegate OnWriteSharedCloudFileCompleteDelegate;
 	FOnReadSharedFileCompleteDelegate OnReadEnumerateSharedFileCompleteDelegate;
 
+	/** Handles to those delegates */
+	FDelegateHandle EnumerationDelegateHandle;
+	FDelegateHandle OnWriteUserCloudFileCompleteDelegateHandle;
+	FDelegateHandle OnReadEnumeratedUserFilesCompleteDelegateHandle;
+	FDelegateHandle OnDeleteEnumeratedUserFilesCompleteDelegateHandle;
+	FDelegateHandle OnWriteSharedCloudFileCompleteDelegateHandle;
+	FDelegateHandle OnReadEnumerateSharedFileCompleteDelegateHandle;
+
 	/** Logged in UserId */
 	TSharedPtr<FUniqueNetId> UserId;
 
@@ -71,7 +79,7 @@
 	 * @param FileCount number of files to write out
 	 * @param Delegate delegate to trigger for each file written
 	 */
-	void WriteNUserCloudFiles(const FUniqueNetId& UserId, const FString& FileNameBase, int32 FileCount, FOnWriteUserFileCompleteDelegate& Delegate);
+	FDelegateHandle WriteNUserCloudFiles(const FUniqueNetId& UserId, const FString& FileNameBase, int32 FileCount, FOnWriteUserFileCompleteDelegate& Delegate);
 
 	/**
 	 *	Write out files marked for the cloud (and shared) for the given user
@@ -80,14 +88,14 @@
 	 * @param FileCount number of files to write out
 	 * @param Delegate delegate to trigger for each file written
 	 */
-	void WriteNSharedCloudFiles(const FUniqueNetId& UserId, const FString& FileNameBase, int32 FileCount, FOnWriteSharedFileCompleteDelegate& Delegate);
+	FDelegateHandle WriteNSharedCloudFiles(const FUniqueNetId& UserId, const FString& FileNameBase, int32 FileCount, FOnWriteSharedFileCompleteDelegate& Delegate);
 
 	/**
 	 *	Read cloud files currently enumerated for the logged in user
 	 * (assumes EnumerateFiles has been called)
 	 * @param Delegate delegate to trigger for each file read	
 	 */	
-	void ReadEnumeratedUserFiles(FOnReadUserFileCompleteDelegate& Delegate);
+	FDelegateHandle ReadEnumeratedUserFiles(FOnReadUserFileCompleteDelegate& Delegate);
 
 	/**
 	 *	Read shared cloud files currently enumerated
@@ -95,7 +103,7 @@
 	 * @param bUseRandom use the random array of existing shared files if true, current users shared files otherwise
 	 * @param Delegate delegate to trigger for each file read	
 	 */	
-	void ReadEnumeratedSharedFiles(bool bUseRandom, FOnReadSharedFileCompleteDelegate& Delegate);
+	FDelegateHandle ReadEnumeratedSharedFiles(bool bUseRandom, FOnReadSharedFileCompleteDelegate& Delegate);
 
 	/**
 	 *	Delete cloud files currently enumerated
@@ -104,7 +112,7 @@
 	 * @param bLocalDelete delete local copy of file, but not cloud version
 	 * @param Delegate delegate to trigger for each file deleted
 	 */
-	void DeleteEnumeratedUserFiles(bool bCloudDelete, bool bLocalDelete, FOnDeleteUserFileCompleteDelegate& Delegate);
+	FDelegateHandle DeleteEnumeratedUserFiles(bool bCloudDelete, bool bLocalDelete, FOnDeleteUserFileCompleteDelegate& Delegate);
 
 	/**
 	 *	Delegate triggered when all user files have been enumerated
diff --git a/Engine/Source/Runtime/Online/OnlineSubsystemUtils/Private/Tests/TestEntitlementsInterface.cpp b/Engine/Source/Runtime/Online/OnlineSubsystemUtils/Private/Tests/TestEntitlementsInterface.cpp
--- a/Engine/Source/Runtime/Online/OnlineSubsystemUtils/Private/Tests/TestEntitlementsInterface.cpp
+++ b/Engine/Source/Runtime/Online/OnlineSubsystemUtils/Private/Tests/TestEntitlementsInterface.cpp
@@ -19,7 +19,7 @@ void FTestEntitlementsInterface::Test(UWorld* InWorld)
 		if (EntitlementsOSS.IsValid())
 		{
 			// Add delegates for the various async calls
-			EntitlementsOSS->AddOnQueryEntitlementsCompleteDelegate(OnQueryEntitlementsCompleteDelegate);
+			OnQueryEntitlementsCompleteDelegateHandle = EntitlementsOSS->AddOnQueryEntitlementsCompleteDelegate_Handle(OnQueryEntitlementsCompleteDelegate);
 
 			// kick off next test
 			StartNextTest();
@@ -61,7 +61,7 @@ void FTestEntitlementsInterface::FinishTest()
 	if (EntitlementsOSS.IsValid())
 	{
 		// Clear delegates for the entitlements async calls
-		EntitlementsOSS->ClearOnQueryEntitlementsCompleteDelegate(OnQueryEntitlementsCompleteDelegate);
+		EntitlementsOSS->ClearOnQueryEntitlementsCompleteDelegate_Handle(OnQueryEntitlementsCompleteDelegateHandle);
 	}
 
 	delete this;
diff --git a/Engine/Source/Runtime/Online/OnlineSubsystemUtils/Private/Tests/TestEntitlementsInterface.h b/Engine/Source/Runtime/Online/OnlineSubsystemUtils/Private/Tests/TestEntitlementsInterface.h
--- a/Engine/Source/Runtime/Online/OnlineSubsystemUtils/Private/Tests/TestEntitlementsInterface.h
+++ b/Engine/Source/Runtime/Online/OnlineSubsystemUtils/Private/Tests/TestEntitlementsInterface.h
@@ -20,6 +20,9 @@
 	/** Delegates for callbacks of each test */
 	FOnQueryEntitlementsCompleteDelegate OnQueryEntitlementsCompleteDelegate;
 
+	/** Handle to the above delegates */
+	FDelegateHandle OnQueryEntitlementsCompleteDelegateHandle;
+
 	/** toggles for whether to run each test */
 	bool bQueryEntitlements;
 
diff --git a/Engine/Source/Runtime/Online/OnlineSubsystemUtils/Private/Tests/TestExternalUIInterface.cpp b/Engine/Source/Runtime/Online/OnlineSubsystemUtils/Private/Tests/TestExternalUIInterface.cpp
--- a/Engine/Source/Runtime/Online/OnlineSubsystemUtils/Private/Tests/TestExternalUIInterface.cpp
+++ b/Engine/Source/Runtime/Online/OnlineSubsystemUtils/Private/Tests/TestExternalUIInterface.cpp
@@ -16,8 +16,8 @@ void FTestExternalUIInterface::Test()
 	check(ExternalUI.IsValid());
 
 	// Define and register delegates
-	ExternalUIChangeDelegate = FOnExternalUIChangeDelegate::CreateRaw(this, &FTestExternalUIInterface::OnExternalUIChange);
-	ExternalUI->AddOnExternalUIChangeDelegate(ExternalUIChangeDelegate);
+	ExternalUIChangeDelegate       = FOnExternalUIChangeDelegate::CreateRaw(this, &FTestExternalUIInterface::OnExternalUIChange);
+	ExternalUIChangeDelegateHandle = ExternalUI->AddOnExternalUIChangeDelegate_Handle(ExternalUIChangeDelegate);
 
 	// Are we testing at least one of our external UI?
 	if (bTestLoginUI == false && bTestFriendsUI == false && bTestInviteUI == false && bTestAchievementsUI == false && bTestWebURL == false && bTestProfileUI == false)
@@ -34,7 +34,7 @@ void FTestExternalUIInterface::Test()
 void FTestExternalUIInterface::FinishTest()
 {
 	UE_LOG(LogOnline, Log, TEXT("FTestExternalUIInterface::FinishTest -- completed testing"));
-	ExternalUI->ClearOnExternalUIChangeDelegate(ExternalUIChangeDelegate);
+	ExternalUI->ClearOnExternalUIChangeDelegate_Handle(ExternalUIChangeDelegateHandle);
 	delete this;
 }
 
diff --git a/Engine/Source/Runtime/Online/OnlineSubsystemUtils/Private/Tests/TestExternalUIInterface.h b/Engine/Source/Runtime/Online/OnlineSubsystemUtils/Private/Tests/TestExternalUIInterface.h
--- a/Engine/Source/Runtime/Online/OnlineSubsystemUtils/Private/Tests/TestExternalUIInterface.h
+++ b/Engine/Source/Runtime/Online/OnlineSubsystemUtils/Private/Tests/TestExternalUIInterface.h
@@ -47,6 +47,9 @@ namespace ETestExternalUIInterfaceState
 	/** Delegate for external UI opening and closing */
 	FOnExternalUIChangeDelegate ExternalUIChangeDelegate;
 
+	/** ExternalUIChange delegate handle */
+	FDelegateHandle ExternalUIChangeDelegateHandle;
+
 	/** Current external UI test */
 	ETestExternalUIInterfaceState::Type State;
 
diff --git a/Engine/Source/Runtime/Online/OnlineSubsystemUtils/Private/Tests/TestFriendsInterface.cpp b/Engine/Source/Runtime/Online/OnlineSubsystemUtils/Private/Tests/TestFriendsInterface.cpp
--- a/Engine/Source/Runtime/Online/OnlineSubsystemUtils/Private/Tests/TestFriendsInterface.cpp
+++ b/Engine/Source/Runtime/Online/OnlineSubsystemUtils/Private/Tests/TestFriendsInterface.cpp
@@ -21,10 +21,9 @@ void FTestFriendsInterface::Test(UWorld* InWorld, const TArray<FString>& Invites
 		// Add our delegate for the async call
 		OnDeleteFriendCompleteDelegate = FOnDeleteFriendCompleteDelegate::CreateRaw(this, &FTestFriendsInterface::OnDeleteFriendComplete);
 		OnQueryRecentPlayersCompleteDelegate = FOnQueryRecentPlayersCompleteDelegate::CreateRaw(this, &FTestFriendsInterface::OnQueryRecentPlayersComplete);
-		
 
-		OnlineSub->GetFriendsInterface()->AddOnDeleteFriendCompleteDelegate(0, OnDeleteFriendCompleteDelegate);
-		OnlineSub->GetFriendsInterface()->AddOnQueryRecentPlayersCompleteDelegate(OnQueryRecentPlayersCompleteDelegate);
+		OnDeleteFriendCompleteDelegateHandle       = OnlineSub->GetFriendsInterface()->AddOnDeleteFriendCompleteDelegate_Handle      (0, OnDeleteFriendCompleteDelegate);
+		OnQueryRecentPlayersCompleteDelegateHandle = OnlineSub->GetFriendsInterface()->AddOnQueryRecentPlayersCompleteDelegate_Handle(OnQueryRecentPlayersCompleteDelegate);
 
 		// list of pending users to send invites to
 		for (int32 Idx=0; Idx < Invites.Num(); Idx++)
@@ -96,8 +95,8 @@ void FTestFriendsInterface::FinishTest()
 		OnlineSub->GetFriendsInterface().IsValid())
 	{
 		// Clear delegates for the various async calls
-		OnlineSub->GetFriendsInterface()->ClearOnDeleteFriendCompleteDelegate(0, OnDeleteFriendCompleteDelegate);
-		OnlineSub->GetFriendsInterface()->ClearOnQueryRecentPlayersCompleteDelegate(OnQueryRecentPlayersCompleteDelegate);
+		OnlineSub->GetFriendsInterface()->ClearOnDeleteFriendCompleteDelegate_Handle      (0, OnDeleteFriendCompleteDelegateHandle);
+		OnlineSub->GetFriendsInterface()->ClearOnQueryRecentPlayersCompleteDelegate_Handle(OnQueryRecentPlayersCompleteDelegateHandle);
 	}
 	delete this;
 }
diff --git a/Engine/Source/Runtime/Online/OnlineSubsystemUtils/Private/Tests/TestFriendsInterface.h b/Engine/Source/Runtime/Online/OnlineSubsystemUtils/Private/Tests/TestFriendsInterface.h
--- a/Engine/Source/Runtime/Online/OnlineSubsystemUtils/Private/Tests/TestFriendsInterface.h
+++ b/Engine/Source/Runtime/Online/OnlineSubsystemUtils/Private/Tests/TestFriendsInterface.h
@@ -19,6 +19,10 @@
 	/** Delegate to use for querying for recent players */
 	FOnQueryRecentPlayersCompleteDelegate OnQueryRecentPlayersCompleteDelegate;
 
+	/** Handles to the above delegates */
+	FDelegateHandle OnDeleteFriendCompleteDelegateHandle;
+	FDelegateHandle OnQueryRecentPlayersCompleteDelegateHandle;
+
 	/** Default name of friends list for running tests */
 	FString FriendsListName;
 	/** Filled in after reading friends list. Invites that are pending will be auto-accepted */
diff --git a/Engine/Source/Runtime/Online/OnlineSubsystemUtils/Private/Tests/TestIdentityInterface.cpp b/Engine/Source/Runtime/Online/OnlineSubsystemUtils/Private/Tests/TestIdentityInterface.cpp
--- a/Engine/Source/Runtime/Online/OnlineSubsystemUtils/Private/Tests/TestIdentityInterface.cpp
+++ b/Engine/Source/Runtime/Online/OnlineSubsystemUtils/Private/Tests/TestIdentityInterface.cpp
@@ -19,8 +19,8 @@ void FTestIdentityInterface::Test(UWorld* InWorld, const FOnlineAccountCredentia
 	if (OnlineIdentity.IsValid())
 	{
 		// Add delegates for the various async calls
-		OnlineIdentity->AddOnLoginCompleteDelegate(LocalUserIdx, OnLoginCompleteDelegate);
-		OnlineIdentity->AddOnLogoutCompleteDelegate(LocalUserIdx, OnLogoutCompleteDelegate);
+		OnLoginCompleteDelegateHandle  = OnlineIdentity->AddOnLoginCompleteDelegate_Handle (LocalUserIdx, OnLoginCompleteDelegate);
+		OnLogoutCompleteDelegateHandle = OnlineIdentity->AddOnLogoutCompleteDelegate_Handle(LocalUserIdx, OnLogoutCompleteDelegate);
 		
 		// kick off next test
 		StartNextTest();
@@ -58,8 +58,8 @@ void FTestIdentityInterface::FinishTest()
 	if (OnlineIdentity.IsValid())
 	{
 		// Clear delegates for the various async calls
-		OnlineIdentity->ClearOnLoginCompleteDelegate(LocalUserIdx, OnLoginCompleteDelegate);
-		OnlineIdentity->ClearOnLogoutCompleteDelegate(LocalUserIdx, OnLogoutCompleteDelegate);
+		OnlineIdentity->ClearOnLoginCompleteDelegate_Handle (LocalUserIdx, OnLoginCompleteDelegateHandle);
+		OnlineIdentity->ClearOnLogoutCompleteDelegate_Handle(LocalUserIdx, OnLogoutCompleteDelegateHandle);
 	}
 	delete this;
 }
diff --git a/Engine/Source/Runtime/Online/OnlineSubsystemUtils/Private/Tests/TestIdentityInterface.h b/Engine/Source/Runtime/Online/OnlineSubsystemUtils/Private/Tests/TestIdentityInterface.h
--- a/Engine/Source/Runtime/Online/OnlineSubsystemUtils/Private/Tests/TestIdentityInterface.h
+++ b/Engine/Source/Runtime/Online/OnlineSubsystemUtils/Private/Tests/TestIdentityInterface.h
@@ -18,6 +18,9 @@
 	FOnLoginCompleteDelegate OnLoginCompleteDelegate;
 	/** Delegate to use for logging out a user */
 	FOnLogoutCompleteDelegate OnLogoutCompleteDelegate;
+	/** Handles to the above delegates */
+	FDelegateHandle OnLoginCompleteDelegateHandle;
+	FDelegateHandle OnLogoutCompleteDelegateHandle;
 	/** true if authentication test should be run */
 	bool bRunLoginTest;
 	/** true if logout test should be run */
diff --git a/Engine/Source/Runtime/Online/OnlineSubsystemUtils/Private/Tests/TestLeaderboardInterface.cpp b/Engine/Source/Runtime/Online/OnlineSubsystemUtils/Private/Tests/TestLeaderboardInterface.cpp
--- a/Engine/Source/Runtime/Online/OnlineSubsystemUtils/Private/Tests/TestLeaderboardInterface.cpp
+++ b/Engine/Source/Runtime/Online/OnlineSubsystemUtils/Private/Tests/TestLeaderboardInterface.cpp
@@ -77,13 +77,13 @@ void FTestLeaderboardInterface::OnLeaderboardFlushComplete(FName SessionName, bo
 	UE_LOG(LogOnline, Verbose, TEXT("OnLeaderboardFlushComplete Session: %s bWasSuccessful: %d"), *SessionName.ToString(), bWasSuccessful);
 	bOverallSuccess = bOverallSuccess && bWasSuccessful;
 
-	Leaderboards->ClearOnLeaderboardFlushCompleteDelegate(LeaderboardFlushDelegate);
+	Leaderboards->ClearOnLeaderboardFlushCompleteDelegate_Handle(LeaderboardFlushDelegateHandle);
 	TestPhase++;
 }
 
 void FTestLeaderboardInterface::FlushLeaderboards()
 {
-	Leaderboards->AddOnLeaderboardFlushCompleteDelegate(LeaderboardFlushDelegate);
+	LeaderboardFlushDelegateHandle = Leaderboards->AddOnLeaderboardFlushCompleteDelegate_Handle(LeaderboardFlushDelegate);
 	Leaderboards->FlushLeaderboards(TEXT("TEST"));
 }
 
@@ -103,7 +103,7 @@ void FTestLeaderboardInterface::OnLeaderboardReadComplete(bool bWasSuccessful)
 		}
 	}
 
-	Leaderboards->ClearOnLeaderboardReadCompleteDelegate(LeaderboardReadCompleteDelegate);
+	Leaderboards->ClearOnLeaderboardReadCompleteDelegate_Handle(LeaderboardReadCompleteDelegateHandle);
 	TestPhase++;
 }
 
@@ -112,7 +112,7 @@ void FTestLeaderboardInterface::ReadLeaderboards()
 	ReadObject = MakeShareable(new TestLeaderboardRead());
 	FOnlineLeaderboardReadRef ReadObjectRef = ReadObject.ToSharedRef();
 
-	Leaderboards->AddOnLeaderboardReadCompleteDelegate(LeaderboardReadCompleteDelegate);
+	LeaderboardReadCompleteDelegateHandle = Leaderboards->AddOnLeaderboardReadCompleteDelegate_Handle(LeaderboardReadCompleteDelegate);
 	Leaderboards->ReadLeaderboardsForFriends(0, ReadObjectRef);
 }
 
diff --git a/Engine/Source/Runtime/Online/OnlineSubsystemUtils/Private/Tests/TestLeaderboardInterface.h b/Engine/Source/Runtime/Online/OnlineSubsystemUtils/Private/Tests/TestLeaderboardInterface.h
--- a/Engine/Source/Runtime/Online/OnlineSubsystemUtils/Private/Tests/TestLeaderboardInterface.h
+++ b/Engine/Source/Runtime/Online/OnlineSubsystemUtils/Private/Tests/TestLeaderboardInterface.h
@@ -30,6 +30,10 @@
 	/** Delegate called when a leaderboard has been successfully read */
 	FOnLeaderboardReadCompleteDelegate LeaderboardReadCompleteDelegate;
 
+	/** Handles to the above delegates */ 
+	FDelegateHandle LeaderboardFlushDelegateHandle;
+	FDelegateHandle LeaderboardReadCompleteDelegateHandle;
+
 	/** Current phase of testing */
 	int32 TestPhase;
 	/** Last phase of testing triggered */
diff --git a/Engine/Source/Runtime/Online/OnlineSubsystemUtils/Private/Tests/TestMessageInterface.cpp b/Engine/Source/Runtime/Online/OnlineSubsystemUtils/Private/Tests/TestMessageInterface.cpp
--- a/Engine/Source/Runtime/Online/OnlineSubsystemUtils/Private/Tests/TestMessageInterface.cpp
+++ b/Engine/Source/Runtime/Online/OnlineSubsystemUtils/Private/Tests/TestMessageInterface.cpp
@@ -61,14 +61,14 @@ void FTestMessageInterface::Test(UWorld* InWorld, const TArray<FString>& InRecip
 	{
 		// Add our delegate for the async call
 		OnEnumerateMessagesCompleteDelegate = FOnEnumerateMessagesCompleteDelegate::CreateRaw(this, &FTestMessageInterface::OnEnumerateMessagesComplete);
-		OnReadMessageCompleteDelegate = FOnReadMessageCompleteDelegate::CreateRaw(this, &FTestMessageInterface::OnReadMessageComplete);
-		OnSendMessageCompleteDelegate = FOnSendMessageCompleteDelegate::CreateRaw(this, &FTestMessageInterface::OnSendMessageComplete);
-		OnDeleteMessageCompleteDelegate = FOnDeleteMessageCompleteDelegate::CreateRaw(this, &FTestMessageInterface::OnDeleteMessageComplete);
+		OnReadMessageCompleteDelegate       = FOnReadMessageCompleteDelegate      ::CreateRaw(this, &FTestMessageInterface::OnReadMessageComplete);
+		OnSendMessageCompleteDelegate       = FOnSendMessageCompleteDelegate      ::CreateRaw(this, &FTestMessageInterface::OnSendMessageComplete);
+		OnDeleteMessageCompleteDelegate     = FOnDeleteMessageCompleteDelegate    ::CreateRaw(this, &FTestMessageInterface::OnDeleteMessageComplete);
 
-		OnlineSub->GetMessageInterface()->AddOnEnumerateMessagesCompleteDelegate(0, OnEnumerateMessagesCompleteDelegate);
-		OnlineSub->GetMessageInterface()->AddOnReadMessageCompleteDelegate(0, OnReadMessageCompleteDelegate);
-		OnlineSub->GetMessageInterface()->AddOnSendMessageCompleteDelegate(0, OnSendMessageCompleteDelegate);
-		OnlineSub->GetMessageInterface()->AddOnDeleteMessageCompleteDelegate(0, OnDeleteMessageCompleteDelegate);
+		OnEnumerateMessagesCompleteDelegateHandle = OnlineSub->GetMessageInterface()->AddOnEnumerateMessagesCompleteDelegate_Handle(0, OnEnumerateMessagesCompleteDelegate);
+		OnReadMessageCompleteDelegateHandle       = OnlineSub->GetMessageInterface()->AddOnReadMessageCompleteDelegate_Handle      (0, OnReadMessageCompleteDelegate);
+		OnSendMessageCompleteDelegateHandle       = OnlineSub->GetMessageInterface()->AddOnSendMessageCompleteDelegate_Handle      (0, OnSendMessageCompleteDelegate);
+		OnDeleteMessageCompleteDelegateHandle     = OnlineSub->GetMessageInterface()->AddOnDeleteMessageCompleteDelegate_Handle    (0, OnDeleteMessageCompleteDelegate);
 
 		// list of users to send messages to
 		for (int32 Idx=0; Idx < InRecipients.Num(); Idx++)
@@ -146,10 +146,10 @@ void FTestMessageInterface::FinishTest()
 		OnlineSub->GetMessageInterface().IsValid())
 	{
 		// Clear delegates for the various async calls
-		OnlineSub->GetMessageInterface()->ClearOnEnumerateMessagesCompleteDelegate(0, OnEnumerateMessagesCompleteDelegate);
-		OnlineSub->GetMessageInterface()->ClearOnReadMessageCompleteDelegate(0, OnReadMessageCompleteDelegate);
-		OnlineSub->GetMessageInterface()->ClearOnSendMessageCompleteDelegate(0, OnSendMessageCompleteDelegate);
-		OnlineSub->GetMessageInterface()->ClearOnDeleteMessageCompleteDelegate(0, OnDeleteMessageCompleteDelegate);
+		OnlineSub->GetMessageInterface()->ClearOnEnumerateMessagesCompleteDelegate_Handle(0, OnEnumerateMessagesCompleteDelegateHandle);
+		OnlineSub->GetMessageInterface()->ClearOnReadMessageCompleteDelegate_Handle      (0, OnReadMessageCompleteDelegateHandle);
+		OnlineSub->GetMessageInterface()->ClearOnSendMessageCompleteDelegate_Handle      (0, OnSendMessageCompleteDelegateHandle);
+		OnlineSub->GetMessageInterface()->ClearOnDeleteMessageCompleteDelegate_Handle    (0, OnDeleteMessageCompleteDelegateHandle);
 	}
 	delete this;
 }
diff --git a/Engine/Source/Runtime/Online/OnlineSubsystemUtils/Private/Tests/TestMessageInterface.h b/Engine/Source/Runtime/Online/OnlineSubsystemUtils/Private/Tests/TestMessageInterface.h
--- a/Engine/Source/Runtime/Online/OnlineSubsystemUtils/Private/Tests/TestMessageInterface.h
+++ b/Engine/Source/Runtime/Online/OnlineSubsystemUtils/Private/Tests/TestMessageInterface.h
@@ -25,6 +25,11 @@ private:
 	/** Delegate to use for deleting messages for a user */
 	FOnDeleteMessageCompleteDelegate OnDeleteMessageCompleteDelegate;
 
+	FDelegateHandle OnEnumerateMessagesCompleteDelegateHandle;
+	FDelegateHandle OnReadMessageCompleteDelegateHandle;
+	FDelegateHandle OnSendMessageCompleteDelegateHandle;
+	FDelegateHandle OnDeleteMessageCompleteDelegateHandle;
+
 	/** true to enable message enumeration */
 	bool bEnumerateMessages;
 	/** true to enable downloading of each message that was enumerated */
diff --git a/Engine/Source/Runtime/Online/OnlineSubsystemUtils/Private/Tests/TestSessionInterface.cpp b/Engine/Source/Runtime/Online/OnlineSubsystemUtils/Private/Tests/TestSessionInterface.cpp
--- a/Engine/Source/Runtime/Online/OnlineSubsystemUtils/Private/Tests/TestSessionInterface.cpp
+++ b/Engine/Source/Runtime/Online/OnlineSubsystemUtils/Private/Tests/TestSessionInterface.cpp
@@ -140,7 +140,7 @@ void FTestSessionInterface::Test(UWorld* InWorld, bool bTestLAN, bool bIsPresenc
 	{
  		HostSettings = MakeShareable(new TestOnlineGameSettings(bTestLAN, bIsPresence));
 		HostSettings->AddWorldSettings(InWorld);
- 		SessionInt->AddOnCreateSessionCompleteDelegate(OnCreateSessionCompleteDelegate);
+ 		OnCreateSessionCompleteDelegateHandle = SessionInt->AddOnCreateSessionCompleteDelegate_Handle(OnCreateSessionCompleteDelegate);
  		SessionInt->CreateSession(0, GameSessionName, *HostSettings);
 	}
 	else
@@ -148,18 +148,18 @@ void FTestSessionInterface::Test(UWorld* InWorld, bool bTestLAN, bool bIsPresenc
 		SearchSettings = MakeShareable(new TestOnlineSearchSettings(bTestLAN, bIsPresence));
 		TSharedRef<FOnlineSessionSearch> SearchSettingsRef = SearchSettings.ToSharedRef();
 
-		SessionInt->AddOnFindSessionsCompleteDelegate(OnFindSessionsCompleteDelegate);
+		OnFindSessionsCompleteDelegateHandle = SessionInt->AddOnFindSessionsCompleteDelegate_Handle(OnFindSessionsCompleteDelegate);
 		SessionInt->FindSessions(0, SearchSettingsRef);
 	}
 
 	// Always on the lookout for invite acceptance (via actual invite or join in an external client)
 	OnSessionInviteAcceptedDelegate = FOnSessionInviteAcceptedDelegate::CreateRaw(this, &FTestSessionInterface::OnSessionInviteAccepted);
-	SessionInt->AddOnSessionInviteAcceptedDelegate(0, OnSessionInviteAcceptedDelegate);
+	OnSessionInviteAcceptedDelegateHandle = SessionInt->AddOnSessionInviteAcceptedDelegate_Handle(0, OnSessionInviteAcceptedDelegate);
 }
 
 void FTestSessionInterface::ClearDelegates()
 {
-	SessionInt->ClearOnSessionInviteAcceptedDelegate(0, OnSessionInviteAcceptedDelegate);
+	SessionInt->ClearOnSessionInviteAcceptedDelegate_Handle(0, OnSessionInviteAcceptedDelegateHandle);
 	GEngine->OnWorldDestroyed().RemoveAll(this);
 }
 
@@ -194,39 +194,39 @@ void FTestSessionInterface::OnSessionInviteAccepted(int32 LocalUserNum, bool bWa
 void FTestSessionInterface::OnEndForJoinSessionComplete(FName SessionName, bool bWasSuccessful)
 {
 	UE_LOG(LogOnline, Verbose, TEXT("OnEndForJoinSessionComplete %s bSuccess: %d"), *SessionName.ToString(), bWasSuccessful);
-	SessionInt->ClearOnEndSessionCompleteDelegate(OnEndForJoinSessionCompleteDelegate);
+	SessionInt->ClearOnEndSessionCompleteDelegate_Handle(OnEndForJoinSessionCompleteDelegateHandle);
 	DestroyExistingSession(SessionName, OnDestroyForJoinSessionCompleteDelegate);
 }
 
 void FTestSessionInterface::EndExistingSession(FName SessionName, FOnEndSessionCompleteDelegate& Delegate)
 {
-	SessionInt->AddOnEndSessionCompleteDelegate(Delegate);
+	SessionInt->AddOnEndSessionCompleteDelegate_Handle(Delegate);
 	SessionInt->EndSession(SessionName);
 }
 
 void FTestSessionInterface::OnDestroyForJoinSessionComplete(FName SessionName, bool bWasSuccessful)
 {
 	UE_LOG(LogOnline, Verbose, TEXT("OnDestroyForJoinSessionComplete %s bSuccess: %d"), *SessionName.ToString(), bWasSuccessful);
-	SessionInt->ClearOnDestroySessionCompleteDelegate(OnDestroyForJoinSessionCompleteDelegate);
+	SessionInt->ClearOnDestroySessionCompleteDelegate_Handle(OnDestroyForJoinSessionCompleteDelegateHandle);
 	JoinSession(0, SessionName, CachedSessionResult);
 }
 
 void FTestSessionInterface::DestroyExistingSession(FName SessionName, FOnDestroySessionCompleteDelegate& Delegate)
 {
-	SessionInt->AddOnDestroySessionCompleteDelegate(Delegate);
+	SessionInt->AddOnDestroySessionCompleteDelegate_Handle(Delegate);
 	SessionInt->DestroySession(SessionName);
 }
 
 void FTestSessionInterface::OnRegisterPlayerComplete(FName SessionName, const TArray< TSharedRef<FUniqueNetId> >& Players, bool bWasSuccessful)
 {
 	UE_LOG(LogOnline, Verbose, TEXT("OnRegisterPlayerComplete %s bSuccess: %d"), *SessionName.ToString(), bWasSuccessful);
-	SessionInt->ClearOnRegisterPlayersCompleteDelegate(OnRegisterPlayersCompleteDelegate);
+	SessionInt->ClearOnRegisterPlayersCompleteDelegate_Handle(OnRegisterPlayersCompleteDelegateHandle);
 }
 
 void FTestSessionInterface::OnUnregisterPlayerComplete(FName SessionName, const TArray< TSharedRef<FUniqueNetId> >& Players, bool bWasSuccessful)
 {
 	UE_LOG(LogOnline, Verbose, TEXT("OnUnregisterPlayerComplete %s bSuccess: %d"), *SessionName.ToString(), bWasSuccessful);
-	SessionInt->ClearOnUnregisterPlayersCompleteDelegate(OnUnregisterPlayersCompleteDelegate);
+	SessionInt->ClearOnUnregisterPlayersCompleteDelegate_Handle(OnUnregisterPlayersCompleteDelegateHandle);
 }
 
 void FTestSessionInterface::JoinSession(int32 LocalUserNum, FName SessionName, const FOnlineSessionSearchResult& SearchResult)
@@ -240,7 +240,7 @@ void FTestSessionInterface::JoinSession(int32 LocalUserNum, FName SessionName, c
 	}
 	else
 	{
-		SessionInt->AddOnJoinSessionCompleteDelegate(OnJoinSessionCompleteDelegate);
+		OnJoinSessionCompleteDelegateHandle = SessionInt->AddOnJoinSessionCompleteDelegate_Handle(OnJoinSessionCompleteDelegate);
 		SessionInt->JoinSession(LocalUserNum, SessionName, SearchResult);
 	}
 }
@@ -248,7 +248,7 @@ void FTestSessionInterface::JoinSession(int32 LocalUserNum, FName SessionName, c
 void FTestSessionInterface::OnCreateSessionComplete(FName SessionName, bool bWasSuccessful)
 {
 	UE_LOG(LogOnline, Verbose, TEXT("OnCreateSessionComplete %s bSuccess: %d"), *SessionName.ToString(), bWasSuccessful);
-	SessionInt->ClearOnCreateSessionCompleteDelegate(OnCreateSessionCompleteDelegate);
+	SessionInt->ClearOnCreateSessionCompleteDelegate_Handle(OnCreateSessionCompleteDelegateHandle);
 
 	bWasLastCommandSuccessful = bWasSuccessful;
 }
@@ -256,7 +256,7 @@ void FTestSessionInterface::OnCreateSessionComplete(FName SessionName, bool bWas
 void FTestSessionInterface::OnStartSessionComplete(FName SessionName, bool bWasSuccessful)
 {
 	UE_LOG(LogOnline, Verbose, TEXT("OnStartSessionComplete %s bSuccess: %d"), *SessionName.ToString(), bWasSuccessful);
-	SessionInt->ClearOnStartSessionCompleteDelegate(OnStartSessionCompleteDelegate);
+	SessionInt->ClearOnStartSessionCompleteDelegate_Handle(OnStartSessionCompleteDelegateHandle);
 
 	bWasLastCommandSuccessful = bWasSuccessful;
 }
@@ -264,7 +264,7 @@ void FTestSessionInterface::OnStartSessionComplete(FName SessionName, bool bWasS
 void FTestSessionInterface::OnEndSessionComplete(FName SessionName, bool bWasSuccessful)
 {
 	UE_LOG(LogOnline, Verbose, TEXT("OnEndSessionComplete %s bSuccess: %d"), *SessionName.ToString(), bWasSuccessful);
-	SessionInt->ClearOnEndSessionCompleteDelegate(OnEndSessionCompleteDelegate);
+	SessionInt->ClearOnEndSessionCompleteDelegate_Handle(OnEndSessionCompleteDelegateHandle);
 
 	bWasLastCommandSuccessful = bWasSuccessful;
 }
@@ -272,19 +272,19 @@ void FTestSessionInterface::OnEndSessionComplete(FName SessionName, bool bWasSuc
 void FTestSessionInterface::OnDestroySessionComplete(FName SessionName, bool bWasSuccessful)
 {
 	UE_LOG(LogOnline, Verbose, TEXT("OnDestroySessionComplete %s bSuccess: %d"), *SessionName.ToString(), bWasSuccessful);
-	SessionInt->ClearOnDestroySessionCompleteDelegate(OnDestroySessionCompleteDelegate);
+	SessionInt->ClearOnDestroySessionCompleteDelegate_Handle(OnDestroySessionCompleteDelegateHandle);
 }
 
 void FTestSessionInterface::OnUpdateSessionComplete(FName SessionName, bool bWasSuccessful)
 {
 	UE_LOG(LogOnline, Verbose, TEXT("OnUpdateSessionComplete %s bSuccess: %d"), *SessionName.ToString(), bWasSuccessful);
-	SessionInt->ClearOnUpdateSessionCompleteDelegate(OnUpdateSessionCompleteDelegate);
+	SessionInt->ClearOnUpdateSessionCompleteDelegate_Handle(OnUpdateSessionCompleteDelegateHandle);
 }
 
 void FTestSessionInterface::OnJoinSessionComplete(FName SessionName, EOnJoinSessionCompleteResult::Type Result)
 {
 	UE_LOG(LogOnline, Verbose, TEXT("OnJoinSessionComplete %s bSuccess: %d"), *SessionName.ToString(), static_cast<int32>(Result));
-	SessionInt->ClearOnJoinSessionCompleteDelegate(OnJoinSessionCompleteDelegate);
+	SessionInt->ClearOnJoinSessionCompleteDelegate_Handle(OnJoinSessionCompleteDelegateHandle);
 
 	if (Result == EOnJoinSessionCompleteResult::Success)
 	{
@@ -307,7 +307,8 @@ void FTestSessionInterface::OnJoinSessionComplete(FName SessionName, EOnJoinSess
 void FTestSessionInterface::OnFindFriendSessionComplete(int32 LocalUserNum, bool bWasSuccessful, const FOnlineSessionSearchResult& SearchResult)
 {
 	UE_LOG(LogOnline, Verbose, TEXT("OnFindFriendSessionComplete LocalUserNum: %d bSuccess: %d"), LocalUserNum, bWasSuccessful);
-	SessionInt->ClearOnFindFriendSessionCompleteDelegate(LocalUserNum, OnFindFriendSessionCompleteDelegate);
+	SessionInt->ClearOnFindFriendSessionCompleteDelegate_Handle(LocalUserNum, OnFindFriendSessionCompleteDelegateHandles.FindRef(LocalUserNum));
+	OnFindFriendSessionCompleteDelegateHandles.Remove(LocalUserNum);
 	if (bWasSuccessful)
 	{
 		if (SearchResult.IsValid())
@@ -324,7 +325,7 @@ void FTestSessionInterface::OnFindFriendSessionComplete(int32 LocalUserNum, bool
 void FTestSessionInterface::OnFindSessionsComplete(bool bWasSuccessful)
 {
 	UE_LOG(LogOnline, Verbose, TEXT("OnFindSessionsComplete bSuccess: %d"), bWasSuccessful);
-	SessionInt->ClearOnFindSessionsCompleteDelegate(OnFindSessionsCompleteDelegate);
+	SessionInt->ClearOnFindSessionsCompleteDelegate_Handle(OnFindSessionsCompleteDelegateHandle);
 
 	UE_LOG(LogOnline, Verbose, TEXT("Num Search Results: %d"), SearchSettings->SearchResults.Num());
 	for (int32 SearchIdx=0; SearchIdx<SearchSettings->SearchResults.Num(); SearchIdx++)
@@ -337,7 +338,7 @@ void FTestSessionInterface::OnFindSessionsComplete(bool bWasSuccessful)
 void FTestSessionInterface::OnCancelFindSessionsComplete(bool bWasSuccessful)
 {
 	UE_LOG(LogOnline, Verbose, TEXT("OnCancelFindSessionsComplete bSuccess: %d"), bWasSuccessful);
-	SessionInt->ClearOnCancelFindSessionsCompleteDelegate(OnCancelFindSessionsCompleteDelegate);
+	SessionInt->ClearOnCancelFindSessionsCompleteDelegate_Handle(OnCancelFindSessionsCompleteDelegateHandle);
 }
 
 bool FTestSessionInterface::Tick(float DeltaTime)
@@ -387,7 +388,7 @@ bool FTestSessionInterface::Exec(UWorld* InWorld, const TCHAR* Cmd, FOutputDevic
 			{
 				TSharedRef<FOnlineSessionSearch> SearchSettingsRef = SearchSettings.ToSharedRef();
 
-				SessionInt->AddOnFindSessionsCompleteDelegate(OnFindSessionsCompleteDelegate);
+				OnFindSessionsCompleteDelegateHandle = SessionInt->AddOnFindSessionsCompleteDelegate_Handle(OnFindSessionsCompleteDelegate);
 				SessionInt->FindSessions(LocalUserNum, SearchSettingsRef);
 			}
 			bWasHandled = true;
@@ -417,7 +418,7 @@ bool FTestSessionInterface::Exec(UWorld* InWorld, const TCHAR* Cmd, FOutputDevic
 						const FOnlineFriend& Friend = *FriendsCache[FriendIdx];
 						if (Friend.GetDisplayName() == FriendNameStr)
 						{
-							SessionInt->AddOnFindFriendSessionCompleteDelegate(LocalUserNum, OnFindFriendSessionCompleteDelegate);
+							OnFindFriendSessionCompleteDelegateHandles.Add(LocalUserNum, SessionInt->AddOnFindFriendSessionCompleteDelegate_Handle(LocalUserNum, OnFindFriendSessionCompleteDelegate));
 							SessionInt->FindFriendSession(LocalUserNum, *Friend.GetUserId());
 							break;
 						}
@@ -430,7 +431,7 @@ bool FTestSessionInterface::Exec(UWorld* InWorld, const TCHAR* Cmd, FOutputDevic
 				if (FParse::Token(Cmd, FriendIdStr, ARRAY_COUNT(FriendIdStr), true))
 				{
 					TSharedPtr<FUniqueNetId> FriendId = Identity->CreateUniquePlayerId((uint8*)FriendIdStr, FCString::Strlen(FriendIdStr));
-					SessionInt->AddOnFindFriendSessionCompleteDelegate(LocalUserNum, OnFindFriendSessionCompleteDelegate);
+					OnFindFriendSessionCompleteDelegateHandles.Add(LocalUserNum, SessionInt->AddOnFindFriendSessionCompleteDelegate_Handle(LocalUserNum, OnFindFriendSessionCompleteDelegate));
 					SessionInt->FindFriendSession(LocalUserNum, *FriendId);
 				}
 			}
@@ -446,20 +447,20 @@ bool FTestSessionInterface::Exec(UWorld* InWorld, const TCHAR* Cmd, FOutputDevic
 			HostSettings->bUsesPresence = bTestPresence;
 			HostSettings->AddWorldSettings(InWorld);
 
-			SessionInt->AddOnCreateSessionCompleteDelegate(OnCreateSessionCompleteDelegate);
+			OnCreateSessionCompleteDelegateHandle = SessionInt->AddOnCreateSessionCompleteDelegate_Handle(OnCreateSessionCompleteDelegate);
 			SessionInt->CreateSession(0, SessionName, *HostSettings);
 			bWasHandled = true;
 		}
 		else if (FParse::Command(&Cmd, TEXT("START")))
 		{
-			SessionInt->AddOnStartSessionCompleteDelegate(OnStartSessionCompleteDelegate);
+			OnStartSessionCompleteDelegateHandle = SessionInt->AddOnStartSessionCompleteDelegate_Handle(OnStartSessionCompleteDelegate);
 			SessionInt->StartSession(SessionName);
 			bWasHandled = true;
 		}
 		else if (FParse::Command(&Cmd, TEXT("UPDATE")))
 		{
 			bool bUpdateOnline = FParse::Command(&Cmd, TEXT("ONLINE")) ? true : false;
-			SessionInt->AddOnUpdateSessionCompleteDelegate(OnUpdateSessionCompleteDelegate);
+			OnUpdateSessionCompleteDelegateHandle = SessionInt->AddOnUpdateSessionCompleteDelegate_Handle(OnUpdateSessionCompleteDelegate);
 			HostSettings->Set(FName(TEXT("UPDATESETTING1")), FString(TEXT("Test1")), EOnlineDataAdvertisementType::ViaOnlineService);
 			HostSettings->Set(FName(TEXT("UPDATESETTING2")), 99, EOnlineDataAdvertisementType::ViaOnlineService);
 			SessionInt->UpdateSession(SessionName, *HostSettings, bUpdateOnline);
@@ -478,13 +479,13 @@ bool FTestSessionInterface::Exec(UWorld* InWorld, const TCHAR* Cmd, FOutputDevic
 		else if (FParse::Command(&Cmd, TEXT("REGISTER")))
 		{
 			bool bWasInvited = FParse::Command(&Cmd, TEXT("INVITED")) ? true : false;
-			SessionInt->AddOnRegisterPlayersCompleteDelegate(OnRegisterPlayersCompleteDelegate);
+			OnRegisterPlayersCompleteDelegateHandle = SessionInt->AddOnRegisterPlayersCompleteDelegate_Handle(OnRegisterPlayersCompleteDelegate);
 			SessionInt->RegisterPlayer(SessionName, *UserId, bWasInvited);
 			bWasHandled = true;
 		}
 		else if (FParse::Command(&Cmd, TEXT("UNREGISTER")))
 		{
-			SessionInt->AddOnUnregisterPlayersCompleteDelegate(OnUnregisterPlayersCompleteDelegate);
+			OnUnregisterPlayersCompleteDelegateHandle = SessionInt->AddOnUnregisterPlayersCompleteDelegate_Handle(OnUnregisterPlayersCompleteDelegate);
 			SessionInt->UnregisterPlayer(SessionName, *UserId);
 			bWasHandled = true;
 		}
diff --git a/Engine/Source/Runtime/Online/OnlineSubsystemUtils/Private/Tests/TestSessionInterface.h b/Engine/Source/Runtime/Online/OnlineSubsystemUtils/Private/Tests/TestSessionInterface.h
--- a/Engine/Source/Runtime/Online/OnlineSubsystemUtils/Private/Tests/TestSessionInterface.h
+++ b/Engine/Source/Runtime/Online/OnlineSubsystemUtils/Private/Tests/TestSessionInterface.h
@@ -54,27 +54,46 @@
     /** Delegate for unregistering player(s) with a session */
 	FOnUnregisterPlayersCompleteDelegate OnUnregisterPlayersCompleteDelegate;
 
+	/** Delegate for searching for sessions */
+	FOnFindSessionsCompleteDelegate OnFindSessionsCompleteDelegate;
+	/** Delegate for canceling a search */
+	FOnCancelFindSessionsCompleteDelegate OnCancelFindSessionsCompleteDelegate;
+
 	/** Delegate after joining a session */
 	FOnJoinSessionCompleteDelegate OnJoinSessionCompleteDelegate;
 
+	/** Delegate for destroying a session after previously ending it */
+	FOnEndSessionCompleteDelegate OnEndForJoinSessionCompleteDelegate;
+	/** Delegate for joining a new session after previously destroying it */
+	FOnDestroySessionCompleteDelegate OnDestroyForJoinSessionCompleteDelegate;
+
+	/** Handles for the above delegates */
+	FDelegateHandle OnReadFriendsListCompleteDelegateHandle;
+	FDelegateHandle OnSessionInviteAcceptedDelegateHandle;
+	FDelegateHandle OnCreateSessionCompleteDelegateHandle;
+	FDelegateHandle OnStartSessionCompleteDelegateHandle;
+	FDelegateHandle OnEndSessionCompleteDelegateHandle;
+	FDelegateHandle OnDestroySessionCompleteDelegateHandle;
+	FDelegateHandle OnUpdateSessionCompleteDelegateHandle;
+	FDelegateHandle OnRegisterPlayersCompleteDelegateHandle;
+	FDelegateHandle OnUnregisterPlayersCompleteDelegateHandle;
+	FDelegateHandle OnFindSessionsCompleteDelegateHandle;
+	FDelegateHandle OnCancelFindSessionsCompleteDelegateHandle;
+	FDelegateHandle OnJoinSessionCompleteDelegateHandle;
+	FDelegateHandle OnEndForJoinSessionCompleteDelegateHandle;
+	FDelegateHandle OnDestroyForJoinSessionCompleteDelegateHandle;
+
 	/** Delegate for joining a friend (JIP) */
 	FOnFindFriendSessionCompleteDelegate OnFindFriendSessionCompleteDelegate;
 
-	/** Delegate for searching for sessions */
-	FOnFindSessionsCompleteDelegate OnFindSessionsCompleteDelegate;
-	/** Delegate for canceling a search */
-	FOnCancelFindSessionsCompleteDelegate OnCancelFindSessionsCompleteDelegate;
+	/** Per-player delegate handles for OnFindFriendSessionComplete */
+	TMap<int32, FDelegateHandle> OnFindFriendSessionCompleteDelegateHandles;
 
 	/** Settings defined when acting as host */
 	TSharedPtr<class TestOnlineGameSettings> HostSettings;
 	/** Search settings defined when searching as client */
 	TSharedPtr<class TestOnlineSearchSettings> SearchSettings;
 
-	/** Delegate for destroying a session after previously ending it */
-	FOnEndSessionCompleteDelegate OnEndForJoinSessionCompleteDelegate;
-	/** Delegate for joining a new session after previously destroying it */
-	FOnDestroySessionCompleteDelegate OnDestroyForJoinSessionCompleteDelegate;
-
 	/** Cached invite/search result while in the process of tearing down an existing session */
 	FOnlineSessionSearchResult CachedSessionResult;
 
diff --git a/Engine/Source/Runtime/Online/OnlineSubsystemUtils/Private/Tests/TestSharingInterface.cpp b/Engine/Source/Runtime/Online/OnlineSubsystemUtils/Private/Tests/TestSharingInterface.cpp
--- a/Engine/Source/Runtime/Online/OnlineSubsystemUtils/Private/Tests/TestSharingInterface.cpp
+++ b/Engine/Source/Runtime/Online/OnlineSubsystemUtils/Private/Tests/TestSharingInterface.cpp
@@ -60,7 +60,7 @@ void FTestSharingInterface::RequestPermissionsToSharePosts()
 
 	for (int32 PlayerIndex = 0; PlayerIndex < MAX_LOCAL_PLAYERS; PlayerIndex++)
 	{
-		SharingInterface->AddOnRequestNewPublishPermissionsCompleteDelegate(PlayerIndex, RequestPermissionsToPostToFeedDelegate);
+		RequestPermissionsToPostToFeedDelegateHandles.Add(PlayerIndex, SharingInterface->AddOnRequestNewPublishPermissionsCompleteDelegate_Handle(PlayerIndex, RequestPermissionsToPostToFeedDelegate));
 		SharingInterface->RequestNewPublishPermissions(PlayerIndex, PublishPermissions, TestStatusUpdate.PostPrivacy);
 	}
 }
@@ -69,7 +69,8 @@ void FTestSharingInterface::RequestPermissionsToSharePosts()
 void FTestSharingInterface::OnStatusPostingPermissionsUpdated(int32 LocalUserNum, bool bWasSuccessful)
 {
 	UE_LOG(LogOnline, Display, TEXT("FTestSharingInterface::OnStatusPostingPermissionsUpdated() - %d"), bWasSuccessful);
-	SharingInterface->ClearOnRequestNewPublishPermissionsCompleteDelegate(LocalUserNum, RequestPermissionsToPostToFeedDelegate);
+	SharingInterface->ClearOnRequestNewPublishPermissionsCompleteDelegate_Handle(LocalUserNum, RequestPermissionsToPostToFeedDelegateHandles.FindRef(LocalUserNum));
+	RequestPermissionsToPostToFeedDelegateHandles.Remove(LocalUserNum);
 
 	if( ++ResponsesReceived == MAX_LOCAL_PLAYERS )
 	{
@@ -87,7 +88,7 @@ void FTestSharingInterface::SharePost()
 
 	for (int32 PlayerIndex = 0; PlayerIndex < MAX_LOCAL_PLAYERS; PlayerIndex++)
 	{
-		SharingInterface->AddOnSharePostCompleteDelegate(PlayerIndex, OnPostSharedDelegate);
+		OnPostSharedDelegateHandles.Add(PlayerIndex, SharingInterface->AddOnSharePostCompleteDelegate_Handle(PlayerIndex, OnPostSharedDelegate));
 		SharingInterface->ShareStatusUpdate(PlayerIndex, TestStatusUpdate);
 	}
 }
@@ -97,7 +98,8 @@ void FTestSharingInterface::OnPostShared(int32 LocalPlayer, bool bWasSuccessful)
 {
 	UE_LOG(LogOnline, Verbose, TEXT("FTestSharingInterface::OnPostShared[PlayerIdx:%i - Successful:%i]"), LocalPlayer, bWasSuccessful);
 
-	SharingInterface->ClearOnSharePostCompleteDelegate(LocalPlayer, OnPostSharedDelegate);
+	SharingInterface->ClearOnSharePostCompleteDelegate_Handle(LocalPlayer, OnPostSharedDelegateHandles.FindRef(LocalPlayer));
+	OnPostSharedDelegateHandles.Remove(LocalPlayer);
 	if( ++ResponsesReceived == MAX_LOCAL_PLAYERS )
 	{
 		RequestPermissionsToReadNewsFeed();
@@ -117,7 +119,7 @@ void FTestSharingInterface::RequestPermissionsToReadNewsFeed()
 
 	for (int32 PlayerIndex = 0; PlayerIndex < MAX_LOCAL_PLAYERS; PlayerIndex++)
 	{
-		SharingInterface->AddOnRequestNewReadPermissionsCompleteDelegate(PlayerIndex, RequestPermissionsToReadFeedDelegate);
+		RequestPermissionsToReadFeedDelegateHandles.Add(PlayerIndex, SharingInterface->AddOnRequestNewReadPermissionsCompleteDelegate_Handle(PlayerIndex, RequestPermissionsToReadFeedDelegate));
 		SharingInterface->RequestNewReadPermissions(PlayerIndex, ReadPermissions);
 	}
 }
@@ -126,7 +128,8 @@ void FTestSharingInterface::RequestPermissionsToReadNewsFeed()
 void FTestSharingInterface::OnReadFeedPermissionsUpdated(int32 LocalUserNum, bool bWasSuccessful)
 {
 	UE_LOG(LogOnline, Display, TEXT("FTestSharingInterface::OnReadFeedPermissionsUpdated() - %d"), bWasSuccessful);
-	SharingInterface->ClearOnRequestNewReadPermissionsCompleteDelegate(LocalUserNum, RequestPermissionsToReadFeedDelegate);
+	SharingInterface->ClearOnRequestNewReadPermissionsCompleteDelegate_Handle(LocalUserNum, RequestPermissionsToReadFeedDelegateHandles.FindRef(LocalUserNum));
+	RequestPermissionsToReadFeedDelegateHandles.Remove(LocalUserNum);
 
 	if( ++ResponsesReceived == MAX_LOCAL_PLAYERS )
 	{
@@ -144,7 +147,7 @@ void FTestSharingInterface::ReadNewsFeed()
 
 	for (int32 PlayerIndex = 0; PlayerIndex < MAX_LOCAL_PLAYERS; PlayerIndex++)
 	{
-		SharingInterface->AddOnReadNewsFeedCompleteDelegate(PlayerIndex, OnNewsFeedReadDelegate);
+		OnNewsFeedReadDelegateHandles.Add(PlayerIndex, SharingInterface->AddOnReadNewsFeedCompleteDelegate_Handle(PlayerIndex, OnNewsFeedReadDelegate));
 		SharingInterface->ReadNewsFeed(PlayerIndex);
 	}
 }
@@ -173,7 +176,7 @@ void FTestSharingInterface::OnNewsFeedRead(int32 LocalPlayer, bool bWasSuccessfu
 		}
 	}
 
-	SharingInterface->ClearOnReadNewsFeedCompleteDelegate(LocalPlayer, OnNewsFeedReadDelegate);
+	SharingInterface->ClearOnReadNewsFeedCompleteDelegate_Handle(LocalPlayer, OnNewsFeedReadDelegateHandles.FindRef(LocalPlayer));
 	if( ++ResponsesReceived == MAX_LOCAL_PLAYERS )
 	{
 		UE_LOG(LogOnline, Display, TEXT("FTestSharingInterface TESTS COMPLETED"), LocalPlayer, bWasSuccessful);
diff --git a/Engine/Source/Runtime/Online/OnlineSubsystemUtils/Private/Tests/TestSharingInterface.h b/Engine/Source/Runtime/Online/OnlineSubsystemUtils/Private/Tests/TestSharingInterface.h
--- a/Engine/Source/Runtime/Online/OnlineSubsystemUtils/Private/Tests/TestSharingInterface.h
+++ b/Engine/Source/Runtime/Online/OnlineSubsystemUtils/Private/Tests/TestSharingInterface.h
@@ -50,7 +50,10 @@ private:
 
 	/** Delegate used to notify this interface that permissions are updated, and that we can now post a status update */
 	FOnRequestNewPublishPermissionsCompleteDelegate RequestPermissionsToPostToFeedDelegate;
-	
+
+	/** Per-player handles for the above delegate */
+	TMap<int32, FDelegateHandle> RequestPermissionsToPostToFeedDelegateHandles;
+
 
 	/**
 	 * Use the OSS to request permissions so we can read the news feed.
@@ -71,6 +74,10 @@ private:
 	FOnRequestNewReadPermissionsCompleteDelegate RequestPermissionsToReadFeedDelegate;
 
 
+	/** Per-player handles for the above delegate */
+	TMap<int32, FDelegateHandle> RequestPermissionsToReadFeedDelegateHandles;
+
+
 private:
 
 	/**
@@ -83,6 +90,10 @@ private:
 	FOnSharePostCompleteDelegate OnPostSharedDelegate;
 
 
+	/** Per-player handles for the above delegate */
+	TMap<int32, FDelegateHandle> OnPostSharedDelegateHandles;
+
+
 	/**
 	 *	Fn bound to the FOnSharePostCompleteDelegate
 	 *
@@ -102,7 +113,11 @@ private:
 
 	/** Delegate which is reported whether reading the news feed was successful or not for the given user */
 	FOnReadNewsFeedCompleteDelegate OnNewsFeedReadDelegate;
-	
+
+
+	/** Per-player handles for the above delegate */
+	TMap<int32, FDelegateHandle> OnNewsFeedReadDelegateHandles;
+
 
 	/**
 	 *	Fn bound to the FOnReadNewsFeedCompleteDelegate
diff --git a/Engine/Source/Runtime/Online/OnlineSubsystemUtils/Private/Tests/TestTimeInterface.cpp b/Engine/Source/Runtime/Online/OnlineSubsystemUtils/Private/Tests/TestTimeInterface.cpp
--- a/Engine/Source/Runtime/Online/OnlineSubsystemUtils/Private/Tests/TestTimeInterface.cpp
+++ b/Engine/Source/Runtime/Online/OnlineSubsystemUtils/Private/Tests/TestTimeInterface.cpp
@@ -10,8 +10,8 @@ void FTestTimeInterface::Test(UWorld* InWorld)
 	if (OnlineTime.IsValid())
 	{
 		// Create and add delegate for the async call
-		OnQueryServerUtcTimeCompleteDelegate = FOnQueryServerUtcTimeCompleteDelegate::CreateRaw(this, &FTestTimeInterface::OnQueryServerUtcTimeComplete);
-		OnlineTime->AddOnQueryServerUtcTimeCompleteDelegate(OnQueryServerUtcTimeCompleteDelegate);
+		OnQueryServerUtcTimeCompleteDelegate       = FOnQueryServerUtcTimeCompleteDelegate::CreateRaw(this, &FTestTimeInterface::OnQueryServerUtcTimeComplete);
+		OnQueryServerUtcTimeCompleteDelegateHandle = OnlineTime->AddOnQueryServerUtcTimeCompleteDelegate_Handle(OnQueryServerUtcTimeCompleteDelegate);
 		// Kick off the async query for server time
 		OnlineTime->QueryServerUtcTime();
 	}
@@ -26,7 +26,7 @@ void FTestTimeInterface::Test(UWorld* InWorld)
 
 void FTestTimeInterface::OnQueryServerUtcTimeComplete(bool bWasSuccessful, const FString& DateTimeStr, const FString& Error)
 {
-	OnlineTime->ClearOnQueryServerUtcTimeCompleteDelegate(OnQueryServerUtcTimeCompleteDelegate);
+	OnlineTime->ClearOnQueryServerUtcTimeCompleteDelegate_Handle(OnQueryServerUtcTimeCompleteDelegateHandle);
 
 	if (bWasSuccessful)
 	{
diff --git a/Engine/Source/Runtime/Online/OnlineSubsystemUtils/Private/Tests/TestTimeInterface.h b/Engine/Source/Runtime/Online/OnlineSubsystemUtils/Private/Tests/TestTimeInterface.h
--- a/Engine/Source/Runtime/Online/OnlineSubsystemUtils/Private/Tests/TestTimeInterface.h
+++ b/Engine/Source/Runtime/Online/OnlineSubsystemUtils/Private/Tests/TestTimeInterface.h
@@ -16,6 +16,8 @@
 	IOnlineTimePtr OnlineTime;
 	/** Delegate to use for querying for server time */
 	FOnQueryServerUtcTimeCompleteDelegate OnQueryServerUtcTimeCompleteDelegate;
+	/** The handle to the above delegate */
+	FDelegateHandle OnQueryServerUtcTimeCompleteDelegateHandle;
 
 	/** Hidden on purpose */
 	FTestTimeInterface()
diff --git a/Engine/Source/Runtime/Online/OnlineSubsystemUtils/Private/Tests/TestTitleFileInterface.cpp b/Engine/Source/Runtime/Online/OnlineSubsystemUtils/Private/Tests/TestTitleFileInterface.cpp
--- a/Engine/Source/Runtime/Online/OnlineSubsystemUtils/Private/Tests/TestTitleFileInterface.cpp
+++ b/Engine/Source/Runtime/Online/OnlineSubsystemUtils/Private/Tests/TestTitleFileInterface.cpp
@@ -8,8 +8,8 @@ void FTestTitleFileInterface::Test(UWorld* InWorld)
 	OnlineTitleFile = Online::GetTitleFileInterface(InWorld, SubsystemName.Len() ? FName(*SubsystemName, FNAME_Find) : NAME_None);
 	if (OnlineTitleFile.IsValid())
 	{
-		OnlineTitleFile->AddOnEnumerateFilesCompleteDelegate(OnEnumerateFilesCompleteDelegate);
-		OnlineTitleFile->AddOnReadFileCompleteDelegate(OnReadFileCompleteDelegate);
+		OnEnumerateFilesCompleteDelegateHandle = OnlineTitleFile->AddOnEnumerateFilesCompleteDelegate_Handle(OnEnumerateFilesCompleteDelegate);
+		OnReadFileCompleteDelegateHandle       = OnlineTitleFile->AddOnReadFileCompleteDelegate_Handle      (OnReadFileCompleteDelegate);
 		OnlineTitleFile->EnumerateFiles();
 	}
 	else
@@ -32,8 +32,8 @@ void FTestTitleFileInterface::FinishTest()
 
 	if (OnlineTitleFile.IsValid())
 	{
-		OnlineTitleFile->ClearOnEnumerateFilesCompleteDelegate(OnEnumerateFilesCompleteDelegate);
-		OnlineTitleFile->ClearOnReadFileCompleteDelegate(OnReadFileCompleteDelegate);
+		OnlineTitleFile->ClearOnEnumerateFilesCompleteDelegate_Handle(OnEnumerateFilesCompleteDelegateHandle);
+		OnlineTitleFile->ClearOnReadFileCompleteDelegate_Handle      (OnReadFileCompleteDelegateHandle);
 	}
 	delete this;
 }
diff --git a/Engine/Source/Runtime/Online/OnlineSubsystemUtils/Private/Tests/TestTitleFileInterface.h b/Engine/Source/Runtime/Online/OnlineSubsystemUtils/Private/Tests/TestTitleFileInterface.h
--- a/Engine/Source/Runtime/Online/OnlineSubsystemUtils/Private/Tests/TestTitleFileInterface.h
+++ b/Engine/Source/Runtime/Online/OnlineSubsystemUtils/Private/Tests/TestTitleFileInterface.h
@@ -29,6 +29,10 @@ private:
 	FOnEnumerateFilesCompleteDelegate OnEnumerateFilesCompleteDelegate;
 	/** Title file download complete delegate */
 	FOnReadFileCompleteDelegate OnReadFileCompleteDelegate;
+	/** OnEnumerateFilesComplete delegate handle */
+	FDelegateHandle OnEnumerateFilesCompleteDelegateHandle;
+	/** OnReadFileComplete delegate handle */
+	FDelegateHandle OnReadFileCompleteDelegateHandle;
 
 	void OnEnumerateFilesComplete(bool bSuccess);
 
diff --git a/Engine/Source/Runtime/Online/OnlineSubsystemUtils/Private/Tests/TestUserInterface.cpp b/Engine/Source/Runtime/Online/OnlineSubsystemUtils/Private/Tests/TestUserInterface.cpp
--- a/Engine/Source/Runtime/Online/OnlineSubsystemUtils/Private/Tests/TestUserInterface.cpp
+++ b/Engine/Source/Runtime/Online/OnlineSubsystemUtils/Private/Tests/TestUserInterface.cpp
@@ -32,8 +32,8 @@ void FTestUserInterface::Test(UWorld* InWorld, const TArray<FString>& InUserIds)
 		OnlineSub->GetUserInterface().IsValid())
 	{
 		// Add our delegate for the async call
-		OnQueryUserInfoCompleteDelegate = FOnQueryUserInfoCompleteDelegate::CreateRaw(this, &FTestUserInterface::OnQueryUserInfoComplete);
-		OnlineSub->GetUserInterface()->AddOnQueryUserInfoCompleteDelegate(0, OnQueryUserInfoCompleteDelegate);
+		OnQueryUserInfoCompleteDelegate       = FOnQueryUserInfoCompleteDelegate::CreateRaw(this, &FTestUserInterface::OnQueryUserInfoComplete);
+		OnQueryUserInfoCompleteDelegateHandle = OnlineSub->GetUserInterface()->AddOnQueryUserInfoCompleteDelegate_Handle(0, OnQueryUserInfoCompleteDelegate);
 
 		// list of users to query
 		for (int32 Idx=0; Idx < InUserIds.Num(); Idx++)
@@ -84,7 +84,7 @@ void FTestUserInterface::FinishTest()
 		OnlineSub->GetUserInterface().IsValid())
 	{
 		// Clear delegates for the various async calls
-		OnlineSub->GetUserInterface()->ClearOnQueryUserInfoCompleteDelegate(0, OnQueryUserInfoCompleteDelegate);
+		OnlineSub->GetUserInterface()->ClearOnQueryUserInfoCompleteDelegate_Handle(0, OnQueryUserInfoCompleteDelegateHandle);
 	}
 	delete this;
 }
diff --git a/Engine/Source/Runtime/Online/OnlineSubsystemUtils/Private/Tests/TestUserInterface.h b/Engine/Source/Runtime/Online/OnlineSubsystemUtils/Private/Tests/TestUserInterface.h
--- a/Engine/Source/Runtime/Online/OnlineSubsystemUtils/Private/Tests/TestUserInterface.h
+++ b/Engine/Source/Runtime/Online/OnlineSubsystemUtils/Private/Tests/TestUserInterface.h
@@ -23,6 +23,8 @@ private:
 	IOnlineSubsystem* OnlineSub;
 	/** Delegate to use for querying user info list */
 	FOnQueryUserInfoCompleteDelegate OnQueryUserInfoCompleteDelegate;
+	/** OnQueryUserInfoComplete delegate handle */
+	FDelegateHandle OnQueryUserInfoCompleteDelegateHandle;
 
 	/** List of User ids to query */
 	TArray< TSharedRef<FUniqueNetId> > QueryUserIds;
diff --git a/Engine/Source/Runtime/ShaderCore/Private/ShaderCore.cpp b/Engine/Source/Runtime/ShaderCore/Private/ShaderCore.cpp
--- a/Engine/Source/Runtime/ShaderCore/Private/ShaderCore.cpp
+++ b/Engine/Source/Runtime/ShaderCore/Private/ShaderCore.cpp
@@ -56,7 +56,7 @@ public:
 			HashState.GetHash(&GGlobalShaderMapHash.Hash[0]);
 		}
 
-		IConsoleManager::Get().RegisterConsoleVariableSink(FConsoleCommandDelegate::CreateStatic(&UpdateShaderDevelopmentMode));
+		IConsoleManager::Get().RegisterConsoleVariableSink_Handle(FConsoleCommandDelegate::CreateStatic(&UpdateShaderDevelopmentMode));
 	}
 };
 
diff --git a/Engine/Source/Runtime/Slate/Private/Framework/Application/SlateApplication.cpp b/Engine/Source/Runtime/Slate/Private/Framework/Application/SlateApplication.cpp
--- a/Engine/Source/Runtime/Slate/Private/Framework/Application/SlateApplication.cpp
+++ b/Engine/Source/Runtime/Slate/Private/Framework/Application/SlateApplication.cpp
@@ -537,7 +537,7 @@ void FSlateApplication::MouseCaptorHelper::InformCurrentCaptorOfCaptureLoss(uint
 }
 
 //////////////////////////////////////////////////////////////////////////
-void FPopupSupport::RegisterClickNotification( const TSharedRef<SWidget>& NotifyWhenClickedOutsideMe, const FOnClickedOutside& InNotification )
+FDelegateHandle FPopupSupport::RegisterClickNotification( const TSharedRef<SWidget>& NotifyWhenClickedOutsideMe, const FOnClickedOutside& InNotification )
 {
 	// If the subscriber or a zone object is destroyed, the subscription is
 	// no longer active. Clean it up here so that consumers of this API have an
@@ -560,13 +560,30 @@ void FPopupSupport::RegisterClickNotification( const TSharedRef<SWidget>& Notify
 
 	// Add a new notification.
 	ClickZoneNotifications.Add( FClickSubscriber( NotifyWhenClickedOutsideMe, InNotification ) );
+
+	return ClickZoneNotifications.Last().Notification.GetHandle();
 }
 
 void FPopupSupport::UnregisterClickNotification( const FOnClickedOutside& InNotification )
 {
 	for (int32 SubscriptionIndex=0; SubscriptionIndex < ClickZoneNotifications.Num();)
 	{
-		if (ClickZoneNotifications[SubscriptionIndex].Notification == InNotification)
+		if (ClickZoneNotifications[SubscriptionIndex].Notification.DEPRECATED_Compare(InNotification))
+		{
+			ClickZoneNotifications.RemoveAtSwap(SubscriptionIndex);
+		}
+		else
+		{
+			SubscriptionIndex++;
+		}
+	}	
+}
+
+void FPopupSupport::UnregisterClickNotification( FDelegateHandle Handle )
+{
+	for (int32 SubscriptionIndex=0; SubscriptionIndex < ClickZoneNotifications.Num();)
+	{
+		if (ClickZoneNotifications[SubscriptionIndex].Notification.GetHandle() == Handle)
 		{
 			ClickZoneNotifications.RemoveAtSwap(SubscriptionIndex);
 		}
diff --git a/Engine/Source/Runtime/Slate/Private/Framework/Docking/TabManager.cpp b/Engine/Source/Runtime/Slate/Private/Framework/Docking/TabManager.cpp
--- a/Engine/Source/Runtime/Slate/Private/Framework/Docking/TabManager.cpp
+++ b/Engine/Source/Runtime/Slate/Private/Framework/Docking/TabManager.cpp
@@ -1561,19 +1561,24 @@ const TSharedRef<FGlobalTabmanager>& FGlobalTabmanager::Get()
 	return Instance;
 }
 
-void FGlobalTabmanager::OnActiveTabChanged_Subscribe( const FOnActiveTabChanged::FDelegate& InDelegate )
+FDelegateHandle FGlobalTabmanager::OnActiveTabChanged_Subscribe( const FOnActiveTabChanged::FDelegate& InDelegate )
 {
-	OnActiveTabChanged.Add( InDelegate );
+	return OnActiveTabChanged.Add( InDelegate );
 }
 
 
 
 void FGlobalTabmanager::OnActiveTabChanged_Unsubscribe( const FOnActiveTabChanged::FDelegate& InDelegate )
 {
-	OnActiveTabChanged.Remove( InDelegate );
+	OnActiveTabChanged.DEPRECATED_Remove( InDelegate );
 }
 
 
+void FGlobalTabmanager::OnActiveTabChanged_Unsubscribe( FDelegateHandle Handle )
+{
+	OnActiveTabChanged.Remove( Handle );
+}
+
 
 TSharedPtr<class SDockTab> FGlobalTabmanager::GetActiveTab() const
 {
diff --git a/Engine/Source/Runtime/Slate/Private/Widgets/Input/SMenuAnchor.cpp b/Engine/Source/Runtime/Slate/Private/Widgets/Input/SMenuAnchor.cpp
--- a/Engine/Source/Runtime/Slate/Private/Widgets/Input/SMenuAnchor.cpp
+++ b/Engine/Source/Runtime/Slate/Private/Widgets/Input/SMenuAnchor.cpp
@@ -249,7 +249,7 @@ void SMenuAnchor::RequestClosePopupWindow( const TSharedRef<SWindow>& PopupWindo
 void SMenuAnchor::OnClickedOutsidePopup()
 {
 	bDismissedThisTick = true;
-	FSlateApplication::Get().GetPopupSupport().UnregisterClickNotification( FOnClickedOutside::CreateSP(this, &SMenuAnchor::OnClickedOutsidePopup) );
+	FSlateApplication::Get().GetPopupSupport().UnregisterClickNotification( OnClickedOutsidePopupDelegateHandle );
 	SetIsOpen(false);	
 }
 
@@ -381,7 +381,7 @@ void SMenuAnchor::SetIsOpen( bool InIsOpen, const bool bFocusMenu )
 						];
 
 						// We want to support dismissing the popup widget when the user clicks outside it.
-						FSlateApplication::Get().GetPopupSupport().RegisterClickNotification( MenuContentRef, FOnClickedOutside::CreateSP( this, &SMenuAnchor::OnClickedOutsidePopup ) );
+						OnClickedOutsidePopupDelegateHandle = FSlateApplication::Get().GetPopupSupport().RegisterClickNotification( MenuContentRef, FOnClickedOutside::CreateSP( this, &SMenuAnchor::OnClickedOutsidePopup ) );
 					}
 
 					if ( bFocusMenu )
diff --git a/Engine/Source/Runtime/Slate/Public/Framework/Application/SlateApplication.h b/Engine/Source/Runtime/Slate/Public/Framework/Application/SlateApplication.h
--- a/Engine/Source/Runtime/Slate/Public/Framework/Application/SlateApplication.h
+++ b/Engine/Source/Runtime/Slate/Public/Framework/Application/SlateApplication.h
@@ -56,7 +56,7 @@ class SLATE_API FPopupSupport
 	* @param NotifyWhenClickedOutsideMe    When the user clicks outside this widget, fire a notification.
 	* @param InNotification                The notification to invoke.
 	*/
-	void RegisterClickNotification(const TSharedRef<SWidget>& NotifyWhenClickedOutsideMe, const FOnClickedOutside& InNotification);
+	FDelegateHandle RegisterClickNotification(const TSharedRef<SWidget>& NotifyWhenClickedOutsideMe, const FOnClickedOutside& InNotification);
 
 	/**
 	* NOTE: Only necessary if notification no longer desired.
@@ -64,8 +64,17 @@ class SLATE_API FPopupSupport
 	*
 	* Unregister the notification because it is no longer desired.
 	*/
+	DELEGATE_DEPRECATED("This overload of UnregisterClickNotification is deprecated, instead pass the result of RegisterClickNotification.")
 	void UnregisterClickNotification(const FOnClickedOutside& InNotification);
 
+	/**
+	* NOTE: Only necessary if notification no longer desired.
+	*       Stale notifications are cleaned up automatically.
+	*
+	* Unregister the notification because it is no longer desired.
+	*/
+	void UnregisterClickNotification(FDelegateHandle InHandle);
+
 	private:
 
 	/** A single subscription about clicks happening outside the widget. */
diff --git a/Engine/Source/Runtime/Slate/Public/Framework/Commands/InputBindingManager.h b/Engine/Source/Runtime/Slate/Public/Framework/Commands/InputBindingManager.h
--- a/Engine/Source/Runtime/Slate/Public/Framework/Commands/InputBindingManager.h
+++ b/Engine/Source/Runtime/Slate/Public/Framework/Commands/InputBindingManager.h
@@ -105,15 +105,26 @@ public:
 	void GetCommandInfosFromContext( const FName InBindingContext, TArray< TSharedPtr<FUICommandInfo> >& OutCommandInfos ) const;
 
 	/** Registers a delegate to be called when a user-defined gesture is edited */
-	void RegisterUserDefinedGestureChanged(const FOnUserDefinedGestureChanged::FDelegate& Delegate)
+	FDelegateHandle RegisterUserDefinedGestureChanged(const FOnUserDefinedGestureChanged::FDelegate& Delegate)
 	{
-		OnUserDefinedGestureChanged.Add(Delegate);
+		return OnUserDefinedGestureChanged.Add(Delegate);
 	}
 
 	/** Unregisters a delegate to be called when a user-defined gesture is edited */
+	DELEGATE_DEPRECATED("This UnregisterUserDefinedGestureChanged overload is deprecated - please remove delegates using the FDelegateHandle returned by the RegisterUserDefinedGestureChanged function.")
 	void UnregisterUserDefinedGestureChanged(const FOnUserDefinedGestureChanged::FDelegate& Delegate)
 	{
+		PRAGMA_DISABLE_DEPRECATION_WARNINGS
+
 		OnUserDefinedGestureChanged.Remove(Delegate);
+
+		PRAGMA_ENABLE_DEPRECATION_WARNINGS
+	}
+
+	/** Unregisters a delegate to be called when a user-defined gesture is edited */
+	void UnregisterUserDefinedGestureChanged(FDelegateHandle DelegateHandle)
+	{
+		OnUserDefinedGestureChanged.Remove(DelegateHandle);
 	}
 
 private:
diff --git a/Engine/Source/Runtime/Slate/Public/Framework/Commands/UIAction.h b/Engine/Source/Runtime/Slate/Public/Framework/Commands/UIAction.h
--- a/Engine/Source/Runtime/Slate/Public/Framework/Commands/UIAction.h
+++ b/Engine/Source/Runtime/Slate/Public/Framework/Commands/UIAction.h
@@ -92,23 +92,6 @@ public:
 		, IsActionVisibleDelegate(InitIsActionVisibleDelegate)
 	{ }
 
-public:
-
-	/**
-	 * Compares this UI action to another for equality.
-	 *
-	 * @param Other The other action to compare to.
-	 * @return true if both actions are equal, false otherwise.
-	 */
-	bool operator==( const FUIAction& Other ) const
-	{
-		return
-			(ExecuteAction == Other.ExecuteAction) &&
-			(CanExecuteAction == Other.CanExecuteAction) &&
-			(IsCheckedDelegate == Other.IsCheckedDelegate) &&
-			(IsActionVisibleDelegate == Other.IsActionVisibleDelegate);
-	}
-
 public:
 
 	/**
diff --git a/Engine/Source/Runtime/Slate/Public/Framework/Docking/TabManager.h b/Engine/Source/Runtime/Slate/Public/Framework/Docking/TabManager.h
--- a/Engine/Source/Runtime/Slate/Public/Framework/Docking/TabManager.h
+++ b/Engine/Source/Runtime/Slate/Public/Framework/Docking/TabManager.h
@@ -753,11 +753,15 @@ public:
 	static const TSharedRef<FGlobalTabmanager>& Get();
 
 	/** Subscribe to notifications about the active tab changing */
-	void OnActiveTabChanged_Subscribe( const FOnActiveTabChanged::FDelegate& InDelegate );
+	FDelegateHandle OnActiveTabChanged_Subscribe( const FOnActiveTabChanged::FDelegate& InDelegate );
 
 	/** Unsubscribe to notifications about the active tab changing */
+	DELEGATE_DEPRECATED("This overload of OnActiveTabChanged_Unsubscribe is deprecated, instead pass the result of OnActiveTabChanged_Subscribe.")
 	void OnActiveTabChanged_Unsubscribe( const FOnActiveTabChanged::FDelegate& InDelegate );
 
+	/** Unsubscribe to notifications about the active tab changing */
+	void OnActiveTabChanged_Unsubscribe( FDelegateHandle Handle );
+
 	/** @return the currently active tab; NULL pointer if there is no active tab */
 	TSharedPtr<SDockTab> GetActiveTab() const;
 
diff --git a/Engine/Source/Runtime/Slate/Public/Widgets/Input/SMenuAnchor.h b/Engine/Source/Runtime/Slate/Public/Widgets/Input/SMenuAnchor.h
--- a/Engine/Source/Runtime/Slate/Public/Widgets/Input/SMenuAnchor.h
+++ b/Engine/Source/Runtime/Slate/Public/Widgets/Input/SMenuAnchor.h
@@ -145,4 +145,7 @@ protected:
 	FVector2D LocalPopupPosition;
 
 	TPanelChildren<FSimpleSlot> Children;
+
+	/** Handle to the registered OnClickedOutsidePopup delegate */
+	FDelegateHandle OnClickedOutsidePopupDelegateHandle;
 };
