Never send multicast RPCs to connection with no viewer

diff --git a/Engine/Source/Runtime/Online/OnlineSubsystemUtils/Private/IpNetDriver.cpp b/Engine/Source/Runtime/Online/OnlineSubsystemUtils/Private/IpNetDriver.cpp
--- a/Engine/Source/Runtime/Online/OnlineSubsystemUtils/Private/IpNetDriver.cpp
+++ b/Engine/Source/Runtime/Online/OnlineSubsystemUtils/Private/IpNetDriver.cpp
@@ -293,7 +293,7 @@ void UIpNetDriver::ProcessRemoteFunction(class AActor* Actor, UFunction* Functio
 			for (int32 i=0; i<ClientConnections.Num(); ++i)
 			{
 				Connection = ClientConnections[i];
-				if (Connection)
+				if (Connection && Connection->Viewer)
 				{
 					// Do relevancy check if unreliable.
 					// Reliables will always go out. This is odd behavior. On one hand we wish to garuntee "reliables always get there". On the other
@@ -305,16 +305,8 @@ void UIpNetDriver::ProcessRemoteFunction(class AActor* Actor, UFunction* Functio
 					bool IsRelevant = true;
 					if ((Function->FunctionFlags & FUNC_NetReliable) == 0)
 					{
-						if (Connection->Viewer)
-						{
-							FNetViewer Viewer(Connection, 0.f);
-							IsRelevant = Actor->IsNetRelevantFor(Viewer.InViewer, Viewer.Viewer, Viewer.ViewLocation);
-						}
-						else
-						{
-							// No viewer for this connection, we're too early so don't send it
-							IsRelevant = false;
-						}
+						FNetViewer Viewer(Connection, 0.f);
+						IsRelevant = Actor->IsNetRelevantFor(Viewer.InViewer, Viewer.Viewer, Viewer.ViewLocation);
 					}
 					
 					if (IsRelevant)
