Turn Based Matches First Pass for iOS. Matchmaking, turn passing, and match quitting work through blueprints.

diff --git a/Engine/Source/Runtime/Engine/Private/RepLayout.cpp b/Engine/Source/Runtime/Engine/Private/RepLayout.cpp
--- a/Engine/Source/Runtime/Engine/Private/RepLayout.cpp
+++ b/Engine/Source/Runtime/Engine/Private/RepLayout.cpp
@@ -1002,6 +1002,22 @@ void FRepLayout::SendProperties_DynamicArray_r(
 	WritePropertyHandle( WriterState.Writer, 0, WriterState.bDoChecksum );		// Signify end of dynamic array
 }
 
+void FRepLayout::SerializeObjectReplicatedProperties(UObject* Object, FArchive & Ar) const
+{
+	for (int32 i = 0; i < Parents.Num(); i++)
+	{
+        UStructProperty* StructProperty = Cast<UStructProperty>(Parents[i].Property);
+        UObjectProperty* ObjectProperty = Cast<UObjectProperty>(Parents[i].Property);
+
+		// We're only able to easily serialize non-object/struct properties, so just do those.
+		if (ObjectProperty == nullptr && StructProperty == nullptr)
+		{
+			bool bHasUnmapped = false;
+			SerializeProperties_r(Ar, NULL, Parents[i].CmdStart, Parents[i].CmdEnd, (uint8*)Object, bHasUnmapped);
+		}
+	}
+}
+
 uint16 FRepLayout::SendProperties_r( 
 	FRepState *	RESTRICT		RepState, 
 	const FReplicationFlags &	RepFlags,
diff --git a/Engine/Source/Runtime/Engine/Public/Net/RepLayout.h b/Engine/Source/Runtime/Engine/Public/Net/RepLayout.h
--- a/Engine/Source/Runtime/Engine/Public/Net/RepLayout.h
+++ b/Engine/Source/Runtime/Engine/Public/Net/RepLayout.h
@@ -6,6 +6,11 @@
 =============================================================================*/
 #pragma once
 
+#include "CoreNet.h"
+#include "Engine/EngineTypes.h"
+
+class FOutBunch;
+
 class FRepChangedParent
 {
 public:
@@ -281,7 +286,7 @@ public:
 		int32 &						LastIndex, 
 		bool &						bContentBlockWritten ) const;
 
-	void InitFromObjectClass( UClass * InObjectClass );
+	ENGINE_API void InitFromObjectClass( UClass * InObjectClass );
 
 	bool ReceiveProperties( UClass * InObjectClass, FRepState * RESTRICT RepState, void* RESTRICT Data, FNetBitReader & InBunch, bool & bOutHasUnmapped ) const;
 	void UpdateUnmappedObjects( FRepState *	RepState, UPackageMap * PackageMap, UObject* Object, bool & bOutSomeObjectsWereMapped, bool & bOutHasMoreUnmapped ) const;
@@ -310,6 +315,9 @@ public:
 	void SerializePropertiesForStruct( UStruct * Struct, FArchive & Ar, UPackageMap	* Map, void* Data, bool & bHasUnmapped ) const;	
 	void InitFromStruct( UStruct * InStruct );
 
+	// Serializes all replicated properties of a UObject in or out of an archive (depending on what type of archive it is)
+	ENGINE_API void SerializeObjectReplicatedProperties(UObject* Object, FArchive & Ar) const;
+
 private:
 	void RebuildConditionalProperties( FRepState * RESTRICT	RepState, const FRepChangedPropertyTracker& ChangedTracker, const FReplicationFlags& RepFlags ) const;
 
diff --git a/Engine/Source/Runtime/Online/IOS/OnlineSubsystemIOS/Private/OnlineSubsystemIOS.cpp b/Engine/Source/Runtime/Online/IOS/OnlineSubsystemIOS/Private/OnlineSubsystemIOS.cpp
--- a/Engine/Source/Runtime/Online/IOS/OnlineSubsystemIOS/Private/OnlineSubsystemIOS.cpp
+++ b/Engine/Source/Runtime/Online/IOS/OnlineSubsystemIOS/Private/OnlineSubsystemIOS.cpp
@@ -116,6 +116,11 @@ IOnlineChatPtr FOnlineSubsystemIOS::GetChatInterface() const
 	return NULL;
 }
 
+IOnlineTurnBasedPtr FOnlineSubsystemIOS::GetTurnBasedInterface() const
+{
+    return TurnBasedInterface;
+}
+
 bool FOnlineSubsystemIOS::Init() 
 {
 	bool bSuccessfullyStartedUp = true;
@@ -140,6 +145,7 @@ bool FOnlineSubsystemIOS::Init()
 		LeaderboardsInterface = MakeShareable(new FOnlineLeaderboardsIOS(this));
 		AchievementsInterface = MakeShareable(new FOnlineAchievementsIOS(this));
 		ExternalUIInterface = MakeShareable(new FOnlineExternalUIIOS());
+        TurnBasedInterface = MakeShareable(new FOnlineTurnBasedIOS());
 	}
 
 	if( IsInAppPurchasingEnabled() )
diff --git a/Engine/Source/Runtime/Online/IOS/OnlineSubsystemIOS/Private/OnlineSubsystemIOSPrivatePCH.h b/Engine/Source/Runtime/Online/IOS/OnlineSubsystemIOS/Private/OnlineSubsystemIOSPrivatePCH.h
--- a/Engine/Source/Runtime/Online/IOS/OnlineSubsystemIOS/Private/OnlineSubsystemIOSPrivatePCH.h
+++ b/Engine/Source/Runtime/Online/IOS/OnlineSubsystemIOS/Private/OnlineSubsystemIOSPrivatePCH.h
@@ -17,6 +17,7 @@
 #include "OnlineStoreInterfaceIOS.h"
 #include "OnlineAchievementsInterfaceIOS.h"
 #include "OnlineExternalUIInterfaceIOS.h"
+#include "OnlineTurnBasedInterfaceIOS.h"
 
 #include "OnlineSubsystemIOS.h"
 #include "OnlineSubsystemIOSModule.h"
\ No newline at end of file
diff --git a/Engine/Source/Runtime/Online/IOS/OnlineSubsystemIOS/Private/OnlineTurnBasedInterfaceIOS.cpp b/Engine/Source/Runtime/Online/IOS/OnlineSubsystemIOS/Private/OnlineTurnBasedInterfaceIOS.cpp
new file mode 100644
--- /dev/null
+++ b/Engine/Source/Runtime/Online/IOS/OnlineSubsystemIOS/Private/OnlineTurnBasedInterfaceIOS.cpp
@@ -0,0 +1,556 @@
+// Copyright 1998-2015 Epic Games, Inc. All Rights Reserved.
+
+#include "OnlineSubsystemIOSPrivatePCH.h"
+#include "OnlineTurnBasedInterfaceIOS.h"
+#include "TurnBasedEventListener.h"
+#include "OnlineIdentityInterface.h"
+#include "OnlineSubsystem.h"
+#include "OnlineAsyncTaskManager.h"
+#include "RepLayout.h"
+#include "TurnBasedMatchInterface.h"
+
+DEFINE_LOG_CATEGORY_STATIC(LogTurnBasedInterfaceIOS, Verbose, All);
+
+#define INCLUDE_LOCAL_PLAYER (true)
+#define DO_NOT_INCLUDE_LOCAL_PLAYER (false)
+
+FTurnBasedMatchIOS::FTurnBasedMatchIOS(GKTurnBasedMatch* _Match, NSArray* PlayerArray)
+	: FTurnBasedMatch()
+	, Match([_Match retain])
+{
+	if (!Match)
+	{
+		UE_LOG(LogTurnBasedInterfaceIOS, Error, TEXT("GKTurnBasedMatch required to create a FTurnBasedMatchIOS"));
+	}
+
+	for (GKPlayer* player in PlayerArray)
+	{
+		PlayerAliasArray.Add(UTF8_TO_TCHAR(player.displayName.UTF8String));
+	}
+}
+
+FTurnBasedMatchIOS::~FTurnBasedMatchIOS()
+{
+	[Match release];
+}
+
+int32 FTurnBasedMatchIOS::GetNumberOfPlayers() const
+{
+	return Match.participants.count;
+}
+
+bool FTurnBasedMatchIOS::GetPlayerDisplayName(int32 PlayerIndex, FString& Name) const
+{
+	if (PlayerIndex >= PlayerAliasArray.Num() || PlayerIndex < 0)
+	{
+		return FTurnBasedMatch::GetPlayerDisplayName(PlayerIndex, Name);
+	}
+
+	Name = PlayerAliasArray[PlayerIndex];
+	return true;
+}
+
+void FTurnBasedMatchIOS::ReloadMatchData(FDownloadMatchDataSignature DownloadCallback)
+{
+	[Match loadMatchDataWithCompletionHandler : ^ (NSData* data, NSError* error)
+	{
+		if (DownloadCallback.IsBound())
+		{
+			DownloadCallback.Execute(GetMatchID(), error == nil);
+		}
+	}];
+}
+
+bool FTurnBasedMatchIOS::HasMatchData() const
+{
+	return Match.matchData != nil && Match.matchData.length != 0;
+}
+
+bool FTurnBasedMatchIOS::GetMatchData(TArray<uint8>& OutMatchData) const
+{
+	if (!Match.matchData || Match.matchData.length == 0)
+	{
+		return false;
+	}
+
+	int32 dataSize = Match.matchData.length;
+	OutMatchData.Empty(dataSize);
+	OutMatchData.AddUninitialized(dataSize - OutMatchData.Num());
+	FMemory::Memcpy(OutMatchData.GetData(), Match.matchData.bytes, dataSize);
+	return true;
+}
+
+void FTurnBasedMatchIOS::SetMatchData(const TArray<uint8>& NewMatchData, FUploadMatchDataSignature UploadCallback)
+{
+	[Match saveCurrentTurnWithMatchData : [NSData dataWithBytes : NewMatchData.GetData() length : NewMatchData.Num()] completionHandler : ^ (NSError* error)
+	{
+		UploadCallback.Execute(GetMatchID(), error == nil);
+	}];
+}
+
+FString FTurnBasedMatchIOS::GetMatchID() const
+{
+	if (!Match)
+	{
+		return FTurnBasedMatch::GetMatchID();
+	}
+	return FString::Printf(TEXT("%s"), UTF8_TO_TCHAR(Match.matchID.UTF8String));
+}
+
+int32 FTurnBasedMatchIOS::GetLocalPlayerIndex() const
+{
+	IOnlineSubsystem* OSS = IOnlineSubsystem::Get();
+	IOnlineIdentityPtr IdentityInterface = OSS ? OSS->GetIdentityInterface() : NULL;
+	if (!IdentityInterface.IsValid())
+	{
+		UE_LOG(LogTurnBasedInterfaceIOS, Warning, TEXT("No Online Identity"));
+		return 0;
+	}
+
+	TSharedPtr<FUniqueNetId> NetID = IdentityInterface->GetUniquePlayerId(0);
+	NSString* playerID = [NSString stringWithFormat : @"%s", TCHAR_TO_UTF8(*(NetID->ToString()))];
+
+	int32 PlayerIndex = 0;
+	NSArray* participantArray = Match.participants;
+	for (GKTurnBasedParticipant* participant in participantArray)
+	{
+		if ([playerID isEqualToString : participant.playerID])
+		{
+			return PlayerIndex;
+		}
+		++PlayerIndex;
+	}
+	return PlayerIndex;
+}
+
+int32 FTurnBasedMatchIOS::GetCurrentPlayerIndex() const
+{
+	GKTurnBasedParticipant* currentParticipant = Match.currentParticipant;
+	if (!currentParticipant)
+	{
+		return 0;
+	}
+
+	return[Match.participants indexOfObject : currentParticipant];
+}
+
+EMPMatchOutcome::Outcome FTurnBasedMatchIOS::GetMatchOutcomeForPlayer(int32 PlayerIndex) const
+{
+	if (PlayerIndex >= Match.participants.count)
+	{
+		return EMPMatchOutcome::None;
+	}
+
+	GKTurnBasedParticipant* participant = Match.participants[PlayerIndex];
+	GKTurnBasedMatchOutcome GKOutcome = participant.matchOutcome;
+	return GetMatchOutcomeFromGKTurnBasedMatchOutcome(GKOutcome);
+}
+
+int32 FTurnBasedMatchIOS::GetPlayerIndexForPlayer(NSString* PlayerID) const
+{
+	int32 playerIndex = 0;
+	for (GKTurnBasedParticipant* participant in Match.participants)
+	{
+		if ([participant.playerID isEqualToString : PlayerID])
+		{
+			return playerIndex;
+		}
+	}
+	UE_LOG(LogTurnBasedInterfaceIOS, Warning, TEXT("Failed to find participant %s in match"), PlayerID.UTF8String);
+	return 0;
+}
+
+bool FTurnBasedMatchIOS::IsGKTurnBasedMatch(GKTurnBasedMatch* Comparison) const
+{
+	if (!Comparison || !Match)
+	{
+		return false;
+	}
+
+	return[Comparison.matchID isEqualToString : Match.matchID];
+}
+
+void FTurnBasedMatchIOS::EndTurnWithMatchData(const TArray<uint8>& MatchData, int32 TurnTimeoutInSeconds, FUploadMatchDataSignature UploadCallback)
+{
+	NSArray* participantArray = GetNextParticipantArray(INCLUDE_LOCAL_PLAYER);
+	[Match endTurnWithNextParticipants : participantArray
+	turnTimeout : TurnTimeoutInSeconds
+			  matchData : [NSData dataWithBytes : MatchData.GetData() length : MatchData.Num()]
+					  completionHandler : ^ (NSError* error)
+	{
+		if (UploadCallback.IsBound())
+		{
+			UploadCallback.Execute(GetMatchID(), error == nil);
+		}
+	}];
+}
+
+void FTurnBasedMatchIOS::QuitMatch(EMPMatchOutcome::Outcome Outcome, int32 TurnTimeoutInSeconds, FQuitMatchSignature QuitMatchCallback)
+{
+	int32 localPlayerIndex = GetLocalPlayerIndex();
+	if (localPlayerIndex < Match.participants.count)
+	{
+		GKTurnBasedParticipant* localParticipant = Match.participants[localPlayerIndex];
+		if (localParticipant.matchOutcome != GKTurnBasedMatchOutcomeNone)
+		{
+			return;
+		}
+	}
+
+	GKTurnBasedMatchOutcome GKOutcome = GetGKTurnBasedMatchOutcomeFromMatchOutcome(Outcome);
+	if (localPlayerIndex == GetCurrentPlayerIndex())
+	{
+		QuitMatchInTurn(GKOutcome, TurnTimeoutInSeconds, QuitMatchCallback);
+	}
+	else
+	{
+		QuitMatchOutOfTurn(GKOutcome, QuitMatchCallback);
+	}
+}
+
+void FTurnBasedMatchIOS::QuitMatchInTurn(GKTurnBasedMatchOutcome Outcome, int32 TurnTimeoutInSeconds, FQuitMatchSignature QuitMatchCallback)
+{
+	FString MatchID = GetMatchID();
+	NSArray* participantArray = GetNextParticipantArray(DO_NOT_INCLUDE_LOCAL_PLAYER);
+	[Match participantQuitInTurnWithOutcome : Outcome nextParticipants : participantArray turnTimeout : TurnTimeoutInSeconds matchData : Match.matchData completionHandler : ^ (NSError* error)
+	{
+		if (QuitMatchCallback.IsBound())
+		{
+			QuitMatchCallback.Execute(MatchID, error == nil);
+		}
+	}];
+}
+
+void FTurnBasedMatchIOS::QuitMatchOutOfTurn(GKTurnBasedMatchOutcome Outcome, FQuitMatchSignature QuitMatchCallback)
+{
+	FString MatchID = GetMatchID();
+	[Match participantQuitOutOfTurnWithOutcome : Outcome withCompletionHandler : ^ (NSError* error)
+	{
+		if (QuitMatchCallback.IsBound())
+		{
+			QuitMatchCallback.Execute(MatchID, error == nil);
+		}
+	}];
+}
+
+NSArray* FTurnBasedMatchIOS::GetNextParticipantArray(bool IncludeLocalPlayer) const
+{
+	NSArray* participantArray = Match.participants;
+
+	int32 localPlayerIndex = GetLocalPlayerIndex();
+	int32 nextPlayerIndex = (localPlayerIndex + 1) % participantArray.count;
+	NSMutableArray* nextParticipantArray = [NSMutableArray array];
+	while (nextPlayerIndex != localPlayerIndex)
+	{
+		[nextParticipantArray addObject : [participantArray objectAtIndex : nextPlayerIndex]];
+		nextPlayerIndex = (nextPlayerIndex + 1) % participantArray.count;
+	}
+
+	if (IncludeLocalPlayer)
+	{
+		[nextParticipantArray addObject : [participantArray objectAtIndex : localPlayerIndex]];
+	}
+
+	return nextParticipantArray;
+}
+
+GKTurnBasedMatchOutcome FTurnBasedMatchIOS::GetGKTurnBasedMatchOutcomeFromMatchOutcome(EMPMatchOutcome::Outcome Outcome) const
+{
+	switch (Outcome)
+	{
+	default:
+	case EMPMatchOutcome::None:        return GKTurnBasedMatchOutcomeNone;
+	case EMPMatchOutcome::Quit:        return GKTurnBasedMatchOutcomeQuit;
+	case EMPMatchOutcome::Won:         return GKTurnBasedMatchOutcomeWon;
+	case EMPMatchOutcome::Lost:        return GKTurnBasedMatchOutcomeLost;
+	case EMPMatchOutcome::Tied:        return GKTurnBasedMatchOutcomeTied;
+	case EMPMatchOutcome::TimeExpired: return GKTurnBasedMatchOutcomeTimeExpired;
+	case EMPMatchOutcome::First:       return GKTurnBasedMatchOutcomeFirst;
+	case EMPMatchOutcome::Second:      return GKTurnBasedMatchOutcomeSecond;
+	case EMPMatchOutcome::Third:       return GKTurnBasedMatchOutcomeThird;
+	case EMPMatchOutcome::Fourth:      return GKTurnBasedMatchOutcomeFourth;
+	}
+}
+
+EMPMatchOutcome::Outcome FTurnBasedMatchIOS::GetMatchOutcomeFromGKTurnBasedMatchOutcome(GKTurnBasedMatchOutcome GKOutcome) const
+{
+	switch (GKOutcome)
+	{
+	default:
+	case GKTurnBasedMatchOutcomeNone:           return EMPMatchOutcome::None;
+	case GKTurnBasedMatchOutcomeQuit:           return EMPMatchOutcome::Quit;
+	case GKTurnBasedMatchOutcomeWon:            return EMPMatchOutcome::Won;
+	case GKTurnBasedMatchOutcomeLost:           return EMPMatchOutcome::Lost;
+	case GKTurnBasedMatchOutcomeTied:           return EMPMatchOutcome::Tied;
+	case GKTurnBasedMatchOutcomeTimeExpired:    return EMPMatchOutcome::TimeExpired;
+	case GKTurnBasedMatchOutcomeFirst:          return EMPMatchOutcome::First;
+	case GKTurnBasedMatchOutcomeSecond:         return EMPMatchOutcome::Second;
+	case GKTurnBasedMatchOutcomeThird:          return EMPMatchOutcome::Third;
+	case GKTurnBasedMatchOutcomeFourth:         return EMPMatchOutcome::Fourth;
+	}
+}
+
+void FTurnBasedMatchIOS::SetGKMatch(GKTurnBasedMatch* GKMatch)
+{
+	[Match release];
+	Match = [GKMatch retain];
+}
+FOnlineTurnBasedIOS::FOnlineTurnBasedIOS()
+	: IOnlineTurnBased()
+	, FTurnBasedMatchmakerDelegate()
+	, Matchmaker(*this)
+	, MatchmakerDelegate(NULL)
+	, EventListener(nil)
+	, EventDelegate(NULL)
+	, NumberOfMatchesBeingLoaded(0)
+{
+	EventListener = [[FTurnBasedEventListenerIOS alloc] initWithOwner:*this];
+}
+
+FOnlineTurnBasedIOS::~FOnlineTurnBasedIOS()
+{
+	[EventListener release];
+}
+
+void FOnlineTurnBasedIOS::SetMatchmakerDelegate(FTurnBasedMatchmakerDelegatePtr Delegate)
+{
+	MatchmakerDelegate = Delegate;
+}
+
+void FOnlineTurnBasedIOS::ShowMatchmaker(const FTurnBasedMatchRequest& MatchRequest)
+{
+	Matchmaker.ShowWithMatchRequest(MatchRequest);
+}
+
+void FOnlineTurnBasedIOS::SetEventDelegate(FTurnBasedEventDelegateWeakPtr Delegate)
+{
+	EventDelegate = Delegate;
+	if (!EventListener && Delegate.IsValid())
+	{
+		EventListener = [[FTurnBasedEventListenerIOS alloc] initWithOwner:*this];
+	}
+	else if (EventListener && !Delegate.IsValid())
+	{
+		[EventListener release];
+		EventListener = nil;
+	}
+}
+
+FTurnBasedEventDelegateWeakPtr FOnlineTurnBasedIOS::GetEventDelegate() const
+{
+	return EventDelegate;
+}
+
+void FOnlineTurnBasedIOS::LoadAllMatches(FLoadTurnBasedMatchesSignature MatchesLoadedCallback)
+{
+	if (NumberOfMatchesBeingLoaded > 0)
+	{
+		UE_LOG(LogTurnBasedInterfaceIOS, Warning, TEXT("Requesting load all matches whilst we are still loading matches"));
+		return;
+	}
+
+	[GKTurnBasedMatch loadMatchesWithCompletionHandler : ^ (NSArray* matches, NSError* error)
+	{
+		MatchArray.Empty();
+
+		NumberOfMatchesBeingLoaded = matches.count;
+		for (GKTurnBasedMatch* match in matches)
+		{
+			NSArray* playerIdentifierArray = FOnlineTurnBasedIOS::GetPlayerIdentifierArrayForMatch(match);
+			[GKPlayer loadPlayersForIdentifiers : playerIdentifierArray withCompletionHandler : ^ (NSArray *players, NSError *nameLoadError)
+			{
+				if (!nameLoadError)
+				{
+					MatchArray.Add(MakeShareable(new FTurnBasedMatchIOS(match, players)));
+				}
+
+				--NumberOfMatchesBeingLoaded;
+				if (NumberOfMatchesBeingLoaded == 0)
+				{
+					TArray<FString> MatchIDArray;
+					for (TArray<FTurnBasedMatchRef>::TConstIterator It = MatchArray.CreateConstIterator(); It; ++It)
+					{
+						MatchIDArray.Add((*It)->GetMatchID());
+					}
+
+					MatchesLoadedCallback.ExecuteIfBound(MatchIDArray, error == nil);
+				}
+			}];
+		}
+	}];
+}
+
+void FOnlineTurnBasedIOS::LoadMatchWithID(FString MatchID, FLoadTurnBasedMatchWithIDSignature MatchLoadedCallback)
+{
+	NSString* IDString = [NSString stringWithUTF8String : TCHAR_TO_UTF8(*MatchID)];
+	[GKTurnBasedMatch loadMatchWithID : IDString
+	withCompletionHandler : ^ (GKTurnBasedMatch* match, NSError* error)
+	{
+		if (!error)
+		{
+			NSArray* playerIdentifierArray = FOnlineTurnBasedIOS::GetPlayerIdentifierArrayForMatch(match);
+			[GKPlayer loadPlayersForIdentifiers : playerIdentifierArray withCompletionHandler : ^ (NSArray *players, NSError *nameLoadError)
+			{
+				if (!nameLoadError)
+				{
+					FTurnBasedMatchPtr PreviousMatchPtr = GetMatchWithID(MatchID);
+					if (PreviousMatchPtr.IsValid())
+					{
+						MatchArray.Remove(PreviousMatchPtr.ToSharedRef());
+					}
+
+					FTurnBasedMatchRef NewMatch = MakeShareable(new FTurnBasedMatchIOS(match, players));
+					MatchArray.Add(NewMatch);
+					MatchLoadedCallback.ExecuteIfBound(NewMatch->GetMatchID(), true);
+				}
+				else
+				{
+					MatchLoadedCallback.ExecuteIfBound(TEXT(""), false);
+				}
+			}];
+		}
+		else
+		{
+			MatchLoadedCallback.ExecuteIfBound(TEXT(""), false);
+		}
+	}];
+}
+
+FTurnBasedMatchPtr FOnlineTurnBasedIOS::GetMatchWithID(FString MatchID) const
+{
+	for (TArray<FTurnBasedMatchRef>::TConstIterator It = MatchArray.CreateConstIterator(); It; ++It)
+	{
+		FTurnBasedMatch& Match = (*It).Get();
+		if (Match.GetMatchID().Compare(MatchID) == 0)
+		{
+			return *It;
+		}
+	}
+	return NULL;
+}
+
+void FOnlineTurnBasedIOS::RemoveMatch(FTurnBasedMatchRef Match, FRemoveMatchSignature RemoveMatchCallback)
+{
+	// AGC: There's almost certainly a better way to do this cast
+	FTurnBasedMatchIOS& MatchIOS = static_cast<FTurnBasedMatchIOS&>(Match.Get());
+
+	FString MatchID = MatchIOS.GetMatchID();
+	[MatchIOS.GetGKMatch() removeWithCompletionHandler:^ (NSError* error)
+	{
+		MatchArray.Remove(Match);
+		if (RemoveMatchCallback.IsBound())
+		{
+			RemoveMatchCallback.Execute(MatchID, error == nil);
+		}
+	}];
+}
+
+void FOnlineTurnBasedIOS::OnMatchmakerCancelled()
+{
+	if (!MatchmakerDelegate.IsValid())
+	{
+		return;
+	}
+	MatchmakerDelegate.Pin()->OnMatchmakerCancelled();
+}
+
+void FOnlineTurnBasedIOS::OnMatchmakerFailed()
+{
+	if (!MatchmakerDelegate.IsValid())
+	{
+		return;
+	}
+	MatchmakerDelegate.Pin()->OnMatchmakerFailed();
+}
+
+void FOnlineTurnBasedIOS::OnMatchFound(FTurnBasedMatchRef Match)
+{
+	MatchArray.Add(Match);
+
+	if (MatchmakerDelegate.IsValid())
+	{
+		dispatch_async(dispatch_get_main_queue(), ^
+		{
+			[FIOSAsyncTask CreateTaskWithBlock : ^ bool(void)
+			{
+				MatchmakerDelegate.Pin()->OnMatchFound(Match);
+				return true;
+			}];
+
+		});
+	}
+}
+
+void FOnlineTurnBasedIOS::OnMatchEnded(FString MatchID)
+{
+	if (EventDelegate.IsValid())
+	{
+		EventDelegate.Pin()->OnMatchEnded(MatchID);
+	}
+}
+
+void FOnlineTurnBasedIOS::OnMatchReceivedTurnEvent(FString MatchID, bool BecameActive, void* Match)
+{
+	GKTurnBasedMatch* IOSTurnBasedMatch = (GKTurnBasedMatch*)Match;
+    NSArray* playerIdentifierArray = FOnlineTurnBasedIOS::GetPlayerIdentifierArrayForMatch(IOSTurnBasedMatch);
+    [GKPlayer loadPlayersForIdentifiers : playerIdentifierArray withCompletionHandler : ^ (NSArray *players, NSError *nameLoadError)
+    {
+        if (!nameLoadError)
+        {
+            FTurnBasedMatchPtr PreviousMatchPtr = GetMatchWithID(MatchID);
+            if (PreviousMatchPtr.IsValid())
+            {
+                MatchArray.Remove(PreviousMatchPtr.ToSharedRef());
+            }
+
+            FTurnBasedMatchRef NewMatch = MakeShareable(new FTurnBasedMatchIOS(IOSTurnBasedMatch, players));
+            MatchArray.Add(NewMatch);
+            dispatch_async(dispatch_get_main_queue(), ^
+            {
+                [FIOSAsyncTask CreateTaskWithBlock : ^ bool(void)
+                {
+					if (TurnBasedMatchInterfaceObject)
+					{
+						TArray<uint8> MatchData;
+						if (GetMatchWithID(MatchID)->GetMatchData(MatchData))
+						{
+							FRepLayout RepLayout;
+							RepLayout.InitFromObjectClass(TurnBasedMatchInterfaceObject->GetClass());
+							FBitReader Reader(MatchData.GetData(), MATCH_DATA_SIZE);
+							RepLayout.SerializeObjectReplicatedProperties(TurnBasedMatchInterfaceObject, Reader);
+						}
+						ITurnBasedMatchInterface::Execute_OnMatchReceivedTurn(TurnBasedMatchInterfaceObject, MatchID, BecameActive);
+					}
+					if (EventDelegate.IsValid())
+					{
+						EventDelegate.Pin()->OnMatchReceivedTurnEvent(MatchID, BecameActive, Match);
+					}
+                    return true;
+                }];
+
+            });
+        }
+    }];
+}
+
+NSArray* FOnlineTurnBasedIOS::GetPlayerIdentifierArrayForMatch(GKTurnBasedMatch* match)
+{
+	NSMutableArray* result = [NSMutableArray array];
+	for (GKTurnBasedParticipant* participant in match.participants)
+	{
+		if (!participant.playerID)
+		{
+			break;
+		}
+		[result addObject : participant.playerID];
+	}
+	return result;
+}
+
+void FOnlineTurnBasedIOS::RegisterTurnBasedMatchInterfaceObject(UObject* Object)
+{
+	if (Object != nullptr && Object->GetClass()->ImplementsInterface(UTurnBasedMatchInterface::StaticClass()))
+	{
+		TurnBasedMatchInterfaceObject = Object;
+	}
+}
\ No newline at end of file
diff --git a/Engine/Source/Runtime/Online/IOS/OnlineSubsystemIOS/Private/OnlineTurnBasedInterfaceIOS.h b/Engine/Source/Runtime/Online/IOS/OnlineSubsystemIOS/Private/OnlineTurnBasedInterfaceIOS.h
new file mode 100644
--- /dev/null
+++ b/Engine/Source/Runtime/Online/IOS/OnlineSubsystemIOS/Private/OnlineTurnBasedInterfaceIOS.h
@@ -0,0 +1,139 @@
+// Copyright 1998-2015 Epic Games, Inc. All Rights Reserved.
+
+#pragma once
+
+#include "OnlineTurnBasedInterface.h"
+#include "TurnBasedMatchmakerIOS.h"
+#import <GameKit/GameKit.h>
+
+@class FTurnBasedEventListenerIOS;
+
+// Wrapper for a Game Center Match
+class FTurnBasedMatchIOS : public FTurnBasedMatch
+{
+public:
+	FTurnBasedMatchIOS(GKTurnBasedMatch* _Match, NSArray* PlayerArray);
+	virtual ~FTurnBasedMatchIOS();
+
+	// Returns the number of players currently participating in a match
+	virtual int32 GetNumberOfPlayers() const override;
+
+	// Takes the index of a player PlayerIndex, and sets the name of that player in Name.
+	// Return true if the PlayerIndex is valid, false if not.
+	virtual bool GetPlayerDisplayName(int32 PlayerIndex, FString& Name) const override;
+
+	// Reloads the custom data for this match.
+	virtual void ReloadMatchData(FDownloadMatchDataSignature DownloadCallback) override;
+
+	// Accessors for the custom data of this match. Custom match data is application specific and can be anything,
+	// up to 8k.
+	virtual bool HasMatchData() const override;
+	virtual bool GetMatchData(TArray<uint8>& OutMatchData) const override;
+	virtual void SetMatchData(const TArray<uint8>& NewMatchData, FUploadMatchDataSignature UploadCallback) override;
+
+	// Returns the GUID of this match as an FString.
+	virtual FString GetMatchID() const override;
+
+	// Get the index of the local player in the participants array of this match.
+	virtual int32 GetLocalPlayerIndex() const override;
+
+	// Get the index of the active player in the participants array of this match.
+	virtual int32 GetCurrentPlayerIndex() const override;
+
+	// Gets the match outcome for a player (win/loss/quit/etc.)
+	virtual EMPMatchOutcome::Outcome GetMatchOutcomeForPlayer(int32 PlayerIndex) const override;
+
+	// Gets the player index in the participants array of the player with PlayerID.
+	int32 GetPlayerIndexForPlayer(NSString* PlayerID) const;
+
+	// Returns true if the match stored in this FTurnBasedMatchIOS object is the same as Comparison.
+	bool IsGKTurnBasedMatch(GKTurnBasedMatch* Comparison) const;
+
+	// Ends the turn for the current player with updated MatchData. TurnTimeoutInSeconds is how long the next player has for their turn, and
+	// UploadCallback will be called when the GameCenter call finishes.
+	virtual void EndTurnWithMatchData(const TArray<uint8>& MatchData, int32 TurnTimeoutInSeconds, FUploadMatchDataSignature UploadCallback) override;
+
+	// Quits this match with the passed in Outcome. TurnTimeoutInSeconds is how long the next player will have (if the player is quitting early
+	// and the match isn't over).
+	virtual void QuitMatch(EMPMatchOutcome::Outcome Outcome, int32 TurnTimeoutInSeconds, FQuitMatchSignature QuitMatchCallback) override;
+
+	// Accessors for the GKTurnBasedMatch object.
+	GKTurnBasedMatch* GetGKMatch() const { return Match; }
+	void SetGKMatch(GKTurnBasedMatch* GKMatch);
+
+private:
+
+	// Helper functions for QuitMatch, depending on whether it is the current player's turn or not.
+	void QuitMatchInTurn(GKTurnBasedMatchOutcome Outcome, int32 TurnTimeoutInSeconds, FQuitMatchSignature QuitMatchCallback);
+	void QuitMatchOutOfTurn(GKTurnBasedMatchOutcome Outcome, FQuitMatchSignature QuitMatchCallback);
+
+	// Gets a list of players, if IncludeLocalPlayer is true, this player will be in the array
+	NSArray* GetNextParticipantArray(bool IncludeLocalPlayer) const;
+
+	// Conversion functions between EMPMatchOutcome::Outcome and GKTurnBasedMatchOutcome enums.
+	GKTurnBasedMatchOutcome GetGKTurnBasedMatchOutcomeFromMatchOutcome(EMPMatchOutcome::Outcome Outcome) const;
+	EMPMatchOutcome::Outcome GetMatchOutcomeFromGKTurnBasedMatchOutcome(GKTurnBasedMatchOutcome GKOutcome) const;
+
+	// The last of active player GUIDs in the match
+	TArray<FString> PlayerAliasArray;
+
+	// The Game Center match object
+	GKTurnBasedMatch* Match;
+};
+
+/**
+*	FOnlineTurnBasedIOS - Implementation of turn based multiplayer for IOS
+*/
+class FOnlineTurnBasedIOS : public IOnlineTurnBased, public FTurnBasedMatchmakerDelegate, public FTurnBasedEventDelegate
+{
+public:
+	FOnlineTurnBasedIOS();
+	virtual ~FOnlineTurnBasedIOS();
+
+	static NSArray* GetPlayerIdentifierArrayForMatch(GKTurnBasedMatch* match);
+
+	// IOnlineTurnBased
+	virtual void SetMatchmakerDelegate(FTurnBasedMatchmakerDelegatePtr Delegate) override;
+	virtual void ShowMatchmaker(const FTurnBasedMatchRequest& MatchRequest) override;
+
+	virtual void SetEventDelegate(FTurnBasedEventDelegateWeakPtr Delegate) override;
+	virtual FTurnBasedEventDelegateWeakPtr GetEventDelegate() const override;
+
+	virtual void LoadAllMatches(FLoadTurnBasedMatchesSignature MatchesLoadedCallback) override;
+	virtual void LoadMatchWithID(FString MatchID, FLoadTurnBasedMatchWithIDSignature MatchesLoadedCallback) override;
+
+	virtual FTurnBasedMatchPtr GetMatchWithID(FString MatchID) const override;
+
+	virtual void RemoveMatch(FTurnBasedMatchRef Match, FRemoveMatchSignature RemoveMatchCallback) override;
+
+	// FTurnBasedMatchmakerDelegate
+	virtual void OnMatchmakerCancelled() override;
+	virtual void OnMatchmakerFailed() override;
+	virtual void OnMatchFound(FTurnBasedMatchRef Match) override;
+
+	// FTurnBasedEventDelegate
+	virtual void OnMatchEnded(FString MatchID) override;
+	virtual void OnMatchReceivedTurnEvent(FString MatchID, bool BecameActive, void* Match) override;
+
+	virtual void RegisterTurnBasedMatchInterfaceObject(UObject* Object);
+	UObject* GetTurnBasedMatchInterfaceObject() { return TurnBasedMatchInterfaceObject; }
+
+	virtual int32 GetMatchDataSize() { return MATCH_DATA_SIZE; }
+
+private:
+	FTurnBasedMatchmakerIOS Matchmaker;
+	FTurnBasedMatchmakerDelegateWeakPtr MatchmakerDelegate;
+
+	FTurnBasedEventListenerIOS* EventListener;
+	FTurnBasedEventDelegateWeakPtr EventDelegate;
+
+	TArray<FTurnBasedMatchRef> MatchArray;
+
+	int32 NumberOfMatchesBeingLoaded;
+
+	UObject* TurnBasedMatchInterfaceObject;
+
+	static const int32 MATCH_DATA_SIZE = 1024 * 8 * 4;
+};
+
+typedef TSharedPtr<FOnlineTurnBasedIOS, ESPMode::ThreadSafe> FOnlineTurnBasedIOSPtr;
diff --git a/Engine/Source/Runtime/Online/IOS/OnlineSubsystemIOS/Private/TurnBasedEventListener.cpp b/Engine/Source/Runtime/Online/IOS/OnlineSubsystemIOS/Private/TurnBasedEventListener.cpp
new file mode 100644
--- /dev/null
+++ b/Engine/Source/Runtime/Online/IOS/OnlineSubsystemIOS/Private/TurnBasedEventListener.cpp
@@ -0,0 +1,66 @@
+// Copyright 1998-2015 Epic Games, Inc. All Rights Reserved.
+
+#include "OnlineSubsystemIOSPrivatePCH.h"
+#include "TurnBasedEventListener.h"
+
+@interface FTurnBasedEventListenerIOS()
+{
+    FTurnBasedEventDelegate* _owner;
+}
+
+@end
+
+@implementation FTurnBasedEventListenerIOS
+
+- (id)initWithOwner:(FTurnBasedEventDelegate&)owner
+{
+    if (self = [super init]) {
+        _owner = &owner;
+        [[GKLocalPlayer localPlayer] registerListener:self];
+    }
+    return self;
+}
+
+- (void)dealloc
+{
+    [[GKLocalPlayer localPlayer] unregisterListener:self];
+    [super dealloc];
+}
+
+// GKChallengeListener
+- (void)player:(GKPlayer *)player didCompleteChallenge:(GKChallenge *)challenge issuedByFriend:(GKPlayer *)friendPlayer {}
+- (void)player:(GKPlayer *)player didReceiveChallenge:(GKChallenge *)challenge {}
+- (void)player:(GKPlayer *)player issuedChallengeWasCompleted:(GKChallenge *)challenge byFriend:(GKPlayer *)friendPlayer {}
+- (void)player:(GKPlayer *)player wantsToPlayChallenge:(GKChallenge *)challenge {}
+
+// GKInviteEventListener
+- (void)player:(GKPlayer *)player didAcceptInvite:(GKInvite *)invite {}
+
+// GKInviteEventListener and GKTurnBasedEventListener
+- (void)player:(GKPlayer *)player didRequestMatchWithPlayers:(NSArray *)playerIDsToInvite {}
+
+// GKTurnBasedEventListener
+- (void)player:(GKPlayer *)player matchEnded:(GKTurnBasedMatch *)match
+{
+    if (_owner) {
+        _owner->OnMatchEnded([self getMatchIDFromMatch:match]);
+    }
+}
+
+- (void)player:(GKPlayer *)player receivedExchangeCancellation:(GKTurnBasedExchange *)exchange forMatch:(GKTurnBasedMatch *)match {}
+- (void)player:(GKPlayer *)player receivedExchangeReplies:(NSArray *)replies forCompletedExchange:(GKTurnBasedExchange *)exchange forMatch:(GKTurnBasedMatch *)match {}
+- (void)player:(GKPlayer *)player receivedExchangeRequest:(GKTurnBasedExchange *)exchange forMatch:(GKTurnBasedMatch *)match {}
+
+- (void)player:(GKPlayer *)player receivedTurnEventForMatch:(GKTurnBasedMatch *)match didBecomeActive:(BOOL)didBecomeActive
+{
+    if (_owner) {
+       _owner->OnMatchReceivedTurnEvent([self getMatchIDFromMatch:match], didBecomeActive, match);
+    }
+}
+
+- (FString)getMatchIDFromMatch:(GKTurnBasedMatch *)match
+{
+    return FString::Printf(TEXT("%s"), UTF8_TO_TCHAR(match.matchID.UTF8String));
+}
+
+@end
diff --git a/Engine/Source/Runtime/Online/IOS/OnlineSubsystemIOS/Private/TurnBasedEventListener.h b/Engine/Source/Runtime/Online/IOS/OnlineSubsystemIOS/Private/TurnBasedEventListener.h
new file mode 100644
--- /dev/null
+++ b/Engine/Source/Runtime/Online/IOS/OnlineSubsystemIOS/Private/TurnBasedEventListener.h
@@ -0,0 +1,31 @@
+// Copyright 1998-2013 Epic Games, Inc. All Rights Reserved.
+
+#pragma once
+
+#include <GameKit/GameKit.h>
+#include "OnlineTurnBasedInterface.h"
+
+@interface FTurnBasedEventListenerIOS : NSObject < GKLocalPlayerListener >
+
+- (id)initWithOwner:(FTurnBasedEventDelegate&)owner;
+
+// GKChallengeListener
+- (void)player:(GKPlayer *)player didCompleteChallenge:(GKChallenge *)challenge issuedByFriend:(GKPlayer *)friendPlayer;
+- (void)player:(GKPlayer *)player didReceiveChallenge:(GKChallenge *)challenge;
+- (void)player:(GKPlayer *)player issuedChallengeWasCompleted:(GKChallenge *)challenge byFriend:(GKPlayer *)friendPlayer;
+- (void)player:(GKPlayer *)player wantsToPlayChallenge:(GKChallenge *)challenge;
+
+// GKInviteEventListener
+- (void)player:(GKPlayer *)player didAcceptInvite:(GKInvite *)invite;
+
+// GKInviteEventListener and GKTurnBasedEventListener
+-(void)player:(GKPlayer *)player didRequestMatchWithPlayers:(NSArray *)playerIDsToInvite;
+
+// GKTurnBasedEventListener
+- (void)player:(GKPlayer *)player matchEnded:(GKTurnBasedMatch *)match;
+- (void)player:(GKPlayer *)player receivedExchangeCancellation:(GKTurnBasedExchange *)exchange forMatch:(GKTurnBasedMatch *)match;
+- (void)player:(GKPlayer *)player receivedExchangeReplies:(NSArray *)replies forCompletedExchange:(GKTurnBasedExchange *)exchange forMatch:(GKTurnBasedMatch *)match;
+- (void)player:(GKPlayer *)player receivedExchangeRequest:(GKTurnBasedExchange *)exchange forMatch:(GKTurnBasedMatch *)match;
+- (void)player:(GKPlayer *)player receivedTurnEventForMatch:(GKTurnBasedMatch *)match didBecomeActive:(BOOL)didBecomeActive;
+
+@end
diff --git a/Engine/Source/Runtime/Online/IOS/OnlineSubsystemIOS/Private/TurnBasedMatchmakerDelegateIOS.cpp b/Engine/Source/Runtime/Online/IOS/OnlineSubsystemIOS/Private/TurnBasedMatchmakerDelegateIOS.cpp
new file mode 100644
--- /dev/null
+++ b/Engine/Source/Runtime/Online/IOS/OnlineSubsystemIOS/Private/TurnBasedMatchmakerDelegateIOS.cpp
@@ -0,0 +1,77 @@
+// Copyright 1998-2013 Epic Games, Inc. All Rights Reserved.
+
+#include "OnlineSubsystemIOSPrivatePCH.h"
+
+#include "OnlineTurnBasedInterface.h"
+#include "TurnBasedMatchmakerDelegateIOS.h"
+
+@interface FTurnBasedMatchmakerDelegateIOS()
+{
+	FTurnBasedMatchmakerDelegate* _delegate;
+}
+
+@end
+
+@implementation FTurnBasedMatchmakerDelegateIOS
+
+- (id)initWithDelegate:(FTurnBasedMatchmakerDelegate*)delegate
+{
+	if (self = [super init])
+	{
+		_delegate = delegate;
+	}
+	return self;
+}
+
+-(void)turnBasedMatchmakerViewController:(GKTurnBasedMatchmakerViewController*)viewController didFailWithError : (NSError*)error
+{
+	if (_delegate) {
+		_delegate->OnMatchmakerFailed();
+	}
+}
+
+-(void)turnBasedMatchmakerViewController:(GKTurnBasedMatchmakerViewController*)viewController didFindMatch : (GKTurnBasedMatch*)match
+{
+	if (_delegate) {
+		[match loadMatchDataWithCompletionHandler : ^ (NSData* matchData, NSError* matchLoadError) {
+			if (matchLoadError) {
+				[self turnBasedMatchmakerViewController : viewController didFailWithError : matchLoadError];
+				return;
+			}
+
+			NSMutableArray* playerIdentifierArray = [NSMutableArray array];
+			for (GKTurnBasedParticipant* participant in match.participants) {
+				if (!participant.playerID)
+				{
+					break;
+				}
+				[playerIdentifierArray addObject : participant.playerID];
+			}
+			[GKPlayer loadPlayersForIdentifiers : playerIdentifierArray withCompletionHandler : ^ (NSArray *players, NSError *nameLoadError) {
+				if (nameLoadError) {
+					[self turnBasedMatchmakerViewController : viewController didFailWithError : nameLoadError];
+					return;
+				}
+
+				FTurnBasedMatchRef MatchRef(new FTurnBasedMatchIOS(match, players));
+				_delegate->OnMatchFound(MatchRef);
+			}];
+		}];
+	}
+	else {
+		[match removeWithCompletionHandler : nil];
+	}
+}
+
+-(void)turnBasedMatchmakerViewController:(GKTurnBasedMatchmakerViewController*)viewController playerQuitForMatch : (GKTurnBasedMatch*)match
+{
+}
+
+-(void)turnBasedMatchmakerViewControllerWasCancelled : (GKTurnBasedMatchmakerViewController*)viewController
+{
+	if (_delegate) {
+		_delegate->OnMatchmakerCancelled();
+	}
+}
+
+@end
diff --git a/Engine/Source/Runtime/Online/IOS/OnlineSubsystemIOS/Private/TurnBasedMatchmakerDelegateIOS.h b/Engine/Source/Runtime/Online/IOS/OnlineSubsystemIOS/Private/TurnBasedMatchmakerDelegateIOS.h
new file mode 100644
--- /dev/null
+++ b/Engine/Source/Runtime/Online/IOS/OnlineSubsystemIOS/Private/TurnBasedMatchmakerDelegateIOS.h
@@ -0,0 +1,16 @@
+// Copyright 1998-2015 Epic Games, Inc. All Rights Reserved.
+
+#pragma once
+
+#include <GameKit/GameKit.h>
+
+@interface FTurnBasedMatchmakerDelegateIOS : NSObject < GKTurnBasedMatchmakerViewControllerDelegate >
+
+- (id)initWithDelegate:(FTurnBasedMatchmakerDelegate*)delegate;
+
+- (void)turnBasedMatchmakerViewController:(GKTurnBasedMatchmakerViewController*)viewController didFailWithError:(NSError*)error;
+- (void)turnBasedMatchmakerViewController:(GKTurnBasedMatchmakerViewController*)viewController didFindMatch:(GKTurnBasedMatch*)match;
+- (void)turnBasedMatchmakerViewController:(GKTurnBasedMatchmakerViewController*)viewController playerQuitForMatch:(GKTurnBasedMatch*)match;
+- (void)turnBasedMatchmakerViewControllerWasCancelled:(GKTurnBasedMatchmakerViewController*)viewController;
+
+@end
diff --git a/Engine/Source/Runtime/Online/IOS/OnlineSubsystemIOS/Private/TurnBasedMatchmakerIOS.cpp b/Engine/Source/Runtime/Online/IOS/OnlineSubsystemIOS/Private/TurnBasedMatchmakerIOS.cpp
new file mode 100644
--- /dev/null
+++ b/Engine/Source/Runtime/Online/IOS/OnlineSubsystemIOS/Private/TurnBasedMatchmakerIOS.cpp
@@ -0,0 +1,105 @@
+// Copyright 1998-2015 Epic Games, Inc. All Rights Reserved.
+
+#include "OnlineSubsystemIOSPrivatePCH.h"
+#include "TurnBasedMatchmakerIOS.h"
+
+#include <GameKit/GameKit.h>
+#include "IOSAppDelegate.h"
+#include "TurnBasedMatchmakerDelegateIOS.h"
+#include "IOSView.h"
+
+DEFINE_LOG_CATEGORY_STATIC(LogTurnBasedMatchmakerIOS, Log, All);
+
+FTurnBasedMatchmakerIOS::FTurnBasedMatchmakerIOS(FTurnBasedMatchmakerDelegate& _Delegate)
+: MatchmakerViewController(nil)
+, IOSDelegate(nil)
+, Delegate(&_Delegate)
+{
+}
+
+FTurnBasedMatchmakerIOS::~FTurnBasedMatchmakerIOS()
+{
+    [IOSDelegate release];
+    [MatchmakerViewController release];
+}
+
+void FTurnBasedMatchmakerIOS::ShowWithMatchRequest(const FTurnBasedMatchRequest& Request)
+{
+    MatchRequest = Request;
+
+    IOSDelegate = [[FTurnBasedMatchmakerDelegateIOS alloc] initWithDelegate:this];
+
+    GKMatchRequest* GKRequest = GetGKMatchRequestFromMPMatchRequest(Request);
+    MatchmakerViewController = [[GKTurnBasedMatchmakerViewController alloc] initWithMatchRequest:GKRequest];
+
+    if (MatchmakerViewController)
+    {
+        MatchmakerViewController.showExistingMatches = Request.GetShowExistingMatches() ? YES : NO;
+        MatchmakerViewController.turnBasedMatchmakerDelegate = IOSDelegate;
+        UIViewController* MainViewController = [IOSAppDelegate GetDelegate].IOSController;
+        [MainViewController presentViewController:MatchmakerViewController animated:YES completion:nil];
+    }
+}
+
+void FTurnBasedMatchmakerIOS::OnMatchmakerCancelled()
+{
+    UE_LOG(LogTurnBasedMatchmakerIOS, Log, TEXT("Matchmaker cancelled"));
+
+    if (Delegate)
+    {
+        Delegate->OnMatchmakerCancelled();
+    }
+    DismissMatchmaker();
+}
+
+void FTurnBasedMatchmakerIOS::OnMatchmakerFailed()
+{
+    UE_LOG(LogTurnBasedMatchmakerIOS, Log, TEXT("Matchmaker failed"));
+
+    if (Delegate)
+    {
+        Delegate->OnMatchmakerFailed();
+    }
+    DismissMatchmaker();
+}
+
+void FTurnBasedMatchmakerIOS::OnMatchFound(FTurnBasedMatchRef Match)
+{
+    UE_LOG(LogTurnBasedMatchmakerIOS, Log, TEXT("Match found"));
+    //Match->SetMatchRequest(MatchRequest);
+
+    if (Delegate)
+    {
+        Delegate->OnMatchFound(Match);
+    }
+
+    DismissMatchmaker();
+}
+
+void FTurnBasedMatchmakerIOS::DismissMatchmaker()
+{
+    if (!MatchmakerViewController)
+    {
+        UE_LOG(LogTurnBasedMatchmakerIOS, Warning, TEXT("No matchmaker was active"));
+        return;
+    }
+
+    [IOSDelegate release];
+    IOSDelegate = nil;
+
+    [MatchmakerViewController dismissViewControllerAnimated:YES completion:nil];
+    [MatchmakerViewController release];
+    MatchmakerViewController = nil;
+}
+
+GKMatchRequest* FTurnBasedMatchmakerIOS::GetGKMatchRequestFromMPMatchRequest(const FTurnBasedMatchRequest& Request) const
+{
+    GKMatchRequest* MatchRequest = [[GKMatchRequest alloc] init];
+
+    MatchRequest.maxPlayers = Request.GetMaxNumberOfPlayers();
+    MatchRequest.minPlayers = Request.GetMinNumberOfPlayers();
+    MatchRequest.defaultNumberOfPlayers = Request.GetMinNumberOfPlayers();
+    MatchRequest.playerGroup = Request.GetPlayerGroup();
+
+    return [MatchRequest autorelease];
+}
diff --git a/Engine/Source/Runtime/Online/IOS/OnlineSubsystemIOS/Private/TurnBasedMatchmakerIOS.h b/Engine/Source/Runtime/Online/IOS/OnlineSubsystemIOS/Private/TurnBasedMatchmakerIOS.h
new file mode 100644
--- /dev/null
+++ b/Engine/Source/Runtime/Online/IOS/OnlineSubsystemIOS/Private/TurnBasedMatchmakerIOS.h
@@ -0,0 +1,32 @@
+// Copyright 1998-2015 Epic Games, Inc. All Rights Reserved.
+
+#pragma once
+
+#include "OnlineTurnBasedInterface.h"
+
+@class FTurnBasedMatchmakerDelegateIOS;
+@class GKMatchRequest;
+@class GKTurnBasedMatchmakerViewController;
+
+class FTurnBasedMatchmakerIOS : public FTurnBasedMatchmakerDelegate
+{
+public:
+    FTurnBasedMatchmakerIOS(FTurnBasedMatchmakerDelegate& _Delegate);
+    ~FTurnBasedMatchmakerIOS();
+
+    void ShowWithMatchRequest(const FTurnBasedMatchRequest& Request);
+
+    virtual void OnMatchmakerCancelled() override;
+    virtual void OnMatchmakerFailed() override;
+    virtual void OnMatchFound(FTurnBasedMatchRef Match) override;
+
+private:
+    GKMatchRequest* GetGKMatchRequestFromMPMatchRequest(const FTurnBasedMatchRequest& Request) const;
+    void DismissMatchmaker();
+
+    FTurnBasedMatchRequest MatchRequest;
+
+    GKTurnBasedMatchmakerViewController* MatchmakerViewController;
+    FTurnBasedMatchmakerDelegateIOS* IOSDelegate;
+    FTurnBasedMatchmakerDelegate* Delegate;
+};
diff --git a/Engine/Source/Runtime/Online/IOS/OnlineSubsystemIOS/Public/OnlineSubsystemIOS.h b/Engine/Source/Runtime/Online/IOS/OnlineSubsystemIOS/Public/OnlineSubsystemIOS.h
--- a/Engine/Source/Runtime/Online/IOS/OnlineSubsystemIOS/Public/OnlineSubsystemIOS.h
+++ b/Engine/Source/Runtime/Online/IOS/OnlineSubsystemIOS/Public/OnlineSubsystemIOS.h
@@ -37,8 +37,9 @@ public:
 	virtual IOnlineSharingPtr GetSharingInterface() const override;
 	virtual IOnlineUserPtr GetUserInterface() const override;
 	virtual IOnlineMessagePtr GetMessageInterface() const override;
-	virtual IOnlinePresencePtr GetPresenceInterface() const override;
+    virtual IOnlinePresencePtr GetPresenceInterface() const override;
 	virtual IOnlineChatPtr GetChatInterface() const override;
+    virtual IOnlineTurnBasedPtr GetTurnBasedInterface() const override;
 	virtual bool Init() override;
 	virtual bool Shutdown() override;
 	virtual FString GetAppId() const override;
@@ -85,6 +86,9 @@ private:
 
 	/** Interface to the external UI services */
 	FOnlineExternalUIIOSPtr ExternalUIInterface;
+
+    /** Interface to the turnbased multiplayer services */
+    FOnlineTurnBasedIOSPtr TurnBasedInterface;
 };
 
 typedef TSharedPtr<FOnlineSubsystemIOS, ESPMode::ThreadSafe> FOnlineSubsystemIOSPtr;
diff --git a/Engine/Source/Runtime/Online/OnlineSubsystem/Private/TurnBasedMatchInterface.cpp b/Engine/Source/Runtime/Online/OnlineSubsystem/Private/TurnBasedMatchInterface.cpp
new file mode 100644
--- /dev/null
+++ b/Engine/Source/Runtime/Online/OnlineSubsystem/Private/TurnBasedMatchInterface.cpp
@@ -0,0 +1,9 @@
+// Copyright 1998-2015 Epic Games, Inc. All Rights Reserved.
+
+#include "OnlineSubsystemPrivatePCH.h"
+#include "TurnBasedMatchInterface.h"
+
+UTurnBasedMatchInterface::UTurnBasedMatchInterface(const FObjectInitializer& ObjectInitializer)
+: Super(ObjectInitializer)
+{
+}
\ No newline at end of file
diff --git a/Engine/Source/Runtime/Online/OnlineSubsystem/Public/Interfaces/OnlineTurnBasedInterface.h b/Engine/Source/Runtime/Online/OnlineSubsystem/Public/Interfaces/OnlineTurnBasedInterface.h
new file mode 100644
--- /dev/null
+++ b/Engine/Source/Runtime/Online/OnlineSubsystem/Public/Interfaces/OnlineTurnBasedInterface.h
@@ -0,0 +1,213 @@
+// Copyright 1998-2015 Epic Games, Inc. All Rights Reserved.
+
+#pragma once
+
+#include "Platform.h"
+#include "OnlineTurnBasedInterface.generated.h"
+
+/** 
+ * FTurnBasedMatchRequest contains all of the information required for a matchmaker to create an FTurnBasedMatch
+ */
+class FTurnBasedMatchRequest
+{
+public:
+    FTurnBasedMatchRequest()
+    : MinNumberOfPlayers(0)
+	, MaxNumberOfPlayers(0)
+    , PlayerGroup(0)
+	, ShowExistingMatches(false)
+    {}
+
+	FTurnBasedMatchRequest(int32 InMinNumberOfPlayers, int32 InMaxNumberOfPlayers, uint32 InPlayerGroup, bool InShowExistingMatches)
+	: MinNumberOfPlayers(InMinNumberOfPlayers)
+	, MaxNumberOfPlayers(InMaxNumberOfPlayers)
+	, PlayerGroup(InPlayerGroup)
+	, ShowExistingMatches(InShowExistingMatches)
+	{}
+
+    void SetMinNumberOfPlayers(int32 InMinNumberOfPlayers)	{ MinNumberOfPlayers = InMinNumberOfPlayers; }
+    int32 GetMinNumberOfPlayers() const						{ return MinNumberOfPlayers; }
+
+	void SetMaxNumberOfPlayers(int32 InMaxNumberOfPlayers)	{ MaxNumberOfPlayers = InMaxNumberOfPlayers; }
+	int32 GetMaxNumberOfPlayers() const						{ return MaxNumberOfPlayers; }
+
+    void SetPlayerGroup(uint32 InPlayerGroup)				{ PlayerGroup = InPlayerGroup;  }
+    uint32 GetPlayerGroup() const							{ return PlayerGroup; }
+
+	void SetShowExistingMatches(bool InShowExistingMatches) { ShowExistingMatches = InShowExistingMatches; }
+	bool GetShowExistingMatches() const						{ return ShowExistingMatches; }
+
+private:
+
+	// The minimum number of players needed for this match
+    int32 MinNumberOfPlayers;
+
+	// The maximum number of players needed for this match
+	int32 MaxNumberOfPlayers;
+
+	// The player group represents another parameter to pass in to matchmaking. For example, 1 could be Deathmatch, 2 could be Capture the Flag, etc.
+	// Only players with the same player group will be matched together.
+    uint32 PlayerGroup;
+
+	// If true, the native matchmaking interface will show matches the player is in already
+	bool ShowExistingMatches;
+};
+
+DECLARE_DELEGATE_TwoParams(FQuitMatchSignature, FString, bool);   // bool Success
+DECLARE_DELEGATE_TwoParams(FRemoveMatchSignature, FString, bool);   // bool Success
+DECLARE_DELEGATE_TwoParams(FUploadMatchDataSignature, FString, bool);   // bool Success
+DECLARE_DELEGATE_TwoParams(FDownloadMatchDataSignature, FString, bool);  // bool Success
+
+/**
+* EMPMatchOutcome represents all the possible outcomes for this player in a match
+*/
+UENUM(BlueprintType)
+namespace EMPMatchOutcome
+{
+    enum Outcome
+    {
+        None,
+        Quit,
+        Won,
+        Lost,
+        Tied,
+        TimeExpired,
+        First,
+        Second,
+        Third,
+        Fourth,
+    };
+}
+
+/**
+* FTurnBasedMatch contains all of the information about an in-progress turn based match
+*/
+class FTurnBasedMatch
+{
+public:
+    virtual ~FTurnBasedMatch() {}
+
+	/* Get the number of players in the match */
+    virtual int32 GetNumberOfPlayers() const { return 2; }
+
+    /* Get the user-friendly display name for the given player.  Can return false if the name of the associated player has not been loaded */
+    virtual bool GetPlayerDisplayName(int32 PlayerIndex, FString& Name) const { return false; }
+
+    /* Request a reload of the match's data.  Once the callback has been called, use GetMatchData to get the data */
+    virtual void ReloadMatchData(FDownloadMatchDataSignature DownloadCallback) {}
+
+    /* HasMatchData is for when you want to know if the match has data, without getting at the data */
+    virtual bool HasMatchData() const { return false; }
+
+    /* Put the match data into OutMatchData, returning whether or not the operation was successful */
+    virtual bool GetMatchData(TArray<uint8>& OutMatchData) const { return false; }
+    
+    /* Update the data for the match with the data provided */
+    virtual void SetMatchData(const TArray<uint8>& NewMatchData, FUploadMatchDataSignature EndTurnCallback) {}
+
+	/* Returns the Match ID for this match */
+    virtual FString GetMatchID() const { return FString(); }
+
+    /* Get the index of the local player in the match's list of participants */
+    virtual int32 GetLocalPlayerIndex() const { return 0; }
+
+    /* Get the index of the current player in the match's list of participants */
+    virtual int32 GetCurrentPlayerIndex() const { return 0; }
+
+	/* Get the outcome of the match (won/lost/quit/etc.) for the player with PlayerIndex */
+    virtual EMPMatchOutcome::Outcome GetMatchOutcomeForPlayer(int32 PlayerIndex) const { return EMPMatchOutcome::None; }
+
+    /* End the turn and upload MatchData. EndTurnCallback is called once the data is uploaded (or failed), with a (bool) Success parameter */
+    virtual void EndTurnWithMatchData(const TArray<uint8>& MatchData, int32 TurnTimeoutInSeconds, FUploadMatchDataSignature EndTurnCallback) {}
+
+    /* Leave the match, providing an outcome for the player */
+    virtual void QuitMatch(EMPMatchOutcome::Outcome Outcome, int32 TurnTimeoutInSeconds, FQuitMatchSignature QuitMatchCallback) {}
+};
+
+typedef TSharedRef<FTurnBasedMatch, ESPMode::ThreadSafe> FTurnBasedMatchRef;
+typedef TSharedPtr<FTurnBasedMatch, ESPMode::ThreadSafe> FTurnBasedMatchPtr;
+
+/**
+* FTurnBasedMatchmakerDelegate provides the interface for all turn based matchmaking callbacks
+*/
+class FTurnBasedMatchmakerDelegate
+{
+public:
+    virtual ~FTurnBasedMatchmakerDelegate() {}
+
+	// This is triggered if the player cancelled the matchmaking process
+    virtual void OnMatchmakerCancelled() {}
+
+	// This is triggered if matchmaking failed for any reason
+    virtual void OnMatchmakerFailed() {}
+
+	// This is triggered once a match has been successfully found
+    virtual void OnMatchFound(FTurnBasedMatchRef Match) {}
+};
+
+typedef TSharedRef<FTurnBasedMatchmakerDelegate, ESPMode::ThreadSafe> FTurnBasedMatchmakerDelegateRef;
+typedef TSharedPtr<FTurnBasedMatchmakerDelegate, ESPMode::ThreadSafe> FTurnBasedMatchmakerDelegatePtr;
+typedef TWeakPtr<FTurnBasedMatchmakerDelegate, ESPMode::ThreadSafe> FTurnBasedMatchmakerDelegateWeakPtr;
+
+/**
+* FTurnBasedEventDelegate provides the interface for responding to events in turn based games
+*/
+class FTurnBasedEventDelegate
+{
+public:
+    virtual ~FTurnBasedEventDelegate() {}
+
+	// This is triggered when the match has ended for any reason
+    virtual void OnMatchEnded(FString MatchID) {}
+
+	// This is triggered when it is the current player's turn
+    virtual void OnMatchReceivedTurnEvent(FString MatchID, bool BecameActive, void* Match) {}
+};
+
+DECLARE_DELEGATE_TwoParams(FLoadTurnBasedMatchesSignature, const TArray<FString>&, bool);   // bool Success, TArray MatchArray
+DECLARE_DELEGATE_TwoParams(FLoadTurnBasedMatchWithIDSignature, FString, bool);   // bool Success
+
+typedef TSharedRef<FTurnBasedEventDelegate, ESPMode::ThreadSafe> FTurnBasedEventDelegateRef;
+typedef TSharedPtr<FTurnBasedEventDelegate, ESPMode::ThreadSafe> FTurnBasedEventDelegatePtr;
+typedef TWeakPtr<FTurnBasedEventDelegate, ESPMode::ThreadSafe> FTurnBasedEventDelegateWeakPtr;
+
+/**
+ *	IOnlineTurnBased - Interface class for turn based multiplayer matches
+ */
+class IOnlineTurnBased
+{
+public:
+    virtual ~IOnlineTurnBased() {}
+
+	// Set a delegate to be called when matchmaking succeeds, fails, or is canceled
+    virtual void SetMatchmakerDelegate(FTurnBasedMatchmakerDelegatePtr Delegate) {}
+
+	// Show the platform specific matchmaker interface with the parameters given in FTurnBasedMatchRequest
+    virtual void ShowMatchmaker(const FTurnBasedMatchRequest& MatchRequest) {}
+   
+	// Accessors for a delegate that gets called when match-specific events happen (such as the match ending or it is now the current players turn)
+    virtual void SetEventDelegate(FTurnBasedEventDelegateWeakPtr Delegate) {}
+	virtual FTurnBasedEventDelegateWeakPtr GetEventDelegate() const { return nullptr; }
+
+	// Load all matches the current player is participating in
+    virtual void LoadAllMatches(FLoadTurnBasedMatchesSignature MatchesLoadedCallback) {}
+
+	// Reload a match with the given MatchID
+    virtual void LoadMatchWithID(FString MatchID, FLoadTurnBasedMatchWithIDSignature MatchesLoadedCallback) {}
+
+	// Returns the match object with the given MatchID
+    virtual FTurnBasedMatchPtr GetMatchWithID(FString MatchID) const { return nullptr; }
+
+	// Deletes a match entirely from the platform specific service
+    virtual void RemoveMatch(FTurnBasedMatchRef Match, FRemoveMatchSignature RemoveMatchCallback) {}
+
+	// Accessors referring to an object that inherits from UTurnBasedMatchInterface that acts like an FTurnBasedEventDelegate
+	// (needed for blueprints)
+	virtual void RegisterTurnBasedMatchInterfaceObject(UObject* Object) {}
+	virtual UObject* GetTurnBasedMatchInterfaceObject() { return nullptr; }
+
+	// The size of the game specific data stored per match
+	virtual int32 GetMatchDataSize() { return 0; }
+};
+
+typedef TSharedPtr<class IOnlineTurnBased, ESPMode::ThreadSafe> IOnlineTurnBasedPtr;
\ No newline at end of file
diff --git a/Engine/Source/Runtime/Online/OnlineSubsystem/Public/Interfaces/TurnBasedMatchInterface.h b/Engine/Source/Runtime/Online/OnlineSubsystem/Public/Interfaces/TurnBasedMatchInterface.h
new file mode 100644
--- /dev/null
+++ b/Engine/Source/Runtime/Online/OnlineSubsystem/Public/Interfaces/TurnBasedMatchInterface.h
@@ -0,0 +1,22 @@
+// Copyright 1998-2015 Epic Games, Inc. All Rights Reserved.
+
+#pragma once
+
+#include "TurnBasedMatchInterface.generated.h"
+
+class IOnlineSubsystem;
+
+UINTERFACE(Blueprintable)
+class UTurnBasedMatchInterface : public UInterface
+{
+	GENERATED_UINTERFACE_BODY()
+};
+
+class ITurnBasedMatchInterface
+{
+	GENERATED_IINTERFACE_BODY()
+
+public:
+	UFUNCTION(BlueprintImplementableEvent, Category = "Online|TurnBased")
+	void OnMatchReceivedTurn(const FString& Match, bool bDidBecomeActive);
+};
diff --git a/Engine/Source/Runtime/Online/OnlineSubsystem/Public/OnlineSubsystem.h b/Engine/Source/Runtime/Online/OnlineSubsystem/Public/OnlineSubsystem.h
--- a/Engine/Source/Runtime/Online/OnlineSubsystem/Public/OnlineSubsystem.h
+++ b/Engine/Source/Runtime/Online/OnlineSubsystem/Public/OnlineSubsystem.h
@@ -50,6 +50,7 @@ typedef TSharedPtr<class IOnlineUser, ESPMode::ThreadSafe> IOnlineUserPtr;
 typedef TSharedPtr<class IOnlineMessage, ESPMode::ThreadSafe> IOnlineMessagePtr;
 typedef TSharedPtr<class IOnlinePresence, ESPMode::ThreadSafe> IOnlinePresencePtr;
 typedef TSharedPtr<class IOnlineChat, ESPMode::ThreadSafe> IOnlineChatPtr;
+typedef TSharedPtr<class IOnlineTurnBased, ESPMode::ThreadSafe> IOnlineTurnBasedPtr;
 typedef TSharedPtr<class FOnlineNotificationHandler, ESPMode::ThreadSafe> FOnlineNotificationHandlerPtr;
 typedef TSharedPtr<class FOnlineNotificationTransportManager, ESPMode::ThreadSafe> FOnlineNotificationTransportManagerPtr;
 
@@ -243,6 +244,7 @@ public:
 	 * @return Interface pointer for the appropriate online user service
 	 */
 	virtual IOnlineChatPtr GetChatInterface() const = 0;
+
 	/**
 	* Get the notification handler instance for this subsystem
 	* @return Pointer for the appropriate notification handler
@@ -252,6 +254,12 @@ public:
 		return OnlineNotificationHandler;
 	}
 
+	/**
+	* Get the interface for managing turn based multiplayer games
+	* @return Interface pointer for the appropriate online user service
+	*/
+	virtual IOnlineTurnBasedPtr GetTurnBasedInterface() const = 0;
+
 	/**
 	* Get the transport manager instance for this subsystem
 	* @return Pointer for the appropriate transport manager
diff --git a/Engine/Source/Runtime/Online/OnlineSubsystemAmazon/Public/OnlineSubsystemAmazon.h b/Engine/Source/Runtime/Online/OnlineSubsystemAmazon/Public/OnlineSubsystemAmazon.h
--- a/Engine/Source/Runtime/Online/OnlineSubsystemAmazon/Public/OnlineSubsystemAmazon.h
+++ b/Engine/Source/Runtime/Online/OnlineSubsystemAmazon/Public/OnlineSubsystemAmazon.h
@@ -118,6 +118,11 @@ public:
 		return NULL;
 	}
 
+    virtual IOnlineTurnBasedPtr GetTurnBasedInterface() const override
+    {
+        return NULL;
+    }
+	
 	virtual bool Init() override;
 	virtual bool Shutdown() override;
 	virtual FString GetAppId() const override;
diff --git a/Engine/Source/Runtime/Online/OnlineSubsystemFacebook/Private/IOS/OnlineSubsystemFacebook.cpp b/Engine/Source/Runtime/Online/OnlineSubsystemFacebook/Private/IOS/OnlineSubsystemFacebook.cpp
--- a/Engine/Source/Runtime/Online/OnlineSubsystemFacebook/Private/IOS/OnlineSubsystemFacebook.cpp
+++ b/Engine/Source/Runtime/Online/OnlineSubsystemFacebook/Private/IOS/OnlineSubsystemFacebook.cpp
@@ -10,6 +10,7 @@ FOnlineSubsystemFacebook::FOnlineSubsystemFacebook()
 	: FacebookIdentity(NULL)
 	, FacebookFriends(NULL)
 	, FacebookSharing(NULL)
+	, FacebookUser(NULL)
 {
 
 }
@@ -19,6 +20,7 @@ FOnlineSubsystemFacebook::~FOnlineSubsystemFacebook()
 	FacebookIdentity = NULL;
 	FacebookFriends = NULL;
 	FacebookSharing = NULL; 
+	FacebookUser = NULL;
 }
 
 IOnlineSessionPtr FOnlineSubsystemFacebook::GetSessionInterface() const
@@ -116,6 +118,11 @@ IOnlineChatPtr FOnlineSubsystemFacebook::GetChatInterface() const
 	return NULL;
 }
 
+IOnlineTurnBasedPtr FOnlineSubsystemFacebook::GetTurnBasedInterface() const
+{
+	return NULL;
+}
+
 bool FOnlineSubsystemFacebook::Init() 
 {
 	bool bSuccessfullyStartedUp = true;
diff --git a/Engine/Source/Runtime/Online/OnlineSubsystemFacebook/Private/Windows/OnlineSubsystemFacebook.cpp b/Engine/Source/Runtime/Online/OnlineSubsystemFacebook/Private/Windows/OnlineSubsystemFacebook.cpp
--- a/Engine/Source/Runtime/Online/OnlineSubsystemFacebook/Private/Windows/OnlineSubsystemFacebook.cpp
+++ b/Engine/Source/Runtime/Online/OnlineSubsystemFacebook/Private/Windows/OnlineSubsystemFacebook.cpp
@@ -114,6 +114,11 @@ IOnlineChatPtr FOnlineSubsystemFacebook::GetChatInterface() const
 	return NULL;
 }
 
+IOnlineTurnBasedPtr FOnlineSubsystemFacebook::GetTurnBasedInterface() const
+{
+    return NULL;
+}
+
 bool FOnlineSubsystemFacebook::Tick(float DeltaTime)
 {
 	if (FacebookIdentity.IsValid())
diff --git a/Engine/Source/Runtime/Online/OnlineSubsystemFacebook/Public/OnlineSubsystemFacebook.h b/Engine/Source/Runtime/Online/OnlineSubsystemFacebook/Public/OnlineSubsystemFacebook.h
--- a/Engine/Source/Runtime/Online/OnlineSubsystemFacebook/Public/OnlineSubsystemFacebook.h
+++ b/Engine/Source/Runtime/Online/OnlineSubsystemFacebook/Public/OnlineSubsystemFacebook.h
@@ -41,6 +41,7 @@ public:
 	virtual IOnlineMessagePtr GetMessageInterface() const override;
 	virtual IOnlinePresencePtr GetPresenceInterface() const override;
 	virtual IOnlineChatPtr GetChatInterface() const override;
+    virtual IOnlineTurnBasedPtr GetTurnBasedInterface() const override;
 	virtual bool Init() override;
 	virtual bool Shutdown() override;
 	virtual FString GetAppId() const override;
diff --git a/Engine/Source/Runtime/Online/OnlineSubsystemNull/Private/OnlineSubsystemNull.cpp b/Engine/Source/Runtime/Online/OnlineSubsystemNull/Private/OnlineSubsystemNull.cpp
--- a/Engine/Source/Runtime/Online/OnlineSubsystemNull/Private/OnlineSubsystemNull.cpp
+++ b/Engine/Source/Runtime/Online/OnlineSubsystemNull/Private/OnlineSubsystemNull.cpp
@@ -108,6 +108,11 @@ IOnlineChatPtr FOnlineSubsystemNull::GetChatInterface() const
 	return NULL;
 }
 
+IOnlineTurnBasedPtr FOnlineSubsystemNull::GetTurnBasedInterface() const
+{
+    return NULL;
+}
+
 bool FOnlineSubsystemNull::Tick(float DeltaTime)
 {
 	if (OnlineAsyncTaskThreadRunnable)
diff --git a/Engine/Source/Runtime/Online/OnlineSubsystemNull/Public/OnlineSubsystemNull.h b/Engine/Source/Runtime/Online/OnlineSubsystemNull/Public/OnlineSubsystemNull.h
--- a/Engine/Source/Runtime/Online/OnlineSubsystemNull/Public/OnlineSubsystemNull.h
+++ b/Engine/Source/Runtime/Online/OnlineSubsystemNull/Public/OnlineSubsystemNull.h
@@ -52,7 +52,8 @@ public:
 	virtual IOnlineMessagePtr GetMessageInterface() const override;
 	virtual IOnlinePresencePtr GetPresenceInterface() const override;
 	virtual IOnlineChatPtr GetChatInterface() const override;
-
+    virtual IOnlineTurnBasedPtr GetTurnBasedInterface() const override;
+	
 	virtual bool Init() override;
 	virtual bool Shutdown() override;
 	virtual FString GetAppId() const override;
diff --git a/Engine/Source/Runtime/Online/OnlineSubsystemSteam/Private/OnlineSubsystemSteam.cpp b/Engine/Source/Runtime/Online/OnlineSubsystemSteam/Private/OnlineSubsystemSteam.cpp
--- a/Engine/Source/Runtime/Online/OnlineSubsystemSteam/Private/OnlineSubsystemSteam.cpp
+++ b/Engine/Source/Runtime/Online/OnlineSubsystemSteam/Private/OnlineSubsystemSteam.cpp
@@ -249,6 +249,11 @@ IOnlineChatPtr FOnlineSubsystemSteam::GetChatInterface() const
 	return NULL;
 }
 
+IOnlineTurnBasedPtr FOnlineSubsystemSteam::GetTurnBasedInterface() const
+{
+    return NULL;
+}
+
 void FOnlineSubsystemSteam::QueueAsyncTask(FOnlineAsyncTask* AsyncTask)
 {
 	check(OnlineAsyncTaskThreadRunnable);
diff --git a/Engine/Source/Runtime/Online/OnlineSubsystemSteam/Public/OnlineSubsystemSteam.h b/Engine/Source/Runtime/Online/OnlineSubsystemSteam/Public/OnlineSubsystemSteam.h
--- a/Engine/Source/Runtime/Online/OnlineSubsystemSteam/Public/OnlineSubsystemSteam.h
+++ b/Engine/Source/Runtime/Online/OnlineSubsystemSteam/Public/OnlineSubsystemSteam.h
@@ -224,6 +224,7 @@ public:
 	virtual IOnlineMessagePtr GetMessageInterface() const override;
 	virtual IOnlinePresencePtr GetPresenceInterface() const override;
 	virtual IOnlineChatPtr GetChatInterface() const override;
+	virtual IOnlineTurnBasedPtr GetTurnBasedInterface() const override;
 	virtual bool IsLocalPlayer(const FUniqueNetId& UniqueId) const override;
 	virtual bool Init() override;
 	virtual bool Shutdown() override;
diff --git a/Engine/Source/Runtime/Online/OnlineSubsystemUtils/Classes/EndTurnCallbackProxy.h b/Engine/Source/Runtime/Online/OnlineSubsystemUtils/Classes/EndTurnCallbackProxy.h
new file mode 100644
--- /dev/null
+++ b/Engine/Source/Runtime/Online/OnlineSubsystemUtils/Classes/EndTurnCallbackProxy.h
@@ -0,0 +1,48 @@
+// Copyright 1998-2015 Epic Games, Inc. All Rights Reserved.
+
+#pragma once
+
+#include "OnlineTurnBasedInterface.h"
+#include "TurnBasedMatchInterface.h"
+#include "EndTurnCallbackProxy.generated.h"
+
+class AActor;
+
+UCLASS(MinimalAPI)
+class UEndTurnCallbackProxy : public UOnlineBlueprintCallProxyBase
+{
+	GENERATED_UCLASS_BODY()
+
+	virtual ~UEndTurnCallbackProxy();
+
+	// Called when there is a successful query
+	UPROPERTY(BlueprintAssignable)
+	FEmptyOnlineDelegate OnSuccess;
+
+	// Called when there is an unsuccessful query
+	UPROPERTY(BlueprintAssignable)
+	FEmptyOnlineDelegate OnFailure;
+
+	// Ends the turn for the current player
+	UFUNCTION(BlueprintCallable, meta = (BlueprintInternalUseOnly = "true", WorldContext = "WorldContextObject"), Category = "Online|TurnBased")
+	static UEndTurnCallbackProxy* EndTurn(UObject* WorldContextObject, class APlayerController* PlayerController, FString MatchID, TScriptInterface<ITurnBasedMatchInterface> TurnBasedMatchInterface);
+
+	// UOnlineBlueprintCallProxyBase interface
+	virtual void Activate() override;
+	// End of UOnlineBlueprintCallProxyBase interface
+
+    void UploadMatchDataDelegate(FString MatchID, bool Succeeded);
+
+private:
+	// The player controller triggering things
+	TWeakObjectPtr<APlayerController> PlayerControllerWeakPtr;
+
+	// The world context object in which this call is taking place
+	UObject* WorldContextObject;
+
+	// The MatchID of the active match
+	FString MatchID;
+
+	// Object that is inheriting from UTurnBasedMatchInterface. Replicated properties will be serialize into the platform specific match data object.
+	UTurnBasedMatchInterface* TurnBasedMatchInterface;
+};
\ No newline at end of file
diff --git a/Engine/Source/Runtime/Online/OnlineSubsystemUtils/Classes/FindTurnBasedMatchCallbackProxy.h b/Engine/Source/Runtime/Online/OnlineSubsystemUtils/Classes/FindTurnBasedMatchCallbackProxy.h
new file mode 100644
--- /dev/null
+++ b/Engine/Source/Runtime/Online/OnlineSubsystemUtils/Classes/FindTurnBasedMatchCallbackProxy.h
@@ -0,0 +1,79 @@
+// Copyright 1998-2015 Epic Games, Inc. All Rights Reserved.
+
+#pragma once
+
+#include "OnlineTurnBasedInterface.h"
+#include "FindTurnBasedMatchCallbackProxy.generated.h"
+
+class UFindTurnBasedMatchCallbackProxy;
+class UTurnBasedMatchInterface;
+
+DECLARE_DYNAMIC_MULTICAST_DELEGATE_OneParam(FOnlineTurnBasedMatchResult, FString, MatchID);
+
+class FFindTurnBasedMatchCallbackProxyMatchmakerDelegate : public FTurnBasedMatchmakerDelegate, public TSharedFromThis<FTurnBasedMatchmakerDelegate>
+{
+public:
+
+	FFindTurnBasedMatchCallbackProxyMatchmakerDelegate();
+	virtual ~FFindTurnBasedMatchCallbackProxyMatchmakerDelegate();
+
+	virtual void OnMatchmakerCancelled() override;
+	virtual void OnMatchmakerFailed() override;
+	virtual void OnMatchFound(FTurnBasedMatchRef Match) override;
+
+	void SetFindTurnBasedMatchCallbackProxy(UFindTurnBasedMatchCallbackProxy* InFindTurnBasedMatchCallbackProxy) { FindTurnBasedMatchCallbackProxy = InFindTurnBasedMatchCallbackProxy; }
+	void SetTurnBasedInterface(IOnlineTurnBasedPtr InTurnBasedInterface) { TurnBasedInterface = InTurnBasedInterface; }
+private:
+	UFindTurnBasedMatchCallbackProxy* FindTurnBasedMatchCallbackProxy;
+	IOnlineTurnBasedPtr TurnBasedInterface;
+};
+
+UCLASS(MinimalAPI)
+class UFindTurnBasedMatchCallbackProxy : public UOnlineBlueprintCallProxyBase
+{
+	GENERATED_UCLASS_BODY()
+
+    virtual ~UFindTurnBasedMatchCallbackProxy();
+    
+	// Called when matchmaking succeeded.
+	UPROPERTY(BlueprintAssignable)
+	FOnlineTurnBasedMatchResult OnSuccess;
+
+	// Called when matchmaking failed
+	UPROPERTY(BlueprintAssignable)
+	FOnlineTurnBasedMatchResult OnFailure;
+
+	// Use the platform matchmaking service (like Game Center) to find a match.
+	UFUNCTION(BlueprintCallable, meta = (BlueprintInternalUseOnly = "true", WorldContext = "WorldContextObject"), Category = "Online|TurnBased")
+	static UFindTurnBasedMatchCallbackProxy* FindTurnBasedMatch(UObject* WorldContextObject, class APlayerController* PlayerController, TScriptInterface<ITurnBasedMatchInterface> MatchActor, int32 MinPlayers, int32 MaxPlayers, int32 PlayerGroup, bool ShowExistingMatches);
+
+	// UOnlineBlueprintCallProxyBase interface
+	virtual void Activate() override;
+	// End of UOnlineBlueprintCallProxyBase interface
+
+	UTurnBasedMatchInterface* GetTurnBasedMatchInterfaceObject() { return TurnBasedMatchInterface; }
+
+private:
+	// The player controller triggering things
+	TWeakObjectPtr<APlayerController> PlayerControllerWeakPtr;
+
+	// The world context object in which this call is taking place
+	UObject* WorldContextObject;
+
+	// TurnBasedMatchInterface object, used to set the match data after a match is found
+	UTurnBasedMatchInterface* TurnBasedMatchInterface;
+
+	// Minimum number of players needed for the match if a match is created
+	uint32 MinPlayers;
+
+	// Maximum number of players needed for the match if a match is created
+	uint32 MaxPlayers;
+
+	// Another matchmaking parameter that must be the same for players to matchmake together - for example, this could be the game mode (1 for deathmatch, 2 for capture the flag, etc.)
+	uint32 PlayerGroup;
+
+	// Show matches that the player is already a part of in the matchmaking interface
+	bool ShowExistingMatches;
+
+	TSharedPtr<FFindTurnBasedMatchCallbackProxyMatchmakerDelegate, ESPMode::ThreadSafe> Delegate;
+};
diff --git a/Engine/Source/Runtime/Online/OnlineSubsystemUtils/Classes/QuitMatchCallbackProxy.h b/Engine/Source/Runtime/Online/OnlineSubsystemUtils/Classes/QuitMatchCallbackProxy.h
new file mode 100644
--- /dev/null
+++ b/Engine/Source/Runtime/Online/OnlineSubsystemUtils/Classes/QuitMatchCallbackProxy.h
@@ -0,0 +1,49 @@
+// Copyright 1998-2015 Epic Games, Inc. All Rights Reserved.
+
+#pragma once
+
+#include "OnlineTurnBasedInterface.h"
+#include "QuitMatchCallbackProxy.generated.h"
+
+UCLASS(MinimalAPI)
+class UQuitMatchCallbackProxy : public UOnlineBlueprintCallProxyBase
+{
+public:
+	GENERATED_UCLASS_BODY()
+
+	virtual ~UQuitMatchCallbackProxy();
+
+	// Called when there is a successful query
+	UPROPERTY(BlueprintAssignable)
+	FEmptyOnlineDelegate OnSuccess;
+
+	// Called when there is an unsuccessful query
+	UPROPERTY(BlueprintAssignable)
+	FEmptyOnlineDelegate OnFailure;
+
+	// Quits the turn based match
+	UFUNCTION(BlueprintCallable, meta = (BlueprintInternalUseOnly = "true", WorldContext = "WorldContextObject"), Category = "Online|TurnBased")
+	static UQuitMatchCallbackProxy* QuitMatch(UObject* WorldContextObject, APlayerController* PlayerController, FString MatchID, EMPMatchOutcome::Outcome Outcome, int32 TurnTimeoutInSeconds);
+
+	// UOnlineBlueprintCallProxyBase interface
+	virtual void Activate() override;
+	// End of UOnlineBlueprintCallProxyBase interface
+
+	void QuitMatchDelegate(FString MatchID, bool Succeeded);
+
+private:
+	// The player controller triggering things
+	TWeakObjectPtr<APlayerController> PlayerControllerWeakPtr;
+
+	// The world context object in which this call is taking place
+	UObject* WorldContextObject;
+
+	// The MatchID of the match to quit
+	FString MatchID;
+
+	// The outcome (won/lost/quit/etc.) of the match
+	EMPMatchOutcome::Outcome Outcome;
+
+	// If the match isn't over, this will be how much time the next player will have to complete their turn
+	int32 TurnTimeoutInSeconds;
+};
\ No newline at end of file
diff --git a/Engine/Source/Runtime/Online/OnlineSubsystemUtils/Classes/TurnBasedBlueprintLibrary.h b/Engine/Source/Runtime/Online/OnlineSubsystemUtils/Classes/TurnBasedBlueprintLibrary.h
new file mode 100644
--- /dev/null
+++ b/Engine/Source/Runtime/Online/OnlineSubsystemUtils/Classes/TurnBasedBlueprintLibrary.h
@@ -0,0 +1,19 @@
+// Copyright 1998-2015 Epic Games, Inc. All Rights Reserved.
+
+#pragma once
+#include "Kismet/BlueprintFunctionLibrary.h"
+#include "OnlineTurnBasedInterface.h"
+#include "TurnBasedBlueprintLibrary.generated.h"
+
+// Library of synchronous achievement calls
+UCLASS()
+class ONLINESUBSYSTEMUTILS_API UTurnBasedBlueprintLibrary : public UBlueprintFunctionLibrary
+{
+	GENERATED_UCLASS_BODY()
+
+	UFUNCTION(BlueprintCallable, meta = (WorldContext = "WorldContextObject"), Category = "Online|TurnBased")
+	static void GetIsMyTurn(UObject* WorldContextObject, APlayerController* PlayerController, FString MatchID, /*out*/ bool& bIsMyTurn);
+
+	UFUNCTION(BlueprintCallable, meta = (WorldContext = "WorldContextObject"), Category = "Online|TurnBased")
+	static void RegisterTurnBasedMatchInterfaceObject(UObject* WorldContextObject, APlayerController* PlayerController, UObject* Object);
+};
diff --git a/Engine/Source/Runtime/Online/OnlineSubsystemUtils/Private/EndTurnCallbackProxy.cpp b/Engine/Source/Runtime/Online/OnlineSubsystemUtils/Private/EndTurnCallbackProxy.cpp
new file mode 100644
--- /dev/null
+++ b/Engine/Source/Runtime/Online/OnlineSubsystemUtils/Private/EndTurnCallbackProxy.cpp
@@ -0,0 +1,83 @@
+#include "OnlineSubsystemUtilsPrivatePCH.h"
+#include "GameFramework/Actor.h"
+#include "RepLayout.h"
+
+//////////////////////////////////////////////////////////////////////////
+// UTurnBasedMatchEndTurnCallbackProxy
+
+UEndTurnCallbackProxy::UEndTurnCallbackProxy(const FObjectInitializer& ObjectInitializer)
+: Super(ObjectInitializer)
+{
+}
+
+UEndTurnCallbackProxy::~UEndTurnCallbackProxy()
+{
+}
+
+UEndTurnCallbackProxy* UEndTurnCallbackProxy::EndTurn(UObject* WorldContextObject, class APlayerController* PlayerController, FString MatchID, TScriptInterface<ITurnBasedMatchInterface> TurnBasedMatchInterface)
+{
+	UEndTurnCallbackProxy* Proxy = NewObject<UEndTurnCallbackProxy>();
+	Proxy->PlayerControllerWeakPtr = PlayerController;
+	Proxy->WorldContextObject = WorldContextObject;
+	Proxy->MatchID = MatchID;
+	Proxy->TurnBasedMatchInterface = (UTurnBasedMatchInterface*)TurnBasedMatchInterface.GetObject();
+	return Proxy;
+}
+
+void UEndTurnCallbackProxy::Activate()
+{
+	FOnlineSubsystemBPCallHelper Helper(TEXT("ConnectToService"), GEngine->GetWorldFromContextObject(WorldContextObject));
+	Helper.QueryIDFromPlayerController(PlayerControllerWeakPtr.Get());
+
+	if (Helper.IsValid())
+	{
+		IOnlineTurnBasedPtr TurnBasedInterface = Helper.OnlineSub->GetTurnBasedInterface();
+		if (TurnBasedInterface.IsValid())
+		{
+			if (TurnBasedMatchInterface != nullptr)
+			{
+				FRepLayout RepLayout;
+				RepLayout.InitFromObjectClass(TurnBasedMatchInterface->GetClass());
+				FBitWriter Writer(TurnBasedInterface->GetMatchDataSize());
+				RepLayout.SerializeObjectReplicatedProperties(TurnBasedMatchInterface, Writer);
+
+				FUploadMatchDataSignature UploadMatchDataSignature;
+				UploadMatchDataSignature.BindUObject(this, &UEndTurnCallbackProxy::UploadMatchDataDelegate);
+
+				FTurnBasedMatchPtr TurnBasedMatchPtr = TurnBasedInterface->GetMatchWithID(MatchID);
+				if (TurnBasedMatchPtr.IsValid())
+				{
+					TurnBasedMatchPtr->EndTurnWithMatchData(*Writer.GetBuffer(), 0, UploadMatchDataSignature);
+					return;
+				}
+
+				FString Message = FString::Printf(TEXT("Match ID %s not found"), *MatchID);
+				FFrame::KismetExecutionMessage(*Message, ELogVerbosity::Warning);
+			}
+			else
+			{
+				FFrame::KismetExecutionMessage(TEXT("No match data passed in to End Turn."), ELogVerbosity::Warning);
+			}
+			return;
+		}
+		else
+		{
+			FFrame::KismetExecutionMessage(TEXT("Turn based games not supported by online subsystem"), ELogVerbosity::Warning);
+		}
+	}
+
+	// Fail immediately
+	OnFailure.Broadcast();
+}
+
+void UEndTurnCallbackProxy::UploadMatchDataDelegate(FString MatchID, bool Succeeded)
+{
+	if (Succeeded)
+	{
+		OnSuccess.Broadcast();
+	}
+	else
+	{
+		OnFailure.Broadcast();
+	}
+}
diff --git a/Engine/Source/Runtime/Online/OnlineSubsystemUtils/Private/FindTurnBasedMatchCallbackProxy.cpp b/Engine/Source/Runtime/Online/OnlineSubsystemUtils/Private/FindTurnBasedMatchCallbackProxy.cpp
new file mode 100644
--- /dev/null
+++ b/Engine/Source/Runtime/Online/OnlineSubsystemUtils/Private/FindTurnBasedMatchCallbackProxy.cpp
@@ -0,0 +1,107 @@
+// Copyright 1998-2015 Epic Games, Inc. All Rights Reserved.
+
+#include "OnlineSubsystemUtilsPrivatePCH.h"
+#include "RepLayout.h"
+#include "TurnBasedMatchInterface.h"
+
+//////////////////////////////////////////////////////////////////////////
+// UFindTurnBasedMatchCallbackProxy
+
+UFindTurnBasedMatchCallbackProxy::UFindTurnBasedMatchCallbackProxy(const FObjectInitializer& ObjectInitializer)
+: Super(ObjectInitializer)
+, WorldContextObject(nullptr)
+, Delegate(new FFindTurnBasedMatchCallbackProxyMatchmakerDelegate)
+{
+}
+
+UFindTurnBasedMatchCallbackProxy::~UFindTurnBasedMatchCallbackProxy()
+{
+}
+
+UFindTurnBasedMatchCallbackProxy* UFindTurnBasedMatchCallbackProxy::FindTurnBasedMatch(UObject* WorldContextObject, class APlayerController* PlayerController, TScriptInterface<ITurnBasedMatchInterface> MatchActor, int32 MinPlayers, int32 MaxPlayers, int32 PlayerGroup, bool ShowExistingMatches)
+{
+	UFindTurnBasedMatchCallbackProxy* Proxy = NewObject<UFindTurnBasedMatchCallbackProxy>();
+	Proxy->PlayerControllerWeakPtr = PlayerController;
+	Proxy->WorldContextObject = WorldContextObject;
+	Proxy->MinPlayers = MinPlayers;
+	Proxy->MaxPlayers = MaxPlayers;
+	Proxy->PlayerGroup = PlayerGroup;
+	Proxy->ShowExistingMatches = ShowExistingMatches;
+	Proxy->TurnBasedMatchInterface = (UTurnBasedMatchInterface*)MatchActor.GetObject();
+	return Proxy;
+}
+
+void UFindTurnBasedMatchCallbackProxy::Activate()
+{
+	FOnlineSubsystemBPCallHelper Helper(TEXT("ConnectToService"), GEngine->GetWorldFromContextObject(WorldContextObject));
+	Helper.QueryIDFromPlayerController(PlayerControllerWeakPtr.Get());
+
+	if (Helper.IsValid())
+	{
+		IOnlineTurnBasedPtr TurnBasedInterface = Helper.OnlineSub->GetTurnBasedInterface();
+		if (TurnBasedInterface.IsValid())
+		{
+			Delegate->SetFindTurnBasedMatchCallbackProxy(this);
+			Delegate->SetTurnBasedInterface(TurnBasedInterface);
+			TurnBasedInterface->SetMatchmakerDelegate(Delegate);
+			FTurnBasedMatchRequest MatchRequest(MinPlayers, MaxPlayers, PlayerGroup, ShowExistingMatches);
+			TurnBasedInterface->ShowMatchmaker(MatchRequest);
+
+			// Results will be handled in the FFindTurnBasedMatchCallbackProxyMatchmakerDelegate object
+			return;
+		}
+		else
+		{
+			FFrame::KismetExecutionMessage(TEXT("Turn based games not supported by online subsystem"), ELogVerbosity::Warning);
+		}
+	}
+
+	// Fail immediately
+	OnFailure.Broadcast(FString());
+}
+
+FFindTurnBasedMatchCallbackProxyMatchmakerDelegate::FFindTurnBasedMatchCallbackProxyMatchmakerDelegate()
+: FTurnBasedMatchmakerDelegate()
+, FindTurnBasedMatchCallbackProxy(nullptr)
+, TurnBasedInterface(nullptr)
+{
+}
+
+FFindTurnBasedMatchCallbackProxyMatchmakerDelegate::~FFindTurnBasedMatchCallbackProxyMatchmakerDelegate()
+{
+}
+
+void FFindTurnBasedMatchCallbackProxyMatchmakerDelegate::OnMatchmakerCancelled()
+{
+	if (FindTurnBasedMatchCallbackProxy)
+	{
+		FindTurnBasedMatchCallbackProxy->OnFailure.Broadcast(FString());
+	}
+}
+
+void FFindTurnBasedMatchCallbackProxyMatchmakerDelegate::OnMatchmakerFailed()
+{
+	if (FindTurnBasedMatchCallbackProxy)
+	{
+		FindTurnBasedMatchCallbackProxy->OnFailure.Broadcast(FString());
+	}
+}
+
+void FFindTurnBasedMatchCallbackProxyMatchmakerDelegate::OnMatchFound(FTurnBasedMatchRef Match)
+{
+    UE_LOG(LogOnline, Verbose, TEXT("Turn-based match found: %s"), *Match->GetMatchID());
+	TArray<uint8> MatchData;
+	
+	if (Match->GetMatchData(MatchData))
+	{
+		FRepLayout RepLayout;
+		RepLayout.InitFromObjectClass(FindTurnBasedMatchCallbackProxy->GetTurnBasedMatchInterfaceObject()->GetClass());
+		FBitReader Reader(MatchData.GetData(), TurnBasedInterface->GetMatchDataSize());
+		RepLayout.SerializeObjectReplicatedProperties(FindTurnBasedMatchCallbackProxy->GetTurnBasedMatchInterfaceObject(), Reader);
+	}
+
+	if (FindTurnBasedMatchCallbackProxy)
+	{
+		FindTurnBasedMatchCallbackProxy->OnSuccess.Broadcast(Match->GetMatchID());
+	}
+}
\ No newline at end of file
diff --git a/Engine/Source/Runtime/Online/OnlineSubsystemUtils/Private/QuitMatchCallbackProxy.cpp b/Engine/Source/Runtime/Online/OnlineSubsystemUtils/Private/QuitMatchCallbackProxy.cpp
new file mode 100644
--- /dev/null
+++ b/Engine/Source/Runtime/Online/OnlineSubsystemUtils/Private/QuitMatchCallbackProxy.cpp
@@ -0,0 +1,70 @@
+#include "OnlineSubsystemUtilsPrivatePCH.h"
+#include "QuitMatchCallbackProxy.h"
+#include "GameFramework/Actor.h"
+#include "RepLayout.h"
+
+UQuitMatchCallbackProxy::UQuitMatchCallbackProxy(const FObjectInitializer& ObjectInitializer)
+	: Super(ObjectInitializer)
+{
+}
+
+UQuitMatchCallbackProxy::~UQuitMatchCallbackProxy()
+{
+}
+
+UQuitMatchCallbackProxy* UQuitMatchCallbackProxy::QuitMatch(UObject* WorldContextObject, APlayerController* PlayerController, FString MatchID, EMPMatchOutcome::Outcome Outcome, int32 TurnTimeoutInSeconds)
+{
+	UQuitMatchCallbackProxy* Proxy = NewObject<UQuitMatchCallbackProxy>();
+	Proxy->PlayerControllerWeakPtr = PlayerController;
+	Proxy->WorldContextObject = WorldContextObject;
+    Proxy->MatchID = MatchID;
+	Proxy->Outcome = Outcome;
+	Proxy->TurnTimeoutInSeconds = TurnTimeoutInSeconds;
+	return Proxy;
+}
+
+void UQuitMatchCallbackProxy::Activate()
+{
+	FOnlineSubsystemBPCallHelper Helper(TEXT("ConnectToService"), GEngine->GetWorldFromContextObject(WorldContextObject));
+	Helper.QueryIDFromPlayerController(PlayerControllerWeakPtr.Get());
+
+	if (Helper.IsValid())
+	{
+		IOnlineTurnBasedPtr OnlineTurnBasedPtr = Helper.OnlineSub->GetTurnBasedInterface();
+		if (OnlineTurnBasedPtr.IsValid())
+		{
+			FTurnBasedMatchPtr Match = OnlineTurnBasedPtr->GetMatchWithID(MatchID);
+			if (Match.IsValid())
+			{
+				FQuitMatchSignature QuitMatchSignature;
+				QuitMatchSignature.BindUObject(this, &UQuitMatchCallbackProxy::QuitMatchDelegate);
+				Match->QuitMatch(Outcome, TurnTimeoutInSeconds, QuitMatchSignature);
+                return;
+			}
+			else
+			{
+				FString Message = FString::Printf(TEXT("Match ID %s not found"), *MatchID);
+				FFrame::KismetExecutionMessage(*Message, ELogVerbosity::Warning);
+			}
+		}
+		else
+		{
+			FFrame::KismetExecutionMessage(TEXT("Turn Based Matches not supported by Online Subsystem"), ELogVerbosity::Warning);
+		}
+	}
+
+	// Fail immediately
+	OnFailure.Broadcast();
+}
+
+void UQuitMatchCallbackProxy::QuitMatchDelegate(FString MatchID, bool Succeeded)
+{
+	if (Succeeded)
+	{
+		OnSuccess.Broadcast();
+	}
+	else
+	{
+		OnFailure.Broadcast();
+	}
+}
diff --git a/Engine/Source/Runtime/Online/OnlineSubsystemUtils/Private/TurnBasedBlueprintLibrary.cpp b/Engine/Source/Runtime/Online/OnlineSubsystemUtils/Private/TurnBasedBlueprintLibrary.cpp
new file mode 100644
--- /dev/null
+++ b/Engine/Source/Runtime/Online/OnlineSubsystemUtils/Private/TurnBasedBlueprintLibrary.cpp
@@ -0,0 +1,55 @@
+// Copyright 1998-2015 Epic Games, Inc. All Rights Reserved.
+
+#include "OnlineSubsystemUtilsPrivatePCH.h"
+#include "TurnBasedBlueprintLibrary.h"
+#include "Runtime/Online/OnlineSubsystem/Public/Interfaces/OnlineTurnBasedInterface.h"
+
+UTurnBasedBlueprintLibrary::UTurnBasedBlueprintLibrary(const FObjectInitializer& ObjectInitializer)
+: Super(ObjectInitializer)
+{
+}
+
+void UTurnBasedBlueprintLibrary::GetIsMyTurn(UObject* WorldContextObject, APlayerController* PlayerController, FString MatchID, /*out*/ bool& bIsMyTurn)
+{
+	bIsMyTurn = false;
+
+	FOnlineSubsystemBPCallHelper Helper(TEXT("GetIsMyTurn"), GEngine->GetWorldFromContextObject(WorldContextObject));
+	Helper.QueryIDFromPlayerController(PlayerController);
+
+	if (Helper.IsValid())
+	{
+		IOnlineTurnBasedPtr OnlineTurnBasedPtr = Helper.OnlineSub->GetTurnBasedInterface();
+		if (OnlineTurnBasedPtr.IsValid())
+		{
+			FTurnBasedMatchPtr Match = OnlineTurnBasedPtr->GetMatchWithID(MatchID);
+			if (Match.IsValid())
+			{
+				bIsMyTurn = Match->GetCurrentPlayerIndex() == Match->GetLocalPlayerIndex();
+			}
+			else
+			{
+				FString Message = FString::Printf(TEXT("Match ID %s not found"), *MatchID);
+				FFrame::KismetExecutionMessage(*Message, ELogVerbosity::Warning);
+			}
+		}
+		else
+		{
+			FFrame::KismetExecutionMessage(TEXT("Turn Based Matches not supported by Online Subsystem"), ELogVerbosity::Warning);
+		}
+	}
+}
+
+void UTurnBasedBlueprintLibrary::RegisterTurnBasedMatchInterfaceObject(UObject* WorldContextObject, APlayerController* PlayerController, UObject* Object)
+{
+	FOnlineSubsystemBPCallHelper Helper(TEXT("GetIsMyTurn"), GEngine->GetWorldFromContextObject(WorldContextObject));
+	Helper.QueryIDFromPlayerController(PlayerController);
+
+	if (Helper.IsValid())
+	{
+		IOnlineTurnBasedPtr OnlineTurnBasedPtr = Helper.OnlineSub->GetTurnBasedInterface();
+		if (OnlineTurnBasedPtr.IsValid())
+		{
+			OnlineTurnBasedPtr->RegisterTurnBasedMatchInterfaceObject(Object);
+		}
+	}
+}
\ No newline at end of file
