Build fix attempt for android

diff --git a/Engine/Source/Runtime/Online/Android/OnlineSubsystemGooglePlay/Public/OnlineSubsystemGooglePlay.h b/Engine/Source/Runtime/Online/Android/OnlineSubsystemGooglePlay/Public/OnlineSubsystemGooglePlay.h
--- a/Engine/Source/Runtime/Online/Android/OnlineSubsystemGooglePlay/Public/OnlineSubsystemGooglePlay.h
+++ b/Engine/Source/Runtime/Online/Android/OnlineSubsystemGooglePlay/Public/OnlineSubsystemGooglePlay.h
@@ -46,6 +46,7 @@ public:
 	virtual IOnlineAchievementsPtr GetAchievementsInterface() const override;
 	virtual IOnlinePresencePtr GetPresenceInterface() const override { return NULL; }
 	virtual IOnlineChatPtr GetChatInterface() const override { return NULL; }
+	virtual IOnlineTurnBasedPtr GetTurnBasedInterface() const override { return NULL; }
 
 	virtual class UObject* GetNamedInterface(FName InterfaceName) override { return NULL; }
 	virtual void SetNamedInterface(FName InterfaceName, class UObject* NewInterface) override {}
