Add macro to allow serialization of nested FOnlineJsonSerializable objects.

#github https://github.com/EpicGames/UnrealEngine/pull/741 by korypostma

diff --git a/Engine/Source/Runtime/Online/OnlineSubsystem/Public/OnlineJsonSerializer.h b/Engine/Source/Runtime/Online/OnlineSubsystem/Public/OnlineJsonSerializer.h
--- a/Engine/Source/Runtime/Online/OnlineSubsystem/Public/OnlineJsonSerializer.h
+++ b/Engine/Source/Runtime/Online/OnlineSubsystem/Public/OnlineJsonSerializer.h
@@ -83,6 +83,28 @@
 			Serializer.EndObject(); \
 		}
 
+#define ONLINE_JSON_SERIALIZE_OBJECT_SERIALIZABLE(JsonName, JsonSerializableObject) \
+		/* Process the JsonName field differently because it is an object */ \
+		if (Serializer.IsLoading()) \
+		{ \
+			/* Read in the value from the JsonName field */ \
+			if (Serializer.GetObject()->HasTypedField<EJson::Object>(JsonName)) \
+			{ \
+				TSharedPtr<FJsonObject> JsonObj = Serializer.GetObject()->GetObjectField(JsonName); \
+				if (JsonObj.IsValid()) \
+				{ \
+					JsonSerializableObject.FromJson(JsonObj); \
+				} \
+			} \
+		} \
+		else \
+		{ \
+			/* Write the value to the JsonName field */ \
+			Serializer.StartObject(JsonName); \
+			JsonSerializableObject.Serialize(Serializer); \
+			Serializer.EndObject(); \
+		}
+
 /** Array of string data */
 typedef TArray<FString> FJsonSerializableArray;
 
