Fixed problems with use of TArrays and NSMutableData in MacHTTP

diff --git a/Engine/Source/Runtime/Online/HTTP/Private/Mac/MacHTTP.cpp b/Engine/Source/Runtime/Online/HTTP/Private/Mac/MacHTTP.cpp
--- a/Engine/Source/Runtime/Online/HTTP/Private/Mac/MacHTTP.cpp
+++ b/Engine/Source/Runtime/Online/HTTP/Private/Mac/MacHTTP.cpp
@@ -100,10 +100,7 @@ const TArray<uint8>& FMacHttpRequest::GetContent()
 {
 	UE_LOG(LogHttp, Verbose, TEXT("FMacHttpRequest::GetContent()"));
 	RequestPayload.Empty();
-	RequestPayload.Reserve( GetContentLength() );
-
-	FMemory::Memcpy( RequestPayload.GetData(), [[Request HTTPBody] bytes], RequestPayload.Num() );
-
+	RequestPayload.Append((const uint8*)[[Request HTTPBody] bytes], GetContentLength());
 	return RequestPayload;
 }
 
@@ -349,7 +346,6 @@ void FMacHttpRequest::Tick(float DeltaSeconds)
 
 @implementation FHttpResponseMacWrapper
 @synthesize Response;
-@synthesize Payload;
 @synthesize bIsReady;
 @synthesize bHadError;
 
@@ -370,15 +366,16 @@ void FMacHttpRequest::Tick(float DeltaSeconds)
 	self.Response = (NSHTTPURLResponse*)response;
 	
 	// presize the payload container if possible
-	self.Payload = [NSMutableData dataWithCapacity:([response expectedContentLength] != NSURLResponseUnknownLength ? [response expectedContentLength] : 0)];
-	UE_LOG(LogHttp, Verbose, TEXT("didReceiveResponse: expectedContentLength = %d. Length = %d: %p"), [response expectedContentLength], [self.Payload length], self);
+	Payload.Empty();
+	Payload.Reserve([response expectedContentLength] != NSURLResponseUnknownLength ? [response expectedContentLength] : 0);
+	UE_LOG(LogHttp, Verbose, TEXT("didReceiveResponse: expectedContentLength = %d. Length = %d: %p"), [response expectedContentLength], Payload.Max(), self);
 }
 
 
 -(void) connection:(NSURLConnection *)connection didReceiveData:(NSData *)data
 {
-	[self.Payload appendData:data];
-	UE_LOG(LogHttp, Verbose, TEXT("didReceiveData with %d bytes. After Append, Payload Length = %d: %p"), [data length], [self.Payload length], self);
+	Payload.Append((const uint8*)[data bytes], [data length]);
+	UE_LOG(LogHttp, Verbose, TEXT("didReceiveData with %d bytes. After Append, Payload Length = %d: %p"), [data length], Payload.Num(), self);
 }
 
 
@@ -399,6 +396,11 @@ void FMacHttpRequest::Tick(float DeltaSeconds)
 	self.bIsReady = YES;
 }
 
+- (TArray<uint8>&)getPayload
+{
+	return Payload;
+}
+
 @end
 
 
@@ -419,8 +421,7 @@ FMacHttpResponse::FMacHttpResponse(const FMacHttpRequest& InRequest)
 FMacHttpResponse::~FMacHttpResponse()
 {
 	UE_LOG(LogHttp, Verbose, TEXT("FMacHttpResponse::~FMacHttpResponse()"));
-	[[ResponseWrapper Payload] release];
-	ResponseWrapper.Payload = nil;
+	[ResponseWrapper getPayload].Empty();
 
 	[ResponseWrapper release];
 	ResponseWrapper = nil;
@@ -491,7 +492,7 @@ int32 FMacHttpResponse::GetContentLength()
 {
 	UE_LOG(LogHttp, Verbose, TEXT("FMacHttpResponse::GetContentLength()"));
 
-	return [[ResponseWrapper Payload] length];
+	return [ResponseWrapper getPayload].Num();
 }
 
 
@@ -503,12 +504,8 @@ const TArray<uint8>& FMacHttpResponse::GetContent()
 	}
 	else
 	{
-		// todo: Find out why we're empty and doing a memcopy every time instead of just once when we receive the payload
-		Payload.Empty();
-		Payload.AddZeroed( [[ResponseWrapper Payload] length] );
-	
-		FMemory::Memcpy( Payload.GetData(), [[ResponseWrapper Payload] bytes], [[ResponseWrapper Payload] length] );
-		UE_LOG(LogHttp, Verbose, TEXT("FMacHttpResponse::GetContent() - Num: %i"), [[ResponseWrapper Payload] length]);
+		Payload = [ResponseWrapper getPayload];
+		UE_LOG(LogHttp, Verbose, TEXT("FMacHttpResponse::GetContent() - Num: %i"), [ResponseWrapper getPayload].Num());
 	}
 
 	return Payload;
@@ -572,5 +569,5 @@ bool FMacHttpResponse::HadError()
 
 const int32 FMacHttpResponse::GetNumBytesReceived() const
 {
-	return [[ResponseWrapper Payload] length];
+	return [ResponseWrapper getPayload].Num();
 }
diff --git a/Engine/Source/Runtime/Online/HTTP/Private/Mac/MacHTTP.h b/Engine/Source/Runtime/Online/HTTP/Private/Mac/MacHTTP.h
--- a/Engine/Source/Runtime/Online/HTTP/Private/Mac/MacHTTP.h
+++ b/Engine/Source/Runtime/Online/HTTP/Private/Mac/MacHTTP.h
@@ -101,17 +101,16 @@ private:
 };
 
 
-
 /**
  * Mac Response Wrapper which will be used for it's delegates to receive responses.
  */
 @interface FHttpResponseMacWrapper : NSObject
 {
-};
+	/** Holds the payload as we receive it. */
+	TArray<uint8> Payload;
+}
 /** A handle for the response */
 @property(retain) NSHTTPURLResponse* Response;
-/** Holds the payload as we receive it. */
-@property(retain) NSMutableData* Payload;
 /** Flag whether the response is ready */
 @property BOOL bIsReady;
 /** When the response is complete, indicates whether the response was received without error. */
@@ -125,6 +124,8 @@ private:
 -(void) connection:(NSURLConnection *)connection didFailWithError:(NSError *)error;
 /** Delegate called with we complete successfully. See Mac docs for when/how this should be used. */
 -(void) connectionDidFinishLoading:(NSURLConnection *)connection;
+
+- (TArray<uint8>&)getPayload;
 @end
 
 
