Fix for UE-7324 freeze with GooglePlay: Added wait for GooglePlay to steal focus after initialization, corrected an issue with return code for early touch events, and removed screen wakeup (caused Samsung freezes).
#android
#ue4

diff --git a/Engine/Build/Android/Java/src/com/epicgames/ue4/GameActivity.java b/Engine/Build/Android/Java/src/com/epicgames/ue4/GameActivity.java
--- a/Engine/Build/Android/Java/src/com/epicgames/ue4/GameActivity.java
+++ b/Engine/Build/Android/Java/src/com/epicgames/ue4/GameActivity.java
@@ -186,12 +186,13 @@ public class GameActivity extends NativeActivity
 
 		_activity = this;
 		
+/*
 		// Turn on and unlock screen.. Assumption is that this
 		// will only really have an effect when for debug launching
 		// as otherwise the screen is already unlocked.
 		this.getWindow().addFlags(
 			WindowManager.LayoutParams.FLAG_TURN_SCREEN_ON |
-			WindowManager.LayoutParams.FLAG_SHOW_WHEN_LOCKED |
+//			WindowManager.LayoutParams.FLAG_SHOW_WHEN_LOCKED |
 			WindowManager.LayoutParams.FLAG_DISMISS_KEYGUARD);
 		// On some devices we can also unlock a key-locked screen by disabling the
 		// keylock guard. To be safe we only do this on < Android 3.2. As the API
@@ -202,6 +203,7 @@ public class GameActivity extends NativeActivity
 			android.app.KeyguardManager.KeyguardLock keylock = keyman.newKeyguardLock("Unlock");
 			keylock.disableKeyguard();
 		}
+*/
 
 		// tell Android that we want volume controls to change the media volume, aka music
 		setVolumeControlStream(AudioManager.STREAM_MUSIC);
diff --git a/Engine/Source/Runtime/Launch/Private/Android/AndroidEventManager.cpp b/Engine/Source/Runtime/Launch/Private/Android/AndroidEventManager.cpp
--- a/Engine/Source/Runtime/Launch/Private/Android/AndroidEventManager.cpp
+++ b/Engine/Source/Runtime/Launch/Private/Android/AndroidEventManager.cpp
@@ -324,3 +324,46 @@ bool FAppEventManager::IsGamePaused()
 {
 	return !bRunning;
 }
+
+void FAppEventManager::WaitForEventInQueue(EAppEventState InState)
+{
+	bool FoundEvent = false;
+
+	TQueue<FAppEventData, EQueueMode::Spsc> HoldingQueue;
+	while (!FoundEvent)
+	{
+		int rc = pthread_mutex_lock(&QueueMutex);
+		check(rc == 0);
+
+		// Copy the existing queue (and check for our event)
+		while (!Queue.IsEmpty())
+		{
+			FAppEventData OutData;
+			Queue.Dequeue(OutData);
+
+			if (OutData.State == InState)
+				FoundEvent = true;
+
+			HoldingQueue.Enqueue(OutData);
+		}
+
+		if (FoundEvent)
+			break;
+
+		// Unlock for new events and wait a bit before trying again
+		rc = pthread_mutex_unlock(&QueueMutex);
+		check(rc == 0);
+		FPlatformProcess::Sleep(0.01f);
+	}
+
+	// Add events back to queue from holding
+	while (!HoldingQueue.IsEmpty())
+	{
+		FAppEventData OutData;
+		HoldingQueue.Dequeue(OutData);
+		Queue.Enqueue(OutData);
+	}
+
+	int rc = pthread_mutex_unlock(&QueueMutex);
+	check(rc == 0);
+}
diff --git a/Engine/Source/Runtime/Launch/Private/Android/AndroidEventManager.h b/Engine/Source/Runtime/Launch/Private/Android/AndroidEventManager.h
--- a/Engine/Source/Runtime/Launch/Private/Android/AndroidEventManager.h
+++ b/Engine/Source/Runtime/Launch/Private/Android/AndroidEventManager.h
@@ -46,6 +46,7 @@ public:
 	void SetEventHandlerEvent(FEvent* InEventHandlerEvent);
 	void HandleWindowCreated(void* InWindow);
 	bool IsGamePaused();
+	void WaitForEventInQueue(EAppEventState InState);
 
 protected:
 	FAppEventManager();
diff --git a/Engine/Source/Runtime/Launch/Private/Android/LaunchAndroid.cpp b/Engine/Source/Runtime/Launch/Private/Android/LaunchAndroid.cpp
--- a/Engine/Source/Runtime/Launch/Private/Android/LaunchAndroid.cpp
+++ b/Engine/Source/Runtime/Launch/Private/Android/LaunchAndroid.cpp
@@ -513,7 +513,7 @@ static int32_t HandleInputCB(struct android_app* app, AInputEvent* event)
 			// make sure OpenGL context created before accepting touch events.. FAndroidWindow::GetScreenRect() may try to create it early from wrong thread if this is the first call
 			if (!GAndroidGPUInfoReady)
 			{
-				return 0;
+				return 1;
 			}
 			FPlatformRect ScreenRect = FAndroidWindow::GetScreenRect();
 
diff --git a/Engine/Source/Runtime/Online/Android/OnlineSubsystemGooglePlay/Private/OnlineSubsystemGooglePlay.cpp b/Engine/Source/Runtime/Online/Android/OnlineSubsystemGooglePlay/Private/OnlineSubsystemGooglePlay.cpp
--- a/Engine/Source/Runtime/Online/Android/OnlineSubsystemGooglePlay/Private/OnlineSubsystemGooglePlay.cpp
+++ b/Engine/Source/Runtime/Online/Android/OnlineSubsystemGooglePlay/Private/OnlineSubsystemGooglePlay.cpp
@@ -99,6 +99,9 @@ IOnlineAchievementsPtr FOnlineSubsystemGooglePlay::GetAchievementsInterface() co
 	return AchievementsInterface;
 }
 
+static bool WaitForLostFocus = false;
+static bool WaitingForLogin = false;
+
 bool FOnlineSubsystemGooglePlay::Init() 
 {
 	FPlatformMisc::LowLevelOutputDebugStringf(TEXT("FOnlineSubsystemAndroid::Init"));
@@ -129,6 +132,8 @@ bool FOnlineSubsystemGooglePlay::Init()
 	CurrentLoginTask = new FOnlineAsyncTaskGooglePlayLogin(this, 0);
 	QueueAsyncTask(CurrentLoginTask);
 
+	WaitingForLogin = true;
+
 	// Create() returns a std::unqiue_ptr, but we convert it to a TUniquePtr.
 	GameServicesPtr.Reset( GameServices::Builder()
 		.SetDefaultOnLog(LogLevel::VERBOSE)
@@ -137,9 +142,29 @@ bool FOnlineSubsystemGooglePlay::Init()
 		})
 		.SetOnAuthActionFinished([this](AuthOperation Op, AuthStatus Status) {
 			OnAuthActionFinished(Op, Status);
+
+			if (Op == AuthOperation::SIGN_IN)
+			{
+				if (Status != AuthStatus::VALID)
+				{
+					WaitForLostFocus = true;
+				}
+				WaitingForLogin = false;
+			}
 		})
 		.Create(PlatformConfiguration).release() );
 
+	// Wait for GooglePlay to complete login task
+	while (WaitingForLogin)
+	{
+		FPlatformProcess::Sleep(0.01f);
+	}
+	if (WaitForLostFocus)
+	{
+		extern void WaitForAndroidLoseFocusEvent();
+		WaitForAndroidLoseFocusEvent();
+	}
+
 	return true;
 }
 
