Add GooglePlay login timeout and pass on onActivityResult to GPG with delegate
#ue4
#android

diff --git a/Engine/Build/Android/Java/src/com/epicgames/ue4/GameActivity.java b/Engine/Build/Android/Java/src/com/epicgames/ue4/GameActivity.java
--- a/Engine/Build/Android/Java/src/com/epicgames/ue4/GameActivity.java
+++ b/Engine/Build/Android/Java/src/com/epicgames/ue4/GameActivity.java
@@ -898,6 +898,7 @@ public class GameActivity extends NativeActivity
 		{
 			super.onActivityResult(requestCode, resultCode, data);
 		}
+		nativeOnActivityResult(this, requestCode, resultCode, data);
 	}
 	
 	public boolean AndroidThunkJava_IapBeginPurchase(String ProductId, boolean bConsumable)
@@ -942,6 +943,8 @@ public class GameActivity extends NativeActivity
 	public native void nativeInitHMDs();
 
 	public native void nativeResumeMainInit();
+
+	public native void nativeOnActivityResult(GameActivity activity, int requestCode, int resultCode, Intent data);
 	
 	static
 	{
diff --git a/Engine/Source/Runtime/Launch/Private/Android/AndroidEventManager.cpp b/Engine/Source/Runtime/Launch/Private/Android/AndroidEventManager.cpp
--- a/Engine/Source/Runtime/Launch/Private/Android/AndroidEventManager.cpp
+++ b/Engine/Source/Runtime/Launch/Private/Android/AndroidEventManager.cpp
@@ -325,9 +325,10 @@ bool FAppEventManager::IsGamePaused()
 	return !bRunning;
 }
 
-void FAppEventManager::WaitForEventInQueue(EAppEventState InState)
+bool FAppEventManager::WaitForEventInQueue(EAppEventState InState, double TimeoutSeconds)
 {
 	bool FoundEvent = false;
+	double StopTime = FPlatformTime::Seconds() + TimeoutSeconds;
 
 	TQueue<FAppEventData, EQueueMode::Spsc> HoldingQueue;
 	while (!FoundEvent)
@@ -350,6 +351,10 @@ void FAppEventManager::WaitForEventInQueue(EAppEventState InState)
 		if (FoundEvent)
 			break;
 
+		// Time expired?
+		if (FPlatformTime::Seconds() > StopTime)
+			break;
+
 		// Unlock for new events and wait a bit before trying again
 		rc = pthread_mutex_unlock(&QueueMutex);
 		check(rc == 0);
@@ -366,4 +371,6 @@ void FAppEventManager::WaitForEventInQueue(EAppEventState InState)
 
 	int rc = pthread_mutex_unlock(&QueueMutex);
 	check(rc == 0);
+
+	return FoundEvent;
 }
diff --git a/Engine/Source/Runtime/Launch/Private/Android/AndroidEventManager.h b/Engine/Source/Runtime/Launch/Private/Android/AndroidEventManager.h
--- a/Engine/Source/Runtime/Launch/Private/Android/AndroidEventManager.h
+++ b/Engine/Source/Runtime/Launch/Private/Android/AndroidEventManager.h
@@ -46,7 +46,7 @@ public:
 	void SetEventHandlerEvent(FEvent* InEventHandlerEvent);
 	void HandleWindowCreated(void* InWindow);
 	bool IsGamePaused();
-	void WaitForEventInQueue(EAppEventState InState);
+	bool WaitForEventInQueue(EAppEventState InState, double TimeoutSeconds);
 
 protected:
 	FAppEventManager();
diff --git a/Engine/Source/Runtime/Launch/Private/Android/AndroidJNI.cpp b/Engine/Source/Runtime/Launch/Private/Android/AndroidJNI.cpp
--- a/Engine/Source/Runtime/Launch/Private/Android/AndroidJNI.cpp
+++ b/Engine/Source/Runtime/Launch/Private/Android/AndroidJNI.cpp
@@ -20,6 +20,8 @@ extern FString GFilePathBase;
 extern FString GFontPathBase;
 extern bool GOBBinAPK;
 
+FOnActivityResult FJavaWrapper::OnActivityResultDelegate;
+
 //////////////////////////////////////////////////////////////////////////
 
 #if UE_BUILD_SHIPPING
@@ -562,3 +564,8 @@ extern "C" bool Java_com_epicgames_ue4_GameActivity_nativeIsShippingBuild(JNIEnv
 	return JNI_FALSE;
 #endif
 }
+
+extern "C" void Java_com_epicgames_ue4_GameActivity_nativeOnActivityResult(JNIEnv* jenv, jobject thiz, jobject activity, jint requestCode, jint resultCode, jobject data)
+{
+	FJavaWrapper::OnActivityResultDelegate.Broadcast(jenv, thiz, activity, requestCode, resultCode, data);
+}
diff --git a/Engine/Source/Runtime/Launch/Private/Android/AndroidJNI.h b/Engine/Source/Runtime/Launch/Private/Android/AndroidJNI.h
--- a/Engine/Source/Runtime/Launch/Private/Android/AndroidJNI.h
+++ b/Engine/Source/Runtime/Launch/Private/Android/AndroidJNI.h
@@ -7,6 +7,8 @@
 
 extern JavaVM* GJavaVM;
 
+DECLARE_MULTICAST_DELEGATE_SixParams(FOnActivityResult, JNIEnv *, jobject, jobject, jint, jint, jobject);
+
 // Define all the Java classes/methods that the game will need to access to
 class FJavaWrapper
 {
@@ -57,4 +59,7 @@ public:
 	static void CallVoidMethod(JNIEnv* Env, jobject Object, jmethodID Method, ...);
 	static jobject CallObjectMethod(JNIEnv* Env, jobject Object, jmethodID Method, ...);
 	static bool CallBooleanMethod(JNIEnv* Env, jobject Object, jmethodID Method, ...);
+
+	// Delegate that can be registered to that is called when an activity is finished
+	static FOnActivityResult OnActivityResultDelegate;
 };
diff --git a/Engine/Source/Runtime/Launch/Private/Android/LaunchAndroid.cpp b/Engine/Source/Runtime/Launch/Private/Android/LaunchAndroid.cpp
--- a/Engine/Source/Runtime/Launch/Private/Android/LaunchAndroid.cpp
+++ b/Engine/Source/Runtime/Launch/Private/Android/LaunchAndroid.cpp
@@ -961,7 +961,7 @@ extern "C" void Java_com_epicgames_ue4_GameActivity_nativeSetAndroidVersionInfor
 	FAndroidMisc::SetVersionInfo( UEAndroidVersion, UEPhoneMake, UEPhoneModel, UEOSLanguage );
 }
 
-void WaitForAndroidLoseFocusEvent()
+bool WaitForAndroidLoseFocusEvent(double TimeoutSeconds)
 {
-	FAppEventManager::GetInstance()->WaitForEventInQueue(EAppEventState::APP_EVENT_STATE_WINDOW_LOST_FOCUS);
+	return FAppEventManager::GetInstance()->WaitForEventInQueue(EAppEventState::APP_EVENT_STATE_WINDOW_LOST_FOCUS, TimeoutSeconds);
 }
diff --git a/Engine/Source/Runtime/Online/Android/OnlineSubsystemGooglePlay/Private/OnlineSubsystemGooglePlay.cpp b/Engine/Source/Runtime/Online/Android/OnlineSubsystemGooglePlay/Private/OnlineSubsystemGooglePlay.cpp
--- a/Engine/Source/Runtime/Online/Android/OnlineSubsystemGooglePlay/Private/OnlineSubsystemGooglePlay.cpp
+++ b/Engine/Source/Runtime/Online/Android/OnlineSubsystemGooglePlay/Private/OnlineSubsystemGooglePlay.cpp
@@ -4,6 +4,7 @@
 
 #include "OnlineSubsystemGooglePlayPrivatePCH.h"
 #include "AndroidApplication.h"
+#include "Android/AndroidJNI.h"
 #include "OnlineAsyncTaskManagerGooglePlay.h"
 #include "OnlineAsyncTaskGooglePlayLogin.h"
 
@@ -13,6 +14,7 @@
 #include "gpg/android_platform_configuration.h"
 #include "gpg/builder.h"
 #include "gpg/debug.h"
+#include "gpg/android_support.h"
 
 using namespace gpg;
 
@@ -133,6 +135,9 @@ bool FOnlineSubsystemGooglePlay::Init()
 	QueueAsyncTask(CurrentLoginTask);
 
 	WaitingForLogin = true;
+	WaitForLostFocus = false;
+
+	FJavaWrapper::OnActivityResultDelegate.AddRaw(this, &FOnlineSubsystemGooglePlay::OnActivityResult);
 
 	// Create() returns a std::unqiue_ptr, but we convert it to a TUniquePtr.
 	GameServicesPtr.Reset( GameServices::Builder()
@@ -161,8 +166,12 @@ bool FOnlineSubsystemGooglePlay::Init()
 	}
 	if (WaitForLostFocus)
 	{
-		extern void WaitForAndroidLoseFocusEvent();
-		WaitForAndroidLoseFocusEvent();
+		extern bool WaitForAndroidLoseFocusEvent(double TimeoutSeconds);
+		if (!WaitForAndroidLoseFocusEvent(3.0))
+		{
+			FPlatformMisc::LowLevelOutputDebugString(TEXT("FOnlineSubsystemAndroid::Init - timed out"));
+			OnAuthActionFinished(AuthOperation::SIGN_IN, AuthStatus::ERROR_NOT_AUTHORIZED);
+		}
 	}
 
 	return true;
@@ -183,6 +192,8 @@ bool FOnlineSubsystemGooglePlay::Shutdown()
 {
 	UE_LOG(LogOnline, Log, TEXT("FOnlineSubsystemAndroid::Shutdown()"));
 
+	FJavaWrapper::OnActivityResultDelegate.RemoveRaw(this, &FOnlineSubsystemGooglePlay::OnActivityResult);
+
 #define DESTRUCT_INTERFACE(Interface) \
 	if (Interface.IsValid()) \
 	{ \
@@ -271,4 +282,10 @@ std::string FOnlineSubsystemGooglePlay::ConvertFStringToStdString(const FString&
 	FPlatformString::Convert(Converted.GetData(), DestLen, *InString, SrcLen);
 
 	return std::string(Converted.GetData());
-}
\ No newline at end of file
+}
+
+void FOnlineSubsystemGooglePlay::OnActivityResult(JNIEnv *env, jobject thiz, jobject activity, jint requestCode, jint resultCode, jobject data)
+{
+	// Pass the result on to google play - otherwise, some callbacks for the turn based system do not get called.
+	AndroidSupport::OnActivityResult(env, activity, requestCode, resultCode, data);
+}
diff --git a/Engine/Source/Runtime/Online/Android/OnlineSubsystemGooglePlay/Public/OnlineSubsystemGooglePlay.h b/Engine/Source/Runtime/Online/Android/OnlineSubsystemGooglePlay/Public/OnlineSubsystemGooglePlay.h
--- a/Engine/Source/Runtime/Online/Android/OnlineSubsystemGooglePlay/Public/OnlineSubsystemGooglePlay.h
+++ b/Engine/Source/Runtime/Online/Android/OnlineSubsystemGooglePlay/Public/OnlineSubsystemGooglePlay.h
@@ -103,6 +103,9 @@ private:
 	/** Google callback when auth is complete */
 	void OnAuthActionFinished(gpg::AuthOperation Op, gpg::AuthStatus Status);
 
+	/** Android callback when an activity is finished */
+	void OnActivityResult(JNIEnv *env, jobject thiz, jobject activity, jint requestCode, jint resultCode, jobject data);
+
 	/** Online async task runnable */
 	TUniquePtr<class FOnlineAsyncTaskManagerGooglePlay> OnlineAsyncTaskThreadRunnable;
 
