Use delegate handle to remove OnActivityResultDelegate on shutdown
#ue4
#android

diff --git a/Engine/Source/Runtime/Online/Android/OnlineSubsystemGooglePlay/Private/OnlineSubsystemGooglePlay.cpp b/Engine/Source/Runtime/Online/Android/OnlineSubsystemGooglePlay/Private/OnlineSubsystemGooglePlay.cpp
--- a/Engine/Source/Runtime/Online/Android/OnlineSubsystemGooglePlay/Private/OnlineSubsystemGooglePlay.cpp
+++ b/Engine/Source/Runtime/Online/Android/OnlineSubsystemGooglePlay/Private/OnlineSubsystemGooglePlay.cpp
@@ -137,7 +137,7 @@ bool FOnlineSubsystemGooglePlay::Init()
 	WaitingForLogin = true;
 	WaitForLostFocus = false;
 
-	FJavaWrapper::OnActivityResultDelegate.AddRaw(this, &FOnlineSubsystemGooglePlay::OnActivityResult);
+	OnActivityResultDelegateHandle = FJavaWrapper::OnActivityResultDelegate.AddRaw(this, &FOnlineSubsystemGooglePlay::OnActivityResult);
 
 	// Create() returns a std::unqiue_ptr, but we convert it to a TUniquePtr.
 	GameServicesPtr.Reset( GameServices::Builder()
@@ -192,7 +192,7 @@ bool FOnlineSubsystemGooglePlay::Shutdown()
 {
 	UE_LOG(LogOnline, Log, TEXT("FOnlineSubsystemAndroid::Shutdown()"));
 
-	FJavaWrapper::OnActivityResultDelegate.RemoveRaw(this, &FOnlineSubsystemGooglePlay::OnActivityResult);
+	FJavaWrapper::OnActivityResultDelegate.Remove(OnActivityResultDelegateHandle);
 
 #define DESTRUCT_INTERFACE(Interface) \
 	if (Interface.IsValid()) \
diff --git a/Engine/Source/Runtime/Online/Android/OnlineSubsystemGooglePlay/Public/OnlineSubsystemGooglePlay.h b/Engine/Source/Runtime/Online/Android/OnlineSubsystemGooglePlay/Public/OnlineSubsystemGooglePlay.h
--- a/Engine/Source/Runtime/Online/Android/OnlineSubsystemGooglePlay/Public/OnlineSubsystemGooglePlay.h
+++ b/Engine/Source/Runtime/Online/Android/OnlineSubsystemGooglePlay/Public/OnlineSubsystemGooglePlay.h
@@ -131,6 +131,9 @@ private:
 
 	/** Track the current login task (if any) so callbacks can notify it */
 	FOnlineAsyncTaskGooglePlayLogin* CurrentLoginTask;
+
+	/** Handle of registered delegate for onActivityResult */
+	FDelegateHandle OnActivityResultDelegateHandle;
 };
 
 typedef TSharedPtr<FOnlineSubsystemGooglePlay, ESPMode::ThreadSafe> FOnlineSubsystemGooglePlayPtr;
