Core - Replaced FPlatformProcess::CreateSyncEvent with FPlatformProcess::GetSynchEventFromPool

diff --git a/Engine/Source/Runtime/Online/OnlineSubsystem/Private/OnlineAsyncTaskManager.cpp b/Engine/Source/Runtime/Online/OnlineSubsystem/Private/OnlineAsyncTaskManager.cpp
--- a/Engine/Source/Runtime/Online/OnlineSubsystem/Private/OnlineAsyncTaskManager.cpp
+++ b/Engine/Source/Runtime/Online/OnlineSubsystem/Private/OnlineAsyncTaskManager.cpp
@@ -9,7 +9,7 @@ int32 FOnlineAsyncTaskManager::InvocationCount = 0;
 #define POLLING_INTERVAL_MS 50
 
 FOnlineAsyncTaskManager::FOnlineAsyncTaskManager() :
-	WorkEvent(NULL),
+	WorkEvent(nullptr),
 	PollingInterval(POLLING_INTERVAL_MS),
 	bRequestingExit(0),
 	OnlineThreadId(0)
@@ -18,7 +18,7 @@ FOnlineAsyncTaskManager::FOnlineAsyncTaskManager() :
 
 bool FOnlineAsyncTaskManager::Init(void)
 {
-	WorkEvent = FPlatformProcess::CreateSynchEvent();
+	WorkEvent = FPlatformProcess::GetSynchEventFromPool();
 	int32 PollingConfig = POLLING_INTERVAL_MS;
 	// Read the polling interval to use from the INI file
 	if (GConfig->GetInt(TEXT("OnlineSubsystem"), TEXT("PollingIntervalInMs"), PollingConfig, GEngineIni))
@@ -26,7 +26,7 @@ bool FOnlineAsyncTaskManager::Init(void)
 		PollingInterval = (uint32)PollingConfig;
 	}
 
-	return WorkEvent != NULL;
+	return WorkEvent != nullptr;
 }
 
 uint32 FOnlineAsyncTaskManager::Run(void)
@@ -150,6 +150,9 @@ void FOnlineAsyncTaskManager::Stop(void)
 
 void FOnlineAsyncTaskManager::Exit(void)
 {
+	FPlatformProcess::ReturnSynchEventToPool(WorkEvent);
+	WorkEvent = nullptr;
+
 	OnlineThreadId = 0;
 	InvocationCount--;
 }
