Merging using UE4-To-UE4-LauncherDev REVERSE

CL # 2427195 on 31/01/2015 16:37
-------------------------------------------
Supporting file attributes list for BuildPatchTool.
Cleaned up and optimised ignore file list parsing.

diff --git a/Engine/Source/Programs/BuildPatchTool/Private/BuildPatchToolMain.cpp b/Engine/Source/Programs/BuildPatchTool/Private/BuildPatchToolMain.cpp
--- a/Engine/Source/Programs/BuildPatchTool/Private/BuildPatchToolMain.cpp
+++ b/Engine/Source/Programs/BuildPatchTool/Private/BuildPatchToolMain.cpp
@@ -18,22 +18,23 @@
 	File based patch data is only recommended for small build images that probably wouldn't benefit from chunking. I.e. Updates are commonly only a few small changed files.
 
 	Required arguments:
-	-BuildRoot=""		Specifies in quotes the directory containing the build image to be read.
-	-CloudDir=""		Specifies in quotes the cloud directory where existing data will be recognised from, and new data added to.
-	-AppID=123456		Specifies without quotes, the ID number for the app
-	-AppName=""			Specifies in quotes, the name of the app
-	-BuildVersion=""	Specifies in quotes, the version string for the build image
-	-AppLaunch=""		Specifies in quotes, the path to the app executable, must be relative to, and inside of BuildRoot.
-	-AppArgs=""			Specifies in quotes, the commandline to send to the app on launch.
+	-BuildRoot=""			Specifies in quotes the directory containing the build image to be read.
+	-CloudDir=""			Specifies in quotes the cloud directory where existing data will be recognised from, and new data added to.
+	-AppID=123456			Specifies without quotes, the ID number for the app
+	-AppName=""				Specifies in quotes, the name of the app
+	-BuildVersion=""		Specifies in quotes, the version string for the build image
+	-AppLaunch=""			Specifies in quotes, the path to the app executable, must be relative to, and inside of BuildRoot.
+	-AppArgs=""				Specifies in quotes, the commandline to send to the app on launch.
 
 	Optional arguments:
-	-stdout				Adds stdout logging to the app.
-	-nochunks			Creates file based patch data instead of chunk based patch data. Should not be used with builds that have large files.
-	-FileIgnoreList=""	Specifies in quotes, the path to a text file containing BuildRoot relative files, separated by \r\n line endings, to not be included in the build.
-	-PrereqName=""		Specifies in quotes, the display name for the prerequisites installer
-	-PrereqPath=""		Specifies in quotes, the prerequisites installer to launch on successful product install
-	-PrereqArgs=""		Specifies in quotes, the commandline to send to prerequisites installer on launch
-	-DataAgeThreshold=12.5		Specified the maximum age (in days) of existing patch data which can be reused in the generated manifest
+	-stdout					Adds stdout logging to the app.
+	-nochunks				Creates file based patch data instead of chunk based patch data. Should not be used with builds that have large files.
+	-FileIgnoreList=""		Specifies in quotes, the path to a text file containing BuildRoot relative files, separated by \r\n line endings, to not be included in the build.
+	-FileAttributeList=""	Specifies in quotes, the path to a text file containing quoted BuildRoot relative files followed by optional attribute keywords readonly compressed executable, separated by \r\n line endings. These attribute will be applied when build is installed client side.
+	-PrereqName=""			Specifies in quotes, the display name for the prerequisites installer
+	-PrereqPath=""			Specifies in quotes, the prerequisites installer to launch on successful product install
+	-PrereqArgs=""			Specifies in quotes, the commandline to send to prerequisites installer on launch
+	-DataAgeThreshold=12.5	Specified the maximum age (in days) of existing patch data which can be reused in the generated manifest
 
 	NB: -DataAgeThreshold can also be present in the [PatchGeneration] section of BuildPatchTool.ini in the cloud directory
 	NB: If -DataAgeThreshold is not supplied, either on the command-line or in BuildPatchTool.ini, then all existing data is eligible for reuse in the generated manifest
@@ -144,6 +145,7 @@ int32 BuildPatchToolMain( const TCHAR* CommandLine )
 	FString LaunchExe;
 	FString LaunchCommand;
 	FString IgnoreListFile;
+	FString AttributeListFile;
 	FString PrereqName;
 	FString PrereqPath;
 	FString PrereqArgs;
@@ -172,6 +174,7 @@ int32 BuildPatchToolMain( const TCHAR* CommandLine )
 		int32 AppLaunchIdx;
 		int32 AppArgsIdx;
 		int32 FileIgnoreListIdx;
+		int32 FileAttributeListIdx;
 		int32 PrereqNameIdx;
 		int32 PrereqPathIdx;
 		int32 PrereqArgsIdx;
@@ -216,6 +219,9 @@ int32 BuildPatchToolMain( const TCHAR* CommandLine )
 		Matcher.Command = TEXT( "FileIgnoreList" );
 		FileIgnoreListIdx = Switches.IndexOfByPredicate(Matcher);
 
+		Matcher.Command = TEXT("FileAttributeList");
+		FileAttributeListIdx = Switches.IndexOfByPredicate(Matcher);
+
 		Matcher.Command = TEXT( "PrereqName" );
 		PrereqNameIdx = Switches.IndexOfByPredicate(Matcher);
 
@@ -264,6 +270,11 @@ int32 BuildPatchToolMain( const TCHAR* CommandLine )
 			FParse::Value( *Switches[ FileIgnoreListIdx ], TEXT( "FileIgnoreList=" ), IgnoreListFile );
 		}
 
+		if( FileAttributeListIdx != INDEX_NONE )
+		{
+			FParse::Value( *Switches[ FileAttributeListIdx ], TEXT( "FileAttributeList=" ), AttributeListFile );
+		}
+
 		if( PrereqNameIdx != INDEX_NONE )
 		{
 			FParse::Value( *Switches[ PrereqNameIdx ], TEXT( "PrereqName=" ), PrereqName );
@@ -376,6 +387,9 @@ int32 BuildPatchToolMain( const TCHAR* CommandLine )
 		}
 	}
 
+	// Initialize the file manager
+	IFileManager::Get().ProcessCommandLineOptions();
+
 	// Check for argument error
 	if( !bSuccess )
 	{
@@ -389,8 +403,17 @@ int32 BuildPatchToolMain( const TCHAR* CommandLine )
 		return 5;
 	}
 
-	// Initialize the file manager
-	IFileManager::Get().ProcessCommandLineOptions();
+	if (!IgnoreListFile.IsEmpty() && !FPaths::FileExists(IgnoreListFile))
+	{
+		GLog->Logf(ELogVerbosity::Error, TEXT("Provided file ignore list was not found %s"), *IgnoreListFile);
+		return 6;
+	}
+
+	if (!AttributeListFile.IsEmpty() && !FPaths::FileExists(AttributeListFile))
+	{
+		GLog->Logf(ELogVerbosity::Error, TEXT("Provided file attribute list was not found %s"), *AttributeListFile);
+		return 7;
+	}
 
 	// Load the BuildPatchServices Module
 	TSharedPtr<IBuildPatchServicesModule> BuildPatchServicesModule = StaticCastSharedPtr<IBuildPatchServicesModule>( FModuleManager::Get().LoadModule( TEXT( "BuildPatchServices" ) ) );
@@ -453,6 +476,7 @@ int32 BuildPatchToolMain( const TCHAR* CommandLine )
 		Settings.LaunchExe = LaunchExe;
 		Settings.LaunchCommand = LaunchCommand;
 		Settings.IgnoreListFile = IgnoreListFile;
+		Settings.AttributeListFile = AttributeListFile;
 		Settings.PrereqName = PrereqName;
 		Settings.PrereqPath = PrereqPath;
 		Settings.PrereqArgs = PrereqArgs;
diff --git a/Engine/Source/Runtime/Online/BuildPatchServices/Private/BuildPatchGeneration.cpp b/Engine/Source/Runtime/Online/BuildPatchServices/Private/BuildPatchGeneration.cpp
--- a/Engine/Source/Runtime/Online/BuildPatchServices/Private/BuildPatchGeneration.cpp
+++ b/Engine/Source/Runtime/Online/BuildPatchServices/Private/BuildPatchGeneration.cpp
@@ -9,6 +9,12 @@
 
 #if WITH_BUILDPATCHGENERATION
 
+FFileAttributes::FFileAttributes()
+	: bReadOnly(false)
+	, bCompressed(false)
+	, bUnixExecutable(false)
+{}
+
 /**
  * Creates a 32bit hash value from a FSHAHashData. This is so that a TMap can be keyed
  * using a FSHAHashData
@@ -983,6 +989,61 @@ static void AddCustomFieldsToBuildManifest(const TMap<FString, FVariant>& Custom
 	}
 }
 
+static void AddFileAttributesToBuildManifest(const FString& AttributesList, FBuildPatchAppManifestRef BuildManifest, TMap<FString, FFileAttributes>& FileAttributesMap)
+{
+	GLog->Logf(TEXT("Parsing file attributes list:-"));
+	checkf(AttributesList.Len() > 0, TEXT("Attributes File List was empty file"));
+	const TCHAR Quote = TEXT('\"');
+	const TCHAR EOFile = TEXT('\0');
+	const TCHAR EOLine = TEXT('\n');
+
+	const TCHAR* CharPtr = *AttributesList;
+	while (*CharPtr != EOFile)
+	{
+		// Parse filename
+		while (*CharPtr != Quote && *CharPtr != EOFile){ ++CharPtr; }
+		if (*CharPtr == EOFile)
+		{
+			break;
+		}
+		const TCHAR* FilenameStart = ++CharPtr;
+		while (*CharPtr != Quote && *CharPtr != EOFile && *CharPtr != EOLine){ ++CharPtr; }
+		checkf(*CharPtr != EOFile, TEXT("Attributes File List: Unexpected end of file before next quote! Pos:%d"), CharPtr - *AttributesList);
+		checkf(*CharPtr != EOLine, TEXT("Attributes File List: Unexpected end of line before next quote! Pos:%d"), CharPtr - *AttributesList);
+		const TCHAR* FilenameEnd = CharPtr++;
+		// Parse keywords
+		while (*CharPtr != Quote && *CharPtr != EOFile && *CharPtr != EOLine){ ++CharPtr; }
+		checkf(*CharPtr != Quote, TEXT("Attributes File List: Unexpected Quote before end of keywords! Pos:%d"), CharPtr - *AttributesList);
+		const TCHAR* EndOfLine = CharPtr;
+		// Grab info
+		FString Filename = FString(FilenameEnd - FilenameStart, FilenameStart).Replace(TEXT("\\"), TEXT("/"));
+		FString Keywords(EndOfLine - FilenameEnd, FilenameEnd);
+		FFileAttributes FileAttributes;
+		GLog->Logf(TEXT("    %s"), *Filename);
+		if (Keywords.Contains(TEXT("readonly")))
+		{
+			FileAttributes.bReadOnly = true;
+			GLog->Logf(TEXT("        readonly"), *Filename);
+		}
+		if (Keywords.Contains(TEXT("compressed")))
+		{
+			FileAttributes.bCompressed = true;
+			GLog->Logf(TEXT("        compressed"), *Filename);
+		}
+		if (Keywords.Contains(TEXT("executable")))
+		{
+			FileAttributes.bUnixExecutable = true;
+			GLog->Logf(TEXT("        executable"), *Filename);
+		}
+		if (!(FileAttributes.bReadOnly || FileAttributes.bCompressed || FileAttributes.bUnixExecutable))
+		{
+			GLog->Logf(TEXT("        none"), *Filename);
+		}
+		FileAttributesMap.Add(MoveTemp(Filename), FileAttributes);
+	}
+
+}
+
 /* FBuildDataGenerator implementation
 *****************************************************************************/
 bool FBuildDataGenerator::GenerateChunksManifestFromDirectory( const FBuildPatchSettings& Settings )
@@ -1003,6 +1064,22 @@ bool FBuildDataGenerator::GenerateChunksManifestFromDirectory( const FBuildPatch
 	// Setup custom fields
 	AddCustomFieldsToBuildManifest(Settings.CustomFields, BuildManifest);
 
+	// Setup File Attributes
+	FString AttributesList;
+	TMap<FString, FFileAttributes> FileAttributesMap;
+	if (Settings.AttributeListFile.Len() > 0)
+	{
+		FFileHelper::LoadFileToString(AttributesList, *Settings.AttributeListFile);
+		if (!AttributesList.IsEmpty())
+		{
+			AddFileAttributesToBuildManifest(AttributesList, BuildManifest, FileAttributesMap);
+		}
+		else
+		{
+			GLog->Logf(TEXT("WARNING: Attributes list file empty"));
+		}
+	}
+
 	// Reset chunk inventory
 	ExistingChunksEnumerated = false;
 	ExistingChunkGuidInventory.Empty();
@@ -1168,6 +1245,28 @@ bool FBuildDataGenerator::GenerateChunksManifestFromDirectory( const FBuildPatch
 	// Fill out lookups
 	BuildManifest->InitLookups();
 
+	// Fill out the file attributes
+	for (const auto& Entry : FileAttributesMap)
+	{
+		const FString& Filename = Entry.Key;
+		const FFileAttributes& Attributes = Entry.Value;
+		if (BuildManifest->FileManifestLookup.Contains(Filename))
+		{
+			FFileManifestData& FileManifest = *BuildManifest->FileManifestLookup[Filename];
+			FileManifest.bIsReadOnly = Attributes.bReadOnly;
+			FileManifest.bIsCompressed = Attributes.bCompressed;
+			// Only overwrite unix exe if true
+			if (Attributes.bUnixExecutable)
+			{
+				FileManifest.bIsUnixExecutable = Attributes.bUnixExecutable;
+			}
+		}
+		else
+		{
+			GLog->Logf(TEXT("File Attributes: File not in build %s"), *Filename);
+		}
+	}
+
 	// Save manifest into the cloud directory
 	FString BaseFilename = FBuildPatchServicesModule::GetCloudDirectory() / FDefaultValueHelper::RemoveWhitespaces(BuildManifest->Data->AppName + BuildManifest->Data->BuildVersion);
 	FString JsonFilename = BaseFilename + TEXT(".manifest");
@@ -1789,35 +1888,52 @@ bool FBuildDataGenerator::CompareDataToChunk( const FString& ChunkFilePath, uint
 
 void FBuildDataGenerator::StripIgnoredFiles( TArray< FString >& AllFiles, const FString& DepotDirectory, const FString& IgnoreListFile )
 {
-	const int32 OriginalNumFiles = AllFiles.Num();
-	FString IgnoreFileList = TEXT( "" );
-	FFileHelper::LoadFileToString( IgnoreFileList, *IgnoreListFile );
-	TArray< FString > IgnoreFiles;
-	IgnoreFileList.ParseIntoArray( &IgnoreFiles, TEXT( "\r\n" ), true );
 	struct FRemoveMatchingStrings
 	{ 
-		const FString* MatchingString;
-		FRemoveMatchingStrings( const FString* InMatch )
-			: MatchingString(InMatch) {}
+		const TSet<FString>& IgnoreList;
+		FRemoveMatchingStrings( const TSet<FString>& IgnoreList )
+			: IgnoreList(IgnoreList) {}
 
 		bool operator()(const FString& RemovalCandidate) const
 		{
-			FString PathA = RemovalCandidate;
-			FString PathB = *MatchingString;
-			FPaths::NormalizeFilename( PathA );
-			FPaths::NormalizeFilename( PathB );
-			return PathA == PathB;
+			const bool bRemove = IgnoreList.Contains(RemovalCandidate);
+			if (bRemove)
+			{
+				GLog->Logf(TEXT("    - %s"), *RemovalCandidate);
+			}
+			return bRemove;
 		}
 	};
 
-	for( int32 IgnoreIdx = 0; IgnoreIdx < IgnoreFiles.Num(); ++IgnoreIdx )
+	GLog->Logf(TEXT("Stripping ignorable files"));
+	const int32 OriginalNumFiles = AllFiles.Num();
+	FString IgnoreFileList = TEXT( "" );
+	FFileHelper::LoadFileToString( IgnoreFileList, *IgnoreListFile );
+	TArray< FString > IgnoreFiles;
+	IgnoreFileList.ParseIntoArray( &IgnoreFiles, TEXT( "\r\n" ), true );
+
+	// Normalize all paths first
+	for (FString& Filename : AllFiles)
 	{
-		const FString& IgnoreFile = IgnoreFiles[IgnoreIdx];
-		const FString FullIgnorePath = DepotDirectory / IgnoreFile;
-		AllFiles.RemoveAll( FRemoveMatchingStrings( &FullIgnorePath ) );
+		FPaths::NormalizeFilename(Filename);
 	}
+	for (FString& Filename : IgnoreFiles)
+	{
+		Filename = DepotDirectory / Filename;
+		FPaths::NormalizeFilename(Filename);
+	}
+
+
+	// Convert ignore list to set
+	TSet<FString> IgnoreSet(MoveTemp(IgnoreFiles));
+
+	// Filter file list
+	FRemoveMatchingStrings FileFilter(IgnoreSet);
+	AllFiles.RemoveAll(FileFilter);
+
 	const int32 NewNumFiles = AllFiles.Num();
 	GLog->Logf( TEXT( "Stripped %d ignorable file(s)" ), ( OriginalNumFiles - NewNumFiles ) );
 }
 
 #endif //WITH_BUILDPATCHGENERATION
+
diff --git a/Engine/Source/Runtime/Online/BuildPatchServices/Private/BuildPatchGeneration.h b/Engine/Source/Runtime/Online/BuildPatchServices/Private/BuildPatchGeneration.h
--- a/Engine/Source/Runtime/Online/BuildPatchServices/Private/BuildPatchGeneration.h
+++ b/Engine/Source/Runtime/Online/BuildPatchServices/Private/BuildPatchGeneration.h
@@ -681,6 +681,17 @@ private:
 
 };
 
+/**
+ * Simple data structure to hold values parsed from the File Attribute List
+ */
+struct FFileAttributes
+{
+	bool bReadOnly;
+	bool bCompressed;
+	bool bUnixExecutable;
+	FFileAttributes();
+};
+
 #endif //WITH_BUILDPATCHGENERATION
 
 #endif // __BuildPatchGeneration_h__
diff --git a/Engine/Source/Runtime/Online/BuildPatchServices/Private/BuildPatchManifest.cpp b/Engine/Source/Runtime/Online/BuildPatchServices/Private/BuildPatchManifest.cpp
--- a/Engine/Source/Runtime/Online/BuildPatchServices/Private/BuildPatchManifest.cpp
+++ b/Engine/Source/Runtime/Online/BuildPatchServices/Private/BuildPatchManifest.cpp
@@ -37,7 +37,7 @@ FString ToStringBlob( const DataType& DataVal )
 }
 
 /**
- * Helper functions to decide whether the passed in data is a JSON string we expect to deserialise a manifest from
+ * Helper functions to decide whether the passed in data is a JSON string we expect to deserialize a manifest from
  */
 bool BufferIsJsonManifest(const TArray<uint8>& DataInput)
 {
@@ -158,6 +158,8 @@ FFileManifestData::FFileManifestData()
 	, FileChunkParts()
 	, bIsUnixExecutable(false)
 	, SymlinkTarget(TEXT(""))
+	, bIsReadOnly(false)
+	, bIsCompressed(false)
 	, FileSize(INDEX_NONE)
 {
 }
@@ -726,6 +728,14 @@ void FBuildPatchAppManifest::SerializeToJSON(FString& JSONOutput)
 				{
 					Writer->WriteValue(TEXT("bIsUnixExecutable"), FileManifest.bIsUnixExecutable);
 				}
+				if (FileManifest.bIsReadOnly)
+				{
+					Writer->WriteValue(TEXT("bIsReadOnly"), FileManifest.bIsReadOnly);
+				}
+				if (FileManifest.bIsCompressed)
+				{
+					Writer->WriteValue(TEXT("bIsCompressed"), FileManifest.bIsCompressed);
+				}
 				const bool bIsSymlink = !FileManifest.SymlinkTarget.IsEmpty();
 				if (bIsSymlink)
 				{
@@ -893,6 +903,8 @@ bool FBuildPatchAppManifest::DeserializeFromJSON( const FString& JSONInput )
 				AllDataGuids.Add(FileChunkPart.Guid);
 			}
 			FileManifest.bIsUnixExecutable = JsonFileManifest->HasField(TEXT("bIsUnixExecutable")) && JsonFileManifest->GetBoolField(TEXT("bIsUnixExecutable"));
+			FileManifest.bIsReadOnly = JsonFileManifest->HasField(TEXT("bIsReadOnly")) && JsonFileManifest->GetBoolField(TEXT("bIsReadOnly"));
+			FileManifest.bIsCompressed = JsonFileManifest->HasField(TEXT("bIsCompressed")) && JsonFileManifest->GetBoolField(TEXT("bIsCompressed"));
 			FileManifest.SymlinkTarget = JsonFileManifest->HasField(TEXT("SymlinkTarget")) ? JsonFileManifest->GetStringField(TEXT("SymlinkTarget")) : TEXT("");
 			FileManifest.Init();
 		}
diff --git a/Engine/Source/Runtime/Online/BuildPatchServices/Private/BuildPatchManifest.h b/Engine/Source/Runtime/Online/BuildPatchServices/Private/BuildPatchManifest.h
--- a/Engine/Source/Runtime/Online/BuildPatchServices/Private/BuildPatchManifest.h
+++ b/Engine/Source/Runtime/Online/BuildPatchServices/Private/BuildPatchManifest.h
@@ -183,6 +183,12 @@ struct FFileManifestData
 	UPROPERTY()
 	FString SymlinkTarget;
 
+	UPROPERTY()
+	bool bIsReadOnly;
+
+	UPROPERTY()
+	bool bIsCompressed;
+
 	FFileManifestData();
 
 	void Init();
diff --git a/Engine/Source/Runtime/Online/BuildPatchServices/Public/BuildPatchSettings.h b/Engine/Source/Runtime/Online/BuildPatchServices/Public/BuildPatchSettings.h
--- a/Engine/Source/Runtime/Online/BuildPatchServices/Public/BuildPatchSettings.h
+++ b/Engine/Source/Runtime/Online/BuildPatchServices/Public/BuildPatchSettings.h
@@ -20,6 +20,8 @@ struct FBuildPatchSettings
 	FString LaunchCommand;
 	// The path to a file containing a \r\n separated list of RootDirectory relative files to ignore.
 	FString IgnoreListFile;
+	// The path to a file containing a \r\n separated list of RootDirectory relative files to followed by attribute keywords.
+	FString AttributeListFile;
 	// The display name of the prerequisites installer
 	FString PrereqName;
 	// The path to the prerequisites installer
