CIS fix
#ios

diff --git a/Engine/Source/Runtime/Online/HTTP/Private/IOS/IOSHTTP.cpp b/Engine/Source/Runtime/Online/HTTP/Private/IOS/IOSHTTP.cpp
--- a/Engine/Source/Runtime/Online/HTTP/Private/IOS/IOSHTTP.cpp
+++ b/Engine/Source/Runtime/Online/HTTP/Private/IOS/IOSHTTP.cpp
@@ -254,7 +254,7 @@ bool FIOSHttpRequest::StartRequest()
 void FIOSHttpRequest::FinishedRequest()
 {
 	UE_LOG(LogHttp, Verbose, TEXT("FIOSHttpRequest::FinishedRequest()"));
-	ElapsedTime = (float)(FPlatformTime.Seconds() - StartRequestTime);
+	ElapsedTime = (float)(FPlatformTime::Seconds() - StartRequestTime);
 	if( Response.IsValid() && !Response->HadError())
 	{
 		UE_LOG(LogHttp, Verbose, TEXT("Request succeeded"));
