Don't write a log message for each session found by the "Find Sessions" node.
Thanks to [REDACTED_PERSON]!
#github https://github.com/EpicGames/UnrealEngine/pull/834

diff --git a/Engine/Source/Runtime/Online/OnlineSubsystemUtils/Private/FindSessionsCallbackProxy.cpp b/Engine/Source/Runtime/Online/OnlineSubsystemUtils/Private/FindSessionsCallbackProxy.cpp
--- a/Engine/Source/Runtime/Online/OnlineSubsystemUtils/Private/FindSessionsCallbackProxy.cpp
+++ b/Engine/Source/Runtime/Online/OnlineSubsystemUtils/Private/FindSessionsCallbackProxy.cpp
@@ -74,13 +74,8 @@ void UFindSessionsCallbackProxy::OnCompleted(bool bSuccess)
 
 	if (bSuccess && SearchObject.IsValid())
 	{
-		// Just log the results for now, will need to add a blueprint-compatible search result struct
 		for (auto& Result : SearchObject->SearchResults)
 		{
-			FString ResultText = FString::Printf(TEXT("Found a session. Ping is %d"), Result.PingInMs);
-
-			FFrame::KismetExecutionMessage(*ResultText, ELogVerbosity::Log);
-
 			FBlueprintSessionResult BPResult;
 			BPResult.OnlineResult = Result;
 			Results.Add(BPResult);
