Dev to Main integration from CL #2446290

diff --git a/Engine/Source/Runtime/Online/OnlineSubsystemUtils/Classes/PartyBeaconHost.h b/Engine/Source/Runtime/Online/OnlineSubsystemUtils/Classes/PartyBeaconHost.h
--- a/Engine/Source/Runtime/Online/OnlineSubsystemUtils/Classes/PartyBeaconHost.h
+++ b/Engine/Source/Runtime/Online/OnlineSubsystemUtils/Classes/PartyBeaconHost.h
@@ -1,9 +1,8 @@
 // Copyright 1998-2015 Epic Games, Inc. All Rights Reserved.
 
 #pragma once
-#include "Runtime/Online/OnlineSubsystemUtils/Classes/OnlineBeaconHost.h"
-#include "Runtime/Online/OnlineSubsystemUtils/Classes/PartyBeaconState.h"
 #include "Runtime/Online/OnlineSubsystemUtils/Classes/OnlineBeaconHostObject.h"
+#include "PartyBeaconState.h"
 #include "PartyBeaconHost.generated.h"
 
 class FUniqueNetId;
@@ -58,8 +57,8 @@ class ONLINESUBSYSTEMUTILS_API APartyBeaconHost : public AOnlineBeaconHostObject
 
 	// Begin AOnlineBeaconHostObject Interface 
 	virtual AOnlineBeaconClient* SpawnBeaconActor(UNetConnection* ClientConnection) override;
-	virtual void ClientConnected(AOnlineBeaconClient* NewClientActor, UNetConnection* ClientConnection);
-	virtual void RemoveClientActor(AOnlineBeaconClient* ClientActor);
+	virtual void ClientConnected(AOnlineBeaconClient* NewClientActor, UNetConnection* ClientConnection) override;
+	virtual void RemoveClientActor(AOnlineBeaconClient* ClientActor) override;
 	// End AOnlineBeaconHost Interface 
 
 	/**
@@ -205,6 +204,7 @@ class ONLINESUBSYSTEMUTILS_API APartyBeaconHost : public AOnlineBeaconHostObject
 	 * Handle a reservation request received from an incoming client
 	 *
 	 * @param Client client beacon making the request
+	 * @param SessionId id of the session that is being checked
 	 * @param ReservationRequest payload of request
 	 */
 	virtual void ProcessReservationRequest(APartyBeaconClient* Client, const FString& SessionId, const FPartyReservation& ReservationRequest);
diff --git a/Engine/Source/Runtime/Online/OnlineSubsystemUtils/Private/IpNetDriver.cpp b/Engine/Source/Runtime/Online/OnlineSubsystemUtils/Private/IpNetDriver.cpp
--- a/Engine/Source/Runtime/Online/OnlineSubsystemUtils/Private/IpNetDriver.cpp
+++ b/Engine/Source/Runtime/Online/OnlineSubsystemUtils/Private/IpNetDriver.cpp
@@ -295,7 +295,7 @@ void UIpNetDriver::ProcessRemoteFunction(class AActor* Actor, UFunction* Functio
 			for (int32 i=0; i<ClientConnections.Num(); ++i)
 			{
 				Connection = ClientConnections[i];
-				if (Connection && Connection->Viewer)
+				if (Connection && Connection->ViewTarget)
 				{
 					// Do relevancy check if unreliable.
 					// Reliables will always go out. This is odd behavior. On one hand we wish to garuntee "reliables always get there". On the other
@@ -308,7 +308,7 @@ void UIpNetDriver::ProcessRemoteFunction(class AActor* Actor, UFunction* Functio
 					if ((Function->FunctionFlags & FUNC_NetReliable) == 0)
 					{
 						FNetViewer Viewer(Connection, 0.f);
-						IsRelevant = Actor->IsNetRelevantFor(Viewer.InViewer, Viewer.Viewer, Viewer.ViewLocation);
+						IsRelevant = Actor->IsNetRelevantFor(Viewer.InViewer, Viewer.ViewTarget, Viewer.ViewLocation);
 					}
 					
 					if (IsRelevant)
diff --git a/Engine/Source/Runtime/Online/OnlineSubsystemUtils/Private/OnlineBeacon.cpp b/Engine/Source/Runtime/Online/OnlineSubsystemUtils/Private/OnlineBeacon.cpp
--- a/Engine/Source/Runtime/Online/OnlineSubsystemUtils/Private/OnlineBeacon.cpp
+++ b/Engine/Source/Runtime/Online/OnlineSubsystemUtils/Private/OnlineBeacon.cpp
@@ -13,10 +13,11 @@ AOnlineBeacon::AOnlineBeacon(const FObjectInitializer& ObjectInitializer) :
 
 bool AOnlineBeacon::InitBase()
 {
-	if (BeaconNetDriverName != NAME_None && GEngine->CreateNamedNetDriver(GetWorld(), BeaconNetDriverName, NAME_BeaconNetDriver))
+	NetDriver = GEngine->CreateNetDriver(GetWorld(), NAME_BeaconNetDriver);
+	if (NetDriver != nullptr)
 	{
 		HandleNetworkFailureDelegateHandle = GEngine->OnNetworkFailure().AddUObject(this, &AOnlineBeacon::HandleNetworkFailure);
-		NetDriver = GEngine->FindNamedNetDriver(GetWorld(), BeaconNetDriverName);
+		NetDriverName = NetDriver->NetDriverName;
 		return true;
 	}
 
@@ -33,7 +34,7 @@ void AOnlineBeacon::DestroyBeacon()
 {
 	UE_LOG(LogBeacon, Verbose, TEXT("Destroying beacon %s, netdriver %s"), *GetName(), NetDriver ? *NetDriver->GetDescription() : TEXT("NULL"));
 	GEngine->OnNetworkFailure().Remove(HandleNetworkFailureDelegateHandle);
-	GEngine->DestroyNamedNetDriver(GetWorld(), BeaconNetDriverName);
+	GEngine->DestroyNamedNetDriver(GetWorld(), NetDriverName);
 	NetDriver = NULL;
 
 	Destroy();
@@ -41,7 +42,7 @@ void AOnlineBeacon::DestroyBeacon()
 
 void AOnlineBeacon::HandleNetworkFailure(UWorld *World, UNetDriver *InNetDriver, ENetworkFailure::Type FailureType, const FString& ErrorString)
 {
-	if (InNetDriver && InNetDriver->NetDriverName == BeaconNetDriverName)
+	if (InNetDriver && InNetDriver->NetDriverName == NetDriverName)
 	{
 		OnFailure();
 	}
@@ -50,7 +51,7 @@ void AOnlineBeacon::HandleNetworkFailure(UWorld *World, UNetDriver *InNetDriver,
 void AOnlineBeacon::OnFailure()
 {
 	GEngine->OnNetworkFailure().Remove(HandleNetworkFailureDelegateHandle);
-	GEngine->DestroyNamedNetDriver(GetWorld(), BeaconNetDriverName);
+	GEngine->DestroyNamedNetDriver(GetWorld(), NetDriverName);
 	NetDriver = NULL;
 }
 
diff --git a/Engine/Source/Runtime/Online/OnlineSubsystemUtils/Private/OnlineBeaconClient.cpp b/Engine/Source/Runtime/Online/OnlineSubsystemUtils/Private/OnlineBeaconClient.cpp
--- a/Engine/Source/Runtime/Online/OnlineSubsystemUtils/Private/OnlineBeaconClient.cpp
+++ b/Engine/Source/Runtime/Online/OnlineSubsystemUtils/Private/OnlineBeaconClient.cpp
@@ -8,9 +8,9 @@
 
 AOnlineBeaconClient::AOnlineBeaconClient(const FObjectInitializer& ObjectInitializer) :
 	Super(ObjectInitializer),
-	BeaconOwner(NULL)
+	BeaconOwner(nullptr)
 {
-	NetDriverName = FName(TEXT("BeaconDriver"));
+	NetDriverName = FName(TEXT("BeaconDriverClient"));
 	bOnlyRelevantToOwner = true;
 }
 
@@ -37,8 +37,8 @@ bool AOnlineBeaconClient::InitClient(FURL& URL)
 			{
 				NetDriver->SetWorld(GetWorld());
 				NetDriver->Notify = this;
-				NetDriver->InitialConnectTimeout = BEACON_CONNECTION_INITIAL_TIMEOUT;
-				NetDriver->ConnectionTimeout = BEACON_CONNECTION_INITIAL_TIMEOUT;
+				NetDriver->InitialConnectTimeout = BeaconConnectionInitialTimeout;
+				NetDriver->ConnectionTimeout = BeaconConnectionInitialTimeout;
 
 				// Send initial message.
 				uint8 IsLittleEndian = uint8(PLATFORM_LITTLE_ENDIAN);
@@ -64,6 +64,13 @@ bool AOnlineBeaconClient::InitClient(FURL& URL)
 	return bSuccess;
 }
 
+void AOnlineBeaconClient::OnFailure()
+{
+	UE_LOG(LogBeacon, Verbose, TEXT("Client beacon (%s) connection failure, handling connection timeout."), *GetName());
+	HostConnectionFailure.ExecuteIfBound();
+	Super::OnFailure();
+}
+
 void AOnlineBeaconClient::ClientOnConnected_Implementation()
 {
 	Role = ROLE_Authority;
@@ -76,8 +83,8 @@ void AOnlineBeaconClient::ClientOnConnected_Implementation()
 	if (NetDriver)
 	{
 		// Increase timeout while we are connected
-		NetDriver->InitialConnectTimeout = BEACON_CONNECTION_TIMEOUT;
-		NetDriver->ConnectionTimeout = BEACON_CONNECTION_TIMEOUT;
+		NetDriver->InitialConnectTimeout = BeaconConnectionTimeout;
+		NetDriver->ConnectionTimeout = BeaconConnectionTimeout;
 	}
 
 	// Call the overloaded function for this client class
diff --git a/Engine/Source/Runtime/Online/OnlineSubsystemUtils/Private/OnlineBeaconHost.cpp b/Engine/Source/Runtime/Online/OnlineSubsystemUtils/Private/OnlineBeaconHost.cpp
--- a/Engine/Source/Runtime/Online/OnlineSubsystemUtils/Private/OnlineBeaconHost.cpp
+++ b/Engine/Source/Runtime/Online/OnlineSubsystemUtils/Private/OnlineBeaconHost.cpp
@@ -7,7 +7,7 @@
 AOnlineBeaconHost::AOnlineBeaconHost(const FObjectInitializer& ObjectInitializer) :
 	Super(ObjectInitializer)
 {
-	NetDriverName = FName(TEXT("BeaconDriver"));
+	NetDriverName = FName(TEXT("BeaconDriverHost"));
 }
 
 void AOnlineBeaconHost::OnNetCleanup(UNetConnection* Connection)
@@ -17,7 +17,7 @@ void AOnlineBeaconHost::OnNetCleanup(UNetConnection* Connection)
 
 bool AOnlineBeaconHost::InitHost()
 {
-	FURL URL(NULL, TEXT(""), TRAVEL_Absolute);
+	FURL URL(nullptr, TEXT(""), TRAVEL_Absolute);
 
 	// Allow the command line to override the default port
 	int32 PortOverride;
@@ -37,8 +37,8 @@ bool AOnlineBeaconHost::InitHost()
 				ListenPort = URL.Port;
 				NetDriver->SetWorld(GetWorld());
 				NetDriver->Notify = this;
-				NetDriver->InitialConnectTimeout = BEACON_CONNECTION_TIMEOUT;
-				NetDriver->ConnectionTimeout = BEACON_CONNECTION_TIMEOUT;
+				NetDriver->InitialConnectTimeout = BeaconConnectionInitialTimeout;
+				NetDriver->ConnectionTimeout = BeaconConnectionTimeout;
 				return true;
 			}
 			else
@@ -55,7 +55,7 @@ bool AOnlineBeaconHost::InitHost()
 
 void AOnlineBeaconHost::HandleNetworkFailure(UWorld* World, class UNetDriver *NetDriver, ENetworkFailure::Type FailureType, const FString& ErrorString)
 {
-	if (NetDriver && NetDriver->NetDriverName == BeaconNetDriverName)
+	if (NetDriver && NetDriver->NetDriverName == NetDriverName)
 	{
 		// Timeouts from clients are ignored
 		if (FailureType != ENetworkFailure::ConnectionTimeout)
@@ -67,7 +67,7 @@ void AOnlineBeaconHost::HandleNetworkFailure(UWorld* World, class UNetDriver *Ne
 
 void AOnlineBeaconHost::NotifyControlMessage(UNetConnection* Connection, uint8 MessageType, class FInBunch& Bunch)
 {
-	if(NetDriver->ServerConnection == NULL)
+	if(NetDriver->ServerConnection == nullptr)
 	{
 		bool bCloseConnection = false;
 
@@ -118,12 +118,12 @@ void AOnlineBeaconHost::NotifyControlMessage(UNetConnection* Connection, uint8 M
 				if (Connection->ClientWorldPackageName == NAME_None)
 				{
 					AOnlineBeaconClient* ClientActor = GetClientActor(Connection);
-					if (ClientActor == NULL)
+					if (ClientActor == nullptr)
 					{
 						UWorld* World = GetWorld();
 						Connection->ClientWorldPackageName = World->GetOutermost()->GetFName();
 
-						AOnlineBeaconClient* NewClientActor = NULL;
+						AOnlineBeaconClient* NewClientActor = nullptr;
 						FOnBeaconSpawned* OnBeaconSpawnedDelegate = OnBeaconSpawnedMapping.Find(BeaconType);
 						if (OnBeaconSpawnedDelegate && OnBeaconSpawnedDelegate->IsBound())
 						{
@@ -137,6 +137,8 @@ void AOnlineBeaconHost::NotifyControlMessage(UNetConnection* Connection, uint8 M
 							Connection->OwningActor = NewClientActor;
 							NewClientActor->Role = ROLE_None;
 							NewClientActor->SetReplicates(false);
+							check(NetDriverName == NetDriver->NetDriverName);
+							NewClientActor->NetDriverName = NetDriverName;
 							ClientActors.Add(NewClientActor);
 							FNetControlMessage<NMT_BeaconAssignGUID>::Send(Connection, NetGUID);
 						}
@@ -243,7 +245,7 @@ AOnlineBeaconClient* AOnlineBeaconHost::GetClientActor(UNetConnection* Connectio
 		}
 	}
 
-	return NULL;
+	return nullptr;
 }
 
 void AOnlineBeaconHost::RemoveClientActor(AOnlineBeaconClient* ClientActor)
@@ -287,7 +289,7 @@ void AOnlineBeaconHost::UnregisterHost(const FString& BeaconType)
 	AOnlineBeaconHostObject* HostObject = GetHost(BeaconType);
 	if (HostObject)
 	{
-		HostObject->SetOwner(NULL);
+		HostObject->SetOwner(nullptr);
 	}
 	
 	OnBeaconSpawned(BeaconType).Unbind();
@@ -305,13 +307,13 @@ AOnlineBeaconHostObject* AOnlineBeaconHost::GetHost(const FString& BeaconType)
 		}
 	}
 
-	return NULL;
+	return nullptr;
 }
 
 AOnlineBeaconHost::FOnBeaconSpawned& AOnlineBeaconHost::OnBeaconSpawned(const FString& BeaconType)
 { 
 	FOnBeaconSpawned* BeaconDelegate = OnBeaconSpawnedMapping.Find(BeaconType);
-	if (BeaconDelegate == NULL)
+	if (BeaconDelegate == nullptr)
 	{
 		FOnBeaconSpawned NewDelegate;
 		OnBeaconSpawnedMapping.Add(BeaconType, NewDelegate);
@@ -324,7 +326,7 @@ AOnlineBeaconHost::FOnBeaconSpawned& AOnlineBeaconHost::OnBeaconSpawned(const FS
 AOnlineBeaconHost::FOnBeaconConnected& AOnlineBeaconHost::OnBeaconConnected(const FString& BeaconType)
 { 
 	FOnBeaconConnected* BeaconDelegate = OnBeaconConnectedMapping.Find(BeaconType);
-	if (BeaconDelegate == NULL)
+	if (BeaconDelegate == nullptr)
 	{
 		FOnBeaconConnected NewDelegate;
 		OnBeaconConnectedMapping.Add(BeaconType, NewDelegate);
diff --git a/Engine/Source/Runtime/Online/OnlineSubsystemUtils/Private/OnlineSessionClient.cpp b/Engine/Source/Runtime/Online/OnlineSubsystemUtils/Private/OnlineSessionClient.cpp
--- a/Engine/Source/Runtime/Online/OnlineSubsystemUtils/Private/OnlineSessionClient.cpp
+++ b/Engine/Source/Runtime/Online/OnlineSubsystemUtils/Private/OnlineSessionClient.cpp
@@ -19,9 +19,16 @@ UOnlineSessionClient::UOnlineSessionClient(const FObjectInitializer& ObjectIniti
 UWorld* UOnlineSessionClient::GetWorld() const
 {
 	ULocalPlayer* LP = Cast<ULocalPlayer>(GetOuter());
-	if (LP != nullptr && LP->PlayerController != nullptr)
+	if (LP != nullptr)
 	{
-		return LP->PlayerController->GetWorld();
+		if (LP->PlayerController != nullptr)
+		{
+			return LP->PlayerController->GetWorld();
+		}
+		else if (LP->ViewportClient != nullptr)
+		{
+			return LP->ViewportClient->GetWorld();
+		}
 	}
 
 	return nullptr;
@@ -54,7 +61,7 @@ TSharedPtr<FUniqueNetId> UOnlineSessionClient::GetUniqueId()
 	ULocalPlayer* LP = Cast<ULocalPlayer>(GetOuter());
 	if (LP != nullptr)
 	{
-		return LP->GetCachedUniqueNetId();
+		return LP->GetPreferredUniqueNetId();
 	}
 
 	return nullptr;
diff --git a/Engine/Source/Runtime/Online/OnlineSubsystemUtils/Private/PartyBeaconClient.cpp b/Engine/Source/Runtime/Online/OnlineSubsystemUtils/Private/PartyBeaconClient.cpp
--- a/Engine/Source/Runtime/Online/OnlineSubsystemUtils/Private/PartyBeaconClient.cpp
+++ b/Engine/Source/Runtime/Online/OnlineSubsystemUtils/Private/PartyBeaconClient.cpp
@@ -81,13 +81,6 @@ void APartyBeaconClient::OnConnected()
 	}
 }
 
-void APartyBeaconClient::OnFailure()
-{
-	UE_LOG(LogBeacon, Verbose, TEXT("Party beacon connection failure, handling connection timeout."));
-	HostConnectionFailure.ExecuteIfBound();
-	Super::OnFailure();
-}
-
 bool APartyBeaconClient::ServerReservationRequest_Validate(const FString& SessionId, FPartyReservation Reservation)
 {
 	return !SessionId.IsEmpty() && Reservation.PartyLeader.IsValid() && Reservation.PartyMembers.Num() > 0;
diff --git a/Engine/Source/Runtime/Online/OnlineSubsystemUtils/Private/PartyBeaconHost.cpp b/Engine/Source/Runtime/Online/OnlineSubsystemUtils/Private/PartyBeaconHost.cpp
--- a/Engine/Source/Runtime/Online/OnlineSubsystemUtils/Private/PartyBeaconHost.cpp
+++ b/Engine/Source/Runtime/Online/OnlineSubsystemUtils/Private/PartyBeaconHost.cpp
@@ -1,12 +1,14 @@
 // Copyright 1998-2015 Epic Games, Inc. All Rights Reserved.
 
 #include "OnlineSubsystemUtilsPrivatePCH.h"
+#include "PartyBeaconHost.h"
+#include "PartyBeaconClient.h"
 
 APartyBeaconHost::APartyBeaconHost(const FObjectInitializer& ObjectInitializer) :
 	Super(ObjectInitializer),
 	State(NULL)
 {
-	BeaconTypeName = TEXT("PartyBeacon");
+	BeaconTypeName = PARTY_BEACON_TYPE;
 	PrimaryActorTick.bCanEverTick = true;
 }
 
@@ -237,7 +239,14 @@ int32 APartyBeaconHost::GetTeamForCurrentPlayer(const FUniqueNetId& PlayerId) co
 
 void APartyBeaconHost::NewPlayerAdded(const FPlayerReservation& NewPlayer)
 {
-	UE_LOG(LogBeacon, Verbose, TEXT("Beacon adding player %s"), *NewPlayer.UniqueId->ToDebugString());
+	if (NewPlayer.UniqueId.IsValid())
+	{
+		UE_LOG(LogBeacon, Verbose, TEXT("Beacon adding player %s"), *NewPlayer.UniqueId->ToDebugString());
+	}
+	else
+	{
+		UE_LOG(LogBeacon, Warning, TEXT("Beacon adding invalid player!"));
+	}
 	if (State)
 	{
 		State->PlayersPendingJoin.Add(NewPlayer.UniqueId.GetUniqueNetId());
@@ -349,20 +358,58 @@ EPartyReservationResult::Type APartyBeaconHost::AddPartyReservation(const FParty
 			const int32 ExistingReservationIdx = State->GetExistingReservation(ReservationRequest.PartyLeader);
 			if (ExistingReservationIdx != INDEX_NONE)
 			{
-				const TArray<FPartyReservation>& Reservations = State->GetReservations();
-				const FPartyReservation& ExistingReservation = Reservations[ExistingReservationIdx];
+				TArray<FPartyReservation>& Reservations = State->GetReservations();
+				FPartyReservation& ExistingReservation = Reservations[ExistingReservationIdx];
 				if (ReservationRequest.PartyMembers.Num() == ExistingReservation.PartyMembers.Num())
 				{
-					// Clean up the game entities for these duplicate players
-					DuplicateReservation.ExecuteIfBound(ReservationRequest);
-
-					// Add all players back into the pending join list
-					for (int32 Count = 0; Count < ReservationRequest.PartyMembers.Num(); Count++)
+					// Verify the reservations are the same
+					int32 NumMatchingReservations = 0;
+					for (const FPlayerReservation& NewPlayerRes : ReservationRequest.PartyMembers)
 					{
-						NewPlayerAdded(ReservationRequest.PartyMembers[Count]);
+						FPlayerReservation* PlayerRes = ExistingReservation.PartyMembers.FindByPredicate(
+							[NewPlayerRes](const FPlayerReservation& ExistingPlayerRes)
+								{ 
+									return NewPlayerRes.UniqueId == ExistingPlayerRes.UniqueId; 
+								});
+
+						if (PlayerRes)
+						{
+							NumMatchingReservations++;
+						}
 					}
 
-					Result = EPartyReservationResult::ReservationDuplicate;
+					if (NumMatchingReservations == ExistingReservation.PartyMembers.Num())
+					{
+						for (const FPlayerReservation& NewPlayerRes : ReservationRequest.PartyMembers)
+						{
+							FPlayerReservation* PlayerRes = ExistingReservation.PartyMembers.FindByPredicate(
+								[NewPlayerRes](const FPlayerReservation& ExistingPlayerRes)
+							{
+								return NewPlayerRes.UniqueId == ExistingPlayerRes.UniqueId;
+							});
+
+							if (PlayerRes)
+							{
+								// Update the validation auth strings because they may have changed with a new login 
+								PlayerRes->ValidationStr = NewPlayerRes.ValidationStr;
+							}
+						}
+
+						// Clean up the game entities for these duplicate players
+						DuplicateReservation.ExecuteIfBound(ReservationRequest);
+
+						// Add all players back into the pending join list
+						for (int32 Count = 0; Count < ReservationRequest.PartyMembers.Num(); Count++)
+						{
+							NewPlayerAdded(ReservationRequest.PartyMembers[Count]);
+						}
+
+						Result = EPartyReservationResult::ReservationDuplicate;
+					}
+					else
+					{
+						Result = EPartyReservationResult::IncorrectPlayerCount;
+					}
 				}
 				else
 				{
@@ -422,7 +469,7 @@ void APartyBeaconHost::RemovePartyReservation(const FUniqueNetIdRepl& PartyLeade
 	}
 	else
 	{
-		UE_LOG(LogBeacon, Warning, TEXT("Failed to find reservation to cancel for leader %s:"), *PartyLeader->ToString());
+		UE_LOG(LogBeacon, Warning, TEXT("Failed to find reservation to cancel for leader %s:"), PartyLeader.IsValid() ? *PartyLeader->ToString() : TEXT("INVALID") );
 	}
 }
 
diff --git a/Engine/Source/Runtime/Online/OnlineSubsystemUtils/Private/VoiceInterfaceImpl.cpp b/Engine/Source/Runtime/Online/OnlineSubsystemUtils/Private/VoiceInterfaceImpl.cpp
--- a/Engine/Source/Runtime/Online/OnlineSubsystemUtils/Private/VoiceInterfaceImpl.cpp
+++ b/Engine/Source/Runtime/Online/OnlineSubsystemUtils/Private/VoiceInterfaceImpl.cpp
@@ -114,6 +114,11 @@ void FOnlineVoiceImpl::StartNetworkedVoice(uint8 LocalUserNum)
 	if (LocalUserNum >= 0 && LocalUserNum < MaxLocalTalkers)
 	{
 		LocalTalkers[LocalUserNum].bHasNetworkedVoice = true;
+		if (VoiceEngine.IsValid())
+		{
+			uint32 Return = VoiceEngine->StartLocalVoiceProcessing(LocalUserNum);
+			UE_LOG(LogVoice, Log, TEXT("StartLocalProcessing(%d) returned 0x%08X"), LocalUserNum, Return);
+		}
 		UE_LOG(LogVoice, Log, TEXT("Starting networked voice for user: %d"), LocalUserNum);
 	}
 	else
@@ -128,6 +133,11 @@ void FOnlineVoiceImpl::StopNetworkedVoice(uint8 LocalUserNum)
 	// Validate the range of the entry
 	if (LocalUserNum >= 0 && LocalUserNum < MaxLocalTalkers)
 	{
+		if (VoiceEngine.IsValid())
+		{
+			uint32 Return = VoiceEngine->StopLocalVoiceProcessing(LocalUserNum);
+			UE_LOG(LogVoice, Log, TEXT("StopLocalVoiceProcessing(%d) returned 0x%08X"), LocalUserNum, Return);
+		}
 		LocalTalkers[LocalUserNum].bHasNetworkedVoice = false;
 		UE_LOG(LogVoice, Log, TEXT("Stopping networked voice for user: %d"), LocalUserNum);
 	}
@@ -158,9 +168,13 @@ bool FOnlineVoiceImpl::RegisterLocalTalker(uint32 LocalUserNum)
 				if (Return == S_OK)
 				{
 					Talker.bIsRegistered = true;
-					// Kick off the processing mode
-					Return = VoiceEngine->StartLocalVoiceProcessing(LocalUserNum);
-					UE_LOG(LogVoice, Log, TEXT("StartLocalProcessing(%d) returned 0x%08X"), LocalUserNum, Return);
+					if (0)
+					{	
+						// If enabled, voice capture is continuous and "push to talk" merely sends packets
+						// Kick off the processing mode
+						Return = VoiceEngine->StartLocalVoiceProcessing(LocalUserNum);
+						UE_LOG(LogVoice, Log, TEXT("StartLocalProcessing(%d) returned 0x%08X"), LocalUserNum, Return);
+					}
 				}
 			}
 			else
@@ -221,7 +235,8 @@ bool FOnlineVoiceImpl::UnregisterLocalTalker(uint32 LocalUserNum)
 			}
 
 			// Remove them from engine too
-			VoiceEngine->StopLocalVoiceProcessing(LocalUserNum);
+			Return = VoiceEngine->StopLocalVoiceProcessing(LocalUserNum);
+			UE_LOG(LogVoice, Log, TEXT("StopLocalVoiceProcessing(%d) returned 0x%08X"), LocalUserNum, Return);
 			Return = VoiceEngine->UnregisterLocalTalker(LocalUserNum);
 			UE_LOG(LogVoice, Log, TEXT("UnregisterLocalTalker(%d) returned 0x%08X"), LocalUserNum, Return);
 			Talker.bIsTalking = false;
diff --git a/Engine/Source/Runtime/Online/Voice/Private/VoicePrivatePCH.h b/Engine/Source/Runtime/Online/Voice/Private/VoicePrivatePCH.h
--- a/Engine/Source/Runtime/Online/Voice/Private/VoicePrivatePCH.h
+++ b/Engine/Source/Runtime/Online/Voice/Private/VoicePrivatePCH.h
@@ -4,6 +4,20 @@
 
 #include "Core.h"
 
+#if PLATFORM_WINDOWS
+
+#include "AllowWindowsPlatformTypes.h"
+
+#include <Audiopolicy.h>
+#include <Mmdeviceapi.h>
+#include <Functiondiscoverykeys_devpkey.h>
+#include <audiodefs.h>
+#include <dsound.h>
+
+#include "HideWindowsPlatformTypes.h"
+
+#endif // PLATFORM_WINDOWS
+
 #define PLATFORM_SUPPORTS_VOICE_CAPTURE (PLATFORM_WINDOWS)
 
 // Module includes
diff --git a/Engine/Source/Runtime/Online/Voice/Private/Windows/VoiceCaptureWindows.cpp b/Engine/Source/Runtime/Online/Voice/Private/Windows/VoiceCaptureWindows.cpp
--- a/Engine/Source/Runtime/Online/Voice/Private/Windows/VoiceCaptureWindows.cpp
+++ b/Engine/Source/Runtime/Online/Voice/Private/Windows/VoiceCaptureWindows.cpp
@@ -5,6 +5,44 @@
 
 #include "AllowWindowsPlatformTypes.h"
 
+struct FVoiceCaptureWindowsVars
+{
+	/** Voice capture device */
+	LPDIRECTSOUNDCAPTURE8 VoiceCaptureDev;
+	/** Voice capture device caps */
+	DSCCAPS VoiceCaptureDevCaps;
+	/** Voice capture buffer */
+	LPDIRECTSOUNDCAPTUREBUFFER8 VoiceCaptureBuffer8;
+	/** Wave format of buffer */
+	WAVEFORMATEX WavFormat;
+	/** Buffer description */
+	DSCBUFFERDESC VoiceCaptureBufferDesc;
+	/** Buffer caps */
+	DSCBCAPS VoiceCaptureBufferCaps8;
+	/** Notification events */
+	HANDLE Events[NUM_EVENTS];
+	/** Current audio position of valid data in capture buffer */
+	DWORD NextCaptureOffset;
+
+	FVoiceCaptureWindowsVars() :
+		VoiceCaptureDev(NULL),
+		VoiceCaptureBuffer8(NULL),
+		NextCaptureOffset(0)
+	{
+		FMemory::Memzero(&VoiceCaptureDevCaps, sizeof(VoiceCaptureDevCaps));
+
+		FMemory::Memzero(&WavFormat, sizeof(WavFormat));
+		FMemory::Memzero(&VoiceCaptureBufferDesc, sizeof(VoiceCaptureBufferDesc));
+		FMemory::Memzero(&VoiceCaptureBufferCaps8, sizeof(VoiceCaptureBufferCaps8));
+		FMemory::Memzero(&Events, ARRAY_COUNT(Events) * sizeof(HANDLE));
+
+		for (int i = 0; i < NUM_EVENTS; i++)
+		{
+			Events[i] = INVALID_HANDLE_VALUE;
+		}
+	}
+};
+
 /** Calculate silence in an audio buffer by using a RMS threshold */
 template<typename T>
 bool IsSilence(T* Buffer, int32 BuffSize)
@@ -39,24 +77,11 @@ bool IsSilence(T* Buffer, int32 BuffSize)
 }
 
 FVoiceCaptureWindows::FVoiceCaptureWindows() :
-	VoiceCaptureDev(NULL),
-	VoiceCaptureState(EVoiceCaptureState::UnInitialized),
-	VoiceCaptureBuffer8(NULL),
+	CV(NULL),
 	LastEventTriggered(STOP_EVENT - 1),
-	NextCaptureOffset(0)
+	VoiceCaptureState(EVoiceCaptureState::UnInitialized)
 {
-	FMemory::Memzero(&VoiceCaptureDevCaps, sizeof(VoiceCaptureDevCaps));
-
-	FMemory::Memzero(&WavFormat, sizeof(WavFormat));
-	FMemory::Memzero(&VoiceCaptureBufferDesc, sizeof(VoiceCaptureBufferDesc));
-	FMemory::Memzero(&VoiceCaptureBufferCaps8, sizeof(VoiceCaptureBufferCaps8));
-	FMemory::Memzero(&Events, ARRAY_COUNT(Events) * sizeof(HANDLE));
-
-	for (int i = 0; i < NUM_EVENTS; i++)
-	{
-		Events[i] = INVALID_HANDLE_VALUE;
-	}
-
+	CV = new FVoiceCaptureWindowsVars();
 	UncompressedAudioBuffer.Empty(MAX_UNCOMPRESSED_VOICE_BUFFER_SIZE);
 }
 
@@ -69,6 +94,12 @@ FVoiceCaptureWindows::~FVoiceCaptureWindows()
 	{
  		VoiceCaptureDev->FreeVoiceCaptureObject(this);
  	}
+	
+	if (CV)
+	{
+		delete CV;
+		CV = NULL;
+	}	
 }
 
 bool FVoiceCaptureWindows::Init(int32 SampleRate, int32 NumChannels)
@@ -93,7 +124,7 @@ bool FVoiceCaptureWindows::Init(int32 SampleRate, int32 NumChannels)
 	}
 
 	// DSDEVID_DefaultCapture WAVEINCAPS 
-	HRESULT hr = DirectSoundCaptureCreate8(&DSDEVID_DefaultVoiceCapture, &VoiceCaptureDev, NULL);
+	HRESULT hr = DirectSoundCaptureCreate8(&DSDEVID_DefaultVoiceCapture, &CV->VoiceCaptureDev, NULL);
 	if (FAILED(hr))
 	{
 		//DSERR_ALLOCATED, DSERR_INVALIDPARAM, DSERR_NOAGGREGATION, DSERR_OUTOFMEMORY
@@ -102,8 +133,8 @@ bool FVoiceCaptureWindows::Init(int32 SampleRate, int32 NumChannels)
 	}
 
 	// Device capabilities
-	VoiceCaptureDevCaps.dwSize = sizeof(DSCCAPS);
-	hr = VoiceCaptureDev->GetCaps(&VoiceCaptureDevCaps);
+	CV->VoiceCaptureDevCaps.dwSize = sizeof(DSCCAPS);
+	hr = CV->VoiceCaptureDev->GetCaps(&CV->VoiceCaptureDevCaps);
 	if (FAILED(hr))
 	{
 		UE_LOG(LogVoiceCapture, Warning, TEXT("Failed to get mic device caps 0x%08x"), hr);
@@ -111,33 +142,33 @@ bool FVoiceCaptureWindows::Init(int32 SampleRate, int32 NumChannels)
 	}
 
 	// Wave format setup
-	WavFormat.wFormatTag = WAVE_FORMAT_PCM;
-	WavFormat.nChannels = NumChannels;
-	WavFormat.wBitsPerSample = 16;
-	WavFormat.nSamplesPerSec = SampleRate;
-	WavFormat.nBlockAlign = (WavFormat.nChannels * WavFormat.wBitsPerSample) / 8;
-	WavFormat.nAvgBytesPerSec = WavFormat.nBlockAlign * WavFormat.nSamplesPerSec;
-	WavFormat.cbSize = 0;
+	CV->WavFormat.wFormatTag = WAVE_FORMAT_PCM;
+	CV->WavFormat.nChannels = NumChannels;
+	CV->WavFormat.wBitsPerSample = 16;
+	CV->WavFormat.nSamplesPerSec = SampleRate;
+	CV->WavFormat.nBlockAlign = (CV->WavFormat.nChannels * CV->WavFormat.wBitsPerSample) / 8;
+	CV->WavFormat.nAvgBytesPerSec = CV->WavFormat.nBlockAlign * CV->WavFormat.nSamplesPerSec;
+	CV->WavFormat.cbSize = 0;
 
 	// Buffer setup
-	VoiceCaptureBufferDesc.dwSize = sizeof(DSCBUFFERDESC);
-	VoiceCaptureBufferDesc.dwFlags = 0;
-	VoiceCaptureBufferDesc.dwBufferBytes = WavFormat.nAvgBytesPerSec / 2; // .5 sec buffer
-	VoiceCaptureBufferDesc.dwReserved = 0;
-	VoiceCaptureBufferDesc.lpwfxFormat = &WavFormat;
-	VoiceCaptureBufferDesc.dwFXCount = 0;
-	VoiceCaptureBufferDesc.lpDSCFXDesc = NULL;
+	CV->VoiceCaptureBufferDesc.dwSize = sizeof(DSCBUFFERDESC);
+	CV->VoiceCaptureBufferDesc.dwFlags = 0;
+	CV->VoiceCaptureBufferDesc.dwBufferBytes = CV->WavFormat.nAvgBytesPerSec / 2; // .5 sec buffer
+	CV->VoiceCaptureBufferDesc.dwReserved = 0;
+	CV->VoiceCaptureBufferDesc.lpwfxFormat = &CV->WavFormat;
+	CV->VoiceCaptureBufferDesc.dwFXCount = 0;
+	CV->VoiceCaptureBufferDesc.lpDSCFXDesc = NULL;
 
 	LPDIRECTSOUNDCAPTUREBUFFER VoiceBuffer = NULL;
 
-	hr = VoiceCaptureDev->CreateCaptureBuffer(&VoiceCaptureBufferDesc, &VoiceBuffer, NULL);
+	hr = CV->VoiceCaptureDev->CreateCaptureBuffer(&CV->VoiceCaptureBufferDesc, &VoiceBuffer, NULL);
 	if (FAILED(hr))
 	{
 		UE_LOG(LogVoiceCapture, Warning, TEXT("Failed to create voice capture buffer 0x%08x"), hr);
 		return false;
 	}
 
-	hr = VoiceBuffer->QueryInterface(IID_IDirectSoundCaptureBuffer8, (LPVOID*)&VoiceCaptureBuffer8);
+	hr = VoiceBuffer->QueryInterface(IID_IDirectSoundCaptureBuffer8, (LPVOID*)&CV->VoiceCaptureBuffer8);
 	VoiceBuffer->Release(); 
 	VoiceBuffer = NULL;
 	if (FAILED(hr))
@@ -146,8 +177,8 @@ bool FVoiceCaptureWindows::Init(int32 SampleRate, int32 NumChannels)
 		return false;
 	}
 
-	VoiceCaptureBufferCaps8.dwSize = sizeof(DSCBCAPS);
-	hr = VoiceCaptureBuffer8->GetCaps(&VoiceCaptureBufferCaps8);
+	CV->VoiceCaptureBufferCaps8.dwSize = sizeof(DSCBCAPS);
+	hr = CV->VoiceCaptureBuffer8->GetCaps(&CV->VoiceCaptureBufferCaps8);
 	if (FAILED(hr))
 	{
 		UE_LOG(LogVoiceCapture, Warning, TEXT("Failed to get voice buffer caps 0x%08x"), hr);
@@ -158,15 +189,15 @@ bool FVoiceCaptureWindows::Init(int32 SampleRate, int32 NumChannels)
 	if (0)
 	{
 		DWORD SizeWritten8 = 0;
-		VoiceCaptureBuffer8->GetFormat(NULL, sizeof(WAVEFORMATEX), &SizeWritten8);
+		CV->VoiceCaptureBuffer8->GetFormat(NULL, sizeof(WAVEFORMATEX), &SizeWritten8);
 
 		LPWAVEFORMATEX BufferFormat8 = (WAVEFORMATEX*)FMemory::Malloc(SizeWritten8);
-		VoiceCaptureBuffer8->GetFormat(BufferFormat8, SizeWritten8, &SizeWritten8);
+		CV->VoiceCaptureBuffer8->GetFormat(BufferFormat8, SizeWritten8, &SizeWritten8);
 		FMemory::Free(BufferFormat8);
 	}
 	// TEST ------------------------
 
-	if (!CreateNotifications(VoiceCaptureBufferCaps8.dwBufferBytes))
+	if (!CreateNotifications(CV->VoiceCaptureBufferCaps8.dwBufferBytes))
 	{
 		UE_LOG(LogVoiceCapture, Warning, TEXT("Failed to create voice buffer notifications"));
 		return false;
@@ -180,29 +211,29 @@ void FVoiceCaptureWindows::Shutdown()
 {
 	Stop();
 
-	for (int i = 0; i < ARRAY_COUNT(Events); i++)
+	for (int i = 0; i < ARRAY_COUNT(CV->Events); i++)
 	{
-		if (Events[i] != INVALID_HANDLE_VALUE)
+		if (CV->Events[i] != INVALID_HANDLE_VALUE)
 		{
-			CloseHandle(Events[i]);
-			Events[i] = INVALID_HANDLE_VALUE;
+			CloseHandle(CV->Events[i]);
+			CV->Events[i] = INVALID_HANDLE_VALUE;
 		}
 	}
 
 	// Free up DirectSound resources
-	if (VoiceCaptureBuffer8)
+	if (CV->VoiceCaptureBuffer8)
 	{
-		VoiceCaptureBuffer8->Release();
-		VoiceCaptureBuffer8 = NULL;
+		CV->VoiceCaptureBuffer8->Release();
+		CV->VoiceCaptureBuffer8 = NULL;
 	}
 
-	if (VoiceCaptureDev)
+	if (CV->VoiceCaptureDev)
 	{
-		VoiceCaptureDev->Release();
-		VoiceCaptureDev = NULL;
+		CV->VoiceCaptureDev->Release();
+		CV->VoiceCaptureDev = NULL;
 	}
 
-	FMemory::Memzero(&VoiceCaptureDevCaps, sizeof(VoiceCaptureDevCaps));
+	FMemory::Memzero(&CV->VoiceCaptureDevCaps, sizeof(CV->VoiceCaptureDevCaps));
 
 	VoiceCaptureState = EVoiceCaptureState::UnInitialized;
 }
@@ -211,7 +242,7 @@ bool FVoiceCaptureWindows::Start()
 {
 	check(VoiceCaptureState != EVoiceCaptureState::UnInitialized);
 
-	HRESULT hr = VoiceCaptureBuffer8->Start(DSCBSTART_LOOPING);
+	HRESULT hr = CV->VoiceCaptureBuffer8->Start(DSCBSTART_LOOPING);
 	if (FAILED(hr))
 	{
 		UE_LOG(LogVoiceCapture, Warning, TEXT("Failed to start capture 0x%08x"), hr);
@@ -224,22 +255,22 @@ bool FVoiceCaptureWindows::Start()
 
 void FVoiceCaptureWindows::Stop()
 {
-	if (VoiceCaptureBuffer8 && 
+	if (CV->VoiceCaptureBuffer8 &&
 		VoiceCaptureState != EVoiceCaptureState::Stopping && 
 		VoiceCaptureState != EVoiceCaptureState::NotCapturing)
 	{
-		VoiceCaptureBuffer8->Stop();
+		CV->VoiceCaptureBuffer8->Stop();
 		VoiceCaptureState = EVoiceCaptureState::Stopping;
 	}
 }
 
 bool FVoiceCaptureWindows::IsCapturing()
 {		
-	if (VoiceCaptureBuffer8)
+	if (CV->VoiceCaptureBuffer8)
 	{
 		DWORD Status = 0;
 
-		HRESULT hr = VoiceCaptureBuffer8->GetStatus(&Status);
+		HRESULT hr = CV->VoiceCaptureBuffer8->GetStatus(&Status);
 		if (FAILED(hr))
 		{
 			UE_LOG(LogVoiceCapture, Warning, TEXT("Failed to get voice buffer status 0x%08x"), hr);
@@ -272,7 +303,7 @@ void FVoiceCaptureWindows::ProcessData()
 	DWORD CurrentCapturePos = 0;
 	DWORD CurrentReadPos = 0;
 
-	HRESULT hr = VoiceCaptureBuffer8 ? VoiceCaptureBuffer8->GetCurrentPosition(&CurrentCapturePos, &CurrentReadPos) : E_FAIL;
+	HRESULT hr = CV->VoiceCaptureBuffer8 ? CV->VoiceCaptureBuffer8->GetCurrentPosition(&CurrentCapturePos, &CurrentReadPos) : E_FAIL;
 	if (FAILED(hr))
 	{
 		UE_LOG(LogVoiceCapture, Warning, TEXT("Failed to get voice buffer cursor position 0x%08x"), hr);
@@ -280,7 +311,7 @@ void FVoiceCaptureWindows::ProcessData()
 		return;
 	}
 
-	DWORD LockSize = ((CurrentReadPos - NextCaptureOffset) + VoiceCaptureBufferCaps8.dwBufferBytes) % VoiceCaptureBufferCaps8.dwBufferBytes;
+	DWORD LockSize = ((CurrentReadPos - CV->NextCaptureOffset) + CV->VoiceCaptureBufferCaps8.dwBufferBytes) % CV->VoiceCaptureBufferCaps8.dwBufferBytes;
 	if(LockSize != 0) 
 	{ 
 		DWORD CaptureFlags = 0;
@@ -288,7 +319,7 @@ void FVoiceCaptureWindows::ProcessData()
 		void* CaptureData = NULL;
 		DWORD CaptureLength2 = 0;
 		void* CaptureData2 = NULL;
-		hr = VoiceCaptureBuffer8->Lock(NextCaptureOffset, LockSize,
+		hr = CV->VoiceCaptureBuffer8->Lock(CV->NextCaptureOffset, LockSize,
 			&CaptureData, &CaptureLength, 
 			&CaptureData2, &CaptureLength2, CaptureFlags);
 		if (SUCCEEDED(hr))
@@ -307,11 +338,11 @@ void FVoiceCaptureWindows::ProcessData()
 				FMemory::Memcpy(AudioBuffer + CaptureLength, CaptureData2, CaptureLength2);
 			}
 
-			VoiceCaptureBuffer8->Unlock(CaptureData, CaptureLength, CaptureData2, CaptureLength2);
+			CV->VoiceCaptureBuffer8->Unlock(CaptureData, CaptureLength, CaptureData2, CaptureLength2);
 
 			// Move the capture offset forward.
-			NextCaptureOffset = (NextCaptureOffset + CaptureLength) % VoiceCaptureBufferCaps8.dwBufferBytes; 
-			NextCaptureOffset = (NextCaptureOffset + CaptureLength2) % VoiceCaptureBufferCaps8.dwBufferBytes; 
+			CV->NextCaptureOffset = (CV->NextCaptureOffset + CaptureLength) % CV->VoiceCaptureBufferCaps8.dwBufferBytes;
+			CV->NextCaptureOffset = (CV->NextCaptureOffset + CaptureLength2) % CV->VoiceCaptureBufferCaps8.dwBufferBytes;
 
 			if (IsSilence((int16*)AudioBuffer, UncompressedAudioBuffer.Num()))
 			{
@@ -366,7 +397,7 @@ bool FVoiceCaptureWindows::CreateNotifications(uint32 BufferSize)
 
 	LPDIRECTSOUNDNOTIFY8 NotifyInt = NULL;
 
-	HRESULT hr = VoiceCaptureBuffer8->QueryInterface(IID_IDirectSoundNotify, (LPVOID*)&NotifyInt);
+	HRESULT hr = CV->VoiceCaptureBuffer8->QueryInterface(IID_IDirectSoundNotify, (LPVOID*)&NotifyInt);
 	if (SUCCEEDED(hr))
 	{
 		DSBPOSITIONNOTIFY NotifyEvents[NUM_EVENTS];
@@ -374,8 +405,8 @@ bool FVoiceCaptureWindows::CreateNotifications(uint32 BufferSize)
 		// Create events.
 		for (int i = 0; i < NUM_EVENTS; i++)
 		{
-			Events[i] = CreateEvent(NULL, true, false, NULL);
-			if (Events[i] == INVALID_HANDLE_VALUE)
+			CV->Events[i] = CreateEvent(NULL, true, false, NULL);
+			if (CV->Events[i] == INVALID_HANDLE_VALUE)
 			{
 				UE_LOG(LogVoiceCapture, Warning, TEXT("Error creating sync event"));
 				bSuccess = false;
@@ -390,12 +421,12 @@ bool FVoiceCaptureWindows::CreateNotifications(uint32 BufferSize)
 			for (uint32 i = 0; i < NumSegments; i++)
 			{
 				NotifyEvents[i].dwOffset = ((i + 1) * BufferSegSize) - 1;
-				NotifyEvents[i].hEventNotify = Events[i];
+				NotifyEvents[i].hEventNotify = CV->Events[i];
 			}
 
 			// when buffer stops
 			NotifyEvents[STOP_EVENT].dwOffset = DSBPN_OFFSETSTOP;
-			NotifyEvents[STOP_EVENT].hEventNotify = Events[STOP_EVENT];
+			NotifyEvents[STOP_EVENT].hEventNotify = CV->Events[STOP_EVENT];
 
 			hr = NotifyInt->SetNotificationPositions(NUM_EVENTS, NotifyEvents);
 			if (FAILED(hr))
@@ -404,7 +435,7 @@ bool FVoiceCaptureWindows::CreateNotifications(uint32 BufferSize)
 				bSuccess = false;
 			}
 
-			float BufferSegMs = BufferSegSize / ((WavFormat.nSamplesPerSec / 1000.0) * (WavFormat.wBitsPerSample / 8.0));
+			float BufferSegMs = BufferSegSize / ((CV->WavFormat.nSamplesPerSec / 1000.0) * (CV->WavFormat.wBitsPerSample / 8.0));
 			UE_LOG(LogVoiceCapture, Display, TEXT("%d notifications, every %d bytes [%f ms]"), NumSegments, BufferSegSize, BufferSegMs);
 		}
 
@@ -419,10 +450,10 @@ bool FVoiceCaptureWindows::CreateNotifications(uint32 BufferSize)
 	{
 		for (int i = 0; i < NUM_EVENTS; i++)
 		{
-			if (Events[i] != INVALID_HANDLE_VALUE)
+			if (CV->Events[i] != INVALID_HANDLE_VALUE)
 			{
-				CloseHandle(Events[i]);
-				Events[i] = INVALID_HANDLE_VALUE;
+				CloseHandle(CV->Events[i]);
+				CV->Events[i] = INVALID_HANDLE_VALUE;
 			}
 		}
 	}
@@ -437,7 +468,7 @@ bool FVoiceCaptureWindows::Tick(float DeltaTime)
 	{
 #if 1
 		uint32 NextEvent = (LastEventTriggered + 1) % (NUM_EVENTS - 1);
-		if (WaitForSingleObject(Events[NextEvent], 0) == WAIT_OBJECT_0)
+		if (WaitForSingleObject(CV->Events[NextEvent], 0) == WAIT_OBJECT_0)
 		{
 #if !UE_BUILD_SHIPPING
 			UE_LOG(LogVoiceCapture, VeryVerbose, TEXT("Event %d Signalled"), NextEvent);
@@ -445,7 +476,7 @@ bool FVoiceCaptureWindows::Tick(float DeltaTime)
 			ProcessData();
 
 			LastEventTriggered = NextEvent;
-			ResetEvent(Events[NextEvent]);
+			ResetEvent(CV->Events[NextEvent]);
 		}
 #else
 		// All "non stop" events
@@ -462,15 +493,15 @@ bool FVoiceCaptureWindows::Tick(float DeltaTime)
 				}
 
 				LastEventTriggered = EventIdx;
-				ResetEvent(Events[EventIdx]);
+				ResetEvent(CV->Events[EventIdx]);
 			}
 		}
 #endif
 
-		if (Events[STOP_EVENT] != INVALID_HANDLE_VALUE && WaitForSingleObject(Events[STOP_EVENT], 0) == WAIT_OBJECT_0)
+		if (CV->Events[STOP_EVENT] != INVALID_HANDLE_VALUE && WaitForSingleObject(CV->Events[STOP_EVENT], 0) == WAIT_OBJECT_0)
 		{
 			UE_LOG(LogVoiceCapture, Verbose, TEXT("Voice capture stopped"));
-			ResetEvent(Events[STOP_EVENT]);
+			ResetEvent(CV->Events[STOP_EVENT]);
 			VoiceCaptureState = EVoiceCaptureState::NotCapturing;
 		}
 	}
@@ -478,4 +509,4 @@ bool FVoiceCaptureWindows::Tick(float DeltaTime)
 	return true;
 }
 
-#include "HideWindowsPlatformTypes.h"
+#include "HideWindowsPlatformTypes.h"
\ No newline at end of file
diff --git a/Engine/Source/Runtime/Online/Voice/Private/Windows/VoiceCaptureWindows.h b/Engine/Source/Runtime/Online/Voice/Private/Windows/VoiceCaptureWindows.h
--- a/Engine/Source/Runtime/Online/Voice/Private/Windows/VoiceCaptureWindows.h
+++ b/Engine/Source/Runtime/Online/Voice/Private/Windows/VoiceCaptureWindows.h
@@ -5,12 +5,11 @@
 #include "VoiceCapture.h"
 #include "VoicePackage.h"
 
-#include "AllowWindowsPlatformTypes.h"
-#include "audiodefs.h"
-#include "dsound.h"
+typedef struct IDirectSound8 *LPDIRECTSOUND8;
 
 class IVoiceCapture;
 class FVoiceCaptureWindows;
+struct FVoiceCaptureWindowsVars;
 
 /** Number of notification events distributed throughout the capture buffer */
 #define NUM_EVENTS 5
@@ -105,28 +104,14 @@ public:
 
 private:
 
-	/** Voice capture device */
-	LPDIRECTSOUNDCAPTURE8 VoiceCaptureDev;
-	/** Voice capture device caps */
-	DSCCAPS VoiceCaptureDevCaps;
-	/** State of capture device */
-	EVoiceCaptureState::Type VoiceCaptureState;
-	/** Voice capture buffer */
-	LPDIRECTSOUNDCAPTUREBUFFER8 VoiceCaptureBuffer8;
-	/** Wave format of buffer */
-	WAVEFORMATEX WavFormat;
-	/** Buffer description */
-	DSCBUFFERDESC VoiceCaptureBufferDesc;
-	/** Buffer caps */
-	DSCBCAPS VoiceCaptureBufferCaps8;
+	/** All windows related variables to hide windows includes */
+	FVoiceCaptureWindowsVars* CV;
 	/** Event sentinel */
 	uint32 LastEventTriggered;
-	/** Notification events */
-	HANDLE Events[NUM_EVENTS];
+	/** State of capture device */
+	EVoiceCaptureState::Type VoiceCaptureState;
 	/** Uncompressed audio buffer */
 	TArray<uint8> UncompressedAudioBuffer;
-	/** Current audio position of valid data in capture buffer */
-	DWORD NextCaptureOffset;
 
 	/**
 	 * Lock the DirectSound audio buffer and copy out the available data based on a notification
@@ -140,5 +125,3 @@ private:
 	 */
 	bool CreateNotifications(uint32 BufferSize);
 };
-
-#include "HideWindowsPlatformTypes.h"
\ No newline at end of file
diff --git a/Engine/Source/Runtime/Online/Voice/Private/Windows/VoiceModuleWindows.cpp b/Engine/Source/Runtime/Online/Voice/Private/Windows/VoiceModuleWindows.cpp
--- a/Engine/Source/Runtime/Online/Voice/Private/Windows/VoiceModuleWindows.cpp
+++ b/Engine/Source/Runtime/Online/Voice/Private/Windows/VoiceModuleWindows.cpp
@@ -1,12 +1,28 @@
 // Copyright 1998-2015 Epic Games, Inc. All Rights Reserved.
 
 #include "VoicePrivatePCH.h"
+
 #include "VoiceCaptureWindows.h"
 #include "VoiceCodecOpus.h"
 #include "Voice.h"
 
 FVoiceCaptureDeviceWindows* FVoiceCaptureDeviceWindows::Singleton = NULL;
 
+/** Helper for printing MS guids */
+FString PrintMSGUID(LPGUID Guid)
+{
+	FString Result;
+	if (Guid)
+	{
+		Result = FString::Printf(TEXT("{%08lX-%04hX-%04hX-%02hhX%02hhX-%02hhX%02hhX%02hhX%02hhX%02hhX%02hhX}"),
+			Guid->Data1, Guid->Data2, Guid->Data3,
+			Guid->Data4[0], Guid->Data4[1], Guid->Data4[2], Guid->Data4[3],
+			Guid->Data4[4], Guid->Data4[5], Guid->Data4[6], Guid->Data4[7]);
+	}
+
+	return Result;
+}
+
 /** Callback to access all the voice capture devices on the platform */
 BOOL CALLBACK CaptureDeviceCallback(
 	LPGUID lpGuid,
@@ -17,7 +33,7 @@ BOOL CALLBACK CaptureDeviceCallback(
 {
 	// @todo identify the proper device
 	FVoiceCaptureWindows* VCPtr = (FVoiceCaptureWindows*)(lpContext);
-	UE_LOG(LogVoiceCapture, Display, TEXT("Device: %s Desc: %s"), lpcstrDescription, lpcstrModule);
+	UE_LOG(LogVoiceCapture, Display, TEXT("Device: %s Desc: %s GUID: %s Context:0x%08x"), lpcstrDescription, lpcstrModule, *PrintMSGUID(lpGuid), lpContext);
 	return true;
 }
 
@@ -61,6 +77,166 @@ FVoiceCaptureDeviceWindows* FVoiceCaptureDeviceWindows::Get()
 	return Singleton;
 }
 
+class FAudioDuckingWindows
+{
+private:
+	
+	FAudioDuckingWindows() {}
+	FAudioDuckingWindows(const FAudioDuckingWindows& Copy) {}
+	~FAudioDuckingWindows() {}
+	void operator=(const FAudioDuckingWindows&) {}
+
+public:
+
+	static HRESULT EnableDuckingOptOut(IMMDevice* pEndpoint, bool bDuckingOptOutChecked)
+	{
+		HRESULT hr = S_OK;
+
+		// Activate session manager.
+		IAudioSessionManager2* pSessionManager2 = NULL;
+		if (SUCCEEDED(hr))
+		{
+			hr = pEndpoint->Activate(__uuidof(IAudioSessionManager2),
+				CLSCTX_INPROC_SERVER,
+				NULL,
+				reinterpret_cast<void **>(&pSessionManager2));
+		}
+
+		IAudioSessionControl* pSessionControl = NULL;
+		if (SUCCEEDED(hr))
+		{
+			hr = pSessionManager2->GetAudioSessionControl(NULL, 0, &pSessionControl);
+
+			pSessionManager2->Release();
+			pSessionManager2 = NULL;
+		}
+
+		IAudioSessionControl2* pSessionControl2 = NULL;
+		if (SUCCEEDED(hr))
+		{
+			hr = pSessionControl->QueryInterface(
+				__uuidof(IAudioSessionControl2),
+				(void**)&pSessionControl2);
+
+			pSessionControl->Release();
+			pSessionControl = NULL;
+		}
+
+		//  Sync the ducking state with the specified preference.
+		if (SUCCEEDED(hr))
+		{
+			hr = pSessionControl2->SetDuckingPreference((::BOOL)bDuckingOptOutChecked);
+			if (FAILED(hr))
+			{
+				UE_LOG(LogVoiceCapture, Display, TEXT("Failed to duck audio endpoint. Error: %0x08x"), hr);
+			}
+
+			pSessionControl2->Release();
+			pSessionControl2 = NULL;
+		}
+		return hr;
+	}
+
+	static bool UpdateAudioDucking(bool bDuckingOptOutChecked)
+	{
+		HRESULT hr = S_OK;
+
+		//  Start with the default endpoint.
+		IMMDeviceEnumerator* pDeviceEnumerator = NULL;
+		hr = CoCreateInstance(__uuidof(MMDeviceEnumerator),
+			NULL,
+			CLSCTX_INPROC_SERVER,
+			IID_PPV_ARGS(&pDeviceEnumerator));
+
+		if (SUCCEEDED(hr))
+		{
+			{
+				IMMDevice* pEndpoint = NULL;
+				hr = pDeviceEnumerator->GetDefaultAudioEndpoint(eRender, eConsole, &pEndpoint);
+				if (SUCCEEDED(hr))
+				{
+					LPWSTR Desc;
+					pEndpoint->GetId(&Desc);
+					UE_LOG(LogVoiceCapture, Display, TEXT("%s ducking on audio device. Desc: %s"), bDuckingOptOutChecked ? TEXT("Disabling") : TEXT("Enabling"), Desc);
+					CoTaskMemFree(Desc);
+
+					FAudioDuckingWindows::EnableDuckingOptOut(pEndpoint, bDuckingOptOutChecked);
+					pEndpoint->Release();
+					pEndpoint = NULL;
+				}
+			}
+
+			if (0) // reference for enumerating all endpoints in case its necessary
+			{
+				IMMDeviceCollection* pDeviceCollection = NULL;
+				IMMDevice* pCollEndpoint = NULL;
+				hr = pDeviceEnumerator->EnumAudioEndpoints(eRender, DEVICE_STATE_ACTIVE, &pDeviceCollection);
+				if (SUCCEEDED(hr))
+				{
+					IPropertyStore *pProps = NULL;
+
+					::UINT DeviceCount = 0;
+					pDeviceCollection->GetCount(&DeviceCount);
+					for (::UINT i = 0; i < DeviceCount; ++i)
+					{
+						hr = pDeviceCollection->Item(i, &pCollEndpoint);
+						if (SUCCEEDED(hr) && pCollEndpoint)
+						{
+							LPWSTR Desc;
+							pCollEndpoint->GetId(&Desc);
+
+							hr = pCollEndpoint->OpenPropertyStore(STGM_READ, &pProps);
+							if (SUCCEEDED(hr))
+							{
+								PROPVARIANT varName;
+								// Initialize container for property value.
+								PropVariantInit(&varName);
+
+								// Get the endpoint's friendly-name property.
+								hr = pProps->GetValue(
+									PKEY_Device_FriendlyName, &varName);
+								if (SUCCEEDED(hr))
+								{
+									// Print endpoint friendly name and endpoint ID.
+									UE_LOG(LogVoiceCapture, Display, TEXT("%s ducking on audio device [%d]: \"%s\" (%s)"),
+										bDuckingOptOutChecked ? TEXT("Disabling") : TEXT("Enabling"),
+										i, varName.pwszVal, Desc);
+
+									CoTaskMemFree(Desc);
+
+									Desc = NULL;
+									PropVariantClear(&varName);
+
+									pProps->Release();
+									pProps = NULL;
+								}
+							}
+
+							FAudioDuckingWindows::EnableDuckingOptOut(pCollEndpoint, bDuckingOptOutChecked);
+
+							pCollEndpoint->Release();
+							pCollEndpoint = NULL;
+						}
+					}
+
+					pDeviceCollection->Release();
+					pDeviceCollection = NULL;
+				}
+			}
+
+			pDeviceEnumerator->Release();
+			pDeviceEnumerator = NULL;
+		}
+
+		if (FAILED(hr))
+		{
+			UE_LOG(LogVoiceCapture, Warning, TEXT("Failed to duck audio endpoint. Error: %0x08x"), hr);
+		}
+
+		return SUCCEEDED(hr);
+	}
+};
+
 bool FVoiceCaptureDeviceWindows::Init()
 {
 	HRESULT hr = DirectSoundCreate8(NULL, &DirectSound, NULL);
@@ -69,14 +245,34 @@ bool FVoiceCaptureDeviceWindows::Init()
 		UE_LOG(LogVoiceCapture, Warning, TEXT("Failed to init DirectSound %d"), hr);
 		return false;
 	}
+	
+	if (0)
+	{
+		hr = DirectSoundEnumerate((LPDSENUMCALLBACK)CaptureDeviceCallback, this);
+		if (FAILED(hr))
+		{
+			UE_LOG(LogVoiceCapture, Warning, TEXT("Failed to enumerate render devices %d"), hr);
+			return false;
+		}
+	}
 
 	hr = DirectSoundCaptureEnumerate((LPDSENUMCALLBACK)CaptureDeviceCallback, this);
 	if (FAILED(hr))
 	{
-		UE_LOG(LogVoiceCapture, Warning, TEXT("Failed to enumerate devices %d"), hr);
+		UE_LOG(LogVoiceCapture, Warning, TEXT("Failed to enumerate capture devices %d"), hr);
 		return false;
 	}
 
+	bool bDuckingOptOut = false;
+	if (GConfig)
+	{
+		if (!GConfig->GetBool(TEXT("Voice"), TEXT("bDuckingOptOut"), bDuckingOptOut, GEngineIni))
+		{
+			bDuckingOptOut = false;
+		}
+	}
+	FAudioDuckingWindows::UpdateAudioDucking(bDuckingOptOut);
+
 	bInitialized = true;
 	return true;
 }
@@ -176,4 +372,4 @@ IVoiceDecoder* CreateVoiceDecoderObject()
 	}
 
 	return NewDecoder; 
-}
\ No newline at end of file
+}
diff --git a/Engine/Source/Runtime/UMG/Private/Components/TextBlock.cpp b/Engine/Source/Runtime/UMG/Private/Components/TextBlock.cpp
--- a/Engine/Source/Runtime/UMG/Private/Components/TextBlock.cpp
+++ b/Engine/Source/Runtime/UMG/Private/Components/TextBlock.cpp
@@ -55,6 +55,15 @@ void UTextBlock::SetShadowOffset(FVector2D InShadowOffset)
 	}
 }
 
+void UTextBlock::SetJustification( ETextJustify::Type InJustification )
+{
+	Justification = InJustification;
+	if ( MyTextBlock.IsValid() )
+	{
+		MyTextBlock->SetJustification( Justification );
+	}
+}
+
 TSharedRef<SWidget> UTextBlock::RebuildWidget()
 {
 	MyTextBlock = SNew(STextBlock);
diff --git a/Engine/Source/Runtime/UMG/Public/Components/OverlaySlot.h b/Engine/Source/Runtime/UMG/Public/Components/OverlaySlot.h
--- a/Engine/Source/Runtime/UMG/Public/Components/OverlaySlot.h
+++ b/Engine/Source/Runtime/UMG/Public/Components/OverlaySlot.h
@@ -4,6 +4,8 @@
 
 #include "SlateWrapperTypes.h"
 
+#include "Components/PanelSlot.h"
+
 #include "OverlaySlot.generated.h"
 
 /** Slot for the UOverlay panel.  Allows content to be hover above other content. */
@@ -46,11 +48,11 @@ public:
 	// End of UPanelSlot interface
 
 	/** Builds the underlying FSlot for the Slate layout panel. */
-	void BuildSlot(TSharedRef<SOverlay> InOverlay);
+	virtual void BuildSlot(TSharedRef<SOverlay> InOverlay);
 
 	virtual void ReleaseSlateResources(bool bReleaseChildren) override;
 
-private:
+protected:
 	//TODO UMG Slots should hold weak or shared refs to slots.
 
 	/** A raw pointer to the slot to allow us to adjust the size, padding...etc at runtime. */
diff --git a/Engine/Source/Runtime/UMG/Public/Components/TextBlock.h b/Engine/Source/Runtime/UMG/Public/Components/TextBlock.h
--- a/Engine/Source/Runtime/UMG/Public/Components/TextBlock.h
+++ b/Engine/Source/Runtime/UMG/Public/Components/TextBlock.h
@@ -41,6 +41,15 @@ public:
 	UFUNCTION(BlueprintCallable, Category="Appearance")
 	void SetShadowOffset(FVector2D InShadowOffset);
 
+
+	/**
+	 *  Set the text justification for this text block
+	 *
+	 *  @param Justification new justification
+	 */
+	UFUNCTION( BlueprintCallable, Category = "Appearance" )
+	void SetJustification( ETextJustify::Type InJustification );
+
 public:
 	UPROPERTY()
 	USlateWidgetStyleAsset* Style_DEPRECATED;
