BuildPatchServices: Downgrade finding a chunk with zero hash from a warning to an info in log output.

diff --git a/Engine/Source/Runtime/Online/BuildPatchServices/Private/BuildPatchGeneration.cpp b/Engine/Source/Runtime/Online/BuildPatchServices/Private/BuildPatchGeneration.cpp
--- a/Engine/Source/Runtime/Online/BuildPatchServices/Private/BuildPatchGeneration.cpp
+++ b/Engine/Source/Runtime/Online/BuildPatchServices/Private/BuildPatchGeneration.cpp
@@ -1500,7 +1500,7 @@ bool FBuildDataGenerator::FindExistingChunkData( const uint64& ChunkHash, const
 								}
 								else
 								{
-									GLog->Logf(TEXT("BuildDataGenerator: WARNING: Ignored an existing chunk %s with a failed hash value of zero to avoid performance problems while chunking"), *ChunkGuid.ToString());
+									GLog->Logf(TEXT("BuildDataGenerator: INFO: Ignored an existing chunk %s with a failed hash value of zero to avoid performance problems while chunking"), *ChunkGuid.ToString());
 								}
 							}
 							else
