PS4 matching2 search attributes

diff --git a/Engine/Source/Runtime/Online/OnlineSubsystem/Public/OnlineSessionSettings.h b/Engine/Source/Runtime/Online/OnlineSubsystem/Public/OnlineSessionSettings.h
--- a/Engine/Source/Runtime/Online/OnlineSubsystem/Public/OnlineSessionSettings.h
+++ b/Engine/Source/Runtime/Online/OnlineSubsystem/Public/OnlineSessionSettings.h
@@ -16,6 +16,16 @@
 /** Setting describing the beacon host port (value is int32) */
 #define SETTING_BEACONPORT FName(TEXT("BEACONPORT"))
 
+/** 8 user defined integer params to be used when filtering searches for sessions */
+#define SETTING_CUSTOMSEARCHINT1 FName(TEXT("CUSTOMSEARCHINT1"))
+#define SETTING_CUSTOMSEARCHINT2 FName(TEXT("CUSTOMSEARCHINT2"))
+#define SETTING_CUSTOMSEARCHINT3 FName(TEXT("CUSTOMSEARCHINT3"))
+#define SETTING_CUSTOMSEARCHINT4 FName(TEXT("CUSTOMSEARCHINT4"))
+#define SETTING_CUSTOMSEARCHINT5 FName(TEXT("CUSTOMSEARCHINT5"))
+#define SETTING_CUSTOMSEARCHINT6 FName(TEXT("CUSTOMSEARCHINT6"))
+#define SETTING_CUSTOMSEARCHINT7 FName(TEXT("CUSTOMSEARCHINT7"))
+#define SETTING_CUSTOMSEARCHINT8 FName(TEXT("CUSTOMSEARCHINT8"))
+
 /** TODO ONLINE Settings to consider */
 /** The server's nonce for this session */
 /** Whether this match is publicly advertised on the online service */
