Implement Voice Module for Mac OS X: [UEPLAT-44]
- This is an extremely simple, OpenAL-based implementation that can probably be shared with Linux.

diff --git a/Engine/Source/Runtime/Online/Voice/Private/Mac/VoiceModuleMac.cpp b/Engine/Source/Runtime/Online/Voice/Private/Mac/VoiceModuleMac.cpp
--- a/Engine/Source/Runtime/Online/Voice/Private/Mac/VoiceModuleMac.cpp
+++ b/Engine/Source/Runtime/Online/Voice/Private/Mac/VoiceModuleMac.cpp
@@ -4,9 +4,204 @@
 #include "VoiceCodecOpus.h"
 #include "Voice.h"
 
+#include <OpenAL/al.h>
+#include <OpenAL/alc.h>
+
+/**
+ * Implementation of voice capture using OpenAL
+ */
+class FVoiceCaptureOpenAL : public IVoiceCapture
+{
+public:
+	FVoiceCaptureOpenAL()
+	: CaptureDevice(nullptr)
+	, VoiceCaptureState(EVoiceCaptureState::UnInitialized)
+	, NumChannels(0)
+	{
+		
+	}
+	~FVoiceCaptureOpenAL()
+	{
+		Shutdown();
+	}
+	
+	// IVoiceCapture
+	virtual bool Init(int32 InSampleRate, int32 InNumChannels) override
+	{
+		check(VoiceCaptureState == EVoiceCaptureState::UnInitialized);
+		
+		if (InSampleRate < 8000 || InSampleRate > 48000)
+		{
+			UE_LOG(LogVoiceCapture, Warning, TEXT("Voice capture doesn't support %d hz"), InSampleRate);
+			return false;
+		}
+		
+		if (InNumChannels < 0 || InNumChannels > 2)
+		{
+			UE_LOG(LogVoiceCapture, Warning, TEXT("Voice capture only supports 1 or 2 channels"));
+			return false;
+		}
+		
+		CaptureDevice = alcCaptureOpenDevice(nullptr, InSampleRate, InNumChannels == 1 ? AL_FORMAT_MONO16 : AL_FORMAT_STEREO16, (InSampleRate * InNumChannels));
+		if (!CaptureDevice)
+		{
+			UE_LOG(LogVoiceCapture, Warning, TEXT("Couldn't create the audio input device"));
+			return false;
+		}
+		
+		ALenum ErrorCode = alcGetError(CaptureDevice);
+		if (ErrorCode != AL_NO_ERROR)
+		{
+			UE_LOG(LogVoiceCapture, Warning, TEXT("Couldn't initialise the audio input device, err: %d"), ErrorCode);
+			return false;
+		}
+		
+		VoiceCaptureState = EVoiceCaptureState::NotCapturing;
+		NumChannels = InNumChannels;
+		
+		return true;
+	}
+	virtual void Shutdown() override
+	{
+		switch (VoiceCaptureState)
+		{
+			case EVoiceCaptureState::Ok:
+			case EVoiceCaptureState::NoData:
+			case EVoiceCaptureState::Stopping:
+			case EVoiceCaptureState::BufferTooSmall:
+			case EVoiceCaptureState::Error:
+				Stop();
+			case EVoiceCaptureState::NotCapturing:
+				check(CaptureDevice);
+				alcCaptureCloseDevice(CaptureDevice);
+				CaptureDevice = nullptr;
+				VoiceCaptureState = EVoiceCaptureState::UnInitialized;
+				break;
+			case EVoiceCaptureState::UnInitialized:
+				break;
+				
+			default:
+				check(false);
+				break;
+		}
+	}
+	virtual bool Start() override
+	{
+		check(CaptureDevice && VoiceCaptureState == EVoiceCaptureState::NotCapturing);
+		alcCaptureStart(CaptureDevice);
+		ALenum Error = alcGetError(CaptureDevice);
+		if ( Error == AL_NO_ERROR )
+		{
+			VoiceCaptureState = EVoiceCaptureState::Ok;
+		}
+		else
+		{
+			UE_LOG(LogVoiceCapture, Warning, TEXT("Failed to start capture %d"), Error);
+		}
+		return ( Error == AL_NO_ERROR );
+	}
+	virtual void Stop() override
+	{
+		check(IsCapturing());
+		alcCaptureStop(CaptureDevice);
+		ALenum Error = alcGetError(CaptureDevice);
+		if ( Error != AL_NO_ERROR )
+		{
+			UE_LOG(LogVoiceCapture, Warning, TEXT("Failed to stop capture %d"), Error);
+		}
+	}
+	virtual bool IsCapturing() override
+	{
+		return CaptureDevice && VoiceCaptureState > EVoiceCaptureState::NotCapturing;
+	}
+	virtual EVoiceCaptureState::Type GetCaptureState(uint32& OutAvailableVoiceData) const override
+	{
+		if (VoiceCaptureState != EVoiceCaptureState::UnInitialized &&
+			VoiceCaptureState != EVoiceCaptureState::NotCapturing)
+		{
+			ALCint Samples = 0;
+			alcGetIntegerv(CaptureDevice, ALC_CAPTURE_SAMPLES, 1, &Samples);
+			ALenum Error = alcGetError(CaptureDevice);
+			if ( Error == AL_NO_ERROR )
+			{
+				OutAvailableVoiceData = (Samples * NumChannels * sizeof(uint16));
+				VoiceCaptureState = (Samples ? EVoiceCaptureState::Ok : EVoiceCaptureState::NoData);
+			}
+			else
+			{
+				UE_LOG(LogVoiceCapture, Warning, TEXT("Failed to get number of available capture samples %d"), Error);
+				OutAvailableVoiceData = 0;
+				VoiceCaptureState = EVoiceCaptureState::Error;
+			}
+		}
+		else
+		{
+			OutAvailableVoiceData = 0;
+		}
+		
+		return VoiceCaptureState;
+	}
+	virtual EVoiceCaptureState::Type GetVoiceData(uint8* OutVoiceBuffer, uint32 InVoiceBufferSize, uint32& OutAvailableVoiceData) override
+	{
+		EVoiceCaptureState::Type State = VoiceCaptureState;
+		
+		if (VoiceCaptureState != EVoiceCaptureState::UnInitialized &&
+			VoiceCaptureState != EVoiceCaptureState::NotCapturing)
+		{
+			ALCint Samples = 0;
+			alcGetIntegerv(CaptureDevice, ALC_CAPTURE_SAMPLES, 1, &Samples);
+			ALenum Error = alcGetError(CaptureDevice);
+			if ( Error == AL_NO_ERROR )
+			{
+				if ( Samples > 0 )
+				{
+					OutAvailableVoiceData = Samples * NumChannels * sizeof(uint16);
+					if ( InVoiceBufferSize >= OutAvailableVoiceData )
+					{
+						alcCaptureSamples(CaptureDevice, OutVoiceBuffer, Samples);
+						Error = alcGetError(CaptureDevice);
+						VoiceCaptureState = EVoiceCaptureState::NoData;
+						State = EVoiceCaptureState::Ok;
+					}
+					else
+					{
+						State = EVoiceCaptureState::BufferTooSmall;
+					}
+				}
+				else
+				{
+					VoiceCaptureState = EVoiceCaptureState::NoData;
+					State = EVoiceCaptureState::Ok;
+				}
+			}
+			if ( Error != AL_NO_ERROR )
+			{
+				UE_LOG(LogVoiceCapture, Warning, TEXT("Failed to capture samples %d"), Error);
+				OutAvailableVoiceData = 0;
+				VoiceCaptureState = EVoiceCaptureState::Error;
+				State = EVoiceCaptureState::Error;
+			}
+		}
+		else
+		{
+			OutAvailableVoiceData = 0;
+		}
+		
+		return State;
+	}
+	
+private:
+	/** OpenAL device for capturing audio */
+	ALCdevice* CaptureDevice;
+	/** State of capture device */
+	mutable EVoiceCaptureState::Type VoiceCaptureState;
+	/** Initialised number of channels */
+	int32 NumChannels;
+};
+
 bool InitVoiceCapture()
 {
-	return false;
+	return true;
 }
 
 void ShutdownVoiceCapture()
@@ -15,15 +210,35 @@ void ShutdownVoiceCapture()
 
 IVoiceCapture* CreateVoiceCaptureObject()
 {
-	return NULL; 
+	IVoiceCapture* Capture = new FVoiceCaptureOpenAL;
+	if ( !Capture || !Capture->Init(VOICE_SAMPLE_RATE, NUM_VOICE_CHANNELS) )
+	{
+		delete Capture;
+		Capture = nullptr;
+	}
+	return Capture;
 }
 
 IVoiceEncoder* CreateVoiceEncoderObject()
 {
-	return NULL; 
+	FVoiceEncoderOpus* NewEncoder = new FVoiceEncoderOpus;
+	if (!NewEncoder->Init(VOICE_SAMPLE_RATE, NUM_VOICE_CHANNELS))
+	{
+		delete NewEncoder;
+		NewEncoder = NULL;
+	}
+	
+	return NewEncoder;
 }
 
 IVoiceDecoder* CreateVoiceDecoderObject()
 {
-	return NULL; 
+	FVoiceDecoderOpus* NewDecoder = new FVoiceDecoderOpus;
+	if (!NewDecoder->Init(VOICE_SAMPLE_RATE, NUM_VOICE_CHANNELS))
+	{
+		delete NewDecoder;
+		NewDecoder = NULL;
+	}
+	
+	return NewDecoder;
 }
diff --git a/Engine/Source/Runtime/Online/Voice/Private/VoicePrivatePCH.h b/Engine/Source/Runtime/Online/Voice/Private/VoicePrivatePCH.h
--- a/Engine/Source/Runtime/Online/Voice/Private/VoicePrivatePCH.h
+++ b/Engine/Source/Runtime/Online/Voice/Private/VoicePrivatePCH.h
@@ -18,7 +18,7 @@
 
 #endif // PLATFORM_WINDOWS
 
-#define PLATFORM_SUPPORTS_VOICE_CAPTURE (PLATFORM_WINDOWS)
+#define PLATFORM_SUPPORTS_VOICE_CAPTURE (PLATFORM_WINDOWS || PLATFORM_MAC)
 
 // Module includes
 #include "VoiceModule.h"
\ No newline at end of file
diff --git a/Engine/Source/Runtime/Online/Voice/Voice.Build.cs b/Engine/Source/Runtime/Online/Voice/Voice.Build.cs
--- a/Engine/Source/Runtime/Online/Voice/Voice.Build.cs
+++ b/Engine/Source/Runtime/Online/Voice/Voice.Build.cs
@@ -25,6 +25,10 @@ public class Voice : ModuleRules
 		{
 			AddThirdPartyPrivateStaticDependencies(Target, "DirectSound");
 		}
+		else if(Target.Platform == UnrealTargetPlatform.Mac)
+		{
+			PublicFrameworks.AddRange(new string[] { "OpenAL" });
+		}
 
 		AddThirdPartyPrivateStaticDependencies(Target, "libOpus");
     }
