OnlineSubsystemLive fixes contributed from Black Tusk:
Fix an issue where people joining the session wouldnÆt get shoulder taps
Session Settings now update and track changes
Session Member Setting is now updated on user join/leave
Added support for Host Token Migration
Fix an issue where bHosting was not correctly set
Add support for JSON objects and arrays in session settings

Integrated CLs from //depot/Partners/Microsoft/UE4-MS/: 2460704 and 2466495.

diff --git a/Engine/Source/Runtime/Online/OnlineSubsystem/Public/OnlineSubsystemSessionSettings.h b/Engine/Source/Runtime/Online/OnlineSubsystem/Public/OnlineSubsystemSessionSettings.h
--- a/Engine/Source/Runtime/Online/OnlineSubsystem/Public/OnlineSubsystemSessionSettings.h
+++ b/Engine/Source/Runtime/Online/OnlineSubsystem/Public/OnlineSubsystemSessionSettings.h
@@ -33,3 +33,7 @@
 #define SETTING_MULTIPLAYER_VISIBILITY FName(TEXT("MULTIPLAYERVISIBILITY"))
 /** Session template name, must match a session template name from the service config. */
 #define SETTING_SESSION_TEMPLATE_NAME FName(TEXT("SESSIONTEMPLATENAME"))
+/** Session change number from shoulder tapped version */
+#define SETTING_CHANGE_NUMBER FName(TEXT("CHANGENUMBER"))
+/** Session can attempt arbiter migration */
+#define SETTING_ALLOW_ARBITER_MIGRATION FName(TEXT("ALLOWARBITERMIGRATION"))
