CoreAudio based VoiceModule for Mac to avoid OpenAL's internal malloc/delete mismatch & a 'killtestvoice' command to silence the test-voice to double-check that the shutdown code works as expected.

diff --git a/Engine/Source/Runtime/Online/OnlineSubsystemUtils/Private/Tests/TestVoice.cpp b/Engine/Source/Runtime/Online/OnlineSubsystemUtils/Private/Tests/TestVoice.cpp
--- a/Engine/Source/Runtime/Online/OnlineSubsystemUtils/Private/Tests/TestVoice.cpp
+++ b/Engine/Source/Runtime/Online/OnlineSubsystemUtils/Private/Tests/TestVoice.cpp
@@ -207,5 +207,10 @@ bool FTestVoice::Tick(float DeltaTime)
 bool FTestVoice::Exec(UWorld* InWorld, const TCHAR* Cmd, FOutputDevice& Ar)
 {
 	bool bWasHandled = false;
+	if (FParse::Command(&Cmd, TEXT("killtestvoice")))
+	{
+		delete this;
+		bWasHandled = true;
+	}
 	return bWasHandled;
 }
diff --git a/Engine/Source/Runtime/Online/Voice/Private/Mac/VoiceModuleMac.cpp b/Engine/Source/Runtime/Online/Voice/Private/Mac/VoiceModuleMac.cpp
--- a/Engine/Source/Runtime/Online/Voice/Private/Mac/VoiceModuleMac.cpp
+++ b/Engine/Source/Runtime/Online/Voice/Private/Mac/VoiceModuleMac.cpp
@@ -4,27 +4,93 @@
 #include "VoiceCodecOpus.h"
 #include "Voice.h"
 
-#include <OpenAL/al.h>
-#include <OpenAL/alc.h>
+#include <CoreAudio/CoreAudio.h>
+#include <AudioToolbox/AudioToolbox.h>
+#include <AudioUnit/AudioUnit.h>
+
+/** Maximum buffer size for storing raw uncompressed audio from the system */
+#define MAX_UNCOMPRESSED_VOICE_BUFFER_SIZE 30 * 1024
 
 /**
  * Implementation of voice capture using OpenAL
  */
-class FVoiceCaptureOpenAL : public IVoiceCapture
+class FVoiceCaptureCoreAudio : public IVoiceCapture
 {
 public:
-	FVoiceCaptureOpenAL()
-	: CaptureDevice(nullptr)
+	FVoiceCaptureCoreAudio()
+	: StreamComponent(nullptr)
+	, StreamConverter(nullptr)
 	, VoiceCaptureState(EVoiceCaptureState::UnInitialized)
-	, NumChannels(0)
+	, BufferSize(0)
+	, InputLatency(0)
+	, ReadableBytes(0)
+	, WriteBuffer(0)
+	, ReadBuffer(0)
+	, ReadOffset(0)
 	{
 		
 	}
-	~FVoiceCaptureOpenAL()
+	~FVoiceCaptureCoreAudio()
 	{
 		Shutdown();
 	}
 	
+	static OSStatus InputProc(void* RefCon,
+							  AudioUnitRenderActionFlags* IOActionFlags,
+							  const AudioTimeStamp* InTimeStamp,
+							  UInt32 InBusNumber,
+							  UInt32 InNumberFrames,
+							  AudioBufferList* IOData)
+	{
+		FVoiceCaptureCoreAudio* This = static_cast<FVoiceCaptureCoreAudio*>(RefCon);
+		if ( This && This->IsCapturing() )
+		{
+			AudioBufferList BufferList;
+			FMemory::Memzero(BufferList);
+			BufferList.mNumberBuffers = 1;
+			BufferList.mBuffers[0].mNumberChannels = This->OutputDesc.mChannelsPerFrame;
+			
+			AudioUnitRenderActionFlags Flags = 0;
+			if ( AudioUnitRender(This->StreamComponent, &Flags, InTimeStamp, InBusNumber, InNumberFrames, &BufferList) == 0 )
+			{
+				if ( This->IsCapturing() && This->ReadBuffer != This->WriteBuffer )
+				{
+					AudioBuffer& WriteBuffer = This->BufferList[This->WriteBuffer];
+					uint32 WritableSize = BufferList.mBuffers[0].mDataByteSize;
+					uint8 const* InputData = (uint8 const*)BufferList.mBuffers[0].mData;
+					while ( WritableSize )
+					{
+						uint32 MaxWritableSize = FMath::Min(This->BufferSize - WriteBuffer.mDataByteSize, WritableSize);
+						uint8* WriteData = (uint8*)(WriteBuffer.mData) + WriteBuffer.mDataByteSize;
+						FMemory::Memcpy(WriteData, InputData, MaxWritableSize);
+						InputData += MaxWritableSize;
+						WritableSize -= MaxWritableSize;
+						
+						WriteBuffer.mDataByteSize += MaxWritableSize;
+						
+						uint32 ReadIndex = (This->WriteBuffer + 1) % This->BufferList.Num();
+						if( WriteBuffer.mDataByteSize == This->BufferSize )
+						{
+							This->WriteBuffer = ReadIndex;
+							FPlatformAtomics::InterlockedAdd(&This->ReadableBytes, WriteBuffer.mDataByteSize);
+							This->VoiceCaptureState = This->VoiceCaptureState == EVoiceCaptureState::NoData ? EVoiceCaptureState::Ok : This->VoiceCaptureState;
+							if ( This->ReadBuffer == ReadIndex )
+							{
+								break; // Clip
+							}
+						}
+					}
+				}
+			}
+			else
+			{
+				This->VoiceCaptureState = EVoiceCaptureState::Error;
+			}
+		}
+		
+		return noErr;
+	}
+	
 	// IVoiceCapture
 	virtual bool Init(int32 InSampleRate, int32 InNumChannels) override
 	{
@@ -42,22 +108,159 @@ public:
 			return false;
 		}
 		
-		CaptureDevice = alcCaptureOpenDevice(nullptr, InSampleRate, InNumChannels == 1 ? AL_FORMAT_MONO16 : AL_FORMAT_STEREO16, (InSampleRate * InNumChannels));
-		if (!CaptureDevice)
+		uint32 PropSize = sizeof(AudioDeviceID);
+		AudioDeviceID InputDevice = 0;
+		AudioObjectPropertyAddress Address = { kAudioHardwarePropertyDefaultInputDevice, kAudioObjectPropertyScopeGlobal, kAudioObjectPropertyElementMaster };
+		if ( AudioObjectGetPropertyData(kAudioObjectSystemObject, &Address, 0, nullptr, &PropSize, &InputDevice) != 0 || InputDevice == kAudioDeviceUnknown )
+		{
+			UE_LOG(LogVoiceCapture, Warning, TEXT("Couldn't get Default CoreAudio input device"));
+			return false;
+		}
+		
+		StreamDesc.mSampleRate = InSampleRate;
+		StreamDesc.mFormatID = kAudioFormatLinearPCM;
+		StreamDesc.mFormatFlags = kAudioFormatFlagIsSignedInteger | kAudioFormatFlagsNativeEndian | kAudioFormatFlagIsPacked;
+		StreamDesc.mBytesPerPacket = InNumChannels == 1 ? 2 : 4;
+		StreamDesc.mBytesPerFrame = InNumChannels == 1 ? 2 : 4;
+		StreamDesc.mFramesPerPacket = 1;
+		StreamDesc.mBitsPerChannel = 16;
+		StreamDesc.mChannelsPerFrame = InNumChannels;
+		StreamDesc.mReserved = 0;
+		
+		AudioComponentDescription ComponentDesc;
+		ComponentDesc.componentType = kAudioUnitType_Output;
+		ComponentDesc.componentSubType = kAudioUnitSubType_HALOutput;
+		ComponentDesc.componentManufacturer = kAudioUnitManufacturer_Apple;
+		ComponentDesc.componentFlags = 0;
+		ComponentDesc.componentFlagsMask = 0;
+		
+		AudioComponent Component = AudioComponentFindNext(nullptr, & ComponentDesc);
+		if ( !Component || AudioComponentInstanceNew(Component, &StreamComponent) != 0 || !StreamComponent || AudioUnitInitialize(StreamComponent) != 0 )
+		{
+			UE_LOG(LogVoiceCapture, Warning, TEXT("Couldn't get CoreAudio input component"));
+			return false;
+		}
+		
+		uint32 EnableInput = 1;
+		uint32 EnableOutput = 0;
+		
+		OSStatus Error = AudioUnitSetProperty(StreamComponent, kAudioOutputUnitProperty_EnableIO, kAudioUnitScope_Output, 0, &EnableOutput, sizeof(EnableOutput));
+		if ( Error != 0 )
 		{
-			UE_LOG(LogVoiceCapture, Warning, TEXT("Couldn't create the audio input device"));
+			UE_LOG(LogVoiceCapture, Warning, TEXT("Couldn't configure CoreAudio I/O settings"));
+			AudioComponentInstanceDispose(StreamComponent);
 			return false;
 		}
 		
-		ALenum ErrorCode = alcGetError(CaptureDevice);
-		if (ErrorCode != AL_NO_ERROR)
+		Error = AudioUnitSetProperty(StreamComponent, kAudioOutputUnitProperty_EnableIO, kAudioUnitScope_Input, 1, &EnableInput, sizeof(EnableInput));
+		if ( Error != 0 )
 		{
-			UE_LOG(LogVoiceCapture, Warning, TEXT("Couldn't initialise the audio input device, err: %d"), ErrorCode);
+			UE_LOG(LogVoiceCapture, Warning, TEXT("Couldn't configure CoreAudio I/O settings"));
+			AudioComponentInstanceDispose(StreamComponent);
+			return false;
+		}
+		
+		Error = AudioUnitSetProperty(StreamComponent, kAudioOutputUnitProperty_CurrentDevice, kAudioUnitScope_Global, 0, &InputDevice, sizeof(InputDevice));
+		if ( Error != 0 )
+		{
+			UE_LOG(LogVoiceCapture, Warning, TEXT("Couldn't configure CoreAudio I/O settings"));
+			AudioComponentInstanceDispose(StreamComponent);
+			return false;
+		}
+		
+		
+		AURenderCallbackStruct InputCb;
+		InputCb.inputProc = &FVoiceCaptureCoreAudio::InputProc;
+		InputCb.inputProcRefCon = this;
+		PropSize = sizeof(NativeDesc);
+		FMemory::Memzero(NativeDesc);
+		if ( AudioUnitSetProperty(StreamComponent, kAudioOutputUnitProperty_SetInputCallback, kAudioUnitScope_Global, 0, &InputCb, sizeof(InputCb)) != 0
+			|| AudioUnitGetProperty(StreamComponent, kAudioUnitProperty_StreamFormat, kAudioUnitScope_Input, 1, &NativeDesc, &PropSize) != 0)
+		{
+			UE_LOG(LogVoiceCapture, Warning, TEXT("Couldn't configure CoreAudio input component"));
+			AudioComponentInstanceDispose(StreamComponent);
+			return false;
+		}
+		
+		OutputDesc = StreamDesc;
+		OutputDesc.mSampleRate = NativeDesc.mSampleRate;
+		if ( AudioUnitSetProperty(StreamComponent, kAudioUnitProperty_StreamFormat, kAudioUnitScope_Output, 1, (void *)&OutputDesc, sizeof(OutputDesc)) != 0 )
+		{
+			UE_LOG(LogVoiceCapture, Warning, TEXT("Couldn't configure CoreAudio input component format"));
+			AudioComponentInstanceDispose(StreamComponent);
+			return false;
+		}
+		
+		if ( StreamDesc.mSampleRate != OutputDesc.mSampleRate )
+		{
+			if ( AudioConverterNew(&OutputDesc, &StreamDesc, &StreamConverter) != 0 || !StreamConverter )
+			{
+				UE_LOG(LogVoiceCapture, Warning, TEXT("Couldn't configure CoreAudio input format converter"));
+				AudioComponentInstanceDispose(StreamComponent);
+				return false;
+			}
+		}
+		
+		PropSize = sizeof(UInt32);
+		UInt32 SafetyOffset = 0;
+		Address = { kAudioDevicePropertySafetyOffset, kAudioDevicePropertyScopeInput, 0 };
+		if ( AudioObjectGetPropertyData(InputDevice, &Address, 0, nullptr, &PropSize, &SafetyOffset) != 0 )
+		{
+			UE_LOG(LogVoiceCapture, Warning, TEXT("Couldn't get CoreAudio input latency"));
+			AudioComponentInstanceDispose(StreamComponent);
+			if(StreamConverter)
+			{
+				AudioConverterDispose(StreamConverter);
+				StreamConverter = nullptr;
+			}
+			return false;
+		}
+		SafetyOffset *= StreamDesc.mBytesPerFrame;
+		
+		PropSize = sizeof(UInt32);
+		BufferSize = 0;
+		Address = { kAudioDevicePropertyBufferFrameSize, kAudioDevicePropertyScopeInput, 0 };
+		if ( AudioObjectGetPropertyData(InputDevice, &Address, 0, nullptr, &PropSize, &BufferSize) != 0 )
+		{
+			UE_LOG(LogVoiceCapture, Warning, TEXT("Couldn't get CoreAudio input latency"));
+			AudioComponentInstanceDispose(StreamComponent);
+			if(StreamConverter)
+			{
+				AudioConverterDispose(StreamConverter);
+				StreamConverter = nullptr;
+			}
 			return false;
 		}
+		BufferSize *= StreamDesc.mBytesPerFrame;
+		
+		InputLatency = SafetyOffset + BufferSize;
+		check(InputLatency < MAX_UNCOMPRESSED_VOICE_BUFFER_SIZE);
+		
+		uint32 NumBuffers = (uint32)ceilf(MAX_UNCOMPRESSED_VOICE_BUFFER_SIZE / (float)BufferSize);
+		
+		BufferList.AddZeroed(NumBuffers);
+		WriteBuffer = 0;
+		for ( uint32 i = 0; i < BufferList.Num(); i++ )
+		{
+			BufferList[i].mData = FMemory::Malloc(BufferSize);
+			if ( !BufferList[i].mData )
+			{
+				UE_LOG(LogVoiceCapture, Warning, TEXT("Couldn't allocate CoreAudio Buffer List backing store"));
+				for ( uint32 j = i; j > 0; j-- )
+				{
+					FMemory::Free(BufferList[j-1].mData);
+				}
+				if(StreamConverter)
+				{
+					AudioConverterDispose(StreamConverter);
+					StreamConverter = nullptr;
+				}
+				AudioComponentInstanceDispose(StreamComponent);
+				return false;
+			}
+		}
 		
 		VoiceCaptureState = EVoiceCaptureState::NotCapturing;
-		NumChannels = InNumChannels;
 		
 		return true;
 	}
@@ -72,9 +275,18 @@ public:
 			case EVoiceCaptureState::Error:
 				Stop();
 			case EVoiceCaptureState::NotCapturing:
-				check(CaptureDevice);
-				alcCaptureCloseDevice(CaptureDevice);
-				CaptureDevice = nullptr;
+				check(StreamComponent);
+				if(StreamConverter)
+				{
+					AudioConverterDispose(StreamConverter);
+					StreamConverter = nullptr;
+				}
+				AudioComponentInstanceDispose(StreamComponent);
+				StreamComponent = nullptr;
+				for ( uint32 i = 0; i < BufferList.Num(); i++ )
+				{
+					FMemory::Free(BufferList[i].mData);
+				}
 				VoiceCaptureState = EVoiceCaptureState::UnInitialized;
 				break;
 			case EVoiceCaptureState::UnInitialized:
@@ -87,52 +299,58 @@ public:
 	}
 	virtual bool Start() override
 	{
-		check(CaptureDevice && VoiceCaptureState == EVoiceCaptureState::NotCapturing);
-		alcCaptureStart(CaptureDevice);
-		ALenum Error = alcGetError(CaptureDevice);
-		if ( Error == AL_NO_ERROR )
+		check(StreamComponent && VoiceCaptureState == EVoiceCaptureState::NotCapturing);
+		check(WriteBuffer == 0 && BufferList.Num());
+		for ( uint32 i = 0; i < BufferList.Num(); i++ )
+		{
+			BufferList[i].mDataByteSize = 0;
+			FMemory::Memzero(BufferList[i].mData, BufferSize);
+		}
+		WriteBuffer = 0;
+		ReadBuffer = BufferList.Num() - 1;
+		ReadOffset = 0;
+		ReadableBytes = 0;
+		OSStatus Error = AudioOutputUnitStart(StreamComponent);
+		if ( Error == 0 )
 		{
 			VoiceCaptureState = EVoiceCaptureState::Ok;
+			return true;
 		}
 		else
 		{
 			UE_LOG(LogVoiceCapture, Warning, TEXT("Failed to start capture %d"), Error);
+			return false;
 		}
-		return ( Error == AL_NO_ERROR );
 	}
 	virtual void Stop() override
 	{
 		check(IsCapturing());
-		alcCaptureStop(CaptureDevice);
-		ALenum Error = alcGetError(CaptureDevice);
-		if ( Error != AL_NO_ERROR )
+		OSStatus Error = AudioOutputUnitStop(StreamComponent);
+		if ( Error != 0 )
 		{
 			UE_LOG(LogVoiceCapture, Warning, TEXT("Failed to stop capture %d"), Error);
 		}
+		VoiceCaptureState = EVoiceCaptureState::NotCapturing;
+		WriteBuffer = 0;
+		ReadBuffer = 0;
+		ReadOffset = 0;
+		ReadableBytes = 0;
+		for ( uint32 i = 0; i < BufferList.Num(); i++ )
+		{
+			BufferList[i].mDataByteSize = 0;
+			FMemory::Memzero(BufferList[i].mData, BufferSize);
+		}
 	}
 	virtual bool IsCapturing() override
 	{
-		return CaptureDevice && VoiceCaptureState > EVoiceCaptureState::NotCapturing;
+		return StreamComponent && VoiceCaptureState > EVoiceCaptureState::NotCapturing;
 	}
 	virtual EVoiceCaptureState::Type GetCaptureState(uint32& OutAvailableVoiceData) const override
 	{
 		if (VoiceCaptureState != EVoiceCaptureState::UnInitialized &&
 			VoiceCaptureState != EVoiceCaptureState::NotCapturing)
 		{
-			ALCint Samples = 0;
-			alcGetIntegerv(CaptureDevice, ALC_CAPTURE_SAMPLES, 1, &Samples);
-			ALenum Error = alcGetError(CaptureDevice);
-			if ( Error == AL_NO_ERROR )
-			{
-				OutAvailableVoiceData = (Samples * NumChannels * sizeof(uint16));
-				VoiceCaptureState = (Samples ? EVoiceCaptureState::Ok : EVoiceCaptureState::NoData);
-			}
-			else
-			{
-				UE_LOG(LogVoiceCapture, Warning, TEXT("Failed to get number of available capture samples %d"), Error);
-				OutAvailableVoiceData = 0;
-				VoiceCaptureState = EVoiceCaptureState::Error;
-			}
+			OutAvailableVoiceData = ReadableBytes >= InputLatency ? (ReadableBytes - InputLatency) : 0;
 		}
 		else
 		{
@@ -141,6 +359,62 @@ public:
 		
 		return VoiceCaptureState;
 	}
+	
+	static OSStatus ConvertInputFormat( AudioConverterRef AudioConverter,
+										UInt32* IONumberDataPackets,
+										AudioBufferList* IOData,
+										AudioStreamPacketDescription** OutDataPacketDescription,
+										void* InUserData)
+	{
+		AudioBufferList* Buffer = static_cast<AudioBufferList*>(InUserData);
+		if ( Buffer )
+		{
+			*IOData = *Buffer;
+		}
+		
+		return noErr;
+	}
+	
+	EVoiceCaptureState::Type CopyBuffer(uint8 ReadBuffer, uint8* OutVoiceBuffer, uint32 ReadOffset, uint32 ReadLen)
+	{
+		EVoiceCaptureState::Type State = VoiceCaptureState;
+		if ( ReadLen )
+		{
+			uint8* Data = (uint8*)(BufferList[ReadBuffer].mData) + ReadOffset;
+			if ( StreamDesc.mSampleRate == OutputDesc.mSampleRate )
+			{
+				FMemory::Memcpy(OutVoiceBuffer, Data, ReadLen);
+				State = EVoiceCaptureState::Ok;
+			}
+			else
+			{
+				UInt32 FramesToCopy = ReadLen / OutputDesc.mBytesPerFrame;
+				AudioBufferList OutputBuffer;
+				OutputBuffer.mNumberBuffers = 1;
+				OutputBuffer.mBuffers[0].mNumberChannels = StreamDesc.mChannelsPerFrame;
+				OutputBuffer.mBuffers[0].mDataByteSize = ReadLen;
+				OutputBuffer.mBuffers[0].mData = OutVoiceBuffer;
+				
+				AudioBufferList InputBuffer;
+				InputBuffer.mNumberBuffers = 1;
+				InputBuffer.mBuffers[0].mNumberChannels = StreamDesc.mChannelsPerFrame;
+				InputBuffer.mBuffers[0].mDataByteSize = BufferList[ReadBuffer].mDataByteSize;
+				InputBuffer.mBuffers[0].mData = Data;
+				
+				if ( AudioConverterFillComplexBuffer(StreamConverter, &FVoiceCaptureCoreAudio::ConvertInputFormat, &InputBuffer, &FramesToCopy, &OutputBuffer, nullptr) == 0 )
+				{
+					State = EVoiceCaptureState::Ok;
+				}
+				else
+				{
+					State = EVoiceCaptureState::Error;
+					VoiceCaptureState = EVoiceCaptureState::Error;
+				}
+			}
+		}
+		return State;
+	}
+	
 	virtual EVoiceCaptureState::Type GetVoiceData(uint8* OutVoiceBuffer, uint32 InVoiceBufferSize, uint32& OutAvailableVoiceData) override
 	{
 		EVoiceCaptureState::Type State = VoiceCaptureState;
@@ -148,38 +422,49 @@ public:
 		if (VoiceCaptureState != EVoiceCaptureState::UnInitialized &&
 			VoiceCaptureState != EVoiceCaptureState::NotCapturing)
 		{
-			ALCint Samples = 0;
-			alcGetIntegerv(CaptureDevice, ALC_CAPTURE_SAMPLES, 1, &Samples);
-			ALenum Error = alcGetError(CaptureDevice);
-			if ( Error == AL_NO_ERROR )
+			GetCaptureState(OutAvailableVoiceData);
+			if ( OutAvailableVoiceData > 0 )
 			{
-				if ( Samples > 0 )
+				if ( InVoiceBufferSize >= OutAvailableVoiceData )
 				{
-					OutAvailableVoiceData = Samples * NumChannels * sizeof(uint16);
-					if ( InVoiceBufferSize >= OutAvailableVoiceData )
+					check(ReadOffset <= BufferSize);
+					uint32 BytesToRead = OutAvailableVoiceData;
+					uint8* Data = OutVoiceBuffer;
+					while (BytesToRead)
 					{
-						alcCaptureSamples(CaptureDevice, OutVoiceBuffer, Samples);
-						Error = alcGetError(CaptureDevice);
-						VoiceCaptureState = EVoiceCaptureState::NoData;
-						State = EVoiceCaptureState::Ok;
-					}
-					else
-					{
-						State = EVoiceCaptureState::BufferTooSmall;
+						uint32 CurrentRead = FMath::Min(FMath::Min(BufferList[ReadBuffer].mDataByteSize - ReadOffset, BytesToRead), BufferSize);
+						State = CopyBuffer(ReadBuffer, Data, ReadOffset, CurrentRead);
+						if ( State == EVoiceCaptureState::Ok )
+						{
+							Data += CurrentRead;
+							BytesToRead -= CurrentRead;
+							ReadOffset += CurrentRead;
+							if( ReadOffset == BufferList[ReadBuffer].mDataByteSize )
+							{
+								BufferList[ReadBuffer].mDataByteSize = 0;
+								FMemory::Memzero(BufferList[ReadBuffer].mData, BufferSize);
+								ReadBuffer = (ReadBuffer + 1) % BufferList.Num();
+								ReadOffset = 0;
+							}
+							check(ReadOffset <= BufferSize);
+						}
+						else
+						{
+							break;
+						}
 					}
+					int32 BytesRead = OutAvailableVoiceData;
+					FPlatformAtomics::InterlockedAdd(&ReadableBytes, -BytesRead);
 				}
 				else
 				{
-					VoiceCaptureState = EVoiceCaptureState::NoData;
-					State = EVoiceCaptureState::Ok;
+					State = EVoiceCaptureState::BufferTooSmall;
 				}
 			}
-			if ( Error != AL_NO_ERROR )
+			else
 			{
-				UE_LOG(LogVoiceCapture, Warning, TEXT("Failed to capture samples %d"), Error);
-				OutAvailableVoiceData = 0;
-				VoiceCaptureState = EVoiceCaptureState::Error;
-				State = EVoiceCaptureState::Error;
+				VoiceCaptureState = EVoiceCaptureState::NoData;
+				State = EVoiceCaptureState::Ok;
 			}
 		}
 		else
@@ -191,12 +476,32 @@ public:
 	}
 	
 private:
-	/** OpenAL device for capturing audio */
-	ALCdevice* CaptureDevice;
+	/** Descriptor for stream format */
+	AudioStreamBasicDescription StreamDesc;
+	/** Descriptor for native HW format */
+	AudioStreamBasicDescription NativeDesc;
+	/** Descriptor for HW output format */
+	AudioStreamBasicDescription OutputDesc;
+	/** Input AudioUnit component instance */
+	AudioComponentInstance StreamComponent;
+	/** Sample rate converter instance */
+	AudioConverterRef StreamConverter;
 	/** State of capture device */
-	mutable EVoiceCaptureState::Type VoiceCaptureState;
-	/** Initialised number of channels */
-	int32 NumChannels;
+	EVoiceCaptureState::Type VoiceCaptureState;
+	/** Outstanding CoreAudio buffer lists */
+	TArray<AudioBuffer> BufferList;
+	/** Input device buffer size */
+	uint32 BufferSize;
+	/** Input device latency */
+	uint32 InputLatency;
+	/** Current readable bytes */
+	int32 ReadableBytes;
+	/** Buffer CoreAudio should write to */
+	uint8 WriteBuffer;
+	/** Buffer game should read from */
+	uint8 ReadBuffer;
+	/** Current offset into ReadBuffer */
+	uint32 ReadOffset;
 };
 
 bool InitVoiceCapture()
@@ -210,7 +515,7 @@ void ShutdownVoiceCapture()
 
 IVoiceCapture* CreateVoiceCaptureObject()
 {
-	IVoiceCapture* Capture = new FVoiceCaptureOpenAL;
+	IVoiceCapture* Capture = new FVoiceCaptureCoreAudio;
 	if ( !Capture || !Capture->Init(VOICE_SAMPLE_RATE, NUM_VOICE_CHANNELS) )
 	{
 		delete Capture;
diff --git a/Engine/Source/Runtime/Online/Voice/Voice.Build.cs b/Engine/Source/Runtime/Online/Voice/Voice.Build.cs
--- a/Engine/Source/Runtime/Online/Voice/Voice.Build.cs
+++ b/Engine/Source/Runtime/Online/Voice/Voice.Build.cs
@@ -27,7 +27,7 @@ public class Voice : ModuleRules
 		}
 		else if(Target.Platform == UnrealTargetPlatform.Mac)
 		{
-			PublicFrameworks.AddRange(new string[] { "OpenAL" });
+			PublicFrameworks.AddRange(new string[] { "CoreAudio", "AudioToolbox", "AudioUnit" });
 		}
 
 		AddThirdPartyPrivateStaticDependencies(Target, "libOpus");
