Comment cleanup for GPS iap

diff --git a/Engine/Build/Android/Java/src/com/epicgames/ue4/GooglePlayStoreHelper.java b/Engine/Build/Android/Java/src/com/epicgames/ue4/GooglePlayStoreHelper.java
--- a/Engine/Build/Android/Java/src/com/epicgames/ue4/GooglePlayStoreHelper.java
+++ b/Engine/Build/Android/Java/src/com/epicgames/ue4/GooglePlayStoreHelper.java
@@ -110,7 +110,6 @@ public class GooglePlayStoreHelper extends StoreHelper
 					}
 					
 					// If we have this item, but somehow failed to consume it on purchase.
-					// make sure to consume it now. TPMB - TODO: make this part of the purchase workflow?
 					if (inventory.hasPurchase(productId))
 					{
 						try {
diff --git a/Engine/Source/Runtime/Online/Android/OnlineSubsystemGooglePlay/Private/OnlineStoreInterfaceGooglePlay.cpp b/Engine/Source/Runtime/Online/Android/OnlineSubsystemGooglePlay/Private/OnlineStoreInterfaceGooglePlay.cpp
--- a/Engine/Source/Runtime/Online/Android/OnlineSubsystemGooglePlay/Private/OnlineStoreInterfaceGooglePlay.cpp
+++ b/Engine/Source/Runtime/Online/Android/OnlineSubsystemGooglePlay/Private/OnlineStoreInterfaceGooglePlay.cpp
@@ -160,7 +160,7 @@ bool FOnlineStoreGooglePlay::BeginPurchase(const FInAppPurchaseProductRequest& P
 		CachedPurchaseStateObject = InPurchaseStateObject;
 
 		extern bool AndroidThunkCpp_Iap_BeginPurchase(const FString&, const bool);
-		bCreatedNewTransaction = AndroidThunkCpp_Iap_BeginPurchase(ProductRequest.ProductIdentifier, ProductRequest.bIsConsumable); // TPMB - update begin purchase to use a consumable flag.
+		bCreatedNewTransaction = AndroidThunkCpp_Iap_BeginPurchase(ProductRequest.ProductIdentifier, ProductRequest.bIsConsumable);
 		UE_LOG(LogOnline, Display, TEXT("Created Transaction? - %s"), 
 			bCreatedNewTransaction ? TEXT("Created a transaction.") : TEXT("Failed to create a transaction."));
 
