No longer saving out binary uobject manifest files
Fixing manifest uobject version so that we can continue to read already in use manifests

diff --git a/Engine/Source/Runtime/Online/BuildPatchServices/Private/BuildPatchGeneration.cpp b/Engine/Source/Runtime/Online/BuildPatchServices/Private/BuildPatchGeneration.cpp
--- a/Engine/Source/Runtime/Online/BuildPatchServices/Private/BuildPatchGeneration.cpp
+++ b/Engine/Source/Runtime/Online/BuildPatchServices/Private/BuildPatchGeneration.cpp
@@ -1268,17 +1268,12 @@ bool FBuildDataGenerator::GenerateChunksManifestFromDirectory( const FBuildPatch
 	}
 
 	// Save manifest into the cloud directory
-	FString BaseFilename = FBuildPatchServicesModule::GetCloudDirectory() / FDefaultValueHelper::RemoveWhitespaces(BuildManifest->Data->AppName + BuildManifest->Data->BuildVersion);
-	FString JsonFilename = BaseFilename + TEXT(".manifest");
-	FString BinaryFilename = BaseFilename + TEXT(".binary.manifest");
+	FString JsonFilename = FBuildPatchServicesModule::GetCloudDirectory() / FDefaultValueHelper::RemoveWhitespaces(BuildManifest->Data->AppName + BuildManifest->Data->BuildVersion) + TEXT(".manifest");
 	BuildManifest->Data->ManifestFileVersion = EBuildPatchAppManifestVersion::GetLatestJsonVersion();
 	BuildManifest->SaveToFile(JsonFilename, false);
-	BuildManifest->Data->ManifestFileVersion = EBuildPatchAppManifestVersion::GetLatestVersion();
-	BuildManifest->SaveToFile(BinaryFilename, true);
 
 	// Output to log for builder info
 	GLog->Logf(TEXT("Saved manifest to %s"), *JsonFilename);
-	GLog->Logf(TEXT("Saved manifest to %s"), *BinaryFilename);
 
 	// Clean up memory
 	delete[] DataBuffer;
@@ -1416,17 +1411,12 @@ bool FBuildDataGenerator::GenerateFilesManifestFromDirectory( const FBuildPatchS
 	}
 
 	// Save manifest into the cloud directory
-	FString BaseFilename = FBuildPatchServicesModule::GetCloudDirectory() / FDefaultValueHelper::RemoveWhitespaces(BuildManifest->Data->AppName + BuildManifest->Data->BuildVersion);
-	FString JsonFilename = BaseFilename + TEXT(".manifest");
-	FString BinaryFilename = BaseFilename + TEXT(".binary.manifest");
+	FString JsonFilename = FBuildPatchServicesModule::GetCloudDirectory() / FDefaultValueHelper::RemoveWhitespaces(BuildManifest->Data->AppName + BuildManifest->Data->BuildVersion) + TEXT(".manifest");
 	BuildManifest->Data->ManifestFileVersion = EBuildPatchAppManifestVersion::GetLatestJsonVersion();
 	BuildManifest->SaveToFile(JsonFilename, false);
-	BuildManifest->Data->ManifestFileVersion = EBuildPatchAppManifestVersion::GetLatestVersion();
-	BuildManifest->SaveToFile(BinaryFilename, true);
 
 	// Output to log for builder info
 	GLog->Logf(TEXT("Saved manifest to %s"), *JsonFilename);
-	GLog->Logf(TEXT("Saved manifest to %s"), *BinaryFilename);
 
 	// Clean up memory
 	delete[] FileReadBuffer;
diff --git a/Engine/Source/Runtime/Online/BuildPatchServices/Private/BuildPatchManifest.cpp b/Engine/Source/Runtime/Online/BuildPatchServices/Private/BuildPatchManifest.cpp
--- a/Engine/Source/Runtime/Online/BuildPatchServices/Private/BuildPatchManifest.cpp
+++ b/Engine/Source/Runtime/Online/BuildPatchServices/Private/BuildPatchManifest.cpp
@@ -633,6 +633,10 @@ bool FBuildPatchAppManifest::DeserializeFromData(const TArray<uint8>& DataInput)
 
 bool FBuildPatchAppManifest::Serialize(FArchive& Ar)
 {
+	// Make sure we use the correct serialization version, this is now fixed and must never use a newer version,
+	// because the property tag has changed in structure meaning older clients would not read correctly.
+	Ar.SetUE4Ver(VER_UE4_STRUCT_GUID_IN_PROPERTY_TAG - 1);
+
 	if (Ar.IsLoading())
 	{
 		DestroyData();
