[INTEGRATE] UE4-LauncherDev->UE4 integrate all up to cl 2475889

diff --git a/Engine/Source/Runtime/Online/HTTP/Private/Windows/HttpWinInet.cpp b/Engine/Source/Runtime/Online/HTTP/Private/Windows/HttpWinInet.cpp
--- a/Engine/Source/Runtime/Online/HTTP/Private/Windows/HttpWinInet.cpp
+++ b/Engine/Source/Runtime/Online/HTTP/Private/Windows/HttpWinInet.cpp
@@ -746,7 +746,7 @@ void FHttpRequestWinInet::Tick(float DeltaSeconds)
 		if(ResponseBytes > ProgressBytesSent)
 		{
 			ProgressBytesSent = ResponseBytes;
-			OnRequestProgress().ExecuteIfBound(SharedThis(this),ResponseBytes);
+			OnRequestProgress().ExecuteIfBound(SharedThis(this), RequestPayload.Num(), ResponseBytes);
 		}
 	}
 
@@ -779,7 +779,8 @@ void FHttpRequestWinInet::Tick(float DeltaSeconds)
 	// No longer waiting for a response and done processing it
 	else if (CompletionStatus == EHttpRequestStatus::Processing &&
 		Response.IsValid() &&
-		Response->bIsReady)
+		Response->bIsReady &&
+		TotalElapsed >= FHttpModule::Get().GetHttpDelayTime())
 	{
 		FinishedRequest();
 	}
diff --git a/Engine/Source/Runtime/Online/HTTP/Public/HttpModule.h b/Engine/Source/Runtime/Online/HTTP/Public/HttpModule.h
--- a/Engine/Source/Runtime/Online/HTTP/Public/HttpModule.h
+++ b/Engine/Source/Runtime/Online/HTTP/Public/HttpModule.h
@@ -134,6 +134,38 @@ public:
 		return bEnableHttp;
 	}
 
+	/**
+	 * toggle null http implementation
+	 */
+	inline void ToggleNullHttp(bool bEnabled)
+	{
+		bUseNullHttp = bEnabled;
+	}
+
+	/**
+	 * @return true if null http is being used
+	 */
+	inline bool IsNullHttpEnabled() const
+	{
+		return bUseNullHttp;
+	}
+
+	/**
+	 * @return min delay time for each http request
+	 */
+	inline float GetHttpDelayTime() const
+	{
+		return HttpDelayTime;
+	}
+
+	/**
+	 * Set the min delay time for each http request
+	 */
+	inline void SetHttpDelayTime(float InHttpDelayTime)
+	{
+		HttpDelayTime = InHttpDelayTime;
+	}
+
 private:
 
 	// IModuleInterface
@@ -161,12 +193,16 @@ private:
 	float HttpReceiveTimeout;
 	/** timeout in seconds to send a request on the connection. -1 for system defaults */
 	float HttpSendTimeout;
+	/** total time to delay the request */
+	float HttpDelayTime;
 	/** Max number of simultaneous connections to a specific server */
 	int32 HttpMaxConnectionsPerServer;
 	/** Max buffer size for individual http reads */
 	int32 MaxReadBufferSize;
 	/** toggles http requests */
 	bool bEnableHttp;
+	/** toggles null (mock) http requests */
+	bool bUseNullHttp;
 	/** singleton for the module while loaded and available */
 	static FHttpModule* Singleton;
 };
diff --git a/Engine/Source/Runtime/Online/HTTP/Public/HttpRequestAdapter.h b/Engine/Source/Runtime/Online/HTTP/Public/HttpRequestAdapter.h
new file mode 100644
--- /dev/null
+++ b/Engine/Source/Runtime/Online/HTTP/Public/HttpRequestAdapter.h
@@ -0,0 +1,54 @@
+// Copyright 1998-2015 Epic Games, Inc. All Rights Reserved.
+
+#pragma once
+
+#include "IHttpRequest.h"
+
+/** 
+  * Adapter class for IHttpRequest abstract interface
+  * does not fully expose the wrapped interface in the base. This allows client defined marshalling of the requests when end point permissions are at issue.
+  */
+
+class FHttpRequestAdapterBase
+{
+public:
+    FHttpRequestAdapterBase(TSharedRef<IHttpRequest>& InHttpRequest) : HttpRequest(InHttpRequest)
+    {}
+
+    FString                       GetURL()    	                                                   { return HttpRequest->GetURL(); }
+    FString                       GetURLParameter(const FString& ParameterName)                    { return HttpRequest->GetURLParameter(ParameterName); }
+    FString                       GetHeader(const FString& HeaderName)                             { return HttpRequest->GetHeader(HeaderName); }
+    TArray<FString>               GetAllHeaders()                                                  { return HttpRequest->GetAllHeaders(); }
+    FString                       GetContentType()                                                 { return HttpRequest->GetContentType(); }
+    int32                         GetContentLength()                                               { return HttpRequest->GetContentLength(); }
+    const TArray<uint8>&          GetContent()                                                     { return HttpRequest->GetContent(); }
+
+    FString                       GetVerb() 						                               { return HttpRequest->GetVerb(); }
+    void                          SetVerb(const FString& Verb) 	                                   { HttpRequest->SetVerb(Verb); }
+    void                          SetURL(const FString& URL) 		                               { HttpRequest->SetURL(URL); }
+    void                          SetContent(const TArray<uint8>& ContentPayload)                  { HttpRequest->SetContent(ContentPayload); }
+    void                          SetContentAsString(const FString& ContentString)                 { HttpRequest->SetContentAsString(ContentString); }
+    void                          SetHeader(const FString& HeaderName, const FString& HeaderValue) { HttpRequest->SetHeader(HeaderName, HeaderValue); }
+    const FHttpResponsePtr        GetResponse() const                                              { return HttpRequest->GetResponse(); }
+    float                         GetElapsedTime()                                                 { return HttpRequest->GetElapsedTime(); }
+
+protected:
+    TSharedRef<IHttpRequest> HttpRequest;
+};
+
+class FHttpRequestAdapter : public FHttpRequestAdapterBase
+{
+public:
+    FHttpRequestAdapter(TSharedRef<IHttpRequest>& InHttpRequest)
+        : FHttpRequestAdapterBase(InHttpRequest)
+    {}
+
+    FHttpRequestCompleteDelegate&   OnProcessRequestComplete()                                 { return HttpRequest->OnProcessRequestComplete(); }
+    FHttpRequestProgressDelegate&   OnRequestProgress()                                        { return HttpRequest->OnRequestProgress(); }
+
+    EHttpRequestStatus::Type        GetStatus()                                                { return HttpRequest->GetStatus(); }
+
+    bool                            ProcessRequest()                                           { return HttpRequest->ProcessRequest(); }
+    void                            CancelRequest()                                            { HttpRequest->CancelRequest(); }
+    void                            Tick(float DeltaSeconds)                                   { HttpRequest->Tick(DeltaSeconds); }
+};
\ No newline at end of file
diff --git a/Engine/Source/Runtime/Online/HTTP/Public/HttpRetrySystem.h b/Engine/Source/Runtime/Online/HTTP/Public/HttpRetrySystem.h
new file mode 100644
--- /dev/null
+++ b/Engine/Source/Runtime/Online/HTTP/Public/HttpRetrySystem.h
@@ -0,0 +1,160 @@
+// Copyright 1998-2015 Epic Games, Inc. All Rights Reserved.
+
+#pragma once
+
+#include "HttpRequestAdapter.h"
+
+/**
+ * Helpers of various types for the retry system
+ */
+namespace FHttpRetrySystem
+{
+    typedef uint32 RetryLimitCountType;
+    typedef double RetryTimeoutRelativeSecondsType;
+
+    inline RetryLimitCountType             RetryLimitCount(uint32 Value)             { return Value; }
+    inline RetryTimeoutRelativeSecondsType RetryTimeoutRelativeSeconds(double Value) { return Value; }
+
+    template <typename  IntrinsicType>
+    IntrinsicType TZero();
+
+    template <> inline float                           TZero<float>()                           { return 0.0f; }
+    template <> inline RetryLimitCountType             TZero<RetryLimitCountType>()             { return RetryLimitCount(0); }
+    template <> inline RetryTimeoutRelativeSecondsType TZero<RetryTimeoutRelativeSecondsType>() { return RetryTimeoutRelativeSeconds(0.0); }
+
+    /**
+     * TOptionalSetting merges a bool and an intrinsic value to remove the need
+     * for having special values to indicate if the option is valid
+     */
+    template <typename IntrinsicType>
+    struct TOptionalSetting
+    {
+        TOptionalSetting() : bUseValue(false), Value(TZero<IntrinsicType>())                          {}
+        explicit TOptionalSetting(IntrinsicType InValue) : bUseValue(true), Value(InValue)                     {}
+        TOptionalSetting(const TOptionalSetting& InValue) : bUseValue(InValue.bUseValue), Value(InValue.Value) {}
+
+        static TOptionalSetting Unused()                       { return TOptionalSetting(); }
+        static TOptionalSetting Create(IntrinsicType InValue)  { return TOptionalSetting(InValue); }
+
+        bool            bUseValue;
+        IntrinsicType   Value;
+    };
+
+    typedef TOptionalSetting<float>                           FRandomFailureRateSetting;
+    typedef TOptionalSetting<RetryLimitCountType>             FRetryLimitCountSetting;
+    typedef TOptionalSetting<RetryTimeoutRelativeSecondsType> FRetryTimeoutRelativeSecondsSetting;
+};
+
+
+namespace FHttpRetrySystem
+{
+    /**
+     * class FRequest is what the retry system accepts as inputs
+     */
+    class FRequest : public FHttpRequestAdapterBase
+    {
+    public:
+        struct EStatus
+        {
+            enum Type
+            {
+                NotStarted = 0,
+                Processing,
+                ProcessingLockout,
+                Cancelled,
+                FailedRetry,
+                FailedTimeout,
+                Succeeded
+            };
+        };
+
+        /**
+        * Delegate called when an FRequest completes
+        *
+        * @param first parameter -  original FRequest
+        * @param second parameter - indicates whether or not the request completed successfully
+        */
+        DECLARE_DELEGATE_TwoParams(FOnProcessRequestCompleteDelegate, TSharedRef<class FRequest>&, bool);
+
+    public:
+        HTTP_API FRequest(TSharedRef<IHttpRequest>& HttpRequest,
+                          FRetryLimitCountSetting InRetryLimitCountOverride = FRetryLimitCountSetting::Unused(),
+                          FRetryTimeoutRelativeSecondsSetting InRetryTimeoutRelativeSecondsOverride = FRetryTimeoutRelativeSecondsSetting::Unused());
+
+        EStatus::Type                      GetStatus()                { return Status; }
+        FOnProcessRequestCompleteDelegate& OnProcessRequestComplete() { return OnProcessRequestCompleteDelegate; }
+
+    protected:
+        friend class FManager;
+
+        bool ProcessRequest() { return HttpRequest->ProcessRequest(); }
+        void CancelRequest()  { HttpRequest->CancelRequest(); }
+
+        EStatus::Type                        Status;
+
+        FRetryLimitCountSetting              RetryLimitCountOverride;
+        FRetryTimeoutRelativeSecondsSetting  RetryTimeoutRelativeSecondsOverride;
+
+        FOnProcessRequestCompleteDelegate    OnProcessRequestCompleteDelegate;
+    };
+}
+
+namespace FHttpRetrySystem
+{
+    class FManager
+    {
+    public:
+        // FManager
+        HTTP_API FManager(FRetryLimitCountSetting InRetryLimitCountDefault, FRetryTimeoutRelativeSecondsSetting InRetryTimeoutRelativeSecondsDefault);
+
+        /**
+         * Updates the entries in the list of retry requests. Optional parameters are for future connection health assessment
+         *
+         * @param FileCount       optional parameter that will be filled with the total files updated
+         * @param FailingCount    optional parameter that will be filled with the total files that have are in a retrying state
+         * @param FailedCount     optional parameter that will be filled with the total files that have failed
+         * @param CompletedCount  optional parameter that will be filled with the total files that have completed
+         *
+         * @return                true if there are no failures or retries
+         */
+        HTTP_API bool Update(uint32* FileCount = NULL, uint32* FailingCount = NULL, uint32* FailedCount = NULL, uint32* CompletedCount = NULL);
+
+        HTTP_API bool ProcessRequest(TSharedRef<FRequest>& HttpRequest);
+        HTTP_API void CancelRequest (TSharedRef<FRequest>& HttpRequest);
+
+        void SetRandomFailureRate(float Value) { RandomFailureRate = FRandomFailureRateSetting::Create(Value); }
+
+    protected:
+        struct FHttpRetryRequestEntry
+        {
+            FHttpRetryRequestEntry(TSharedRef<FRequest>& HttpRequestValue);
+
+            bool                    bShouldCancel;
+            uint32                  CurrentRetryCount;
+            double                  RequestStartTimeAbsoluteSeconds;
+            double                  LockoutEndTimeAbsoluteSeconds;
+
+            TSharedRef<FRequest>	HttpRequest;
+        };
+
+        // @return true if there is a no formal response to the request
+        // @TODO return true if a variety of 5xx errors are the result of a formal response
+        bool ShouldRetry(const FHttpRetryRequestEntry& HttpRetryRequestEntry);
+
+        // @return true if retry chances have not been exhausted
+        bool CanRetry(const FHttpRetryRequestEntry& HttpRetryRequestEntry);
+
+        // @return true if the retry request has timed out
+        bool HasTimedOut(const FHttpRetryRequestEntry& HttpRetryRequestEntry, const double NowAbsoluteSeconds);
+
+        // @return number of seconds to lockout for
+        float GetLockoutPeriodSeconds(const FHttpRetryRequestEntry& HttpRetryRequestEntry);
+
+        // Default configuration for the retry system
+        FRandomFailureRateSetting            RandomFailureRate;
+        FRetryLimitCountSetting              RetryLimitCountDefault;
+        FRetryTimeoutRelativeSecondsSetting  RetryTimeoutRelativeSecondsDefault;
+
+        TArray<FHttpRetryRequestEntry>        RequestList;
+    };
+}
\ No newline at end of file
diff --git a/Engine/Source/Runtime/Online/HTTP/Public/Interfaces/IHttpRequest.h b/Engine/Source/Runtime/Online/HTTP/Public/Interfaces/IHttpRequest.h
--- a/Engine/Source/Runtime/Online/HTTP/Public/Interfaces/IHttpRequest.h
+++ b/Engine/Source/Runtime/Online/HTTP/Public/Interfaces/IHttpRequest.h
@@ -60,12 +60,13 @@ typedef TSharedPtr<class IHttpResponse,ESPMode::ThreadSafe> FHttpResponsePtr;
 DECLARE_DELEGATE_ThreeParams(FHttpRequestCompleteDelegate, FHttpRequestPtr, FHttpResponsePtr, bool);
 
 /**
- * Delegate called per tick to update an Http request download size progress
+ * Delegate called per tick to update an Http request upload or download size progress
  *
  * @param first parameter - original Http request that started things
- * @param second parameter - the number of bytes downloaded so far.
+ * @param second parameter - the number of bytes sent / uploaded in the request so far.
+ * @param third parameter - the number of bytes received / downloaded in the response so far.
  */
-DECLARE_DELEGATE_TwoParams(FHttpRequestProgressDelegate, FHttpRequestPtr, int32);
+DECLARE_DELEGATE_ThreeParams(FHttpRequestProgressDelegate, FHttpRequestPtr, int32, int32);
 
 /**
  * Interface for Http requests (created using FHttpFactory)
@@ -142,7 +143,7 @@ public:
 	virtual FHttpRequestCompleteDelegate& OnProcessRequestComplete() = 0;
 
 	/**
-	 * Delegate called to update the request progress. See FHttpRequestProgressDelegate
+	 * Delegate called to update the request/response progress. See FHttpRequestProgressDelegate
 	 */
 	virtual FHttpRequestProgressDelegate& OnRequestProgress() = 0;
 
diff --git a/Engine/Source/Runtime/Online/IOS/OnlineSubsystemIOS/Private/OnlineSubsystemIOS.cpp b/Engine/Source/Runtime/Online/IOS/OnlineSubsystemIOS/Private/OnlineSubsystemIOS.cpp
--- a/Engine/Source/Runtime/Online/IOS/OnlineSubsystemIOS/Private/OnlineSubsystemIOS.cpp
+++ b/Engine/Source/Runtime/Online/IOS/OnlineSubsystemIOS/Private/OnlineSubsystemIOS.cpp
@@ -36,6 +36,10 @@ IOnlineUserCloudPtr FOnlineSubsystemIOS::GetUserCloudInterface() const
 	return NULL;
 }
 
+IOnlineUserCloudPtr FOnlineSubsystemIOS::GetUserCloudInterface(const FString& Key) const
+{
+	return NULL;
+}
 
 IOnlineLeaderboardsPtr FOnlineSubsystemIOS::GetLeaderboardsInterface() const
 {
diff --git a/Engine/Source/Runtime/Online/IOS/OnlineSubsystemIOS/Public/OnlineSubsystemIOS.h b/Engine/Source/Runtime/Online/IOS/OnlineSubsystemIOS/Public/OnlineSubsystemIOS.h
--- a/Engine/Source/Runtime/Online/IOS/OnlineSubsystemIOS/Public/OnlineSubsystemIOS.h
+++ b/Engine/Source/Runtime/Online/IOS/OnlineSubsystemIOS/Public/OnlineSubsystemIOS.h
@@ -24,6 +24,7 @@ public:
 	virtual IOnlineGroupsPtr GetGroupsInterface() const override;
 	virtual IOnlineSharedCloudPtr GetSharedCloudInterface() const override;
 	virtual IOnlineUserCloudPtr GetUserCloudInterface() const override;
+	virtual IOnlineUserCloudPtr GetUserCloudInterface(const FString& Key) const override;
 	virtual IOnlineLeaderboardsPtr GetLeaderboardsInterface() const override;
 	virtual IOnlineVoicePtr GetVoiceInterface() const  override;
 	virtual IOnlineExternalUIPtr GetExternalUIInterface() const override;
diff --git a/Engine/Source/Runtime/Online/OnlineSubsystem/Public/Interfaces/OnlineChatInterface.h b/Engine/Source/Runtime/Online/OnlineSubsystem/Public/Interfaces/OnlineChatInterface.h
--- a/Engine/Source/Runtime/Online/OnlineSubsystem/Public/Interfaces/OnlineChatInterface.h
+++ b/Engine/Source/Runtime/Online/OnlineSubsystem/Public/Interfaces/OnlineChatInterface.h
@@ -18,7 +18,7 @@ public:
 	virtual ~FChatRoomInfo() {}
 
 	virtual const FChatRoomId& GetRoomId() const = 0;
-	virtual const FUniqueNetId& GetOwnerId() const = 0;
+	virtual const TSharedRef<FUniqueNetId>& GetOwnerId() const = 0;
 	virtual const FString& GetSubject() const = 0;
 	virtual bool IsPrivate() const = 0;
 	virtual bool IsJoined() const = 0;
@@ -32,7 +32,7 @@ class FChatRoomMember
 public:
 	virtual ~FChatRoomMember() {}
 
-	virtual const FUniqueNetId& GetUserId() const = 0;
+	virtual const TSharedRef<FUniqueNetId>& GetUserId() const = 0;
 	virtual const FString& GetNickname() const = 0;
 };
 
@@ -44,7 +44,7 @@ class FChatMessage
 public:
 	virtual ~FChatMessage() {}
 
-	virtual const FUniqueNetId& GetUserId() const = 0;
+	virtual const TSharedRef<FUniqueNetId>& GetUserId() const = 0;
 	virtual const FString& GetNickname() const = 0;
 	virtual const FString& GetBody() const = 0;
 	virtual const FDateTime& GetTimestamp() const = 0;
diff --git a/Engine/Source/Runtime/Online/OnlineSubsystem/Public/Interfaces/OnlineEntitlementsInterface.h b/Engine/Source/Runtime/Online/OnlineSubsystem/Public/Interfaces/OnlineEntitlementsInterface.h
--- a/Engine/Source/Runtime/Online/OnlineSubsystem/Public/Interfaces/OnlineEntitlementsInterface.h
+++ b/Engine/Source/Runtime/Online/OnlineSubsystem/Public/Interfaces/OnlineEntitlementsInterface.h
@@ -20,6 +20,8 @@ struct FOnlineEntitlement
 	FString Name;
 	/** Id for the item that this entitlement is associated with */
 	FString ItemId;
+	/**  Namespace of the entitlement */
+	FString Namespace;
 
 	/**
 	 * @return Any additional data associated with the entitlement
@@ -44,9 +46,10 @@ struct FOnlineEntitlement
  *
  * @param bWasSuccessful true if server was contacted and a valid result received
  * @param UserId of the user who was granted entitlements in this callback
+ * @param Namespace optional namespace that was used to query. Empty means all entitlements were queried
  * @param Error string representing the error condition
  */
-DECLARE_MULTICAST_DELEGATE_ThreeParams(FOnQueryEntitlementsComplete, bool, const FUniqueNetId&, const FString&);
+DECLARE_MULTICAST_DELEGATE_FourParams(FOnQueryEntitlementsComplete, bool, const FUniqueNetId&, const FString&, const FString&);
 typedef FOnQueryEntitlementsComplete::FDelegate FOnQueryEntitlementsCompleteDelegate;
 
 /**
@@ -79,27 +82,30 @@ public:
 	 * Gets the cached entitlement set for the requested user
 	 *
 	 * @param UserId the ID of the user to get entitlements for
+	 * @param Namespace optional namespace to filter on
 	 * @param OutUserEntitlements out parameter to copy the user's entitlements into
 	 */
-	virtual void GetAllEntitlements(const FUniqueNetId& UserId, TArray<TSharedRef<FOnlineEntitlement>>& OutUserEntitlements) = 0;
+	virtual void GetAllEntitlements(const FUniqueNetId& UserId, const FString& Namespace, TArray<TSharedRef<FOnlineEntitlement>>& OutUserEntitlements) = 0;
 
 	/**
 	 * Contacts server and retrieves the list of the user's entitlements, caching them locally
 	 *
 	 * @param UserId the ID of the user to act on
+	 * @param Namespace optional namespace to filter on
 	 *
 	 * @return true if the operation started successfully
 	 */
-	virtual bool QueryEntitlements(const FUniqueNetId& UserId, const FPagedQuery& Page = FPagedQuery()) = 0;
+	virtual bool QueryEntitlements(const FUniqueNetId& UserId, const FString& Namespace, const FPagedQuery& Page = FPagedQuery()) = 0;
 	
 	/**
 	 * Delegate instanced called when enumerating entitlements has completed
 	 *
 	 * @param bWasSuccessful true if server was contacted and a valid result received
 	 * @param UserId of the user who was granted entitlements in this callback
+	 * @param Namespace optional namespace that was used to query. Empty means all entitlements were queried
 	 * @param Error string representing the error condition
 	 */
-	DEFINE_ONLINE_DELEGATE_THREE_PARAM(OnQueryEntitlementsComplete, bool, const FUniqueNetId&, const FString&);
+	DEFINE_ONLINE_DELEGATE_FOUR_PARAM(OnQueryEntitlementsComplete, bool, const FUniqueNetId&, const FString&, const FString&);
 };
 
 typedef TSharedPtr<IOnlineEntitlements, ESPMode::ThreadSafe> IOnlineEntitlementsPtr;
diff --git a/Engine/Source/Runtime/Online/OnlineSubsystem/Public/Interfaces/OnlinePresenceInterface.h b/Engine/Source/Runtime/Online/OnlineSubsystem/Public/Interfaces/OnlinePresenceInterface.h
--- a/Engine/Source/Runtime/Online/OnlineSubsystem/Public/Interfaces/OnlinePresenceInterface.h
+++ b/Engine/Source/Runtime/Online/OnlineSubsystem/Public/Interfaces/OnlinePresenceInterface.h
@@ -24,7 +24,7 @@ const FString DefaultSessionIdKey = "SessionId";
 
 namespace EOnlinePresenceState
 {
-	enum Type
+	enum Type : uint8
 	{
 		Online,
 		Offline,
diff --git a/Engine/Source/Runtime/Online/OnlineSubsystem/Public/Interfaces/OnlineSessionInterface.h b/Engine/Source/Runtime/Online/OnlineSubsystem/Public/Interfaces/OnlineSessionInterface.h
--- a/Engine/Source/Runtime/Online/OnlineSubsystem/Public/Interfaces/OnlineSessionInterface.h
+++ b/Engine/Source/Runtime/Online/OnlineSubsystem/Public/Interfaces/OnlineSessionInterface.h
@@ -179,9 +179,10 @@ typedef FOnSessionUserInviteAccepted::FDelegate FOnSessionUserInviteAcceptedDele
  *
  * @param UserId the user being invited
  * @param FromId the user that sent the invite
+ * @param AppId the id of the client/app user was in when sending hte game invite
  * @param InviteResult the search/settings for the session we're joining via invite
  */
-DECLARE_MULTICAST_DELEGATE_ThreeParams(FOnSessionInviteReceived, const FUniqueNetId& /*UserId*/, const FUniqueNetId& /*FromId*/, const FOnlineSessionSearchResult& /*InviteResult*/);
+DECLARE_MULTICAST_DELEGATE_FourParams(FOnSessionInviteReceived, const FUniqueNetId& /*UserId*/, const FUniqueNetId& /*FromId*/, const FString& /*AppId*/, const FOnlineSessionSearchResult& /*InviteResult*/);
 typedef FOnSessionInviteReceived::FDelegate FOnSessionInviteReceivedDelegate;
 
 /**
@@ -660,9 +661,10 @@ public:
 	 *
 	 * @param UserId the user being invited
 	 * @param FromId the user that sent the invite
+	 * @param AppId the id of the client/app user was in when sending hte game invite
 	 * @param InviteResult the search/settings for the session we're joining via invite
 	 */
-	DEFINE_ONLINE_DELEGATE_THREE_PARAM(OnSessionInviteReceived, const FUniqueNetId& /*UserId*/, const FUniqueNetId& /*FromId*/, const FOnlineSessionSearchResult& /*InviteResult*/);
+	DEFINE_ONLINE_DELEGATE_FOUR_PARAM(OnSessionInviteReceived, const FUniqueNetId& /*UserId*/, const FUniqueNetId& /*FromId*/, const FString& /*AppId*/, const FOnlineSessionSearchResult& /*InviteResult*/);
 
 	/**
 	 * Returns the platform specific connection information for joining the match.
diff --git a/Engine/Source/Runtime/Online/OnlineSubsystem/Public/Interfaces/OnlineUserCloudInterface.h b/Engine/Source/Runtime/Online/OnlineSubsystem/Public/Interfaces/OnlineUserCloudInterface.h
--- a/Engine/Source/Runtime/Online/OnlineSubsystem/Public/Interfaces/OnlineUserCloudInterface.h
+++ b/Engine/Source/Runtime/Online/OnlineSubsystem/Public/Interfaces/OnlineUserCloudInterface.h
@@ -13,6 +13,17 @@
 DECLARE_MULTICAST_DELEGATE_TwoParams(FOnEnumerateUserFilesComplete, bool, const FUniqueNetId&);
 typedef FOnEnumerateUserFilesComplete::FDelegate FOnEnumerateUserFilesCompleteDelegate;
 
+/**
+ * Delegate fired at intervals during a user file write to the network platform's storage
+ *
+ * @param BytesWritten the number of bytes written so far
+ * @param UserId User owning the storage
+ * @param FileName the name of the file this was for
+ *
+ */
+DECLARE_MULTICAST_DELEGATE_ThreeParams(FOnWriteUserFileProgress, int32, const FUniqueNetId&, const FString&);
+typedef FOnWriteUserFileProgress::FDelegate FOnWriteUserFileProgressDelegate;
+
 /**
  * Delegate fired when a user file write to the network platform's storage is complete
  *
@@ -24,6 +35,17 @@ typedef FOnEnumerateUserFilesComplete::FDelegate FOnEnumerateUserFilesCompleteDe
 DECLARE_MULTICAST_DELEGATE_ThreeParams(FOnWriteUserFileComplete, bool, const FUniqueNetId&, const FString&);
 typedef FOnWriteUserFileComplete::FDelegate FOnWriteUserFileCompleteDelegate;
 
+/**
+ * Delegate fired when a user file write to the network platform's storage is canceled
+ *
+ * @param bWasSuccessful whether the cancellation was successful or not
+ * @param UserId User owning the storage
+ * @param FileName the name of the file this was for
+ *
+ */
+DECLARE_MULTICAST_DELEGATE_ThreeParams(FOnWriteUserFileCanceled, bool, const FUniqueNetId&, const FString&);
+typedef FOnWriteUserFileCanceled::FDelegate FOnWriteUserFileCanceledDelegate;
+
 /**
  * Delegate fired when a user file read from the network platform's storage is complete
  *
@@ -45,6 +67,17 @@ typedef FOnReadUserFileComplete::FDelegate FOnReadUserFileCompleteDelegate;
 DECLARE_MULTICAST_DELEGATE_ThreeParams(FOnDeleteUserFileComplete, bool, const FUniqueNetId&, const FString&);
 typedef FOnDeleteUserFileComplete::FDelegate FOnDeleteUserFileCompleteDelegate;
 
+/**
+ * Delegate fired when getting usage statistics from the network platform's storage is complete
+ *
+ * @param bWasSuccessful  Whether the operation was successful or not
+ * @param UserId          User owning the storage
+ * @param BytesUsed       Will be set to the total number of bytes used by the user
+ * @param TotalQuota      The total number of bytes the user is entitled to write
+ */
+DECLARE_MULTICAST_DELEGATE_FourParams(FOnRequestUsageInfoComplete, bool, const FUniqueNetId&, int64, TOptional<int64>);
+typedef FOnRequestUsageInfoComplete::FDelegate FOnRequestUsageInfoCompleteDelegate;
+
 
 /**
  * Provides access to per user cloud file storage
@@ -145,6 +178,16 @@ public:
 	 */
 	virtual bool WriteUserFile(const FUniqueNetId& UserId, const FString& FileName, TArray<uint8>& FileContents) = 0;
 
+	/**
+	* Delegate fired at intervals during a user file write to the network platform's storage
+	*
+	* @param BytesWritten the number of bytes written so far
+	* @param UserId User owning the storage
+	* @param FileName the name of the file this was for
+	*
+	*/
+	DEFINE_ONLINE_DELEGATE_THREE_PARAM(OnWriteUserFileProgress, int32, const FUniqueNetId&, const FString&);
+
 	/**
 	 * Delegate fired when a user file write to the network platform's storage is complete
 	 *
@@ -155,6 +198,24 @@ public:
 	 */
 	DEFINE_ONLINE_DELEGATE_THREE_PARAM(OnWriteUserFileComplete, bool, const FUniqueNetId&, const FString&);
 
+	/**
+	 * Cancel the ongoing upload of the specified file, if it is in progress
+	 * @param UserId User owning the storage
+	 * @param FileNam the name of the file to cancel
+	 *
+	 */
+	virtual void CancelWriteUserFile(const FUniqueNetId& UserId, const FString& FileName) = 0;
+
+	/**
+	 * Delegate fired when a user file write to the network platform's storage is canceled
+	 *
+	 * @param bWasSuccessful whether the file upload cancellations was successful or not
+	 * @param UserId User owning the storage
+	 * @param FileName the name of the file this was for
+	 *
+	 */
+	DEFINE_ONLINE_DELEGATE_THREE_PARAM(OnWriteUserFileCanceled, bool, const FUniqueNetId&, const FString&);
+
 	/**
 	 * Starts an asynchronous delete of the specified user file from the network platform's file store
 	 *
@@ -176,6 +237,25 @@ public:
 	 */
 	DEFINE_ONLINE_DELEGATE_THREE_PARAM(OnDeleteUserFileComplete, bool, const FUniqueNetId&, const FString&);
 
+	/**
+	 * Starts an asynchronous request to get the usage statistics from the cloud storage service
+	 *
+	 * @param UserId      IN      User owning the storage
+	 *
+	 * @return true if the call starts successfully, false otherwise
+	 */
+	virtual bool RequestUsageInfo(const FUniqueNetId& UserId) = 0;
+
+	/**
+	 * Delegate fired when getting usage statistics from the network platform's storage is complete
+	 *
+	 * @param bWasSuccessful  Whether the operation was successful or not
+	 * @param UserId          User owning the storage
+	 * @param BytesUsed       Will be set to the total number of bytes used by the user
+	 * @param TotalQuota      The total number of bytes the user is entitled to write
+	 */
+	DEFINE_ONLINE_DELEGATE_FOUR_PARAM(OnRequestUsageInfoComplete, bool, const FUniqueNetId&, int64, TOptional<int64>);
+
 	/**
 	 *	Print out the state of the cloud for this service
 	 *
diff --git a/Engine/Source/Runtime/Online/OnlineSubsystem/Public/OnlineSubsystem.h b/Engine/Source/Runtime/Online/OnlineSubsystem/Public/OnlineSubsystem.h
--- a/Engine/Source/Runtime/Online/OnlineSubsystem/Public/OnlineSubsystem.h
+++ b/Engine/Source/Runtime/Online/OnlineSubsystem/Public/OnlineSubsystem.h
@@ -58,9 +58,10 @@ typedef TSharedPtr<class FOnlineNotificationTransportManager, ESPMode::ThreadSaf
 /**
  * Called when the connection state as reported by the online platform changes
  *
- * @param ConnectionState state of the connection
+ * @param LastConnectionState last state of the connection
+ * @param ConnectionState current state of the connection
  */
-DECLARE_MULTICAST_DELEGATE_OneParam(FOnConnectionStatusChanged, EOnlineServerConnectionStatus::Type);
+DECLARE_MULTICAST_DELEGATE_TwoParams(FOnConnectionStatusChanged, EOnlineServerConnectionStatus::Type /*LastConnectionState*/, EOnlineServerConnectionStatus::Type /*ConnectionState*/);
 typedef FOnConnectionStatusChanged::FDelegate FOnConnectionStatusChangedDelegate;
 
 /**
@@ -156,13 +157,20 @@ public:
 	 */
 	virtual IOnlineSharedCloudPtr GetSharedCloudInterface() const = 0;
 
-	/** 
-	 * Get the interface for accessing user files in the cloud
-	 * @return Interface pointer for the appropriate cloud service
-	 */
+	/**
+	* Get the interface for accessing user files in the cloud
+	* @return Interface pointer for the appropriate cloud service
+	*/
 	virtual IOnlineUserCloudPtr GetUserCloudInterface() const = 0;
 
-	/** 
+	/**
+	* Get the interface for accessing user files in the cloud for a specific service
+	* @param Key   The key for the required user cloud interface
+	* @return      Interface pointer for the appropriate cloud service
+	*/
+	virtual IOnlineUserCloudPtr GetUserCloudInterface(const FString& Key) const = 0;
+
+	/**
 	 * Get the interface for accessing user entitlements
 	 * @return Interface pointer for the appropriate entitlements service
 	 */
@@ -362,9 +370,10 @@ public:
 	/**
 	 * Called when the connection state as reported by the online platform changes
 	 *
-	 * @param ConnectionState state of the connection
+	 * @param LastConnectionState last state of the connection
+	 * @param ConnectionState current state of the connection
 	 */
-	DEFINE_ONLINE_DELEGATE_ONE_PARAM(OnConnectionStatusChanged, EOnlineServerConnectionStatus::Type);
+	DEFINE_ONLINE_DELEGATE_TWO_PARAM(OnConnectionStatusChanged, EOnlineServerConnectionStatus::Type /*LastConnectionState*/, EOnlineServerConnectionStatus::Type /*ConnectionState*/);
 };
 
 /** Public references to the online subsystem pointer should use this */
diff --git a/Engine/Source/Runtime/Online/OnlineSubsystemAmazon/Public/OnlineSubsystemAmazon.h b/Engine/Source/Runtime/Online/OnlineSubsystemAmazon/Public/OnlineSubsystemAmazon.h
--- a/Engine/Source/Runtime/Online/OnlineSubsystemAmazon/Public/OnlineSubsystemAmazon.h
+++ b/Engine/Source/Runtime/Online/OnlineSubsystemAmazon/Public/OnlineSubsystemAmazon.h
@@ -53,6 +53,10 @@ public:
 	{
 		return NULL;
 	}
+	virtual IOnlineUserCloudPtr GetUserCloudInterface(const FString& Key) const override
+	{
+		return NULL;
+	}
 	virtual IOnlineLeaderboardsPtr GetLeaderboardsInterface() const override
 	{
 		return NULL;
diff --git a/Engine/Source/Runtime/Online/OnlineSubsystemFacebook/Private/IOS/OnlineSubsystemFacebook.cpp b/Engine/Source/Runtime/Online/OnlineSubsystemFacebook/Private/IOS/OnlineSubsystemFacebook.cpp
--- a/Engine/Source/Runtime/Online/OnlineSubsystemFacebook/Private/IOS/OnlineSubsystemFacebook.cpp
+++ b/Engine/Source/Runtime/Online/OnlineSubsystemFacebook/Private/IOS/OnlineSubsystemFacebook.cpp
@@ -53,6 +53,11 @@ IOnlineUserCloudPtr FOnlineSubsystemFacebook::GetUserCloudInterface() const
 	return NULL;
 }
 
+IOnlineUserCloudPtr FOnlineSubsystemFacebook::GetUserCloudInterface(const FString& Key) const
+{
+	return NULL;
+}
+
 IOnlineLeaderboardsPtr FOnlineSubsystemFacebook::GetLeaderboardsInterface() const
 {
 	return NULL;
diff --git a/Engine/Source/Runtime/Online/OnlineSubsystemFacebook/Private/Windows/OnlineSubsystemFacebook.cpp b/Engine/Source/Runtime/Online/OnlineSubsystemFacebook/Private/Windows/OnlineSubsystemFacebook.cpp
--- a/Engine/Source/Runtime/Online/OnlineSubsystemFacebook/Private/Windows/OnlineSubsystemFacebook.cpp
+++ b/Engine/Source/Runtime/Online/OnlineSubsystemFacebook/Private/Windows/OnlineSubsystemFacebook.cpp
@@ -35,6 +35,10 @@ IOnlineUserCloudPtr FOnlineSubsystemFacebook::GetUserCloudInterface() const
 	return NULL;
 }
 
+IOnlineUserCloudPtr FOnlineSubsystemFacebook::GetUserCloudInterface(const FString& Key) const
+{
+	return NULL;
+}
 
 IOnlineLeaderboardsPtr FOnlineSubsystemFacebook::GetLeaderboardsInterface() const
 {
diff --git a/Engine/Source/Runtime/Online/OnlineSubsystemNull/Private/OnlineIdentityNull.cpp b/Engine/Source/Runtime/Online/OnlineSubsystemNull/Private/OnlineIdentityNull.cpp
--- a/Engine/Source/Runtime/Online/OnlineSubsystemNull/Private/OnlineIdentityNull.cpp
+++ b/Engine/Source/Runtime/Online/OnlineSubsystemNull/Private/OnlineIdentityNull.cpp
@@ -53,7 +53,7 @@ bool FOnlineIdentityNull::Login(int32 LocalUserNum, const FOnlineAccountCredenti
 	}
 	else if (AccountCredentials.Id.IsEmpty())
 	{
-		ErrorStr = FString::Printf(TEXT("Invalid account id=%s"), *AccountCredentials.Id);
+		ErrorStr = TEXT("Invalid account id, string empty");
 	}
 	else
 	{
diff --git a/Engine/Source/Runtime/Online/OnlineSubsystemNull/Private/OnlineSubsystemNull.cpp b/Engine/Source/Runtime/Online/OnlineSubsystemNull/Private/OnlineSubsystemNull.cpp
--- a/Engine/Source/Runtime/Online/OnlineSubsystemNull/Private/OnlineSubsystemNull.cpp
+++ b/Engine/Source/Runtime/Online/OnlineSubsystemNull/Private/OnlineSubsystemNull.cpp
@@ -37,6 +37,11 @@ IOnlineUserCloudPtr FOnlineSubsystemNull::GetUserCloudInterface() const
 	return NULL;
 }
 
+IOnlineUserCloudPtr FOnlineSubsystemNull::GetUserCloudInterface(const FString& Key) const
+{
+	return NULL;
+}
+
 IOnlineEntitlementsPtr FOnlineSubsystemNull::GetEntitlementsInterface() const
 {
 	return NULL;
diff --git a/Engine/Source/Runtime/Online/OnlineSubsystemSteam/Private/OnlineAsyncTaskManagerSteam.cpp b/Engine/Source/Runtime/Online/OnlineSubsystemSteam/Private/OnlineAsyncTaskManagerSteam.cpp
--- a/Engine/Source/Runtime/Online/OnlineSubsystemSteam/Private/OnlineAsyncTaskManagerSteam.cpp
+++ b/Engine/Source/Runtime/Online/OnlineSubsystemSteam/Private/OnlineAsyncTaskManagerSteam.cpp
@@ -628,7 +628,7 @@ public:
 	 */
 	virtual void Finalize() override
 	{
-		Subsystem->TriggerOnConnectionStatusChangedDelegates(ConnectionState);
+		Subsystem->TriggerOnConnectionStatusChangedDelegates(EOnlineServerConnectionStatus::Normal, ConnectionState);
 	}
 };
 
@@ -779,7 +779,7 @@ public:
 		if (bTriggerConnectionStatusUpdate)
 		{
 			EOnlineServerConnectionStatus::Type ConnectionState = SteamConnectionResult(CallbackResults.m_eResult);
-			Subsystem->TriggerOnConnectionStatusChangedDelegates(ConnectionState);
+			Subsystem->TriggerOnConnectionStatusChangedDelegates(EOnlineServerConnectionStatus::Normal, ConnectionState);
 		}
 	}
 };
diff --git a/Engine/Source/Runtime/Online/OnlineSubsystemSteam/Private/OnlineSubsystemSteam.cpp b/Engine/Source/Runtime/Online/OnlineSubsystemSteam/Private/OnlineSubsystemSteam.cpp
--- a/Engine/Source/Runtime/Online/OnlineSubsystemSteam/Private/OnlineSubsystemSteam.cpp
+++ b/Engine/Source/Runtime/Online/OnlineSubsystemSteam/Private/OnlineSubsystemSteam.cpp
@@ -179,6 +179,11 @@ IOnlineUserCloudPtr FOnlineSubsystemSteam::GetUserCloudInterface() const
 	return UserCloudInterface;
 }
 
+IOnlineUserCloudPtr FOnlineSubsystemSteam::GetUserCloudInterface(const FString& Key) const
+{
+	return NULL;
+}
+
 IOnlineLeaderboardsPtr FOnlineSubsystemSteam::GetLeaderboardsInterface() const
 {
 	return LeaderboardsInterface;
diff --git a/Engine/Source/Runtime/Online/OnlineSubsystemSteam/Private/OnlineUserCloudInterfaceSteam.cpp b/Engine/Source/Runtime/Online/OnlineSubsystemSteam/Private/OnlineUserCloudInterfaceSteam.cpp
--- a/Engine/Source/Runtime/Online/OnlineSubsystemSteam/Private/OnlineUserCloudInterfaceSteam.cpp
+++ b/Engine/Source/Runtime/Online/OnlineSubsystemSteam/Private/OnlineUserCloudInterfaceSteam.cpp
@@ -333,12 +333,24 @@ bool FOnlineUserCloudSteam::WriteUserFile(const FUniqueNetId& UserId, const FStr
 	return false;
 }
 
+void FOnlineUserCloudSteam::CancelWriteUserFile(const FUniqueNetId& UserId, const FString& FileName)
+{
+	// Not implemented
+}
+
+
 bool FOnlineUserCloudSteam::DeleteUserFile(const FUniqueNetId& UserId, const FString& FileName, bool bShouldCloudDelete, bool bShouldLocallyDelete)
 {
 	SteamSubsystem->QueueAsyncTask(new FOnlineAsyncTaskSteamDeleteUserFile(SteamSubsystem, FUniqueNetIdSteam(*(uint64*)UserId.GetBytes()), FileName, bShouldCloudDelete, bShouldLocallyDelete));
 	return true;
 }
 
+bool FOnlineUserCloudSteam::RequestUsageInfo(const FUniqueNetId& UserId)
+{
+	// Not implemented
+	return false;
+}
+
 void FOnlineUserCloudSteam::DumpCloudState(const FUniqueNetId& UserId)
 {
 	int32 TotalBytes, TotalAvailable;
diff --git a/Engine/Source/Runtime/Online/OnlineSubsystemSteam/Private/OnlineUserCloudInterfaceSteam.h b/Engine/Source/Runtime/Online/OnlineSubsystemSteam/Private/OnlineUserCloudInterfaceSteam.h
--- a/Engine/Source/Runtime/Online/OnlineSubsystemSteam/Private/OnlineUserCloudInterfaceSteam.h
+++ b/Engine/Source/Runtime/Online/OnlineSubsystemSteam/Private/OnlineUserCloudInterfaceSteam.h
@@ -231,7 +231,9 @@ public:
 	virtual void GetUserFileList(const FUniqueNetId& UserId, TArray<FCloudFileHeader>& UserFiles) override;
 	virtual bool ReadUserFile(const FUniqueNetId& UserId, const FString& FileName) override;
 	virtual bool WriteUserFile(const FUniqueNetId& UserId, const FString& FileName, TArray<uint8>& FileContents) override;
+	virtual void CancelWriteUserFile(const FUniqueNetId& UserId, const FString& FileName) override;
 	virtual bool DeleteUserFile(const FUniqueNetId& UserId, const FString& FileName, bool bShouldCloudDelete, bool bShouldLocallyDelete) override;
+	virtual bool RequestUsageInfo(const FUniqueNetId& UserId) override;
 	virtual void DumpCloudState(const FUniqueNetId& UserId) override;
 	virtual void DumpCloudFileState(const FUniqueNetId& UserId, const FString& FileName) override;
 
diff --git a/Engine/Source/Runtime/Online/OnlineSubsystemUtils/Private/Tests/TestEntitlementsInterface.cpp b/Engine/Source/Runtime/Online/OnlineSubsystemUtils/Private/Tests/TestEntitlementsInterface.cpp
--- a/Engine/Source/Runtime/Online/OnlineSubsystemUtils/Private/Tests/TestEntitlementsInterface.cpp
+++ b/Engine/Source/Runtime/Online/OnlineSubsystemUtils/Private/Tests/TestEntitlementsInterface.cpp
@@ -48,7 +48,7 @@ void FTestEntitlementsInterface::StartNextTest()
 {
 	if (bQueryEntitlements)
 	{
-		EntitlementsOSS->QueryEntitlements(*UserId);
+		EntitlementsOSS->QueryEntitlements(*UserId, FString());
 	}
 	else
 	{
@@ -67,14 +67,14 @@ void FTestEntitlementsInterface::FinishTest()
 	delete this;
 }
 
-void FTestEntitlementsInterface::OnQueryEntitlementsComplete(bool bWasSuccessful, const FUniqueNetId& UserId, const FString& Error)
+void FTestEntitlementsInterface::OnQueryEntitlementsComplete(bool bWasSuccessful, const FUniqueNetId& UserId, const FString& Namespace, const FString& Error)
 {
 	UE_LOG(LogOnline, Log, TEXT("enumerated entitlements. UserId=%s Result=%s Error=[%s]"), 
 		*UserId.ToDebugString(), bWasSuccessful ? TEXT("true") : TEXT("false"), *Error);
 
 	// Now check em out
 	TArray<TSharedRef<FOnlineEntitlement>> Entitlements;
-	EntitlementsOSS->GetAllEntitlements(UserId, Entitlements);
+	EntitlementsOSS->GetAllEntitlements(UserId, FString(), Entitlements);
 
 	for (auto It = Entitlements.CreateConstIterator(); It; ++It)
 	{
diff --git a/Engine/Source/Runtime/Online/OnlineSubsystemUtils/Private/Tests/TestEntitlementsInterface.h b/Engine/Source/Runtime/Online/OnlineSubsystemUtils/Private/Tests/TestEntitlementsInterface.h
--- a/Engine/Source/Runtime/Online/OnlineSubsystemUtils/Private/Tests/TestEntitlementsInterface.h
+++ b/Engine/Source/Runtime/Online/OnlineSubsystemUtils/Private/Tests/TestEntitlementsInterface.h
@@ -55,7 +55,7 @@
 	 * @param UserId of the user who was granted entitlements in this callback
 	 * @param Error string representing the error condition if any
 	 */
-	void OnQueryEntitlementsComplete(bool bWasSuccessful, const FUniqueNetId& UserId, const FString& Error);
+	void OnQueryEntitlementsComplete(bool bWasSuccessful, const FUniqueNetId& UserId, const FString& Namespace, const FString& Error);
 
  public:
 	
diff --git a/Engine/Source/Runtime/Portal/Proxies/PortalAccountProxies/PortalAccountProxies.Build.cs b/Engine/Source/Runtime/Portal/Proxies/PortalAccountProxies/PortalAccountProxies.Build.cs
new file mode 100644
--- /dev/null
+++ b/Engine/Source/Runtime/Portal/Proxies/PortalAccountProxies/PortalAccountProxies.Build.cs
@@ -0,0 +1,24 @@
+// Copyright 1998-2015 Epic Games, Inc. All Rights Reserved.
+
+namespace UnrealBuildTool.Rules
+{
+	public class PortalAccountProxies : ModuleRules
+	{
+		public PortalAccountProxies(TargetInfo Target)
+		{
+			PublicDependencyModuleNames.AddRange(
+				new string[]
+				{
+					"Core",
+					"CoreUObject",
+				}
+			); 
+
+			PrivateIncludePaths.AddRange(
+				new string[] {
+					"Runtime/Portal/Proxies/PortalAccountProxies/Private",
+				}
+			);
+		}
+	}
+}
diff --git a/Engine/Source/Runtime/Portal/Proxies/PortalAccountProxies/Private/Module.cpp b/Engine/Source/Runtime/Portal/Proxies/PortalAccountProxies/Private/Module.cpp
new file mode 100644
--- /dev/null
+++ b/Engine/Source/Runtime/Portal/Proxies/PortalAccountProxies/Private/Module.cpp
@@ -0,0 +1,32 @@
+// Copyright 1998-2015 Epic Games, Inc. All Rights Reserved.
+
+#include "PrivatePCH.h"
+#include "ModuleManager.h"
+
+
+/**
+ * Implements the PortalAccountProxies module.
+ */
+class FPortalAccountProxiesModule
+	: public IPortalAccountProxiesModule
+{
+public:
+
+	// IModuleInterface interface
+
+	virtual void StartupModule() override
+	{
+	}
+
+	virtual void ShutdownModule() override
+	{
+	}
+
+	virtual bool SupportsDynamicReloading() override
+	{
+		return false;
+	}
+};
+
+
+IMPLEMENT_MODULE(FPortalAccountProxiesModule, PortalAccountProxies);
diff --git a/Engine/Source/Runtime/Portal/Proxies/PortalChatProxies/PortalChatProxies.Build.cs b/Engine/Source/Runtime/Portal/Proxies/PortalChatProxies/PortalChatProxies.Build.cs
new file mode 100644
--- /dev/null
+++ b/Engine/Source/Runtime/Portal/Proxies/PortalChatProxies/PortalChatProxies.Build.cs
@@ -0,0 +1,24 @@
+// Copyright 1998-2015 Epic Games, Inc. All Rights Reserved.
+
+namespace UnrealBuildTool.Rules
+{
+	public class PortalChatProxies : ModuleRules
+	{
+		public PortalChatProxies(TargetInfo Target)
+		{
+			PublicDependencyModuleNames.AddRange(
+				new string[]
+				{
+					"Core",
+					"CoreUObject",
+				}
+			); 
+
+			PrivateIncludePaths.AddRange(
+				new string[] {
+					"Runtime/Portal/Proxies/PortalChatProxies/Private",
+				}
+			);
+		}
+	}
+}
diff --git a/Engine/Source/Runtime/Portal/Proxies/PortalChatProxies/Private/Module.cpp b/Engine/Source/Runtime/Portal/Proxies/PortalChatProxies/Private/Module.cpp
new file mode 100644
--- /dev/null
+++ b/Engine/Source/Runtime/Portal/Proxies/PortalChatProxies/Private/Module.cpp
@@ -0,0 +1,32 @@
+// Copyright 1998-2015 Epic Games, Inc. All Rights Reserved.
+
+#include "PrivatePCH.h"
+#include "ModuleManager.h"
+
+
+/**
+ * Implements the PortalChatProxies module.
+ */
+class FPortalChatProxiesModule
+	: public IPortalChatProxiesModule
+{
+public:
+
+	// IModuleInterface interface
+
+	virtual void StartupModule() override
+	{
+	}
+
+	virtual void ShutdownModule() override
+	{
+	}
+
+	virtual bool SupportsDynamicReloading() override
+	{
+		return false;
+	}
+};
+
+
+IMPLEMENT_MODULE(FPortalChatProxiesModule, PortalChatProxies);
diff --git a/Engine/Source/Runtime/Portal/Proxies/PortalPackageProxies/PortalPackageProxies.Build.cs b/Engine/Source/Runtime/Portal/Proxies/PortalPackageProxies/PortalPackageProxies.Build.cs
new file mode 100644
--- /dev/null
+++ b/Engine/Source/Runtime/Portal/Proxies/PortalPackageProxies/PortalPackageProxies.Build.cs
@@ -0,0 +1,38 @@
+// Copyright 1998-2015 Epic Games, Inc. All Rights Reserved.
+
+namespace UnrealBuildTool.Rules
+{
+	public class PortalPackageProxies : ModuleRules
+	{
+		public PortalPackageProxies(TargetInfo Target)
+		{
+            DynamicallyLoadedModuleNames.AddRange(
+                new string[] {
+                    "Messaging",
+                    "PortalServices",
+                }
+            );
+
+			PublicDependencyModuleNames.AddRange(
+				new string[]
+				{
+					"Core",
+					"CoreUObject",
+				}
+			);
+
+            PrivateIncludePathModuleNames.AddRange(
+                new string[] {
+                    "Messaging",
+                    "PortalServices",
+                }
+            );
+
+			PrivateIncludePaths.AddRange(
+				new string[] {
+					"Runtime/Portal/Proxies/PortalPackageProxies/Private",
+				}
+			);
+		}
+	}
+}
diff --git a/Engine/Source/Runtime/Portal/Proxies/PortalPackageProxies/Private/Module.cpp b/Engine/Source/Runtime/Portal/Proxies/PortalPackageProxies/Private/Module.cpp
new file mode 100644
--- /dev/null
+++ b/Engine/Source/Runtime/Portal/Proxies/PortalPackageProxies/Private/Module.cpp
@@ -0,0 +1,53 @@
+// Copyright 1998-2015 Epic Games, Inc. All Rights Reserved.
+
+#include "PrivatePCH.h"
+#include "ModuleManager.h"
+#include "IPortalServicesModule.h"
+#include "TypeContainer.h"
+
+
+/**
+ * Implements the PortalPackageProxies module.
+ */
+class FPortalPackageProxiesModule
+	: public IModuleInterface
+{
+public:
+
+	// IModuleInterface interface
+
+	virtual void StartupModule() override
+	{
+		// load dependencies
+		IPortalServicesModule* ServicesModule = FModuleManager::GetModulePtr<IPortalServicesModule>("PortalServices");
+
+		if (ServicesModule == nullptr)
+		{
+			GLog->Log(TEXT("Error: The required module 'PortalServices' failed to load. PortalPackageProxies cannot be used."));
+
+			return;
+		}
+
+		// register service proxies
+		ServicesModule->GetServiceContainer().RegisterClass<IPortalPackageInstaller, FPackageInstallerProxy>(ETypeContainerScope::Instance);
+	}
+
+	virtual void ShutdownModule() override
+	{
+		// unregister service proxies
+		IPortalServicesModule* ServicesModule = FModuleManager::GetModulePtr<IPortalServicesModule>("PortalServices");
+
+		if (ServicesModule != nullptr)
+		{
+			ServicesModule->GetServiceContainer().Unregister<IPortalPackageInstaller>();
+		}
+	}
+
+	virtual bool SupportsDynamicReloading() override
+	{
+		return false;
+	}
+};
+
+
+IMPLEMENT_MODULE(FPortalPackageProxiesModule, PortalPackageProxies);
diff --git a/Engine/Source/Runtime/Portal/Proxies/PortalPackageProxies/Private/PackageInstallerProxy.cpp b/Engine/Source/Runtime/Portal/Proxies/PortalPackageProxies/Private/PackageInstallerProxy.cpp
new file mode 100644
--- /dev/null
+++ b/Engine/Source/Runtime/Portal/Proxies/PortalPackageProxies/Private/PackageInstallerProxy.cpp
@@ -0,0 +1,65 @@
+// Copyright 1998-2015 Epic Games, Inc. All Rights Reserved.
+
+#include "PrivatePCH.h"
+#include "IPortalPackageInstaller.h"
+#include "PortalPackageInstallerMessages.h"
+
+
+/* FPackageInstallerProxy structors
+ *****************************************************************************/
+
+FPackageInstallerProxy::FPackageInstallerProxy()
+{
+	MessageEndpoint = FMessageEndpoint::Builder("FPackageInstallerProxy")
+		.Handling<FPortalPackageInstallResponse>(this, &FPackageInstallerProxy::HandlePackageInstallResponse)
+		.Handling<FPortalPackageUninstallResponse>(this, &FPackageInstallerProxy::HandlePackageUninstallResponse);
+}
+
+
+/* IPortalPackageInstaller interface
+ *****************************************************************************/
+
+TSlowResult<bool> FPackageInstallerProxy::InstallPackage(const FString& Path, const FPortalPackageId& Id)
+{
+	if (!MessageEndpoint.IsValid())
+	{
+		return TSlowResult<bool>();
+	}
+
+	auto Message = new FPortalPackageInstallRequest();
+	{
+		Message->AppName = Id.AppName;
+		Message->BuildLabel = Id.BuildLabel;
+		Message->RequestId = FGuid::NewGuid();
+	}
+
+	//MessageEndpoint->Send(Message, )
+
+	return TSlowResult<bool>();
+}
+
+
+TSlowResult<bool> FPackageInstallerProxy::UninstallPackage(const FString& Path, const FString& Id, bool RemoveUserFiles)
+{
+	if (!MessageEndpoint.IsValid())
+	{
+		return TSlowResult<bool>();
+	}
+
+	return TSlowResult<bool>();
+}
+
+
+/* FPackageInstallerProxy callbacks
+ *****************************************************************************/
+
+void FPackageInstallerProxy::HandlePackageInstallResponse(const FPortalPackageInstallResponse& Message, const IMessageContextRef& Context)
+{
+
+}
+
+
+void FPackageInstallerProxy::HandlePackageUninstallResponse(const FPortalPackageUninstallResponse& Message, const IMessageContextRef& Context)
+{
+
+}
diff --git a/Engine/Source/Runtime/Portal/Proxies/PortalPackageProxies/Private/PackageInstallerProxy.h b/Engine/Source/Runtime/Portal/Proxies/PortalPackageProxies/Private/PackageInstallerProxy.h
new file mode 100644
--- /dev/null
+++ b/Engine/Source/Runtime/Portal/Proxies/PortalPackageProxies/Private/PackageInstallerProxy.h
@@ -0,0 +1,43 @@
+// Copyright 1998-2015 Epic Games, Inc. All Rights Reserved.
+
+#pragma once
+
+#include "IPortalPackageInstaller.h"
+#include "SlowResult.h"
+
+
+class FMessageEndpoint;
+class IMessageContext;
+struct FPortalPackageInstallResponse;
+struct FPortalPackageUninstallResponse;
+
+
+/**
+ * Implements a proxy for the Portal's package installer service.
+ */
+class FPackageInstallerProxy
+	: public IPortalPackageInstaller
+{
+public:
+
+	/** Default constructor. */
+	FPackageInstallerProxy();
+
+public:
+
+	// IPortalPackageInstaller interface
+
+	virtual TSlowResult<bool> InstallPackage(const FString& Path, const FPortalPackageId& Id) override;
+	virtual TSlowResult<bool> UninstallPackage(const FString& Path, const FString& Id, bool RemoveUserFiles) override;
+
+private:
+
+	/** Handles FPortalPackageInstallResponse messages. */
+	void HandlePackageInstallResponse(const FPortalPackageInstallResponse& Message, const TSharedRef<IMessageContext, ESPMode::ThreadSafe>& Context);
+	void HandlePackageUninstallResponse(const FPortalPackageUninstallResponse& Message, const TSharedRef<IMessageContext, ESPMode::ThreadSafe>& Context);
+
+private:
+
+	/** Holds the message endpoint. */
+	TSharedPtr<FMessageEndpoint, ESPMode::ThreadSafe> MessageEndpoint;
+};
diff --git a/Engine/Source/Runtime/Portal/Proxies/PortalPackageProxies/Public/PortalPackageInstallerMessages.h b/Engine/Source/Runtime/Portal/Proxies/PortalPackageProxies/Public/PortalPackageInstallerMessages.h
new file mode 100644
--- /dev/null
+++ b/Engine/Source/Runtime/Portal/Proxies/PortalPackageProxies/Public/PortalPackageInstallerMessages.h
@@ -0,0 +1,106 @@
+// Copyright 1998-2015 Epic Games, Inc. All Rights Reserved.
+
+#pragma once
+
+#include "PortalPackageInstallerMessages.generated.h"
+
+
+/* Service discovery messages
+ *****************************************************************************/
+
+/**
+ * Implements a request message for Portal package installation.
+ */
+USTRUCT()
+struct FPortalPackageInstallRequest
+{
+	GENERATED_USTRUCT_BODY()
+
+	/** The application name. */
+	UPROPERTY()
+	FString AppName;
+
+	/** The application build label. */
+	UPROPERTY()
+	FString BuildLabel;
+
+	/** The path to install the package contents to. */
+	UPROPERTY()
+	FString DestinationPath;
+
+	/** Unique identifier for this request. */
+	UPROPERTY()
+	FGuid RequestId;
+};
+
+template<> struct TStructOpsTypeTraits<FPortalPackageInstallRequest> : public TStructOpsTypeTraitsBase { enum { WithMessageHandling = true }; };
+
+
+/**
+ * Implements a reply message to Portal package installation requests.
+ */
+USTRUCT()
+struct FPortalPackageInstallResponse
+{
+	GENERATED_USTRUCT_BODY()
+
+	/** Identifier of the request that this response is for. */
+	UPROPERTY()
+	FGuid RequestId;
+
+	UPROPERTY()
+	bool Success;
+};
+
+template<> struct TStructOpsTypeTraits<FPortalPackageInstallResponse> : public TStructOpsTypeTraitsBase{ enum { WithMessageHandling = true }; };
+
+
+/**
+ * Implements a message for Portal package uninstallation.
+ */
+USTRUCT()
+struct FPortalPackageUninstallRequest
+{
+	GENERATED_USTRUCT_BODY()
+
+	/** The application name. */
+	UPROPERTY()
+	FString AppName;
+
+	/** The application build label. */
+	UPROPERTY()
+	FString BuildLabel;
+
+	/** The path to where the package is installed. */
+	UPROPERTY()
+	FString InstallationPath;
+
+	/** Whether any user generated files should be removed as well. */
+	UPROPERTY()
+	bool RemoveUserFiles;
+
+	/** Unique identifier for this request. */
+	UPROPERTY()
+	FGuid RequestId;
+};
+
+template<> struct TStructOpsTypeTraits<FPortalPackageUninstallRequest> : public TStructOpsTypeTraitsBase{ enum { WithMessageHandling = true }; };
+
+
+/**
+ * Implements a reply message to Portal package uninstallation requests.
+ */
+USTRUCT()
+struct FPortalPackageUninstallResponse
+{
+	GENERATED_USTRUCT_BODY()
+
+	/** Identifier of the request that this response is for. */
+	UPROPERTY()
+	FGuid RequestId;
+
+	UPROPERTY()
+	bool Success;
+};
+
+template<> struct TStructOpsTypeTraits<FPortalPackageUninstallResponse> : public TStructOpsTypeTraitsBase{ enum { WithMessageHandling = true }; };
diff --git a/Engine/Source/Runtime/Portal/Services/PortalServices.Build.cs b/Engine/Source/Runtime/Portal/Services/PortalServices.Build.cs
new file mode 100644
--- /dev/null
+++ b/Engine/Source/Runtime/Portal/Services/PortalServices.Build.cs
@@ -0,0 +1,17 @@
+// Copyright 1998-2015 Epic Games, Inc. All Rights Reserved.
+
+namespace UnrealBuildTool.Rules
+{
+	public class PortalServices : ModuleRules
+	{
+		public PortalServices(TargetInfo Target)
+		{
+			PublicDependencyModuleNames.AddRange(
+				new string[]
+				{
+					"Core",
+				}
+			); 
+		}
+	}
+}
diff --git a/Engine/Source/Runtime/Portal/Services/Private/Module.cpp b/Engine/Source/Runtime/Portal/Services/Private/Module.cpp
new file mode 100644
--- /dev/null
+++ b/Engine/Source/Runtime/Portal/Services/Private/Module.cpp
@@ -0,0 +1,47 @@
+// Copyright 1998-2015 Epic Games, Inc. All Rights Reserved.
+
+#include "PortalServicesPrivatePCH.h"
+#include "TypeContainer.h"
+#include "ModuleManager.h"
+
+
+/**
+ * Implements the PortalServices module.
+ */
+class FPortalServicesModule
+	: public IPortalServicesModule
+{
+public:
+
+	// IPortalServicesModule
+
+	virtual FTypeContainer& GetServiceContainer() override
+	{
+		return ServiceContainer;
+	}
+
+public:
+
+	// IModuleInterface interface
+
+	virtual void StartupModule() override
+	{
+	}
+
+	virtual void ShutdownModule() override
+	{
+	}
+
+	virtual bool SupportsDynamicReloading() override
+	{
+		return false;
+	}
+
+private:
+
+	/** Holds registered service instances. */
+	FTypeContainer ServiceContainer;
+};
+
+
+IMPLEMENT_MODULE(FPortalServicesModule, PortalServices);
diff --git a/Engine/Source/Runtime/Portal/Services/Private/PortalServicesPrivatePCH.h b/Engine/Source/Runtime/Portal/Services/Private/PortalServicesPrivatePCH.h
new file mode 100644
--- /dev/null
+++ b/Engine/Source/Runtime/Portal/Services/Private/PortalServicesPrivatePCH.h
@@ -0,0 +1,10 @@
+// Copyright 1998-2015 Epic Games, Inc. All Rights Reserved.
+
+#pragma once
+
+
+/* Private dependencies
+*****************************************************************************/
+
+#include "Core.h"
+#include "IPortalServicesModule.h"
diff --git a/Engine/Source/Runtime/Portal/Services/Public/IPortalServicesModule.h b/Engine/Source/Runtime/Portal/Services/Public/IPortalServicesModule.h
new file mode 100644
--- /dev/null
+++ b/Engine/Source/Runtime/Portal/Services/Public/IPortalServicesModule.h
@@ -0,0 +1,26 @@
+// Copyright 1998-2015 Epic Games, Inc. All Rights Reserved.
+
+#pragma once
+
+#include "ModuleInterface.h"
+
+
+class FTypeContainer;
+class IPortalService;
+
+
+/**
+ * Interface for the PortalServices module.
+ */
+class IPortalServicesModule
+	: public IModuleInterface
+{
+public:
+
+	/**
+	 * Gets a container with registered Portal services.
+	 *
+	 * @return The type container.
+	 */
+	virtual FTypeContainer& GetServiceContainer() = 0;
+};
diff --git a/Engine/Source/Runtime/Portal/Services/Public/Package/IPortalPackageInstaller.h b/Engine/Source/Runtime/Portal/Services/Public/Package/IPortalPackageInstaller.h
new file mode 100644
--- /dev/null
+++ b/Engine/Source/Runtime/Portal/Services/Public/Package/IPortalPackageInstaller.h
@@ -0,0 +1,54 @@
+// Copyright 1998-2015 Epic Games, Inc. All Rights Reserved.
+
+#pragma once
+
+#include "SlowResult.h"
+#include "TypeContainer.h"
+
+
+/**
+ * Structure for package identifiers.
+ *
+ * Package identifiers are used to uniquely identify a package.
+ */
+struct FPortalPackageId
+{
+	/** The application name. */
+	FString AppName;
+
+	/** The application build label. */
+	FString BuildLabel;
+};
+
+
+/**
+ * Interface for package installer services.
+ */
+class IPortalPackageInstaller
+{
+public:
+
+	/**
+	 * Install the specified package using the given request object.
+	 *
+	 * @param Path The path at which the package should be installed.
+	 * @param Id The identifier of the package to install.
+	 * @return The result of the task.
+	 * @see UninstallPackage
+	 */
+	virtual TSlowResult<bool> InstallPackage(const FString& Path, const FPortalPackageId& Id) = 0;
+
+	/**
+	 * Attempts to uninstall the specified package using the given request object.
+	 *
+	 * @param Path The path at which the package is installed.
+	 * @param Id The identifier of the package to uninstall.
+	 * @param RemoveUserFiles Whether user created files should be removed as well.
+	 * @return The result of the task.
+	 * @see InstallPackage
+	 */
+	virtual TSlowResult<bool> UninstallPackage(const FString& Path, const FString& Id, bool RemoveUserFiles) = 0;
+};
+
+
+Expose_TNameOf(IPortalPackageInstaller)
diff --git a/Engine/Source/Runtime/Slate/Private/Framework/Notifications/NotificationManager.cpp b/Engine/Source/Runtime/Slate/Private/Framework/Notifications/NotificationManager.cpp
--- a/Engine/Source/Runtime/Slate/Private/Framework/Notifications/NotificationManager.cpp
+++ b/Engine/Source/Runtime/Slate/Private/Framework/Notifications/NotificationManager.cpp
@@ -46,14 +46,15 @@ void FSlateNotificationManager::FRegionalNotificationList::Arrange()
 		if( PinnedWindow.IsValid() )
 		{
 			const FVector2D DesiredSize = PinnedWindow->GetDesiredSize();
-			const FVector2D NewPosition( AnchorPoint.X - DesiredSize.X, AnchorPoint.Y - DesiredSize.Y );
+			const FVector2D NewPosition(AnchorPoint.X - DesiredSize.X, AnchorPoint.Y - DesiredSize.Y);
 			if( NewPosition != PinnedWindow->GetPositionInScreen() && DesiredSize != PinnedWindow->GetSizeInScreen() )
 			{
 				PinnedWindow->ReshapeWindow( NewPosition, DesiredSize );
 			}
 			else if( NewPosition != PinnedWindow->GetPositionInScreen() )
 			{
-				PinnedWindow->MoveWindowTo( NewPosition );
+				float StackOffset = NotificationManagerConstants::NotificationOffset.Y * ((Notifications.Num()-1) - ListIndex);
+				PinnedWindow->MoveWindowTo(NewPosition - FVector2D(0, StackOffset));
 			}
 			AnchorPoint.Y -= DesiredSize.Y;
 		}
diff --git a/Engine/Source/Runtime/Slate/Private/Widgets/Input/SMultiLineEditableTextBox.cpp b/Engine/Source/Runtime/Slate/Private/Widgets/Input/SMultiLineEditableTextBox.cpp
--- a/Engine/Source/Runtime/Slate/Private/Widgets/Input/SMultiLineEditableTextBox.cpp
+++ b/Engine/Source/Runtime/Slate/Private/Widgets/Input/SMultiLineEditableTextBox.cpp
@@ -37,11 +37,16 @@ void SMultiLineEditableTextBox::Construct( const FArguments& InArgs )
 	TAttribute<FMargin> Padding = InArgs._Padding.IsSet() ? InArgs._Padding : InArgs._Style->Padding;
 	TAttribute<FMargin> HScrollBarPadding = InArgs._HScrollBarPadding.IsSet() ? InArgs._HScrollBarPadding : InArgs._Style->HScrollBarPadding;
 	TAttribute<FMargin> VScrollBarPadding = InArgs._VScrollBarPadding.IsSet() ? InArgs._VScrollBarPadding : InArgs._Style->VScrollBarPadding;
-	TAttribute<FSlateFontInfo> Font = InArgs._Font.IsSet() ? InArgs._Font : InArgs._Style->Font;
 	TAttribute<FSlateColor> ForegroundColor = InArgs._ForegroundColor.IsSet() ? InArgs._ForegroundColor : InArgs._Style->ForegroundColor;
 	TAttribute<FSlateColor> BackgroundColor = InArgs._BackgroundColor.IsSet() ? InArgs._BackgroundColor : InArgs._Style->BackgroundColor;
 	ReadOnlyForegroundColor = InArgs._ReadOnlyForegroundColor.IsSet() ? InArgs._ReadOnlyForegroundColor : InArgs._Style->ReadOnlyForegroundColor;
 
+	TAttribute<FSlateFontInfo> Font;
+	if (InArgs._Font.IsSet())
+	{
+		Font = InArgs._Font;
+	}
+
 	TSharedPtr<SScrollBar> HScrollBar = InArgs._HScrollBar;
 	if (!HScrollBar.IsValid())
 	{
diff --git a/Engine/Source/Runtime/Slate/Private/Widgets/Views/STableViewBase.cpp b/Engine/Source/Runtime/Slate/Private/Widgets/Views/STableViewBase.cpp
--- a/Engine/Source/Runtime/Slate/Private/Widgets/Views/STableViewBase.cpp
+++ b/Engine/Source/Runtime/Slate/Private/Widgets/Views/STableViewBase.cpp
@@ -257,6 +257,7 @@ void STableViewBase::Tick( const FGeometry& AllottedGeometry, const double InCur
 			UpdateSelectionSet();
 
 			// Update scrollbar
+			if (NumItemsBeingObserved > 0)
 			{
 				// The thumb size is whatever fraction of the items we are currently seeing (including partially seen items).
 				// e.g. if we are seeing 0.5 of the first generated widget and 0.75 of the last widget, that's 1.25 widgets.
@@ -264,6 +265,12 @@ void STableViewBase::Tick( const FGeometry& AllottedGeometry, const double InCur
 				const double OffsetFraction = ScrollOffset / NumItemsBeingObserved;
 				ScrollBar->SetState( OffsetFraction, ThumbSizeFraction );
 			}
+			else
+			{
+				const double ThumbSizeFraction = 1;
+				const double OffsetFraction = 0;
+				ScrollBar->SetState( OffsetFraction, ThumbSizeFraction );
+			}
 
 			NotifyItemScrolledIntoView();
 
@@ -458,8 +465,7 @@ FReply STableViewBase::OnKeyDown( const FGeometry& MyGeometry, const FKeyEvent&
 {
 	if ( InKeyEvent.IsControlDown() && InKeyEvent.GetKey() == EKeys::End )
 	{
-		ScrollOffset = GetNumItemsBeingObserved();
-		RequestListRefresh();
+		ScrollToBottom();
 		return FReply::Handled();
 	}
 
@@ -734,6 +740,18 @@ float STableViewBase::GetScrollRateInItems() const
 		: 0.5f;
 }
 
+void STableViewBase::ScrollToTop()
+{
+	ScrollOffset = 0;
+	RequestListRefresh();
+}
+
+void STableViewBase::ScrollToBottom()
+{
+	ScrollOffset = GetNumItemsBeingObserved();
+	RequestListRefresh();
+}
+
 FVector2D STableViewBase::GetScrollDistance()
 {
 	return FVector2D( 0, ScrollBar->DistanceFromTop() );
diff --git a/Engine/Source/Runtime/Slate/Public/Widgets/Views/STableViewBase.h b/Engine/Source/Runtime/Slate/Public/Widgets/Views/STableViewBase.h
--- a/Engine/Source/Runtime/Slate/Public/Widgets/Views/STableViewBase.h
+++ b/Engine/Source/Runtime/Slate/Public/Widgets/Views/STableViewBase.h
@@ -92,6 +92,12 @@ public:
 	/** Is this list backing a tree or just a standalone list */
 	const ETableViewMode::Type TableViewMode;
 
+	/** Scrolls the view to the top */
+	void ScrollToTop();
+
+	/** Scrolls the view to the bottom */
+	void ScrollToBottom();
+
 public:
 
 	// SWidget interface
diff --git a/Engine/Source/Runtime/SlateCore/Public/Layout/SlateRect.h b/Engine/Source/Runtime/SlateCore/Public/Layout/SlateRect.h
--- a/Engine/Source/Runtime/SlateCore/Public/Layout/SlateRect.h
+++ b/Engine/Source/Runtime/SlateCore/Public/Layout/SlateRect.h
@@ -226,6 +226,34 @@ public:
 	{
 		return (A.Left <= B.Left) && (A.Right >= B.Right) && (A.Top <= B.Top) && (A.Bottom >= B.Bottom);
 	}
+
+	/**
+	* Returns a string of containing the coordinates of the rect
+	*
+	* @return	A string of the rect coordinates 
+	*/
+	FString ToString() const
+	{
+		return FString::Printf(TEXT("Left=%3.3f Top=%3.3f Right=%3.3f Bottom=%3.3f"), Left, Top, Right, Bottom);
+	}
+
+	/**
+	* Returns a string of containing the coordinates of the rect
+	*
+	* @param InSourceString A string containing the values to initialize this rect in format Left=Value Top=Value...
+	*
+	* @return	True if initialized successfully
+	*/
+	bool InitFromString(const FString& InSourceString)
+	{
+		// The initialization is only successful if the values can all be parsed from the string
+		const bool bSuccessful = FParse::Value(*InSourceString, TEXT("Left="), Left) && 
+								FParse::Value(*InSourceString, TEXT("Top="), Top) && 
+								FParse::Value(*InSourceString, TEXT("Right="), Right) && 
+								FParse::Value(*InSourceString, TEXT("Bottom="), Bottom);
+
+		return bSuccessful;
+	}
 };
 
 /**
diff --git a/Engine/Source/Runtime/SlateReflector/Private/Widgets/SWidgetReflector.h b/Engine/Source/Runtime/SlateReflector/Private/Widgets/SWidgetReflector.h
new file mode 100644
--- /dev/null
+++ b/Engine/Source/Runtime/SlateReflector/Private/Widgets/SWidgetReflector.h
@@ -0,0 +1,202 @@
+// Copyright 1998-2014 Epic Games, Inc. All Rights Reserved.
+
+/*=============================================================================
+	SWidgetReflector.h: Declares the SWidgetReflector class.
+=============================================================================*/
+
+#pragma once
+
+
+/** A Delegate for passing along a string of a source code location to access */
+DECLARE_DELEGATE_RetVal_ThreeParams(bool, FAccessSourceCode, const FString& /*FileName*/, int32 /*InLineNumber*/, int32 /*InColumnNumber*/);
+
+
+/**
+ * Abstract base class for widget reflectors.
+ */
+class SWidgetReflector
+	: public SCompoundWidget
+	, public IWidgetReflector
+{
+	// The reflector uses a tree that observes FReflectorNodes.
+	typedef STreeView<TSharedPtr<FReflectorNode>> SReflectorTree;
+
+public:
+
+	SLATE_BEGIN_ARGS(SWidgetReflector) { }
+	SLATE_END_ARGS()
+
+	/** Create widgets that comprise this WidgetReflector implementation */
+	/**
+	 * Creates and initializes a new widget reflector widget.
+	 *
+	 * @param InArgs The construction arguments.
+	 */
+	void Construct( const FArguments& InArgs );
+
+public:
+
+	// Begin SCompoundWidget overrides
+
+	virtual void Tick( const FGeometry& AllottedGeometry, const double InCurrentTime, const float InDeltaTime ) OVERRIDE;
+
+	// End SCompoundWidget overrides
+
+public:
+
+	// Begin IWidgetReflector interface
+
+	virtual bool IsInPickingMode( ) const OVERRIDE
+	{
+		return bIsPicking;
+	}
+
+	virtual bool IsShowingFocus( ) const OVERRIDE
+	{
+		return bShowFocus;
+	}
+
+	virtual bool IsVisualizingLayoutUnderCursor( ) const OVERRIDE
+	{
+		return bIsPicking;
+	}
+
+	virtual void OnWidgetPicked( ) OVERRIDE
+	{
+		bIsPicking = false;
+	}
+
+	virtual bool ReflectorNeedsToDrawIn( TSharedRef<SWindow> ThisWindow ) const OVERRIDE;
+
+	virtual void SetSourceAccessDelegate( FAccessSourceCode InDelegate ) OVERRIDE
+	{
+		SourceAccessDelegate = InDelegate;
+	}
+
+	virtual void SetWidgetsToVisualize( const FWidgetPath& InWidgetsToVisualize ) OVERRIDE;
+
+	virtual int32 Visualize( const FWidgetPath& InWidgetsToVisualize, FSlateWindowElementList& OutDrawElements, int32 LayerId ) OVERRIDE;
+
+	// End IWidgetReflector interface
+
+protected:
+
+	/**
+	 * Generates a tool tip for the given reflector tree node.
+	 *
+	 * @param InReflectorNode The node to generate the tool tip for.
+	 *
+	 * @return The tool tip widget.
+	 */
+	TSharedRef<SToolTip> GenerateToolTipForReflectorNode( TSharedPtr<FReflectorNode> InReflectorNode );
+
+	/**
+	 * Mark the provided reflector nodes such that they stand out in the tree and are visible.
+	 *
+	 * @param WidgetPathToObserve The nodes to mark.
+	 */
+	void VisualizeAsTree( const TArray< TSharedPtr<FReflectorNode> >& WidgetPathToVisualize );
+
+	/**
+	 * Draw the widget path to the picked widget as the widgets' outlines.
+	 *
+	 * @param InWidgetsToVisualize A widget path whose widgets' outlines to draw.
+	 * @param OutDrawElements A list of draw elements; we will add the output outlines into it.
+	 * @param LayerId The maximum layer achieved in OutDrawElements so far.
+	 *
+	 * @return The maximum layer ID we achieved while painting.
+	 */
+	int32 VisualizePickAsRectangles( const FWidgetPath& InWidgetsToVisualize, FSlateWindowElementList& OutDrawElements, int32 LayerId );
+
+	/**
+	 * Draw an outline for the specified nodes.
+	 *
+	 * @param InNodesToDraw A widget path whose widgets' outlines to draw.
+	 * @param WindowGeometry The geometry of the window in which to draw.
+	 * @param OutDrawElements A list of draw elements; we will add the output outlines into it.
+	 * @param LayerId the maximum layer achieved in OutDrawElements so far.
+	 *
+	 * @return The maximum layer ID we achieved while painting.
+	 */
+	int32 VisualizeSelectedNodesAsRectangles( const TArray<TSharedPtr<FReflectorNode>>& InNodesToDraw, const TSharedRef<SWindow>& VisualizeInWindow, FSlateWindowElementList& OutDrawElements, int32 LayerId );
+
+private:
+
+	// Callback for changing the application scale slider
+	void HandleAppScaleSliderChanged( float NewValue )
+	{
+		FSlateApplication::Get().SetApplicationScale(NewValue);
+	}
+
+	// Callback for getting the value of the application scale slider.
+	float HandleAppScaleSliderValue( ) const
+	{
+		return FSlateApplication::Get().GetApplicationScale();
+	}
+
+	// Callback for checked state changes of the focus check box.
+	void HandleFocusCheckBoxCheckedStateChanged( ESlateCheckBoxState::Type NewValue );
+
+	// Callback for getting the checked state of the focus check box.
+	ESlateCheckBoxState::Type HandleFocusCheckBoxIsChecked( ) const
+	{
+		return bShowFocus ? ESlateCheckBoxState::Checked : ESlateCheckBoxState::Unchecked;
+	}
+
+	// Callback for getting the text of the frame rate text block.
+	FString HandleFrameRateText( ) const;
+
+	// Callback for clicking the pick button.
+	FReply HandlePickButtonClicked( )
+	{
+		bIsPicking = !bIsPicking;
+
+		if (bIsPicking)
+		{
+			bShowFocus = false;
+		}
+
+		return FReply::Handled();
+	}
+
+	// Callback for getting the color of the pick button text.
+	FSlateColor HandlePickButtonColorAndOpacity( ) const
+	{
+		return bIsPicking
+			? FCoreStyle::Get().GetSlateColor("SelectionColor")
+			: FCoreStyle::Get().GetSlateColor("DefaultForeground");
+	}
+
+	// Callback for getting the text of the pick button.
+	FString HandlePickButtonText() const;
+
+	// Callback for generating a row in the reflector tree view.
+	TSharedRef<ITableRow> HandleReflectorTreeGenerateRow( TSharedPtr<FReflectorNode> InReflectorNode, const TSharedRef<STableViewBase>& OwnerTable );
+
+	// Callback for getting the child items of the given reflector tree node.
+	void HandleReflectorTreeGetChildren( TSharedPtr<FReflectorNode> InWidgetGeometry, TArray<TSharedPtr<FReflectorNode>>& OutChildren )
+	{
+		OutChildren = InWidgetGeometry->ChildNodes;
+	}
+
+	// Callback for when the selection in the reflector tree has changed.
+	void HandleReflectorTreeSelectionChanged( TSharedPtr< FReflectorNode >, ESelectInfo::Type /*SelectInfo*/ )
+	{
+		SelectedNodes = ReflectorTree->GetSelectedItems();
+	}
+
+private:
+
+	TSharedPtr<SReflectorTree> ReflectorTree;
+
+	TArray<TSharedPtr<FReflectorNode>> SelectedNodes;
+	TArray<TSharedPtr<FReflectorNode>> ReflectorTreeRoot;
+	TArray<TSharedPtr<FReflectorNode>> PickedPath;
+
+	SSplitter::FSlot* WidgetInfoLocation;
+
+	FAccessSourceCode SourceAccessDelegate;
+
+	bool bShowFocus;
+	bool bIsPicking;
+};
diff --git a/Engine/Source/Runtime/WebBrowser/Private/SWebBrowser.cpp b/Engine/Source/Runtime/WebBrowser/Private/SWebBrowser.cpp
--- a/Engine/Source/Runtime/WebBrowser/Private/SWebBrowser.cpp
+++ b/Engine/Source/Runtime/WebBrowser/Private/SWebBrowser.cpp
@@ -16,6 +16,11 @@ SWebBrowser::SWebBrowser()
 
 void SWebBrowser::Construct(const FArguments& InArgs)
 {
+	OnLoadCompleted = InArgs._OnLoadCompleted;
+	OnLoadError = InArgs._OnLoadError;
+	OnLoadStarted = InArgs._OnLoadStarted;
+	OnTitleChanged = InArgs._OnTitleChanged;
+
 	void* OSWindowHandle = nullptr;
 	if (InArgs._ParentWindow.IsValid())
 	{
@@ -23,12 +28,15 @@ void SWebBrowser::Construct(const FArguments& InArgs)
 		OSWindowHandle = NativeWindow->GetOSWindowHandle();
 	}
 
-	BrowserWindow = IWebBrowserModule::Get().GetSingleton()->CreateBrowserWindow(OSWindowHandle,
-	                                                                             InArgs._InitialURL,
-	                                                                             InArgs._ViewportSize.Get().X,
-	                                                                             InArgs._ViewportSize.Get().Y,
-	                                                                             InArgs._SupportsTransparency,
-	                                                                             InArgs._ContentsToLoad);
+	BrowserWindow = IWebBrowserModule::Get().GetSingleton()->CreateBrowserWindow(
+		OSWindowHandle,
+		InArgs._InitialURL,
+		InArgs._ViewportSize.Get().X,
+		InArgs._ViewportSize.Get().Y,
+		InArgs._SupportsTransparency,
+		InArgs._ContentsToLoad,
+		InArgs._ShowErrorMessage
+	);
 
 	TSharedPtr<SViewport> ViewportWidget;
 
@@ -101,6 +109,7 @@ void SWebBrowser::Construct(const FArguments& InArgs)
 
 	if (BrowserWindow.IsValid())
 	{
+		BrowserWindow->OnDocumentStateChanged().AddSP(this, &SWebBrowser::HandleBrowserWindowDocumentStateChanged);
 		BrowserViewport = MakeShareable(new FWebBrowserViewport(BrowserWindow, ViewportWidget));
 		ViewportWidget->SetViewportInterface(BrowserViewport.ToSharedRef());
 	}
@@ -131,6 +140,36 @@ FText SWebBrowser::GetTitleText() const
 	return LOCTEXT("InvalidWindow", "Browser Window is not valid/supported");
 }
 
+FString SWebBrowser::GetUrl()
+{
+	if (BrowserWindow->IsValid())
+	{
+		return BrowserWindow->GetUrl();
+	}
+
+	return FString();
+}
+
+bool SWebBrowser::IsLoaded() const
+{
+	if (BrowserWindow.IsValid())
+	{
+		return (BrowserWindow->GetDocumentLoadingState() == EWebBrowserDocumentState::Completed);
+	}
+
+	return false;
+}
+
+bool SWebBrowser::IsLoading() const
+{
+	if (BrowserWindow.IsValid())
+	{
+		return (BrowserWindow->GetDocumentLoadingState() == EWebBrowserDocumentState::Loading);
+	}
+
+	return false;
+}
+
 bool SWebBrowser::CanGoBack() const
 {
 	if (BrowserWindow.IsValid())
@@ -216,4 +255,24 @@ EVisibility SWebBrowser::GetLoadingThrobberVisibility() const
 	return EVisibility::Hidden;
 }
 
+
+void SWebBrowser::HandleBrowserWindowDocumentStateChanged(EWebBrowserDocumentState NewState)
+{
+	switch (NewState)
+	{
+	case EWebBrowserDocumentState::Completed:
+		OnLoadCompleted.ExecuteIfBound();
+		break;
+
+	case EWebBrowserDocumentState::Error:
+		OnLoadError.ExecuteIfBound();
+		break;
+
+	case EWebBrowserDocumentState::Loading:
+		OnLoadStarted.ExecuteIfBound();
+		break;
+	}
+}
+
+
 #undef LOCTEXT_NAMESPACE
diff --git a/Engine/Source/Runtime/WebBrowser/Private/WebBrowserHandler.cpp b/Engine/Source/Runtime/WebBrowser/Private/WebBrowserHandler.cpp
--- a/Engine/Source/Runtime/WebBrowser/Private/WebBrowserHandler.cpp
+++ b/Engine/Source/Runtime/WebBrowser/Private/WebBrowserHandler.cpp
@@ -9,30 +9,36 @@
 
 #if WITH_CEF3
 FWebBrowserHandler::FWebBrowserHandler()
-{
-}
+	: ShowErrorMessage(true)
+{ }
 
 void FWebBrowserHandler::OnTitleChange(CefRefPtr<CefBrowser> Browser, const CefString& Title)
 {
+	TSharedPtr<FWebBrowserWindow> BrowserWindow = BrowserWindowPtr.Pin();
+
 	if (BrowserWindow.IsValid())
 	{
-		BrowserWindow.Pin()->SetTitle(Title);
+		BrowserWindow->SetTitle(Title);
 	}
 }
 
 void FWebBrowserHandler::OnAfterCreated(CefRefPtr<CefBrowser> Browser)
 {
+	TSharedPtr<FWebBrowserWindow> BrowserWindow = BrowserWindowPtr.Pin();
+
 	if (BrowserWindow.IsValid())
 	{
-		BrowserWindow.Pin()->BindCefBrowser(Browser);
+		BrowserWindow->BindCefBrowser(Browser);
 	}
 }
 
 void FWebBrowserHandler::OnBeforeClose(CefRefPtr<CefBrowser> Browser)
 {
+	TSharedPtr<FWebBrowserWindow> BrowserWindow = BrowserWindowPtr.Pin();
+
 	if (BrowserWindow.IsValid())
 	{
-		BrowserWindow.Pin()->BindCefBrowser(nullptr);
+		BrowserWindow->BindCefBrowser(nullptr);
 	}
 }
 
@@ -44,26 +50,56 @@ void FWebBrowserHandler::OnLoadError(CefRefPtr<CefBrowser> Browser,
 {
 	// Don't display an error for downloaded files.
 	if (InErrorCode == ERR_ABORTED)
+	{
 		return;
+	}
+
+	// notify browser window
+	if (Frame->IsMain())
+	{
+		TSharedPtr<FWebBrowserWindow> BrowserWindow = BrowserWindowPtr.Pin();
+
+		if (BrowserWindow.IsValid())
+		{
+			BrowserWindow->NotifyDocumentError();
+		}
+	}
 
 	// Display a load error message.
-	FFormatNamedArguments Args;
-	Args.Add(TEXT("FailedUrl"), FText::FromString(FailedUrl.ToWString().c_str()));
-	Args.Add(TEXT("ErrorText"), FText::FromString(ErrorText.ToWString().c_str()));
-	Args.Add(TEXT("ErrorCode"), FText::AsNumber(InErrorCode));
-	FText ErrorMsg = FText::Format(LOCTEXT("WebBrowserLoadError", "Failed to load URL {FailedUrl} with error {ErrorText} ({ErrorCode})."), Args);
-	FString ErrorHTML = TEXT("<html><body bgcolor=\"white\"><h2>")
-						+ ErrorMsg.ToString()
-						+ TEXT("</h2></body></html>");
-	Frame->LoadString(*ErrorHTML, FailedUrl);
+	if (ShowErrorMessage)
+	{
+		FFormatNamedArguments Args;
+		{
+			Args.Add(TEXT("FailedUrl"), FText::FromString(FailedUrl.ToWString().c_str()));
+			Args.Add(TEXT("ErrorText"), FText::FromString(ErrorText.ToWString().c_str()));
+			Args.Add(TEXT("ErrorCode"), FText::AsNumber(InErrorCode));
+		}
+		FText ErrorMsg = FText::Format(LOCTEXT("WebBrowserLoadError", "Failed to load URL {FailedUrl} with error {ErrorText} ({ErrorCode})."), Args);
+		FString ErrorHTML = TEXT("<html><body bgcolor=\"white\"><h2>") + ErrorMsg.ToString() + TEXT("</h2></body></html>");
+
+		Frame->LoadString(*ErrorHTML, FailedUrl);
+	}
+}
+
+void FWebBrowserHandler::OnLoadingStateChange(CefRefPtr<CefBrowser> browser, bool isLoading, bool canGoBack, bool canGoForward)
+{
+	TSharedPtr<FWebBrowserWindow> BrowserWindow = BrowserWindowPtr.Pin();
+
+	if (BrowserWindow.IsValid())
+	{
+		BrowserWindow->NotifyDocumentLoadingStateChange(isLoading);
+	}
 }
 
 bool FWebBrowserHandler::GetViewRect(CefRefPtr<CefBrowser> Browser, CefRect& Rect)
 {
+	TSharedPtr<FWebBrowserWindow> BrowserWindow = BrowserWindowPtr.Pin();
+
 	if (BrowserWindow.IsValid())
 	{
-		return BrowserWindow.Pin()->GetViewRect(Rect);
+		return BrowserWindow->GetViewRect(Rect);
 	}
+
 	return false;
 }
 
@@ -73,17 +109,21 @@ void FWebBrowserHandler::OnPaint(CefRefPtr<CefBrowser> Browser,
 	const void* Buffer,
 	int Width, int Height)
 {
+	TSharedPtr<FWebBrowserWindow> BrowserWindow = BrowserWindowPtr.Pin();
+
 	if (BrowserWindow.IsValid())
 	{
-		BrowserWindow.Pin()->OnPaint(Type, DirtyRects, Buffer, Width, Height);
+		BrowserWindow->OnPaint(Type, DirtyRects, Buffer, Width, Height);
 	}
 }
 
 void FWebBrowserHandler::OnCursorChange(CefRefPtr<CefBrowser> Browser, CefCursorHandle Cursor)
 {
+	TSharedPtr<FWebBrowserWindow> BrowserWindow = BrowserWindowPtr.Pin();
+
 	if (BrowserWindow.IsValid())
 	{
-		BrowserWindow.Pin()->OnCursorChange(Cursor);
+		BrowserWindow->OnCursorChange(Cursor);
 	}
 }
 
@@ -108,7 +148,7 @@ bool FWebBrowserHandler::OnBeforeResourceLoad(CefRefPtr<CefBrowser> Browser, Cef
 
 void FWebBrowserHandler::SetBrowserWindow(TSharedPtr<FWebBrowserWindow> InBrowserWindow)
 {
-	BrowserWindow = InBrowserWindow;
+	BrowserWindowPtr = InBrowserWindow;
 }
 #endif
 
diff --git a/Engine/Source/Runtime/WebBrowser/Private/WebBrowserHandler.h b/Engine/Source/Runtime/WebBrowser/Private/WebBrowserHandler.h
--- a/Engine/Source/Runtime/WebBrowser/Private/WebBrowserHandler.h
+++ b/Engine/Source/Runtime/WebBrowser/Private/WebBrowserHandler.h
@@ -5,69 +5,112 @@
 #if WITH_CEF3
 
 #if PLATFORM_WINDOWS
-#include "AllowWindowsPlatformTypes.h"
+	#include "AllowWindowsPlatformTypes.h"
 #endif
+
 #include "include/cef_client.h"
+
 #if PLATFORM_WINDOWS
-#include "HideWindowsPlatformTypes.h"
+	#include "HideWindowsPlatformTypes.h"
 #endif
 
-// Forward Declarations
+
 class FWebBrowserWindow;
 
+
 /**
- * Implements CEF Client and other Browser level interfaces 
+ * Implements CEF Client and other Browser level interfaces.
  */
-class FWebBrowserHandler : public CefClient,
-	public CefDisplayHandler,
-	public CefLifeSpanHandler,
-	public CefLoadHandler,
-	public CefRenderHandler,
-	public CefRequestHandler
+class FWebBrowserHandler
+	: public CefClient
+	, public CefDisplayHandler
+	, public CefLifeSpanHandler
+	, public CefLoadHandler
+	, public CefRenderHandler
+	, public CefRequestHandler
 {
 public:
+
+	/** Default constructor. */
+	FWebBrowserHandler();
+
+public:
+
 	/**
-	 * Default Constructor
+	 * Pass in a pointer to our Browser Window so that events can be passed on.
+	 *
+	 * @param InBrowserWindow The browser window this will be handling.
 	 */
-	FWebBrowserHandler();
+	void SetBrowserWindow(TSharedPtr<FWebBrowserWindow> InBrowserWindow);
+
+	/** Sets whether to show messages on loading errors. */
+	void SetShowErrorMessage(bool InShowErrorMessage)
+	{
+		ShowErrorMessage = InShowErrorMessage;
+	}
+
+public:
 
 	// CefClient Interface
+
 	virtual CefRefPtr<CefDisplayHandler> GetDisplayHandler() override
 	{
 		return this;
 	}
+
 	virtual CefRefPtr<CefLifeSpanHandler> GetLifeSpanHandler() override
 	{
 		return this;
 	}
+
 	virtual CefRefPtr<CefLoadHandler> GetLoadHandler() override
 	{
 		return this;
 	}
+
 	virtual CefRefPtr<CefRenderHandler> GetRenderHandler() override
 	{
 		return this;
 	}
+
 	virtual CefRefPtr<CefRequestHandler> GetRequestHandler() override
 	{
 		return this;
 	}
 
+public:
+
 	// CefDisplayHandler Interface
+
 	virtual void OnTitleChange(CefRefPtr<CefBrowser> Browser, const CefString& Title) override;
 
+public:
+
 	// CefLifeSpanHandler Interface
+
 	virtual void OnAfterCreated(CefRefPtr<CefBrowser> Browser) override;
 	virtual void OnBeforeClose(CefRefPtr<CefBrowser> Browser) override;
 
+public:
+
 	// CefLoadHandler Interface
+
 	virtual void OnLoadError(CefRefPtr<CefBrowser> Browser,
 		CefRefPtr<CefFrame> Frame,
 		CefLoadHandler::ErrorCode InErrorCode,
 		const CefString& ErrorText,
 		const CefString& FailedUrl) override;
 
+	virtual void OnLoadingStateChange(
+		CefRefPtr<CefBrowser> browser,
+		bool isLoading,
+		bool canGoBack,
+		bool canGoForward) override;
+
+public:
+
 	// CefRenderHandler Interface
+
 	virtual bool GetViewRect(CefRefPtr<CefBrowser> Browser, CefRect& Rect) override;
 	virtual void OnPaint(CefRefPtr<CefBrowser> Browser,
 		PaintElementType Type,
@@ -76,23 +119,25 @@ public:
 		int Width, int Height) override;
 	virtual void OnCursorChange(CefRefPtr<CefBrowser> Browser, CefCursorHandle Cursor) override;
 
+public:
+
 	// CefRequestHandler Interface
+
 	virtual bool OnBeforeResourceLoad(CefRefPtr<CefBrowser> Browser,
 		CefRefPtr<CefFrame> Frame,
 		CefRefPtr<CefRequest> Request) override;
 
-	/**
-	 * Pass in a pointer to our Browser Window so that events can be passed on
-	 *
-	 * @param InBrowserWindow The browser window this will be handling
-	 */
-	void SetBrowserWindow(TSharedPtr<FWebBrowserWindow> InBrowserWindow);
-
 private:
+
 	/** Weak Pointer to our Web Browser window so that events can be passed on while it's valid*/
-	TWeakPtr<FWebBrowserWindow> BrowserWindow;
+	TWeakPtr<FWebBrowserWindow> BrowserWindowPtr;
+
+	/** Whether to show an error message in case of loading errors. */
+	bool ShowErrorMessage;
 
 	// Include the default reference counting implementation.
 	IMPLEMENT_REFCOUNTING(FWebBrowserHandler);
 };
+
+
 #endif
diff --git a/Engine/Source/Runtime/WebBrowser/Private/WebBrowserSingleton.cpp b/Engine/Source/Runtime/WebBrowser/Private/WebBrowserSingleton.cpp
--- a/Engine/Source/Runtime/WebBrowser/Private/WebBrowserSingleton.cpp
+++ b/Engine/Source/Runtime/WebBrowser/Private/WebBrowserSingleton.cpp
@@ -114,11 +114,11 @@ FWebBrowserSingleton::~FWebBrowserSingleton()
 #endif
 }
 
-TSharedPtr<IWebBrowserWindow> FWebBrowserSingleton::CreateBrowserWindow(void* OSWindowHandle, FString InitialURL, uint32 Width, uint32 Height, bool bUseTransparency, TOptional<FString> ContentsToLoad)
+TSharedPtr<IWebBrowserWindow> FWebBrowserSingleton::CreateBrowserWindow(void* OSWindowHandle, FString InitialURL, uint32 Width, uint32 Height, bool bUseTransparency, TOptional<FString> ContentsToLoad, bool ShowErrorMessage)
 {
 #if WITH_CEF3
 	// Create new window
-	TSharedPtr<FWebBrowserWindow> NewWindow(new FWebBrowserWindow(FIntPoint(Width, Height), InitialURL, ContentsToLoad));
+	TSharedPtr<FWebBrowserWindow> NewWindow(new FWebBrowserWindow(FIntPoint(Width, Height), InitialURL, ContentsToLoad, ShowErrorMessage));
 
 	// WebBrowserHandler implements browser-level callbacks.
 	CefRefPtr<FWebBrowserHandler> NewHandler(new FWebBrowserHandler);
diff --git a/Engine/Source/Runtime/WebBrowser/Private/WebBrowserSingleton.h b/Engine/Source/Runtime/WebBrowser/Private/WebBrowserSingleton.h
--- a/Engine/Source/Runtime/WebBrowser/Private/WebBrowserSingleton.h
+++ b/Engine/Source/Runtime/WebBrowser/Private/WebBrowserSingleton.h
@@ -6,46 +6,55 @@
 
 #if WITH_CEF3
 #if PLATFORM_WINDOWS
-#include "AllowWindowsPlatformTypes.h"
+	#include "AllowWindowsPlatformTypes.h"
 #endif
+
 #include "include/internal/cef_ptr.h"
+
 #if PLATFORM_WINDOWS
-#include "HideWindowsPlatformTypes.h"
+	#include "HideWindowsPlatformTypes.h"
 #endif
 #endif
 
-// Forward Declarations
+
 class FWebBrowserApp;
 class FWebBrowserHandler;
 class FWebBrowserWindow;
 
+
 /**
  * Implementation of singleton class that takes care of general web browser tasks
  */
-class FWebBrowserSingleton : public IWebBrowserSingleton, public FTickerObjectBase
+class FWebBrowserSingleton
+	: public IWebBrowserSingleton
+	, public FTickerObjectBase
 {
 public:
-	/**
-	 * Default Constructor
-	 */
+
+	/** Default constructor. */
 	FWebBrowserSingleton();
-	/**
-	 * Virtual Destructor
-	 */
+
+	/** Virtual destructor. */
 	virtual ~FWebBrowserSingleton();
 
+	/**
+	* Gets the Current Locale Code in the format CEF expects
+	*
+	* @return Locale code as either "xx" or "xx-YY"
+	*/
+	static FString GetCurrentLocaleCode();
+
+public:
+
 	// IWebBrowserSingleton Interface
-	TSharedPtr<IWebBrowserWindow> CreateBrowserWindow(void* OSWindowHandle, FString InitialURL, uint32 Width, uint32 Height, bool bUseTransparency, TOptional<FString> ContentsToLoad = TOptional<FString>()) override;
+
+	TSharedPtr<IWebBrowserWindow> CreateBrowserWindow(void* OSWindowHandle, FString InitialURL, uint32 Width, uint32 Height, bool bUseTransparency, TOptional<FString> ContentsToLoad = TOptional<FString>(), bool ShowErrorMessage = true) override;
+
+public:
 
 	// FTickerObjectBase Interface
-	virtual bool Tick(float DeltaTime) override;
 
-	/**
-	 * Gets the Current Locale Code in the format CEF expects
-	 *
-	 * @return Locale code as either "xx" or "xx-YY"
-	 */
-	static FString GetCurrentLocaleCode();
+	virtual bool Tick(float DeltaTime) override;
 
 private:
 #if WITH_CEF3
diff --git a/Engine/Source/Runtime/WebBrowser/Private/WebBrowserWindow.cpp b/Engine/Source/Runtime/WebBrowser/Private/WebBrowserWindow.cpp
--- a/Engine/Source/Runtime/WebBrowser/Private/WebBrowserWindow.cpp
+++ b/Engine/Source/Runtime/WebBrowser/Private/WebBrowserWindow.cpp
@@ -7,13 +7,15 @@
 #include "RHI.h"
 
 #if WITH_CEF3
-FWebBrowserWindow::FWebBrowserWindow(FIntPoint InViewportSize, FString InInitialURL, TOptional<FString> InContentsToLoad)
-	: UpdatableTexture(nullptr)
+FWebBrowserWindow::FWebBrowserWindow(FIntPoint InViewportSize, FString InInitialURL, TOptional<FString> InContentsToLoad, bool InShowErrorMessage)
+	: DocumentState(EWebBrowserDocumentState::NoDocument)
+	, UpdatableTexture(nullptr)
 	, ViewportSize(InViewportSize)
 	, bIsClosing(false)
 	, bHasBeenPainted(false)
 	, InitialURL(InInitialURL)
 	, ContentsToLoad(InContentsToLoad)
+	, ShowErrorMessage(InShowErrorMessage)
 {
 	TextureData.Reserve(ViewportSize.X * ViewportSize.Y * 4);
 	TextureData.SetNumZeroed(ViewportSize.X * ViewportSize.Y * 4);
@@ -120,11 +122,31 @@ bool FWebBrowserWindow::IsClosing() const
 	return bIsClosing;
 }
 
+EWebBrowserDocumentState FWebBrowserWindow::GetDocumentLoadingState() const
+{
+	return DocumentState;
+}
+
 FString FWebBrowserWindow::GetTitle() const
 {
 	return Title;
 }
 
+FString FWebBrowserWindow::GetUrl() const
+{
+	if (InternalCefBrowser != nullptr)
+	{
+		CefRefPtr<CefFrame> MainFrame = InternalCefBrowser->GetMainFrame();
+
+		if (MainFrame != nullptr)
+		{
+			return MainFrame->GetURL().ToWString().c_str();
+		}
+	}
+
+	return FString();
+}
+
 void FWebBrowserWindow::OnKeyDown(const FKeyEvent& InKeyEvent)
 {
 	if (IsValid())
@@ -364,6 +386,7 @@ void FWebBrowserWindow::SetHandler(CefRefPtr<FWebBrowserHandler> InHandler)
 	{
 		Handler = InHandler;
 		Handler->SetBrowserWindow(SharedThis(this));
+		Handler->SetShowErrorMessage(ShowErrorMessage);
 	}
 }
 
@@ -404,6 +427,26 @@ bool FWebBrowserWindow::GetViewRect(CefRect& Rect)
 	return true;
 }
 
+void FWebBrowserWindow::NotifyDocumentError()
+{
+	DocumentState = EWebBrowserDocumentState::Error;
+	DocumentStateChangedEvent.Broadcast(DocumentState);
+}
+
+void FWebBrowserWindow::NotifyDocumentLoadingStateChange(bool IsLoading)
+{
+	EWebBrowserDocumentState NewState = IsLoading
+		? EWebBrowserDocumentState::Loading
+		: EWebBrowserDocumentState::Completed;
+
+	if (DocumentState != EWebBrowserDocumentState::Error)
+	{
+		DocumentState = NewState;
+	}
+
+	DocumentStateChangedEvent.Broadcast(NewState);
+}
+
 void FWebBrowserWindow::OnPaint(CefRenderHandler::PaintElementType Type, const CefRenderHandler::RectList& DirtyRects, const void* Buffer, int Width, int Height)
 {
 	const int32 BufferSize = Width*Height*4;
diff --git a/Engine/Source/Runtime/WebBrowser/Private/WebBrowserWindow.h b/Engine/Source/Runtime/WebBrowser/Private/WebBrowserWindow.h
--- a/Engine/Source/Runtime/WebBrowser/Private/WebBrowserWindow.h
+++ b/Engine/Source/Runtime/WebBrowser/Private/WebBrowserWindow.h
@@ -41,8 +41,9 @@ public:
 	 * @param InViewportSize Initial size of the browser window.
 	 * @param InInitialURL The Initial URL that will be loaded.
 	 * @param InContentsToLoad Optional string to load as a web page.
+	 * @param InShowErrorMessage Whether to show an error message in case of loading errors.
 	 */
-	FWebBrowserWindow(FIntPoint InViewportSize, FString InInitialURL, TOptional<FString> InContentsToLoad);
+	FWebBrowserWindow(FIntPoint InViewportSize, FString InInitialURL, TOptional<FString> InContentsToLoad, bool InShowErrorMessage);
 
 	/** Virtual Destructor. */
 	virtual ~FWebBrowserWindow();
@@ -70,7 +71,9 @@ public:
 	virtual bool IsValid() const override;
 	virtual bool HasBeenPainted() const override;
 	virtual bool IsClosing() const override;
+	virtual EWebBrowserDocumentState GetDocumentLoadingState() const override;
 	virtual FString GetTitle() const override;
+	virtual FString GetUrl() const override;
 	virtual void OnKeyDown(const FKeyEvent& InKeyEvent) override;
 	virtual void OnKeyUp(const FKeyEvent& InKeyEvent) override;
 	virtual void OnKeyChar(const FCharacterEvent& InCharacterEvent) override;
@@ -89,6 +92,12 @@ public:
 	virtual void Reload() override;
 	virtual void StopLoad() override;
 
+	DECLARE_DERIVED_EVENT(FWebBrowserWindow, IWebBrowserWindow::FOnDocumentStateChanged, FOnDocumentStateChanged);
+	virtual FOnDocumentStateChanged& OnDocumentStateChanged() override
+	{
+		return DocumentStateChangedEvent;
+	}
+
 	DECLARE_DERIVED_EVENT(FWebBrowserWindow, IWebBrowserWindow::FOnTitleChanged, FOnTitleChanged);
 	virtual FOnTitleChanged& OnTitleChanged() override
 	{
@@ -120,6 +129,16 @@ private:
 	 */
 	bool GetViewRect(CefRect& Rect);
 
+	/** Notifies clients that document loading has failed. */
+	void NotifyDocumentError();
+
+	/**
+	 * Notifies clients that the loading state of the document has changed.
+	 *
+	 * @param IsLoading Whether the document is loading (false = completed).
+	 */
+	void NotifyDocumentLoadingStateChange(bool IsLoading);
+
 	/**
 	 * Called when there is an update to the rendered web page.
 	 *
@@ -138,6 +157,8 @@ private:
 	 */
 	void OnCursorChange(CefCursorHandle Cursor);
 
+public:
+
 	/**
 	 * Gets the Cef Keyboard Modifiers based on a Key Event.
 	 * 
@@ -164,6 +185,9 @@ private:
 
 private:
 
+	/** Current state of the document being loaded. */
+	EWebBrowserDocumentState DocumentState;
+
 	/** Interface to the texture we are rendering to. */
 	FSlateUpdatableTexture* UpdatableTexture;
 
@@ -194,6 +218,12 @@ private:
 	/** Optional text to load as a web page. */
 	TOptional<FString> ContentsToLoad;
 
+	/** Delegate for broadcasting load state changes. */
+	FOnDocumentStateChanged DocumentStateChangedEvent;
+
+	/** Whether to show an error message in case of loading errors. */
+	bool ShowErrorMessage;
+
 	/** Delegate for broadcasting title changes. */
 	FOnTitleChanged TitleChangedEvent;
 };
diff --git a/Engine/Source/Runtime/WebBrowser/Public/IWebBrowserSingleton.h b/Engine/Source/Runtime/WebBrowser/Public/IWebBrowserSingleton.h
--- a/Engine/Source/Runtime/WebBrowser/Public/IWebBrowserSingleton.h
+++ b/Engine/Source/Runtime/WebBrowser/Public/IWebBrowserSingleton.h
@@ -26,7 +26,8 @@ public:
 	 * @param Height Initial height of the browser
 	 * @param bUseTransparency Whether to allow transparent rendering of pages
 	 * @param ContentsToLoad Optional string to load as a web page
+	 * @param ShowErrorMessage Whether to show an error message in case of loading errors.
 	 * @return New Web Browser Window Interface (may be null if not supported)
 	 */
-	virtual TSharedPtr<IWebBrowserWindow> CreateBrowserWindow(void* OSWindowHandle, FString InitialURL, uint32 Width, uint32 Height, bool bUseTransparency, TOptional<FString> ContentsToLoad = TOptional<FString>()) = 0;
+	virtual TSharedPtr<IWebBrowserWindow> CreateBrowserWindow(void* OSWindowHandle, FString InitialURL, uint32 Width, uint32 Height, bool bUseTransparency, TOptional<FString> ContentsToLoad = TOptional<FString>(), bool ShowErrorMessage = true) = 0;
 };
diff --git a/Engine/Source/Runtime/WebBrowser/Public/IWebBrowserWindow.h b/Engine/Source/Runtime/WebBrowser/Public/IWebBrowserWindow.h
--- a/Engine/Source/Runtime/WebBrowser/Public/IWebBrowserWindow.h
+++ b/Engine/Source/Runtime/WebBrowser/Public/IWebBrowserWindow.h
@@ -10,6 +10,15 @@ struct FPointerEvent;
 class FSlateShaderResource;
 
 
+enum class EWebBrowserDocumentState
+{
+	Completed,
+	Error,
+	Loading,
+	NoDocument
+};
+
+
 /**
  * Interface for dealing with a Web Browser window
  */
@@ -61,11 +70,21 @@ public:
 	 */
 	virtual bool IsClosing() const = 0;
 
+	/** Gets the loading state of the current document. */
+	virtual EWebBrowserDocumentState GetDocumentLoadingState() const = 0;
+
 	/**
 	 * Gets the current title of the Browser page
 	 */
 	virtual FString GetTitle() const = 0;
 
+	/**
+	 * Gets the currently loaded URL.
+	 *
+	 * @return The URL, or empty string if no document is loaded.
+	 */
+	virtual FString GetUrl() const = 0;
+
 	/**
 	 * Notify the browser that a key has been pressed
 	 *
@@ -164,8 +183,12 @@ public:
 
 public:
 
+	/** A delegate that is invoked when the loading state of a document changed. */
+	DECLARE_EVENT_OneParam(IWebBrowserWindow, FOnDocumentStateChanged, EWebBrowserDocumentState /*NewState*/);
+	virtual FOnDocumentStateChanged& OnDocumentStateChanged() = 0;
+
 	/** A delegate to allow callbacks when a browser title changes. */
-	DECLARE_EVENT_OneParam(IWebBrowserWindow, FOnTitleChanged, FString);
+	DECLARE_EVENT_OneParam(IWebBrowserWindow, FOnTitleChanged, FString /*NewTitle*/);
 	virtual FOnTitleChanged& OnTitleChanged() = 0;
 
 protected:
