fix for warning from calling ProcHandleClose directly

diff --git a/Engine/Source/Runtime/Online/BuildPatchServices/Private/BuildPatchInstaller.cpp b/Engine/Source/Runtime/Online/BuildPatchServices/Private/BuildPatchInstaller.cpp
--- a/Engine/Source/Runtime/Online/BuildPatchServices/Private/BuildPatchInstaller.cpp
+++ b/Engine/Source/Runtime/Online/BuildPatchServices/Private/BuildPatchInstaller.cpp
@@ -843,7 +843,7 @@ bool FBuildPatchInstaller::RunPrereqInstaller()
 		FPlatformProcess::WaitForProc(ProcessHandle);
 		FPlatformProcess::GetProcReturnCode(ProcessHandle, &ReturnCode);
 
-		ProcessHandle.Close();
+		FPlatformProcess::CloseProc(ProcessHandle);
 
 		if (ReturnCode != 0)
 		{
