Enabled UHT to digest GENERATED_BODY instead of GENERATED_UCLASS_BODY, GENERATED_USTRUCT_BODY, GENERATED_UINTERFACE_BODY or GENERATED_IINTERFACE_BODY, changed every occurence to the new syntax and fix

diff --git a/Engine/Source/Runtime/Online/OnlineSubsystemUtils/Classes/FindSessionsCallbackProxy.h b/Engine/Source/Runtime/Online/OnlineSubsystemUtils/Classes/FindSessionsCallbackProxy.h
--- a/Engine/Source/Runtime/Online/OnlineSubsystemUtils/Classes/FindSessionsCallbackProxy.h
+++ b/Engine/Source/Runtime/Online/OnlineSubsystemUtils/Classes/FindSessionsCallbackProxy.h
@@ -7,7 +7,7 @@
 USTRUCT(BlueprintType)
 struct FBlueprintSessionResult
 {
-	GENERATED_USTRUCT_BODY()
+	GENERATED_BODY()
 
 	FOnlineSessionSearchResult OnlineResult;
 };
@@ -17,7 +17,9 @@ DECLARE_DYNAMIC_MULTICAST_DELEGATE_OneParam(FBlueprintFindSessionsResultDelegate
 UCLASS(MinimalAPI)
 class UFindSessionsCallbackProxy : public UOnlineBlueprintCallProxyBase
 {
-	GENERATED_UCLASS_BODY()
+	GENERATED_BODY()
+public:
+	ONLINESUBSYSTEMUTILS_API UFindSessionsCallbackProxy(const FObjectInitializer& ObjectInitializer = FObjectInitializer::Get());
 
 	// Called when there is a successful query
 	UPROPERTY(BlueprintAssignable)
diff --git a/Engine/Source/Runtime/Online/OnlineSubsystemUtils/Classes/FindTurnBasedMatchCallbackProxy.h b/Engine/Source/Runtime/Online/OnlineSubsystemUtils/Classes/FindTurnBasedMatchCallbackProxy.h
--- a/Engine/Source/Runtime/Online/OnlineSubsystemUtils/Classes/FindTurnBasedMatchCallbackProxy.h
+++ b/Engine/Source/Runtime/Online/OnlineSubsystemUtils/Classes/FindTurnBasedMatchCallbackProxy.h
@@ -31,7 +31,9 @@ private:
 UCLASS(MinimalAPI)
 class UFindTurnBasedMatchCallbackProxy : public UOnlineBlueprintCallProxyBase
 {
-	GENERATED_UCLASS_BODY()
+	GENERATED_BODY()
+public:
+	ONLINESUBSYSTEMUTILS_API UFindTurnBasedMatchCallbackProxy(const FObjectInitializer& ObjectInitializer = FObjectInitializer::Get());
 
     virtual ~UFindTurnBasedMatchCallbackProxy();
     
diff --git a/Engine/Source/Runtime/Online/OnlineSubsystemUtils/Classes/InAppPurchaseCallbackProxy.h b/Engine/Source/Runtime/Online/OnlineSubsystemUtils/Classes/InAppPurchaseCallbackProxy.h
--- a/Engine/Source/Runtime/Online/OnlineSubsystemUtils/Classes/InAppPurchaseCallbackProxy.h
+++ b/Engine/Source/Runtime/Online/OnlineSubsystemUtils/Classes/InAppPurchaseCallbackProxy.h
@@ -9,7 +9,9 @@ DECLARE_DYNAMIC_MULTICAST_DELEGATE_TwoParams(FInAppPurchaseResult, EInAppPurchas
 UCLASS(MinimalAPI)
 class UInAppPurchaseCallbackProxy : public UObject
 {
-	GENERATED_UCLASS_BODY()
+	GENERATED_BODY()
+public:
+	ONLINESUBSYSTEMUTILS_API UInAppPurchaseCallbackProxy(const FObjectInitializer& ObjectInitializer = FObjectInitializer::Get());
 
 	// Called when there is a successful In-App Purchase transaction
 	UPROPERTY(BlueprintAssignable)
diff --git a/Engine/Source/Runtime/Online/OnlineSubsystemUtils/Classes/InAppPurchaseQueryCallbackProxy.h b/Engine/Source/Runtime/Online/OnlineSubsystemUtils/Classes/InAppPurchaseQueryCallbackProxy.h
--- a/Engine/Source/Runtime/Online/OnlineSubsystemUtils/Classes/InAppPurchaseQueryCallbackProxy.h
+++ b/Engine/Source/Runtime/Online/OnlineSubsystemUtils/Classes/InAppPurchaseQueryCallbackProxy.h
@@ -9,7 +9,9 @@ DECLARE_DYNAMIC_MULTICAST_DELEGATE_OneParam(FInAppPurchaseQueryResult, const TAr
 UCLASS(MinimalAPI)
 class UInAppPurchaseQueryCallbackProxy : public UObject
 {
-	GENERATED_UCLASS_BODY()
+	GENERATED_BODY()
+public:
+	ONLINESUBSYSTEMUTILS_API UInAppPurchaseQueryCallbackProxy(const FObjectInitializer& ObjectInitializer = FObjectInitializer::Get());
 
 	// Called when there is a successful InAppPurchase query
 	UPROPERTY(BlueprintAssignable)
diff --git a/Engine/Source/Runtime/Online/OnlineSubsystemUtils/Classes/IpConnection.h b/Engine/Source/Runtime/Online/OnlineSubsystemUtils/Classes/IpConnection.h
--- a/Engine/Source/Runtime/Online/OnlineSubsystemUtils/Classes/IpConnection.h
+++ b/Engine/Source/Runtime/Online/OnlineSubsystemUtils/Classes/IpConnection.h
@@ -10,7 +10,9 @@
 UCLASS(transient, config=Engine)
 class ONLINESUBSYSTEMUTILS_API UIpConnection : public UNetConnection
 {
-    GENERATED_UCLASS_BODY()
+    GENERATED_BODY()
+public:
+    UIpConnection(const FObjectInitializer& ObjectInitializer = FObjectInitializer::Get());
 	// Variables.
 	TSharedPtr<FInternetAddr>	RemoteAddr;
 	class FSocket*				Socket;
diff --git a/Engine/Source/Runtime/Online/OnlineSubsystemUtils/Classes/IpNetDriver.h b/Engine/Source/Runtime/Online/OnlineSubsystemUtils/Classes/IpNetDriver.h
--- a/Engine/Source/Runtime/Online/OnlineSubsystemUtils/Classes/IpNetDriver.h
+++ b/Engine/Source/Runtime/Online/OnlineSubsystemUtils/Classes/IpNetDriver.h
@@ -10,7 +10,9 @@
 UCLASS(transient, config=Engine)
 class ONLINESUBSYSTEMUTILS_API UIpNetDriver : public UNetDriver
 {
-    GENERATED_UCLASS_BODY()
+    GENERATED_BODY()
+public:
+    UIpNetDriver(const FObjectInitializer& ObjectInitializer = FObjectInitializer::Get());
 
 	/** Should port unreachable messages be logged */
     UPROPERTY(Config)
diff --git a/Engine/Source/Runtime/Online/OnlineSubsystemUtils/Classes/JoinSessionCallbackProxy.h b/Engine/Source/Runtime/Online/OnlineSubsystemUtils/Classes/JoinSessionCallbackProxy.h
--- a/Engine/Source/Runtime/Online/OnlineSubsystemUtils/Classes/JoinSessionCallbackProxy.h
+++ b/Engine/Source/Runtime/Online/OnlineSubsystemUtils/Classes/JoinSessionCallbackProxy.h
@@ -7,7 +7,9 @@
 UCLASS(MinimalAPI)
 class UJoinSessionCallbackProxy : public UOnlineBlueprintCallProxyBase
 {
-	GENERATED_UCLASS_BODY()
+	GENERATED_BODY()
+public:
+	ONLINESUBSYSTEMUTILS_API UJoinSessionCallbackProxy(const FObjectInitializer& ObjectInitializer = FObjectInitializer::Get());
 
 	// Called when there is a successful join
 	UPROPERTY(BlueprintAssignable)
diff --git a/Engine/Source/Runtime/Online/OnlineSubsystemUtils/Classes/LeaderboardBlueprintLibrary.h b/Engine/Source/Runtime/Online/OnlineSubsystemUtils/Classes/LeaderboardBlueprintLibrary.h
--- a/Engine/Source/Runtime/Online/OnlineSubsystemUtils/Classes/LeaderboardBlueprintLibrary.h
+++ b/Engine/Source/Runtime/Online/OnlineSubsystemUtils/Classes/LeaderboardBlueprintLibrary.h
@@ -10,7 +10,9 @@
 UCLASS()
 class ONLINESUBSYSTEMUTILS_API ULeaderboardBlueprintLibrary : public UBlueprintFunctionLibrary
 {
-	GENERATED_UCLASS_BODY()
+	GENERATED_BODY()
+public:
+	ULeaderboardBlueprintLibrary(const FObjectInitializer& ObjectInitializer = FObjectInitializer::Get());
 
 	/** Writes an integer value to the specified leaderboard */
 	UFUNCTION(BlueprintCallable, Category = "Online|Leaderboard")
diff --git a/Engine/Source/Runtime/Online/OnlineSubsystemUtils/Classes/LeaderboardFlushCallbackProxy.h b/Engine/Source/Runtime/Online/OnlineSubsystemUtils/Classes/LeaderboardFlushCallbackProxy.h
--- a/Engine/Source/Runtime/Online/OnlineSubsystemUtils/Classes/LeaderboardFlushCallbackProxy.h
+++ b/Engine/Source/Runtime/Online/OnlineSubsystemUtils/Classes/LeaderboardFlushCallbackProxy.h
@@ -9,7 +9,9 @@ DECLARE_DYNAMIC_MULTICAST_DELEGATE_OneParam(FOnLeaderboardFlushed, FName, Sessio
 UCLASS(MinimalAPI)
 class ULeaderboardFlushCallbackProxy : public UObject
 {
-	GENERATED_UCLASS_BODY()
+	GENERATED_BODY()
+public:
+	ONLINESUBSYSTEMUTILS_API ULeaderboardFlushCallbackProxy(const FObjectInitializer& ObjectInitializer = FObjectInitializer::Get());
 
 	// Called when there is a successful leaderboard flush
 	UPROPERTY(BlueprintAssignable)
diff --git a/Engine/Source/Runtime/Online/OnlineSubsystemUtils/Classes/LeaderboardQueryCallbackProxy.h b/Engine/Source/Runtime/Online/OnlineSubsystemUtils/Classes/LeaderboardQueryCallbackProxy.h
--- a/Engine/Source/Runtime/Online/OnlineSubsystemUtils/Classes/LeaderboardQueryCallbackProxy.h
+++ b/Engine/Source/Runtime/Online/OnlineSubsystemUtils/Classes/LeaderboardQueryCallbackProxy.h
@@ -9,7 +9,9 @@ DECLARE_DYNAMIC_MULTICAST_DELEGATE_OneParam(FLeaderboardQueryResult, int32, Lead
 UCLASS(MinimalAPI)
 class ULeaderboardQueryCallbackProxy : public UObject
 {
-	GENERATED_UCLASS_BODY()
+	GENERATED_BODY()
+public:
+	ONLINESUBSYSTEMUTILS_API ULeaderboardQueryCallbackProxy(const FObjectInitializer& ObjectInitializer = FObjectInitializer::Get());
 
 	// Called when there is a successful leaderboard query
 	UPROPERTY(BlueprintAssignable)
diff --git a/Engine/Source/Runtime/Online/OnlineSubsystemUtils/Classes/OnlineBeacon.h b/Engine/Source/Runtime/Online/OnlineSubsystemUtils/Classes/OnlineBeacon.h
--- a/Engine/Source/Runtime/Online/OnlineSubsystemUtils/Classes/OnlineBeacon.h
+++ b/Engine/Source/Runtime/Online/OnlineSubsystemUtils/Classes/OnlineBeacon.h
@@ -24,7 +24,9 @@ namespace EBeaconState
 UCLASS(transient, config=Engine, notplaceable)
 class ONLINESUBSYSTEMUTILS_API AOnlineBeacon : public AActor, public FNetworkNotify
 {
-	GENERATED_UCLASS_BODY()
+	GENERATED_BODY()
+public:
+	AOnlineBeacon(const FObjectInitializer& ObjectInitializer = FObjectInitializer::Get());
 
 	// Begin AActor Interface
 	virtual void OnActorChannelOpen(class FInBunch& InBunch, class UNetConnection* Connection) override;
diff --git a/Engine/Source/Runtime/Online/OnlineSubsystemUtils/Classes/OnlineBeaconClient.h b/Engine/Source/Runtime/Online/OnlineSubsystemUtils/Classes/OnlineBeaconClient.h
--- a/Engine/Source/Runtime/Online/OnlineSubsystemUtils/Classes/OnlineBeaconClient.h
+++ b/Engine/Source/Runtime/Online/OnlineSubsystemUtils/Classes/OnlineBeaconClient.h
@@ -19,7 +19,9 @@ DECLARE_DELEGATE(FOnHostConnectionFailure);
 UCLASS(transient, notplaceable, config=Engine)
 class ONLINESUBSYSTEMUTILS_API AOnlineBeaconClient : public AOnlineBeacon
 {
-	GENERATED_UCLASS_BODY()
+	GENERATED_BODY()
+public:
+	AOnlineBeaconClient(const FObjectInitializer& ObjectInitializer = FObjectInitializer::Get());
 
 	// Begin AActor Interface
 	void OnNetCleanup(class UNetConnection* Connection);
@@ -87,9 +89,10 @@ private:
 	/**
 	 * Called on the server side to open up the actor channel that will allow RPCs to occur
 	 */
-	UFUNCTION(client, reliable)
+	UFUNCTION(client="ClientOnConnected_Implementation", reliable)
 	virtual void ClientOnConnected();
+	virtual void ClientOnConnected_Implementation();
 
 	friend class AOnlineBeaconHost;
 	friend class AOnlineBeaconHostObject;
-};
\ No newline at end of file
+};
diff --git a/Engine/Source/Runtime/Online/OnlineSubsystemUtils/Classes/OnlineBeaconHost.h b/Engine/Source/Runtime/Online/OnlineSubsystemUtils/Classes/OnlineBeaconHost.h
--- a/Engine/Source/Runtime/Online/OnlineSubsystemUtils/Classes/OnlineBeaconHost.h
+++ b/Engine/Source/Runtime/Online/OnlineSubsystemUtils/Classes/OnlineBeaconHost.h
@@ -19,7 +19,9 @@ class AOnlineBeaconClient;
 UCLASS(transient, notplaceable, config=Engine)
 class ONLINESUBSYSTEMUTILS_API AOnlineBeaconHost : public AOnlineBeacon
 {
-	GENERATED_UCLASS_BODY()
+	GENERATED_BODY()
+public:
+	AOnlineBeaconHost(const FObjectInitializer& ObjectInitializer = FObjectInitializer::Get());
 
 	/** Configured listen port for this beacon host */
 	UPROPERTY(Config)
@@ -116,4 +118,4 @@ private:
 
 	/** Mapping of beacon types to the OnBeaconConnected delegates */
 	TMap<FString, FOnBeaconConnected> OnBeaconConnectedMapping;
-};
\ No newline at end of file
+};
diff --git a/Engine/Source/Runtime/Online/OnlineSubsystemUtils/Classes/OnlineBeaconHostObject.h b/Engine/Source/Runtime/Online/OnlineSubsystemUtils/Classes/OnlineBeaconHostObject.h
--- a/Engine/Source/Runtime/Online/OnlineSubsystemUtils/Classes/OnlineBeaconHostObject.h
+++ b/Engine/Source/Runtime/Online/OnlineSubsystemUtils/Classes/OnlineBeaconHostObject.h
@@ -19,7 +19,9 @@ class UNetConnection;
 UCLASS(transient, config=Engine, notplaceable)
 class ONLINESUBSYSTEMUTILS_API AOnlineBeaconHostObject : public AActor
 {
-	GENERATED_UCLASS_BODY()
+	GENERATED_BODY()
+public:
+	AOnlineBeaconHostObject(const FObjectInitializer& ObjectInitializer = FObjectInitializer::Get());
 
 	/** Custom name for this beacon */
 	UPROPERTY(Transient)
@@ -52,4 +54,4 @@ class ONLINESUBSYSTEMUTILS_API AOnlineBeaconHostObject : public AActor
 	 * @param ClientActor client beacon actor to remove
 	 */
 	virtual void RemoveClientActor(AOnlineBeaconClient* ClientActor);
-};
\ No newline at end of file
+};
diff --git a/Engine/Source/Runtime/Online/OnlineSubsystemUtils/Classes/OnlineBlueprintCallProxyBase.h b/Engine/Source/Runtime/Online/OnlineSubsystemUtils/Classes/OnlineBlueprintCallProxyBase.h
--- a/Engine/Source/Runtime/Online/OnlineSubsystemUtils/Classes/OnlineBlueprintCallProxyBase.h
+++ b/Engine/Source/Runtime/Online/OnlineSubsystemUtils/Classes/OnlineBlueprintCallProxyBase.h
@@ -8,7 +8,9 @@ DECLARE_DYNAMIC_MULTICAST_DELEGATE(FEmptyOnlineDelegate);
 UCLASS(MinimalAPI)
 class UOnlineBlueprintCallProxyBase : public UObject
 {
-	GENERATED_UCLASS_BODY()
+	GENERATED_BODY()
+public:
+	ONLINESUBSYSTEMUTILS_API UOnlineBlueprintCallProxyBase(const FObjectInitializer& ObjectInitializer = FObjectInitializer::Get());
 
 	// Called to trigger the actual online action once the delegates have been bound
 	UFUNCTION(BlueprintCallable, meta=(BlueprintInternalUseOnly="true"), Category="Online")
diff --git a/Engine/Source/Runtime/Online/OnlineSubsystemUtils/Classes/OnlineSessionClient.h b/Engine/Source/Runtime/Online/OnlineSubsystemUtils/Classes/OnlineSessionClient.h
--- a/Engine/Source/Runtime/Online/OnlineSubsystemUtils/Classes/OnlineSessionClient.h
+++ b/Engine/Source/Runtime/Online/OnlineSubsystemUtils/Classes/OnlineSessionClient.h
@@ -16,7 +16,9 @@ class APlayerController;
 UCLASS(config=Game)
 class ONLINESUBSYSTEMUTILS_API UOnlineSessionClient : public UOnlineSession
 {
-	GENERATED_UCLASS_BODY()
+	GENERATED_BODY()
+public:
+	UOnlineSessionClient(const FObjectInitializer& ObjectInitializer = FObjectInitializer::Get());
 
 protected:
 
diff --git a/Engine/Source/Runtime/Online/OnlineSubsystemUtils/Classes/PartyBeaconClient.h b/Engine/Source/Runtime/Online/OnlineSubsystemUtils/Classes/PartyBeaconClient.h
--- a/Engine/Source/Runtime/Online/OnlineSubsystemUtils/Classes/PartyBeaconClient.h
+++ b/Engine/Source/Runtime/Online/OnlineSubsystemUtils/Classes/PartyBeaconClient.h
@@ -22,7 +22,9 @@ DECLARE_DELEGATE_OneParam(FOnReservationRequestComplete, EPartyReservationResult
 UCLASS(transient, notplaceable, config=Engine)
 class ONLINESUBSYSTEMUTILS_API APartyBeaconClient : public AOnlineBeaconClient
 {
-	GENERATED_UCLASS_BODY()
+	GENERATED_BODY()
+public:
+	APartyBeaconClient(const FObjectInitializer& ObjectInitializer = FObjectInitializer::Get());
 
 	// Begin AOnlineBeacon Interface
 	virtual FString GetBeaconType() override { return PARTY_BEACON_TYPE; }
@@ -55,8 +57,9 @@ class ONLINESUBSYSTEMUTILS_API APartyBeaconClient : public AOnlineBeaconClient
 	 *
 	 * @param ReservationResponse response from server
 	 */
- 	UFUNCTION(client, reliable)
+ 	UFUNCTION(client="ClientReservationResponse_Implementation", reliable)
 	virtual void ClientReservationResponse(EPartyReservationResult::Type ReservationResponse);
+ 	virtual void ClientReservationResponse_Implementation(EPartyReservationResult::Type ReservationResponse);
 
 	/**
 	 * Delegate triggered when a response from the party beacon host has been received
@@ -90,9 +93,13 @@ protected:
 	 *
 	 * @param Reservation pending reservation request to make with server
 	 */
-	UFUNCTION(server, reliable, WithValidation)
+	UFUNCTION(server="ServerReservationRequest_Implementation", reliable, WithValidation="ServerReservationRequest_Validate")
 	virtual void ServerReservationRequest(const FString& SessionId, struct FPartyReservation Reservation);
+	virtual void ServerReservationRequest_Implementation(const FString& SessionId, FPartyReservation Reservation);
+	virtual bool ServerReservationRequest_Validate(const FString& , FPartyReservation );
 
-	UFUNCTION(server, reliable, WithValidation)
+	UFUNCTION(server="ServerCancelReservationRequest_Implementation", reliable, WithValidation="ServerCancelReservationRequest_Validate")
 	virtual void ServerCancelReservationRequest(struct FUniqueNetIdRepl PartyLeader);
+	virtual void ServerCancelReservationRequest_Implementation(FUniqueNetIdRepl PartyLeader);
+	virtual bool ServerCancelReservationRequest_Validate(FUniqueNetIdRepl );
 };
diff --git a/Engine/Source/Runtime/Online/OnlineSubsystemUtils/Classes/PartyBeaconHost.h b/Engine/Source/Runtime/Online/OnlineSubsystemUtils/Classes/PartyBeaconHost.h
--- a/Engine/Source/Runtime/Online/OnlineSubsystemUtils/Classes/PartyBeaconHost.h
+++ b/Engine/Source/Runtime/Online/OnlineSubsystemUtils/Classes/PartyBeaconHost.h
@@ -49,7 +49,9 @@ DECLARE_DELEGATE(FOnReservationsFull);
 UCLASS(transient, notplaceable, config=Engine)
 class ONLINESUBSYSTEMUTILS_API APartyBeaconHost : public AOnlineBeaconHostObject
 {
-	GENERATED_UCLASS_BODY()
+	GENERATED_BODY()
+public:
+	APartyBeaconHost(const FObjectInitializer& ObjectInitializer = FObjectInitializer::Get());
 
 	// Begin AActor Interface
 	virtual void Tick(float DeltaTime) override;
diff --git a/Engine/Source/Runtime/Online/OnlineSubsystemUtils/Classes/PartyBeaconState.h b/Engine/Source/Runtime/Online/OnlineSubsystemUtils/Classes/PartyBeaconState.h
--- a/Engine/Source/Runtime/Online/OnlineSubsystemUtils/Classes/PartyBeaconState.h
+++ b/Engine/Source/Runtime/Online/OnlineSubsystemUtils/Classes/PartyBeaconState.h
@@ -127,7 +127,7 @@ namespace EPartyReservationResult
 USTRUCT()
 struct FPlayerReservation
 {
-	GENERATED_USTRUCT_BODY()
+	GENERATED_BODY()
 	
 	#if CPP
 	FPlayerReservation() :
@@ -152,7 +152,7 @@ struct FPlayerReservation
 USTRUCT()
 struct ONLINESUBSYSTEMUTILS_API FPartyReservation
 {
-	GENERATED_USTRUCT_BODY()
+	GENERATED_BODY()
 
 	/** Team assigned to this party */
 	UPROPERTY(Transient)
@@ -176,7 +176,9 @@ struct ONLINESUBSYSTEMUTILS_API FPartyReservation
 UCLASS(transient, notplaceable, config=Engine)
 class ONLINESUBSYSTEMUTILS_API UPartyBeaconState : public UObject
 {
-	GENERATED_UCLASS_BODY()
+	GENERATED_BODY()
+public:
+	UPartyBeaconState(const FObjectInitializer& ObjectInitializer = FObjectInitializer::Get());
 
 	/**
 	 * Initialize this state object 
diff --git a/Engine/Source/Runtime/Online/OnlineSubsystemUtils/Classes/QuitMatchCallbackProxy.h b/Engine/Source/Runtime/Online/OnlineSubsystemUtils/Classes/QuitMatchCallbackProxy.h
--- a/Engine/Source/Runtime/Online/OnlineSubsystemUtils/Classes/QuitMatchCallbackProxy.h
+++ b/Engine/Source/Runtime/Online/OnlineSubsystemUtils/Classes/QuitMatchCallbackProxy.h
@@ -8,8 +8,10 @@
 UCLASS(MinimalAPI)
 class UQuitMatchCallbackProxy : public UOnlineBlueprintCallProxyBase
 {
+	GENERATED_BODY()
+public:
+	ONLINESUBSYSTEMUTILS_API UQuitMatchCallbackProxy(const FObjectInitializer& ObjectInitializer = FObjectInitializer::Get());
 public:
-	GENERATED_UCLASS_BODY()
 
 	virtual ~UQuitMatchCallbackProxy();
 
@@ -46,4 +48,4 @@ private:
 
 	// If the match isn't over, this will be how much time the next player will have to complete their turn
 	int32 TurnTimeoutInSeconds;
-};
\ No newline at end of file
+};
diff --git a/Engine/Source/Runtime/Online/OnlineSubsystemUtils/Classes/TestBeaconClient.h b/Engine/Source/Runtime/Online/OnlineSubsystemUtils/Classes/TestBeaconClient.h
--- a/Engine/Source/Runtime/Online/OnlineSubsystemUtils/Classes/TestBeaconClient.h
+++ b/Engine/Source/Runtime/Online/OnlineSubsystemUtils/Classes/TestBeaconClient.h
@@ -10,7 +10,9 @@
 UCLASS(transient, notplaceable, config=Engine)
 class ONLINESUBSYSTEMUTILS_API ATestBeaconClient : public AOnlineBeaconClient
 {
-	GENERATED_UCLASS_BODY()
+	GENERATED_BODY()
+public:
+	ATestBeaconClient(const FObjectInitializer& ObjectInitializer = FObjectInitializer::Get());
 
 	// Begin AOnlineBeacon Interface
 	virtual FString GetBeaconType() override { return TEXT("TestBeacon"); }
@@ -21,10 +23,13 @@ class ONLINESUBSYSTEMUTILS_API ATestBeaconClient : public AOnlineBeaconClient
 	// End AOnlineBeaconClient Interface
 
 	/** Send a ping RPC to the client */
-	UFUNCTION(client, reliable)
+	UFUNCTION(client="ClientPing_Implementation", reliable)
 	virtual void ClientPing();
+	virtual void ClientPing_Implementation();
 
 	/** Send a pong RPC to the host */
-	UFUNCTION(server, reliable, WithValidation)
+	UFUNCTION(server="ServerPong_Implementation", reliable, WithValidation="ServerPong_Validate")
 	virtual void ServerPong();
+	virtual void ServerPong_Implementation();
+	virtual bool ServerPong_Validate();
 };
diff --git a/Engine/Source/Runtime/Online/OnlineSubsystemUtils/Classes/TestBeaconHost.h b/Engine/Source/Runtime/Online/OnlineSubsystemUtils/Classes/TestBeaconHost.h
--- a/Engine/Source/Runtime/Online/OnlineSubsystemUtils/Classes/TestBeaconHost.h
+++ b/Engine/Source/Runtime/Online/OnlineSubsystemUtils/Classes/TestBeaconHost.h
@@ -10,7 +10,9 @@
 UCLASS(transient, notplaceable, config=Engine)
 class ONLINESUBSYSTEMUTILS_API ATestBeaconHost : public AOnlineBeaconHostObject
 {
-	GENERATED_UCLASS_BODY()
+	GENERATED_BODY()
+public:
+	ATestBeaconHost(const FObjectInitializer& ObjectInitializer = FObjectInitializer::Get());
 
 	// Begin AOnlineBeaconHost Interface 
 	virtual AOnlineBeaconClient* SpawnBeaconActor(class UNetConnection* ClientConnection) override;
diff --git a/Engine/Source/Runtime/Online/OnlineSubsystemUtils/Classes/TurnBasedBlueprintLibrary.h b/Engine/Source/Runtime/Online/OnlineSubsystemUtils/Classes/TurnBasedBlueprintLibrary.h
--- a/Engine/Source/Runtime/Online/OnlineSubsystemUtils/Classes/TurnBasedBlueprintLibrary.h
+++ b/Engine/Source/Runtime/Online/OnlineSubsystemUtils/Classes/TurnBasedBlueprintLibrary.h
@@ -9,7 +9,9 @@
 UCLASS()
 class ONLINESUBSYSTEMUTILS_API UTurnBasedBlueprintLibrary : public UBlueprintFunctionLibrary
 {
-	GENERATED_UCLASS_BODY()
+	GENERATED_BODY()
+public:
+	UTurnBasedBlueprintLibrary(const FObjectInitializer& ObjectInitializer = FObjectInitializer::Get());
 
 	UFUNCTION(BlueprintCallable, meta = (WorldContext = "WorldContextObject"), Category = "Online|TurnBased")
 	static void GetIsMyTurn(UObject* WorldContextObject, APlayerController* PlayerController, FString MatchID, /*out*/ bool& bIsMyTurn);
diff --git a/Engine/Source/Runtime/Slate/Public/Framework/Styling/ButtonWidgetStyle.h b/Engine/Source/Runtime/Slate/Public/Framework/Styling/ButtonWidgetStyle.h
--- a/Engine/Source/Runtime/Slate/Public/Framework/Styling/ButtonWidgetStyle.h
+++ b/Engine/Source/Runtime/Slate/Public/Framework/Styling/ButtonWidgetStyle.h
@@ -10,10 +10,9 @@
 UCLASS(BlueprintType, hidecategories=Object, MinimalAPI)
 class UButtonWidgetStyle : public USlateWidgetStyleContainerBase
 {
-public:
 	GENERATED_BODY()
-
 public:
+
 	/** The actual data describing the button's appearance. */
 	UPROPERTY(Category=Appearance, EditAnywhere, BlueprintReadWrite, meta=( ShowOnlyInnerProperties ))
 	FButtonStyle ButtonStyle;
diff --git a/Engine/Source/Runtime/SlateCore/Public/Fonts/FontProviderInterface.h b/Engine/Source/Runtime/SlateCore/Public/Fonts/FontProviderInterface.h
--- a/Engine/Source/Runtime/SlateCore/Public/Fonts/FontProviderInterface.h
+++ b/Engine/Source/Runtime/SlateCore/Public/Fonts/FontProviderInterface.h
@@ -8,12 +8,15 @@
 UINTERFACE(MinimalAPI, meta=(CannotImplementInterfaceInBlueprint))
 class UFontProviderInterface : public UInterface
 {
-	GENERATED_UINTERFACE_BODY()
+	GENERATED_BODY()
+public:
+	SLATECORE_API UFontProviderInterface(const FObjectInitializer& ObjectInitializer = FObjectInitializer::Get());
 };
 
 class IFontProviderInterface
 {
-	GENERATED_IINTERFACE_BODY()
+	GENERATED_BODY()
+public:
 
 	virtual const FCompositeFont* GetCompositeFont() const = 0;
 };
diff --git a/Engine/Source/Runtime/SlateCore/Public/Input/Events.h b/Engine/Source/Runtime/SlateCore/Public/Input/Events.h
--- a/Engine/Source/Runtime/SlateCore/Public/Input/Events.h
+++ b/Engine/Source/Runtime/SlateCore/Public/Input/Events.h
@@ -56,7 +56,7 @@ namespace EKeyboardFocusCause
 USTRUCT()
 struct FFocusEvent
 {
-	GENERATED_USTRUCT_BODY()
+	GENERATED_BODY()
 
 public:
 
@@ -128,7 +128,7 @@ struct FVirtualPointerPosition
 USTRUCT(BlueprintType)
 struct FInputEvent
 {
-	GENERATED_USTRUCT_BODY()
+	GENERATED_BODY()
 
 public:
 
@@ -358,7 +358,7 @@ USTRUCT(BlueprintType)
 struct FKeyEvent
 	: public FInputEvent
 {
-	GENERATED_USTRUCT_BODY()
+	GENERATED_BODY()
 
 public:
 	/**
@@ -445,7 +445,7 @@ USTRUCT(BlueprintType)
 struct FKeyboardEvent : public FKeyEvent
 //struct DEPRECATED(4.6, "Use FKeyEvent") FKeyboardEvent : public FKeyEvent
 {
-	GENERATED_USTRUCT_BODY()
+	GENERATED_BODY()
 public:
 	FKeyEvent() {}
 
@@ -461,7 +461,7 @@ USTRUCT(BlueprintType)
 struct FAnalogInputEvent
 	: public FKeyEvent
 {
-	GENERATED_USTRUCT_BODY()
+	GENERATED_BODY()
 
 public:
 	/**
@@ -515,7 +515,7 @@ USTRUCT()
 struct FCharacterEvent
 	: public FInputEvent
 {
-	GENERATED_USTRUCT_BODY()
+	GENERATED_BODY()
 
 public:
 	/**
@@ -587,7 +587,7 @@ USTRUCT(BlueprintType)
 struct FPointerEvent
 	: public FInputEvent
 {
-	GENERATED_USTRUCT_BODY()
+	GENERATED_BODY()
 public:
 
 	/**
@@ -782,7 +782,7 @@ USTRUCT(BlueprintType)
 struct FControllerEvent
 	: public FInputEvent
 {
-	GENERATED_USTRUCT_BODY()
+	GENERATED_BODY()
 public:
 
 	FControllerEvent()
@@ -816,7 +816,7 @@ USTRUCT(BlueprintType)
 struct FMotionEvent
 	: public FInputEvent
 {
-	GENERATED_USTRUCT_BODY()
+	GENERATED_BODY()
 
 public:
 	/**
@@ -884,7 +884,7 @@ USTRUCT(BlueprintType)
 struct FNavigationEvent
 	: public FInputEvent
 {
-	GENERATED_USTRUCT_BODY()
+	GENERATED_BODY()
 
 public:
 	/**
diff --git a/Engine/Source/Runtime/SlateCore/Public/Styling/SlateWidgetStyleContainerInterface.h b/Engine/Source/Runtime/SlateCore/Public/Styling/SlateWidgetStyleContainerInterface.h
--- a/Engine/Source/Runtime/SlateCore/Public/Styling/SlateWidgetStyleContainerInterface.h
+++ b/Engine/Source/Runtime/SlateCore/Public/Styling/SlateWidgetStyleContainerInterface.h
@@ -8,14 +8,17 @@
 UINTERFACE()
 class SLATECORE_API USlateWidgetStyleContainerInterface : public UInterface
 {
-	GENERATED_UINTERFACE_BODY()
+	GENERATED_BODY()
+public:
+	USlateWidgetStyleContainerInterface(const FObjectInitializer& ObjectInitializer = FObjectInitializer::Get());
 };
 
 class ISlateWidgetStyleContainerInterface
 {
-	GENERATED_IINTERFACE_BODY()
+	GENERATED_BODY()
+public:
 
 public:
 
 	virtual const struct FSlateWidgetStyle* const GetStyle() const = 0;
-};
\ No newline at end of file
+};
diff --git a/Engine/Source/Runtime/UMG/Public/Animation/MovieScene2DTransformSection.h b/Engine/Source/Runtime/UMG/Public/Animation/MovieScene2DTransformSection.h
--- a/Engine/Source/Runtime/UMG/Public/Animation/MovieScene2DTransformSection.h
+++ b/Engine/Source/Runtime/UMG/Public/Animation/MovieScene2DTransformSection.h
@@ -14,7 +14,9 @@ struct FWidgetTransform;
 UCLASS(MinimalAPI)
 class UMovieScene2DTransformSection : public UMovieSceneSection
 {
-	GENERATED_UCLASS_BODY()
+	GENERATED_BODY()
+public:
+	UMG_API UMovieScene2DTransformSection(const FObjectInitializer& ObjectInitializer = FObjectInitializer::Get());
 public:
 
 	/** MovieSceneSection interface */
@@ -52,4 +54,4 @@ private:
 	/** Shear curve */
 	UPROPERTY()
 	FRichCurve Shear[2];
-};
\ No newline at end of file
+};
diff --git a/Engine/Source/Runtime/UMG/Public/Animation/MovieScene2DTransformTrack.h b/Engine/Source/Runtime/UMG/Public/Animation/MovieScene2DTransformTrack.h
--- a/Engine/Source/Runtime/UMG/Public/Animation/MovieScene2DTransformTrack.h
+++ b/Engine/Source/Runtime/UMG/Public/Animation/MovieScene2DTransformTrack.h
@@ -20,7 +20,9 @@ struct F2DTransformKey
 UCLASS( MinimalAPI )
 class UMovieScene2DTransformTrack : public UMovieScenePropertyTrack
 {
-	GENERATED_UCLASS_BODY()
+	GENERATED_BODY()
+public:
+	UMG_API UMovieScene2DTransformTrack(const FObjectInitializer& ObjectInitializer = FObjectInitializer::Get());
 public:
 	/** UMovieSceneTrack interface */
 	virtual UMovieSceneSection* CreateNewSection() override;
@@ -40,4 +42,4 @@ public:
 
 	UMG_API bool AddKeyToSection( float Time, const F2DTransformKey& TransformKey );
 
-};
\ No newline at end of file
+};
diff --git a/Engine/Source/Runtime/UMG/Public/Animation/MovieSceneMarginSection.h b/Engine/Source/Runtime/UMG/Public/Animation/MovieSceneMarginSection.h
--- a/Engine/Source/Runtime/UMG/Public/Animation/MovieSceneMarginSection.h
+++ b/Engine/Source/Runtime/UMG/Public/Animation/MovieSceneMarginSection.h
@@ -14,7 +14,9 @@ struct FMarginKey;
 UCLASS(MinimalAPI)
 class UMovieSceneMarginSection : public UMovieSceneSection
 {
-	GENERATED_UCLASS_BODY()
+	GENERATED_BODY()
+public:
+	UMG_API UMovieSceneMarginSection(const FObjectInitializer& ObjectInitializer = FObjectInitializer::Get());
 public:
 	/** UMovieSceneSection interface */
 	virtual void MoveSection( float DeltaPosition ) override;
diff --git a/Engine/Source/Runtime/UMG/Public/Animation/MovieSceneMarginTrack.h b/Engine/Source/Runtime/UMG/Public/Animation/MovieSceneMarginTrack.h
--- a/Engine/Source/Runtime/UMG/Public/Animation/MovieSceneMarginTrack.h
+++ b/Engine/Source/Runtime/UMG/Public/Animation/MovieSceneMarginTrack.h
@@ -18,7 +18,9 @@ struct FMarginKey
 UCLASS( MinimalAPI )
 class UMovieSceneMarginTrack : public UMovieScenePropertyTrack
 {
-	GENERATED_UCLASS_BODY()
+	GENERATED_BODY()
+public:
+	UMG_API UMovieSceneMarginTrack(const FObjectInitializer& ObjectInitializer = FObjectInitializer::Get());
 
 public:
 	/** UMovieSceneTrack interface */
diff --git a/Engine/Source/Runtime/UMG/Public/Animation/UMGSequencePlayer.h b/Engine/Source/Runtime/UMG/Public/Animation/UMGSequencePlayer.h
--- a/Engine/Source/Runtime/UMG/Public/Animation/UMGSequencePlayer.h
+++ b/Engine/Source/Runtime/UMG/Public/Animation/UMGSequencePlayer.h
@@ -11,7 +11,9 @@ class UMovieSceneBindings;
 UCLASS(transient)
 class UMG_API UUMGSequencePlayer : public UObject, public IMovieScenePlayer
 {
-	GENERATED_UCLASS_BODY()
+	GENERATED_BODY()
+public:
+	UUMGSequencePlayer(const FObjectInitializer& ObjectInitializer = FObjectInitializer::Get());
 
 public:
 	void InitSequencePlayer( const UWidgetAnimation& InAnimation, UUserWidget& UserWidget );
diff --git a/Engine/Source/Runtime/UMG/Public/Animation/WidgetAnimation.h b/Engine/Source/Runtime/UMG/Public/Animation/WidgetAnimation.h
--- a/Engine/Source/Runtime/UMG/Public/Animation/WidgetAnimation.h
+++ b/Engine/Source/Runtime/UMG/Public/Animation/WidgetAnimation.h
@@ -12,7 +12,7 @@ class UWidgetTree;
 USTRUCT()
 struct FWidgetAnimationBinding
 {
-	GENERATED_USTRUCT_BODY()
+	GENERATED_BODY()
 
 	UPROPERTY()
 	FName WidgetName;
@@ -48,7 +48,9 @@ public:
 UCLASS(MinimalAPI)
 class UWidgetAnimation : public UObject
 {
-	GENERATED_UCLASS_BODY()
+	GENERATED_BODY()
+public:
+	UMG_API UWidgetAnimation(const FObjectInitializer& ObjectInitializer = FObjectInitializer::Get());
 
 public:
 
diff --git a/Engine/Source/Runtime/UMG/Public/Binding/DynamicPropertyPath.h b/Engine/Source/Runtime/UMG/Public/Binding/DynamicPropertyPath.h
--- a/Engine/Source/Runtime/UMG/Public/Binding/DynamicPropertyPath.h
+++ b/Engine/Source/Runtime/UMG/Public/Binding/DynamicPropertyPath.h
@@ -1,4 +1,4 @@
-﻿// Copyright 1998-2015 Epic Games, Inc. All Rights Reserved.
+// Copyright 1998-2015 Epic Games, Inc. All Rights Reserved.
 
 #pragma once
 
@@ -123,7 +123,7 @@ inline bool IsConcreteTypeCompatibleWithReflectedType<UObject*>(UProperty* Prope
 USTRUCT()
 struct UMG_API FPropertyPathSegment
 {
-	GENERATED_USTRUCT_BODY()
+	GENERATED_BODY()
 
 public:
 	/** Implementation detail, don't use this constructor. */
@@ -166,7 +166,7 @@ private:
 USTRUCT()
 struct UMG_API FDynamicPropertyPath
 {
-	GENERATED_USTRUCT_BODY()
+	GENERATED_BODY()
 
 public:
 
diff --git a/Engine/Source/Runtime/UMG/Public/Blueprint/AsyncTaskDownloadImage.h b/Engine/Source/Runtime/UMG/Public/Blueprint/AsyncTaskDownloadImage.h
--- a/Engine/Source/Runtime/UMG/Public/Blueprint/AsyncTaskDownloadImage.h
+++ b/Engine/Source/Runtime/UMG/Public/Blueprint/AsyncTaskDownloadImage.h
@@ -12,7 +12,9 @@ DECLARE_DYNAMIC_MULTICAST_DELEGATE_OneParam(FDownloadImageDelegate, UTexture2D*,
 UCLASS()
 class UMG_API UAsyncTaskDownloadImage : public UBlueprintAsyncActionBase
 {
-	GENERATED_UCLASS_BODY()
+	GENERATED_BODY()
+public:
+	UAsyncTaskDownloadImage(const FObjectInitializer& ObjectInitializer = FObjectInitializer::Get());
 
 public:
 	UFUNCTION(BlueprintCallable, meta=( BlueprintInternalUseOnly="true" ))
diff --git a/Engine/Source/Runtime/UMG/Public/Blueprint/DragDropOperation.h b/Engine/Source/Runtime/UMG/Public/Blueprint/DragDropOperation.h
--- a/Engine/Source/Runtime/UMG/Public/Blueprint/DragDropOperation.h
+++ b/Engine/Source/Runtime/UMG/Public/Blueprint/DragDropOperation.h
@@ -33,7 +33,9 @@ enum class EDragPivot : uint8
 UCLASS(BlueprintType, Blueprintable)
 class UMG_API UDragDropOperation : public UObject
 {
-	GENERATED_UCLASS_BODY()
+	GENERATED_BODY()
+public:
+	UDragDropOperation(const FObjectInitializer& ObjectInitializer = FObjectInitializer::Get());
 	
 public:
 	/** A simple string tag you can optionally use to provide extra metadata about the operation. */
@@ -80,12 +82,15 @@ public:
 	/**  */
 	UFUNCTION(BlueprintNativeEvent, Category="Drag and Drop")
 	void Drop(const FPointerEvent& PointerEvent);
+	virtual void Drop_Implementation(const FPointerEvent& PointerEvent);
 
 	/**  */
 	UFUNCTION(BlueprintNativeEvent, Category="Drag and Drop")
 	void DragCancelled(const FPointerEvent& PointerEvent);
+	virtual void DragCancelled_Implementation(const FPointerEvent& PointerEvent);
 
 	/**  */
 	UFUNCTION(BlueprintNativeEvent, Category="Drag and Drop")
 	void Dragged(const FPointerEvent& PointerEvent);
+	virtual void Dragged_Implementation(const FPointerEvent& PointerEvent);
 };
diff --git a/Engine/Source/Runtime/UMG/Public/Blueprint/SlateBlueprintLibrary.h b/Engine/Source/Runtime/UMG/Public/Blueprint/SlateBlueprintLibrary.h
--- a/Engine/Source/Runtime/UMG/Public/Blueprint/SlateBlueprintLibrary.h
+++ b/Engine/Source/Runtime/UMG/Public/Blueprint/SlateBlueprintLibrary.h
@@ -7,7 +7,9 @@
 UCLASS(MinimalAPI)
 class USlateBlueprintLibrary : public UBlueprintFunctionLibrary
 {
-	GENERATED_UCLASS_BODY()
+	GENERATED_BODY()
+public:
+	UMG_API USlateBlueprintLibrary(const FObjectInitializer& ObjectInitializer = FObjectInitializer::Get());
 
 public:
 
diff --git a/Engine/Source/Runtime/UMG/Public/Blueprint/WidgetBlueprintGeneratedClass.h b/Engine/Source/Runtime/UMG/Public/Blueprint/WidgetBlueprintGeneratedClass.h
--- a/Engine/Source/Runtime/UMG/Public/Blueprint/WidgetBlueprintGeneratedClass.h
+++ b/Engine/Source/Runtime/UMG/Public/Blueprint/WidgetBlueprintGeneratedClass.h
@@ -22,7 +22,7 @@ enum class EBindingKind : uint8
 USTRUCT()
 struct FDelegateRuntimeBinding
 {
-	GENERATED_USTRUCT_BODY()
+	GENERATED_BODY()
 
 	/** The widget that will be bound to the live data. */
 	UPROPERTY()
@@ -55,7 +55,9 @@ struct FDelegateRuntimeBinding
 UCLASS()
 class UMG_API UWidgetBlueprintGeneratedClass : public UBlueprintGeneratedClass
 {
-	GENERATED_UCLASS_BODY()
+	GENERATED_BODY()
+public:
+	UWidgetBlueprintGeneratedClass(const FObjectInitializer& ObjectInitializer = FObjectInitializer::Get());
 
 public:
 
@@ -82,4 +84,4 @@ public:
 	 * binding and wiring necessary to have the user's widget perform as desired.
 	 */
 	void InitializeWidget(UUserWidget* UserWidget) const;
-};
\ No newline at end of file
+};
diff --git a/Engine/Source/Runtime/UMG/Public/Blueprint/WidgetBlueprintLibrary.h b/Engine/Source/Runtime/UMG/Public/Blueprint/WidgetBlueprintLibrary.h
--- a/Engine/Source/Runtime/UMG/Public/Blueprint/WidgetBlueprintLibrary.h
+++ b/Engine/Source/Runtime/UMG/Public/Blueprint/WidgetBlueprintLibrary.h
@@ -13,7 +13,9 @@ struct FPaintContext;
 UCLASS(MinimalAPI)
 class UWidgetBlueprintLibrary : public UBlueprintFunctionLibrary
 {
-	GENERATED_UCLASS_BODY()
+	GENERATED_BODY()
+public:
+	UMG_API UWidgetBlueprintLibrary(const FObjectInitializer& ObjectInitializer = FObjectInitializer::Get());
 
 public:
 
diff --git a/Engine/Source/Runtime/UMG/Public/Blueprint/WidgetLayoutLibrary.h b/Engine/Source/Runtime/UMG/Public/Blueprint/WidgetLayoutLibrary.h
--- a/Engine/Source/Runtime/UMG/Public/Blueprint/WidgetLayoutLibrary.h
+++ b/Engine/Source/Runtime/UMG/Public/Blueprint/WidgetLayoutLibrary.h
@@ -7,7 +7,9 @@
 UCLASS(MinimalAPI)
 class UWidgetLayoutLibrary : public UBlueprintFunctionLibrary
 {
-	GENERATED_UCLASS_BODY()
+	GENERATED_BODY()
+public:
+	UMG_API UWidgetLayoutLibrary(const FObjectInitializer& ObjectInitializer = FObjectInitializer::Get());
 
 public:
 
diff --git a/Engine/Source/Runtime/UMG/Public/Blueprint/WidgetNavigation.h b/Engine/Source/Runtime/UMG/Public/Blueprint/WidgetNavigation.h
--- a/Engine/Source/Runtime/UMG/Public/Blueprint/WidgetNavigation.h
+++ b/Engine/Source/Runtime/UMG/Public/Blueprint/WidgetNavigation.h
@@ -13,7 +13,7 @@
 USTRUCT()
 struct UMG_API FWidgetNavigationData
 {
-	GENERATED_USTRUCT_BODY()
+	GENERATED_BODY()
 
 public:
 	UPROPERTY(EditAnywhere, BlueprintReadOnly, Category="Navigation")
@@ -29,7 +29,9 @@ public:
 UCLASS()
 class UMG_API UWidgetNavigation : public UObject
 {
-	GENERATED_UCLASS_BODY()
+	GENERATED_BODY()
+public:
+	UWidgetNavigation(const FObjectInitializer& ObjectInitializer = FObjectInitializer::Get());
 	
 public:
 	/**  */
diff --git a/Engine/Source/Runtime/UMG/Public/Blueprint/WidgetTree.h b/Engine/Source/Runtime/UMG/Public/Blueprint/WidgetTree.h
--- a/Engine/Source/Runtime/UMG/Public/Blueprint/WidgetTree.h
+++ b/Engine/Source/Runtime/UMG/Public/Blueprint/WidgetTree.h
@@ -10,7 +10,9 @@
 UCLASS()
 class UMG_API UWidgetTree : public UObject
 {
-	GENERATED_UCLASS_BODY()
+	GENERATED_BODY()
+public:
+	UWidgetTree(const FObjectInitializer& ObjectInitializer = FObjectInitializer::Get());
 
 public:
 
diff --git a/Engine/Source/Runtime/UMG/Public/Components/Border.h b/Engine/Source/Runtime/UMG/Public/Components/Border.h
--- a/Engine/Source/Runtime/UMG/Public/Components/Border.h
+++ b/Engine/Source/Runtime/UMG/Public/Components/Border.h
@@ -1,4 +1,4 @@
-﻿// Copyright 1998-2015 Epic Games, Inc. All Rights Reserved.
+// Copyright 1998-2015 Epic Games, Inc. All Rights Reserved.
 
 #pragma once
 
@@ -16,7 +16,9 @@ class USlateBrushAsset;
 UCLASS()
 class UMG_API UBorder : public UContentWidget
 {
-	GENERATED_UCLASS_BODY()
+	GENERATED_BODY()
+public:
+	UBorder(const FObjectInitializer& ObjectInitializer = FObjectInitializer::Get());
 
 public:
 
diff --git a/Engine/Source/Runtime/UMG/Public/Components/BorderSlot.h b/Engine/Source/Runtime/UMG/Public/Components/BorderSlot.h
--- a/Engine/Source/Runtime/UMG/Public/Components/BorderSlot.h
+++ b/Engine/Source/Runtime/UMG/Public/Components/BorderSlot.h
@@ -14,7 +14,9 @@ class UBorder;
 UCLASS()
 class UMG_API UBorderSlot : public UPanelSlot
 {
-	GENERATED_UCLASS_BODY()
+	GENERATED_BODY()
+public:
+	UBorderSlot(const FObjectInitializer& ObjectInitializer = FObjectInitializer::Get());
 
 public:
 
diff --git a/Engine/Source/Runtime/UMG/Public/Components/Button.h b/Engine/Source/Runtime/UMG/Public/Components/Button.h
--- a/Engine/Source/Runtime/UMG/Public/Components/Button.h
+++ b/Engine/Source/Runtime/UMG/Public/Components/Button.h
@@ -1,4 +1,4 @@
-﻿// Copyright 1998-2015 Epic Games, Inc. All Rights Reserved.
+// Copyright 1998-2015 Epic Games, Inc. All Rights Reserved.
 
 #pragma once
 
@@ -21,7 +21,9 @@ DECLARE_DYNAMIC_MULTICAST_DELEGATE(FOnButtonReleasedEvent);
 UCLASS()
 class UMG_API UButton : public UContentWidget
 {
-	GENERATED_UCLASS_BODY()
+	GENERATED_BODY()
+public:
+	UButton(const FObjectInitializer& ObjectInitializer = FObjectInitializer::Get());
 
 public:
 	/** The template style asset, used to seed the mutable instance of the style. */
diff --git a/Engine/Source/Runtime/UMG/Public/Components/ButtonSlot.h b/Engine/Source/Runtime/UMG/Public/Components/ButtonSlot.h
--- a/Engine/Source/Runtime/UMG/Public/Components/ButtonSlot.h
+++ b/Engine/Source/Runtime/UMG/Public/Components/ButtonSlot.h
@@ -10,7 +10,9 @@
 UCLASS()
 class UMG_API UButtonSlot : public UPanelSlot
 {
-	GENERATED_UCLASS_BODY()
+	GENERATED_BODY()
+public:
+	UButtonSlot(const FObjectInitializer& ObjectInitializer = FObjectInitializer::Get());
 
 public:
 	
diff --git a/Engine/Source/Runtime/UMG/Public/Components/CanvasPanel.h b/Engine/Source/Runtime/UMG/Public/Components/CanvasPanel.h
--- a/Engine/Source/Runtime/UMG/Public/Components/CanvasPanel.h
+++ b/Engine/Source/Runtime/UMG/Public/Components/CanvasPanel.h
@@ -1,4 +1,4 @@
-﻿// Copyright 1998-2015 Epic Games, Inc. All Rights Reserved.
+// Copyright 1998-2015 Epic Games, Inc. All Rights Reserved.
 
 #pragma once
 #include "Components/PanelWidget.h"
@@ -19,7 +19,9 @@ class UCanvasPanelSlot;
 UCLASS()
 class UMG_API UCanvasPanel : public UPanelWidget
 {
-	GENERATED_UCLASS_BODY()
+	GENERATED_BODY()
+public:
+	UCanvasPanel(const FObjectInitializer& ObjectInitializer = FObjectInitializer::Get());
 
 public:
 
diff --git a/Engine/Source/Runtime/UMG/Public/Components/CanvasPanelSlot.h b/Engine/Source/Runtime/UMG/Public/Components/CanvasPanelSlot.h
--- a/Engine/Source/Runtime/UMG/Public/Components/CanvasPanelSlot.h
+++ b/Engine/Source/Runtime/UMG/Public/Components/CanvasPanelSlot.h
@@ -10,8 +10,8 @@
 USTRUCT(BlueprintType)
 struct FAnchorData
 {
+	GENERATED_BODY()
 public:
-	GENERATED_USTRUCT_BODY()
 
 public:
 
@@ -31,7 +31,9 @@ public:
 UCLASS()
 class UMG_API UCanvasPanelSlot : public UPanelSlot
 {
-	GENERATED_UCLASS_BODY()
+	GENERATED_BODY()
+public:
+	UCanvasPanelSlot(const FObjectInitializer& ObjectInitializer = FObjectInitializer::Get());
 
 public:
 
@@ -156,4 +158,4 @@ private:
 
 	TOptional<FVector2D> DesiredPosition;
 #endif
-};
\ No newline at end of file
+};
diff --git a/Engine/Source/Runtime/UMG/Public/Components/CheckBox.h b/Engine/Source/Runtime/UMG/Public/Components/CheckBox.h
--- a/Engine/Source/Runtime/UMG/Public/Components/CheckBox.h
+++ b/Engine/Source/Runtime/UMG/Public/Components/CheckBox.h
@@ -1,4 +1,4 @@
-﻿// Copyright 1998-2015 Epic Games, Inc. All Rights Reserved.
+// Copyright 1998-2015 Epic Games, Inc. All Rights Reserved.
 
 #pragma once
 
@@ -19,7 +19,9 @@ DECLARE_DYNAMIC_MULTICAST_DELEGATE_OneParam( FOnCheckBoxComponentStateChanged, b
 UCLASS()
 class UMG_API UCheckBox : public UContentWidget
 {
-	GENERATED_UCLASS_BODY()
+	GENERATED_BODY()
+public:
+	UCheckBox(const FObjectInitializer& ObjectInitializer = FObjectInitializer::Get());
 
 public:
 	/** Whether the check box is currently in a checked state */
diff --git a/Engine/Source/Runtime/UMG/Public/Components/CircularThrobber.h b/Engine/Source/Runtime/UMG/Public/Components/CircularThrobber.h
--- a/Engine/Source/Runtime/UMG/Public/Components/CircularThrobber.h
+++ b/Engine/Source/Runtime/UMG/Public/Components/CircularThrobber.h
@@ -1,4 +1,4 @@
-﻿// Copyright 1998-2015 Epic Games, Inc. All Rights Reserved.
+// Copyright 1998-2015 Epic Games, Inc. All Rights Reserved.
 
 #pragma once
 
@@ -15,7 +15,9 @@ class USlateBrushAsset;
 UCLASS()
 class UMG_API UCircularThrobber : public UWidget
 {
-	GENERATED_UCLASS_BODY()
+	GENERATED_BODY()
+public:
+	UCircularThrobber(const FObjectInitializer& ObjectInitializer = FObjectInitializer::Get());
 
 public:
 
diff --git a/Engine/Source/Runtime/UMG/Public/Components/ComboBox.h b/Engine/Source/Runtime/UMG/Public/Components/ComboBox.h
--- a/Engine/Source/Runtime/UMG/Public/Components/ComboBox.h
+++ b/Engine/Source/Runtime/UMG/Public/Components/ComboBox.h
@@ -10,7 +10,9 @@
 UCLASS(Experimental, meta=( DisplayName="ComboBox (Object)" ))
 class UMG_API UComboBox : public UWidget
 {
-	GENERATED_UCLASS_BODY()
+	GENERATED_BODY()
+public:
+	UComboBox(const FObjectInitializer& ObjectInitializer = FObjectInitializer::Get());
 
 public:
 
diff --git a/Engine/Source/Runtime/UMG/Public/Components/ComboBoxString.h b/Engine/Source/Runtime/UMG/Public/Components/ComboBoxString.h
--- a/Engine/Source/Runtime/UMG/Public/Components/ComboBoxString.h
+++ b/Engine/Source/Runtime/UMG/Public/Components/ComboBoxString.h
@@ -12,7 +12,9 @@
 UCLASS( meta=( DisplayName="ComboBox (String)"), ClassGroup=UserInterface)
 class UMG_API UComboBoxString : public UWidget
 {
-	GENERATED_UCLASS_BODY()
+	GENERATED_BODY()
+public:
+	UComboBoxString(const FObjectInitializer& ObjectInitializer = FObjectInitializer::Get());
 
 	DECLARE_DYNAMIC_MULTICAST_DELEGATE_TwoParams(FOnSelectionChangedEvent, FString, SelectedItem, ESelectInfo::Type, SelectionType);
 	DECLARE_DYNAMIC_MULTICAST_DELEGATE(FOnOpeningEvent);
diff --git a/Engine/Source/Runtime/UMG/Public/Components/ContentWidget.h b/Engine/Source/Runtime/UMG/Public/Components/ContentWidget.h
--- a/Engine/Source/Runtime/UMG/Public/Components/ContentWidget.h
+++ b/Engine/Source/Runtime/UMG/Public/Components/ContentWidget.h
@@ -8,7 +8,9 @@
 UCLASS(Abstract)
 class UMG_API UContentWidget : public UPanelWidget
 {
-	GENERATED_UCLASS_BODY()
+	GENERATED_BODY()
+public:
+	UContentWidget(const FObjectInitializer& ObjectInitializer = FObjectInitializer::Get());
 
 public:
 	/**  */
diff --git a/Engine/Source/Runtime/UMG/Public/Components/EditableText.h b/Engine/Source/Runtime/UMG/Public/Components/EditableText.h
--- a/Engine/Source/Runtime/UMG/Public/Components/EditableText.h
+++ b/Engine/Source/Runtime/UMG/Public/Components/EditableText.h
@@ -12,7 +12,9 @@
 UCLASS()
 class UMG_API UEditableText : public UWidget
 {
-	GENERATED_UCLASS_BODY()
+	GENERATED_BODY()
+public:
+	UEditableText(const FObjectInitializer& ObjectInitializer = FObjectInitializer::Get());
 
 public:
 
diff --git a/Engine/Source/Runtime/UMG/Public/Components/EditableTextBox.h b/Engine/Source/Runtime/UMG/Public/Components/EditableTextBox.h
--- a/Engine/Source/Runtime/UMG/Public/Components/EditableTextBox.h
+++ b/Engine/Source/Runtime/UMG/Public/Components/EditableTextBox.h
@@ -1,4 +1,4 @@
-﻿// Copyright 1998-2015 Epic Games, Inc. All Rights Reserved.
+// Copyright 1998-2015 Epic Games, Inc. All Rights Reserved.
 
 #pragma once
 
@@ -13,7 +13,9 @@
 UCLASS(meta=(DisplayName="Text Box"))
 class UMG_API UEditableTextBox : public UWidget
 {
-	GENERATED_UCLASS_BODY()
+	GENERATED_BODY()
+public:
+	UEditableTextBox(const FObjectInitializer& ObjectInitializer = FObjectInitializer::Get());
 
 public:
 
diff --git a/Engine/Source/Runtime/UMG/Public/Components/ExpandableArea.h b/Engine/Source/Runtime/UMG/Public/Components/ExpandableArea.h
--- a/Engine/Source/Runtime/UMG/Public/Components/ExpandableArea.h
+++ b/Engine/Source/Runtime/UMG/Public/Components/ExpandableArea.h
@@ -1,4 +1,4 @@
-﻿// Copyright 1998-2015 Epic Games, Inc. All Rights Reserved.
+// Copyright 1998-2015 Epic Games, Inc. All Rights Reserved.
 
 #pragma once
 
@@ -12,7 +12,9 @@ DECLARE_DYNAMIC_MULTICAST_DELEGATE_OneParam(FOnExpandableAreaExpansionChanged, b
 UCLASS(Experimental)
 class UMG_API UExpandableArea : public UWidget, public INamedSlotInterface
 {
-	GENERATED_UCLASS_BODY()
+	GENERATED_BODY()
+public:
+	UExpandableArea(const FObjectInitializer& ObjectInitializer = FObjectInitializer::Get());
 
 public:
 	/**  */
diff --git a/Engine/Source/Runtime/UMG/Public/Components/GridPanel.h b/Engine/Source/Runtime/UMG/Public/Components/GridPanel.h
--- a/Engine/Source/Runtime/UMG/Public/Components/GridPanel.h
+++ b/Engine/Source/Runtime/UMG/Public/Components/GridPanel.h
@@ -1,4 +1,4 @@
-﻿// Copyright 1998-2015 Epic Games, Inc. All Rights Reserved.
+// Copyright 1998-2015 Epic Games, Inc. All Rights Reserved.
 
 #pragma once
 
@@ -12,7 +12,9 @@
 UCLASS()
 class UMG_API UGridPanel : public UPanelWidget
 {
-	GENERATED_UCLASS_BODY()
+	GENERATED_BODY()
+public:
+	UGridPanel(const FObjectInitializer& ObjectInitializer = FObjectInitializer::Get());
 
 public:
 
diff --git a/Engine/Source/Runtime/UMG/Public/Components/GridSlot.h b/Engine/Source/Runtime/UMG/Public/Components/GridSlot.h
--- a/Engine/Source/Runtime/UMG/Public/Components/GridSlot.h
+++ b/Engine/Source/Runtime/UMG/Public/Components/GridSlot.h
@@ -13,7 +13,9 @@
 UCLASS()
 class UMG_API UGridSlot : public UPanelSlot
 {
-	GENERATED_UCLASS_BODY()
+	GENERATED_BODY()
+public:
+	UGridSlot(const FObjectInitializer& ObjectInitializer = FObjectInitializer::Get());
 
 public:
 
@@ -87,4 +89,4 @@ private:
 
 	/** A raw pointer to the slot to allow us to adjust the size, padding...etc at runtime. */
 	SGridPanel::FSlot* Slot;
-};
\ No newline at end of file
+};
diff --git a/Engine/Source/Runtime/UMG/Public/Components/HorizontalBox.h b/Engine/Source/Runtime/UMG/Public/Components/HorizontalBox.h
--- a/Engine/Source/Runtime/UMG/Public/Components/HorizontalBox.h
+++ b/Engine/Source/Runtime/UMG/Public/Components/HorizontalBox.h
@@ -1,4 +1,4 @@
-﻿// Copyright 1998-2015 Epic Games, Inc. All Rights Reserved.
+// Copyright 1998-2015 Epic Games, Inc. All Rights Reserved.
 
 #pragma once
 #include "Components/PanelWidget.h"
@@ -15,7 +15,9 @@ class UHorizontalBoxSlot;
 UCLASS(ClassGroup=UserInterface)
 class UMG_API UHorizontalBox : public UPanelWidget
 {
-	GENERATED_UCLASS_BODY()
+	GENERATED_BODY()
+public:
+	UHorizontalBox(const FObjectInitializer& ObjectInitializer = FObjectInitializer::Get());
 
 	/**  */
 	UFUNCTION(BlueprintCallable, Category="Widget")
diff --git a/Engine/Source/Runtime/UMG/Public/Components/HorizontalBoxSlot.h b/Engine/Source/Runtime/UMG/Public/Components/HorizontalBoxSlot.h
--- a/Engine/Source/Runtime/UMG/Public/Components/HorizontalBoxSlot.h
+++ b/Engine/Source/Runtime/UMG/Public/Components/HorizontalBoxSlot.h
@@ -9,7 +9,9 @@
 UCLASS()
 class UMG_API UHorizontalBoxSlot : public UPanelSlot
 {
-	GENERATED_UCLASS_BODY()
+	GENERATED_BODY()
+public:
+	UHorizontalBoxSlot(const FObjectInitializer& ObjectInitializer = FObjectInitializer::Get());
 
 public:
 	
@@ -51,4 +53,4 @@ public:
 
 private:
 	SHorizontalBox::FSlot* Slot;
-};
\ No newline at end of file
+};
diff --git a/Engine/Source/Runtime/UMG/Public/Components/Image.h b/Engine/Source/Runtime/UMG/Public/Components/Image.h
--- a/Engine/Source/Runtime/UMG/Public/Components/Image.h
+++ b/Engine/Source/Runtime/UMG/Public/Components/Image.h
@@ -1,4 +1,4 @@
-﻿// Copyright 1998-2015 Epic Games, Inc. All Rights Reserved.
+// Copyright 1998-2015 Epic Games, Inc. All Rights Reserved.
 
 #pragma once
 
@@ -14,7 +14,9 @@ class USlateBrushAsset;
 UCLASS()
 class UMG_API UImage : public UWidget
 {
-	GENERATED_UCLASS_BODY()
+	GENERATED_BODY()
+public:
+	UImage(const FObjectInitializer& ObjectInitializer = FObjectInitializer::Get());
 
 public:
 
diff --git a/Engine/Source/Runtime/UMG/Public/Components/ListView.h b/Engine/Source/Runtime/UMG/Public/Components/ListView.h
--- a/Engine/Source/Runtime/UMG/Public/Components/ListView.h
+++ b/Engine/Source/Runtime/UMG/Public/Components/ListView.h
@@ -10,7 +10,9 @@
 UCLASS(Experimental, ClassGroup=UserInterface)
 class UMG_API UListView : public UTableViewBase
 {
-	GENERATED_UCLASS_BODY()
+	GENERATED_BODY()
+public:
+	UListView(const FObjectInitializer& ObjectInitializer = FObjectInitializer::Get());
 
 public:
 
diff --git a/Engine/Source/Runtime/UMG/Public/Components/MenuAnchor.h b/Engine/Source/Runtime/UMG/Public/Components/MenuAnchor.h
--- a/Engine/Source/Runtime/UMG/Public/Components/MenuAnchor.h
+++ b/Engine/Source/Runtime/UMG/Public/Components/MenuAnchor.h
@@ -15,7 +15,9 @@ DECLARE_DYNAMIC_MULTICAST_DELEGATE_OneParam(FOnMenuOpenChangedEvent, bool, bIsOp
 UCLASS()
 class UMG_API UMenuAnchor : public UContentWidget
 {
-	GENERATED_UCLASS_BODY()
+	GENERATED_BODY()
+public:
+	UMenuAnchor(const FObjectInitializer& ObjectInitializer = FObjectInitializer::Get());
 
 public:
 
diff --git a/Engine/Source/Runtime/UMG/Public/Components/MultiLineEditableText.h b/Engine/Source/Runtime/UMG/Public/Components/MultiLineEditableText.h
--- a/Engine/Source/Runtime/UMG/Public/Components/MultiLineEditableText.h
+++ b/Engine/Source/Runtime/UMG/Public/Components/MultiLineEditableText.h
@@ -10,7 +10,9 @@
 UCLASS(meta=( DisplayName="Editable Text (Multi-Line)" ))
 class UMG_API UMultiLineEditableText : public UWidget
 {
-	GENERATED_UCLASS_BODY()
+	GENERATED_BODY()
+public:
+	UMultiLineEditableText(const FObjectInitializer& ObjectInitializer = FObjectInitializer::Get());
 
 public:
 
diff --git a/Engine/Source/Runtime/UMG/Public/Components/MultiLineEditableTextBox.h b/Engine/Source/Runtime/UMG/Public/Components/MultiLineEditableTextBox.h
--- a/Engine/Source/Runtime/UMG/Public/Components/MultiLineEditableTextBox.h
+++ b/Engine/Source/Runtime/UMG/Public/Components/MultiLineEditableTextBox.h
@@ -10,7 +10,9 @@
 UCLASS(meta=(DisplayName="Text Box (Multi-Line)"))
 class UMG_API UMultiLineEditableTextBox : public UWidget
 {
-	GENERATED_UCLASS_BODY()
+	GENERATED_BODY()
+public:
+	UMultiLineEditableTextBox(const FObjectInitializer& ObjectInitializer = FObjectInitializer::Get());
 
 public:
 
diff --git a/Engine/Source/Runtime/UMG/Public/Components/NamedSlot.h b/Engine/Source/Runtime/UMG/Public/Components/NamedSlot.h
--- a/Engine/Source/Runtime/UMG/Public/Components/NamedSlot.h
+++ b/Engine/Source/Runtime/UMG/Public/Components/NamedSlot.h
@@ -11,7 +11,9 @@
 UCLASS()
 class UMG_API UNamedSlot : public UContentWidget
 {
-	GENERATED_UCLASS_BODY()
+	GENERATED_BODY()
+public:
+	UNamedSlot(const FObjectInitializer& ObjectInitializer = FObjectInitializer::Get());
 
 public:
 
diff --git a/Engine/Source/Runtime/UMG/Public/Components/NamedSlotInterface.h b/Engine/Source/Runtime/UMG/Public/Components/NamedSlotInterface.h
--- a/Engine/Source/Runtime/UMG/Public/Components/NamedSlotInterface.h
+++ b/Engine/Source/Runtime/UMG/Public/Components/NamedSlotInterface.h
@@ -10,12 +10,15 @@ class UWidget;
 UINTERFACE(MinimalAPI, meta=( CannotImplementInterfaceInBlueprint ))
 class UNamedSlotInterface : public UInterface
 {
-	GENERATED_UINTERFACE_BODY()
+	GENERATED_BODY()
+public:
+	UMG_API UNamedSlotInterface(const FObjectInitializer& ObjectInitializer = FObjectInitializer::Get());
 };
 
 class UMG_API INamedSlotInterface
 {
-	GENERATED_IINTERFACE_BODY()
+	GENERATED_BODY()
+public:
 
 	/**  */
 	virtual void GetSlotNames(TArray<FName>& SlotNames) const = 0;
diff --git a/Engine/Source/Runtime/UMG/Public/Components/NativeWidgetHost.h b/Engine/Source/Runtime/UMG/Public/Components/NativeWidgetHost.h
--- a/Engine/Source/Runtime/UMG/Public/Components/NativeWidgetHost.h
+++ b/Engine/Source/Runtime/UMG/Public/Components/NativeWidgetHost.h
@@ -11,7 +11,9 @@
 UCLASS()
 class UMG_API UNativeWidgetHost : public UWidget
 {
-	GENERATED_UCLASS_BODY()
+	GENERATED_BODY()
+public:
+	UNativeWidgetHost(const FObjectInitializer& ObjectInitializer = FObjectInitializer::Get());
 
 	void SetContent(TSharedRef<SWidget> InContent);
 
diff --git a/Engine/Source/Runtime/UMG/Public/Components/Overlay.h b/Engine/Source/Runtime/UMG/Public/Components/Overlay.h
--- a/Engine/Source/Runtime/UMG/Public/Components/Overlay.h
+++ b/Engine/Source/Runtime/UMG/Public/Components/Overlay.h
@@ -10,7 +10,9 @@
 UCLASS()
 class UMG_API UOverlay : public UPanelWidget
 {
-	GENERATED_UCLASS_BODY()
+	GENERATED_BODY()
+public:
+	UOverlay(const FObjectInitializer& ObjectInitializer = FObjectInitializer::Get());
 
 	virtual void ReleaseSlateResources(bool bReleaseChildren) override;
 
diff --git a/Engine/Source/Runtime/UMG/Public/Components/OverlaySlot.h b/Engine/Source/Runtime/UMG/Public/Components/OverlaySlot.h
--- a/Engine/Source/Runtime/UMG/Public/Components/OverlaySlot.h
+++ b/Engine/Source/Runtime/UMG/Public/Components/OverlaySlot.h
@@ -12,7 +12,9 @@
 UCLASS()
 class UMG_API UOverlaySlot : public UPanelSlot
 {
-	GENERATED_UCLASS_BODY()
+	GENERATED_BODY()
+public:
+	UOverlaySlot(const FObjectInitializer& ObjectInitializer = FObjectInitializer::Get());
 
 public:
 	
diff --git a/Engine/Source/Runtime/UMG/Public/Components/PanelSlot.h b/Engine/Source/Runtime/UMG/Public/Components/PanelSlot.h
--- a/Engine/Source/Runtime/UMG/Public/Components/PanelSlot.h
+++ b/Engine/Source/Runtime/UMG/Public/Components/PanelSlot.h
@@ -8,7 +8,9 @@
 UCLASS(BlueprintType)
 class UMG_API UPanelSlot : public UVisual
 {
-	GENERATED_UCLASS_BODY()
+	GENERATED_BODY()
+public:
+	UPanelSlot(const FObjectInitializer& ObjectInitializer = FObjectInitializer::Get());
 
 public:
 
diff --git a/Engine/Source/Runtime/UMG/Public/Components/PanelWidget.h b/Engine/Source/Runtime/UMG/Public/Components/PanelWidget.h
--- a/Engine/Source/Runtime/UMG/Public/Components/PanelWidget.h
+++ b/Engine/Source/Runtime/UMG/Public/Components/PanelWidget.h
@@ -9,7 +9,9 @@
 UCLASS(Abstract)
 class UMG_API UPanelWidget : public UWidget
 {
-	GENERATED_UCLASS_BODY()
+	GENERATED_BODY()
+public:
+	UPanelWidget(const FObjectInitializer& ObjectInitializer = FObjectInitializer::Get());
 
 protected:
 
diff --git a/Engine/Source/Runtime/UMG/Public/Components/ProgressBar.h b/Engine/Source/Runtime/UMG/Public/Components/ProgressBar.h
--- a/Engine/Source/Runtime/UMG/Public/Components/ProgressBar.h
+++ b/Engine/Source/Runtime/UMG/Public/Components/ProgressBar.h
@@ -1,4 +1,4 @@
-﻿// Copyright 1998-2015 Epic Games, Inc. All Rights Reserved.
+// Copyright 1998-2015 Epic Games, Inc. All Rights Reserved.
 
 #pragma once
 
@@ -14,7 +14,9 @@
 UCLASS(ClassGroup=UserInterface)
 class UMG_API UProgressBar : public UWidget
 {
-	GENERATED_UCLASS_BODY()
+	GENERATED_BODY()
+public:
+	UProgressBar(const FObjectInitializer& ObjectInitializer = FObjectInitializer::Get());
 	
 public:
 
diff --git a/Engine/Source/Runtime/UMG/Public/Components/ScaleBox.h b/Engine/Source/Runtime/UMG/Public/Components/ScaleBox.h
--- a/Engine/Source/Runtime/UMG/Public/Components/ScaleBox.h
+++ b/Engine/Source/Runtime/UMG/Public/Components/ScaleBox.h
@@ -1,4 +1,4 @@
-﻿// Copyright 1998-2015 Epic Games, Inc. All Rights Reserved.
+// Copyright 1998-2015 Epic Games, Inc. All Rights Reserved.
 
 #pragma once
 
@@ -15,7 +15,9 @@
 UCLASS()
 class UMG_API UScaleBox : public UContentWidget
 {
-	GENERATED_UCLASS_BODY()
+	GENERATED_BODY()
+public:
+	UScaleBox(const FObjectInitializer& ObjectInitializer = FObjectInitializer::Get());
 
 public:
 
diff --git a/Engine/Source/Runtime/UMG/Public/Components/ScaleBoxSlot.h b/Engine/Source/Runtime/UMG/Public/Components/ScaleBoxSlot.h
--- a/Engine/Source/Runtime/UMG/Public/Components/ScaleBoxSlot.h
+++ b/Engine/Source/Runtime/UMG/Public/Components/ScaleBoxSlot.h
@@ -10,7 +10,9 @@
 UCLASS()
 class UMG_API UScaleBoxSlot : public UPanelSlot
 {
-	GENERATED_UCLASS_BODY()
+	GENERATED_BODY()
+public:
+	UScaleBoxSlot(const FObjectInitializer& ObjectInitializer = FObjectInitializer::Get());
 
 public:
 	
diff --git a/Engine/Source/Runtime/UMG/Public/Components/ScrollBar.h b/Engine/Source/Runtime/UMG/Public/Components/ScrollBar.h
--- a/Engine/Source/Runtime/UMG/Public/Components/ScrollBar.h
+++ b/Engine/Source/Runtime/UMG/Public/Components/ScrollBar.h
@@ -8,7 +8,9 @@
 UCLASS(Experimental)
 class UMG_API UScrollBar : public UWidget
 {
-	GENERATED_UCLASS_BODY()
+	GENERATED_BODY()
+public:
+	UScrollBar(const FObjectInitializer& ObjectInitializer = FObjectInitializer::Get());
 
 public:
 
diff --git a/Engine/Source/Runtime/UMG/Public/Components/ScrollBox.h b/Engine/Source/Runtime/UMG/Public/Components/ScrollBox.h
--- a/Engine/Source/Runtime/UMG/Public/Components/ScrollBox.h
+++ b/Engine/Source/Runtime/UMG/Public/Components/ScrollBox.h
@@ -10,7 +10,9 @@
 UCLASS()
 class UMG_API UScrollBox : public UPanelWidget
 {
-	GENERATED_UCLASS_BODY()
+	GENERATED_BODY()
+public:
+	UScrollBox(const FObjectInitializer& ObjectInitializer = FObjectInitializer::Get());
 
 public:
 
diff --git a/Engine/Source/Runtime/UMG/Public/Components/ScrollBoxSlot.h b/Engine/Source/Runtime/UMG/Public/Components/ScrollBoxSlot.h
--- a/Engine/Source/Runtime/UMG/Public/Components/ScrollBoxSlot.h
+++ b/Engine/Source/Runtime/UMG/Public/Components/ScrollBoxSlot.h
@@ -10,7 +10,9 @@
 UCLASS()
 class UMG_API UScrollBoxSlot : public UPanelSlot
 {
-	GENERATED_UCLASS_BODY()
+	GENERATED_BODY()
+public:
+	UScrollBoxSlot(const FObjectInitializer& ObjectInitializer = FObjectInitializer::Get());
 
 public:
 	
diff --git a/Engine/Source/Runtime/UMG/Public/Components/SizeBox.h b/Engine/Source/Runtime/UMG/Public/Components/SizeBox.h
--- a/Engine/Source/Runtime/UMG/Public/Components/SizeBox.h
+++ b/Engine/Source/Runtime/UMG/Public/Components/SizeBox.h
@@ -1,4 +1,4 @@
-﻿// Copyright 1998-2015 Epic Games, Inc. All Rights Reserved.
+// Copyright 1998-2015 Epic Games, Inc. All Rights Reserved.
 
 #pragma once
 
@@ -16,7 +16,9 @@ class USlateBrushAsset;
 UCLASS()
 class UMG_API USizeBox : public UContentWidget
 {
-	GENERATED_UCLASS_BODY()
+	GENERATED_BODY()
+public:
+	USizeBox(const FObjectInitializer& ObjectInitializer = FObjectInitializer::Get());
 
 public:
 
diff --git a/Engine/Source/Runtime/UMG/Public/Components/SizeBoxSlot.h b/Engine/Source/Runtime/UMG/Public/Components/SizeBoxSlot.h
--- a/Engine/Source/Runtime/UMG/Public/Components/SizeBoxSlot.h
+++ b/Engine/Source/Runtime/UMG/Public/Components/SizeBoxSlot.h
@@ -10,7 +10,9 @@
 UCLASS()
 class UMG_API USizeBoxSlot : public UPanelSlot
 {
-	GENERATED_UCLASS_BODY()
+	GENERATED_BODY()
+public:
+	USizeBoxSlot(const FObjectInitializer& ObjectInitializer = FObjectInitializer::Get());
 
 public:
 	/** The padding area between the slot and the content it contains. */
diff --git a/Engine/Source/Runtime/UMG/Public/Components/Slider.h b/Engine/Source/Runtime/UMG/Public/Components/Slider.h
--- a/Engine/Source/Runtime/UMG/Public/Components/Slider.h
+++ b/Engine/Source/Runtime/UMG/Public/Components/Slider.h
@@ -1,4 +1,4 @@
-﻿// Copyright 1998-2015 Epic Games, Inc. All Rights Reserved.
+// Copyright 1998-2015 Epic Games, Inc. All Rights Reserved.
 
 #pragma once
 
@@ -16,7 +16,9 @@ DECLARE_DYNAMIC_MULTICAST_DELEGATE_OneParam(FOnFloatValueChangedEvent, float, Va
 UCLASS()
 class UMG_API USlider : public UWidget
 {
-	GENERATED_UCLASS_BODY()
+	GENERATED_BODY()
+public:
+	USlider(const FObjectInitializer& ObjectInitializer = FObjectInitializer::Get());
 
 public:
 	/** The volume value to display. */
diff --git a/Engine/Source/Runtime/UMG/Public/Components/Spacer.h b/Engine/Source/Runtime/UMG/Public/Components/Spacer.h
--- a/Engine/Source/Runtime/UMG/Public/Components/Spacer.h
+++ b/Engine/Source/Runtime/UMG/Public/Components/Spacer.h
@@ -1,4 +1,4 @@
-﻿// Copyright 1998-2015 Epic Games, Inc. All Rights Reserved.
+// Copyright 1998-2015 Epic Games, Inc. All Rights Reserved.
 
 #pragma once
 
@@ -12,7 +12,9 @@
 UCLASS()
 class UMG_API USpacer : public UWidget
 {
-	GENERATED_UCLASS_BODY()
+	GENERATED_BODY()
+public:
+	USpacer(const FObjectInitializer& ObjectInitializer = FObjectInitializer::Get());
 
 public:
 
diff --git a/Engine/Source/Runtime/UMG/Public/Components/SpinBox.h b/Engine/Source/Runtime/UMG/Public/Components/SpinBox.h
--- a/Engine/Source/Runtime/UMG/Public/Components/SpinBox.h
+++ b/Engine/Source/Runtime/UMG/Public/Components/SpinBox.h
@@ -10,7 +10,9 @@
 UCLASS()
 class UMG_API USpinBox : public UWidget
 {
-	GENERATED_UCLASS_BODY()
+	GENERATED_BODY()
+public:
+	USpinBox(const FObjectInitializer& ObjectInitializer = FObjectInitializer::Get());
 
 public:
 	DECLARE_DYNAMIC_MULTICAST_DELEGATE_OneParam(FOnSpinBoxValueChangedEvent, float, InValue);
@@ -208,4 +210,4 @@ protected:
 
 protected:
 	TSharedPtr<SSpinBox<float>> MySpinBox;
-};
\ No newline at end of file
+};
diff --git a/Engine/Source/Runtime/UMG/Public/Components/TableViewBase.h b/Engine/Source/Runtime/UMG/Public/Components/TableViewBase.h
--- a/Engine/Source/Runtime/UMG/Public/Components/TableViewBase.h
+++ b/Engine/Source/Runtime/UMG/Public/Components/TableViewBase.h
@@ -8,7 +8,9 @@
 UCLASS(Abstract)
 class UMG_API UTableViewBase : public UWidget
 {
-	GENERATED_UCLASS_BODY()
+	GENERATED_BODY()
+public:
+	UTableViewBase(const FObjectInitializer& ObjectInitializer = FObjectInitializer::Get());
 
 	/** Delegate for constructing a UWidget based on a UObject */
 	DECLARE_DYNAMIC_DELEGATE_RetVal_OneParam(UWidget*, FOnGenerateRowUObject, UObject*, Item);
diff --git a/Engine/Source/Runtime/UMG/Public/Components/TextBlock.h b/Engine/Source/Runtime/UMG/Public/Components/TextBlock.h
--- a/Engine/Source/Runtime/UMG/Public/Components/TextBlock.h
+++ b/Engine/Source/Runtime/UMG/Public/Components/TextBlock.h
@@ -1,4 +1,4 @@
-﻿// Copyright 1998-2015 Epic Games, Inc. All Rights Reserved.
+// Copyright 1998-2015 Epic Games, Inc. All Rights Reserved.
 
 #pragma once
 
@@ -13,7 +13,9 @@
 UCLASS(meta=(DisplayName="Text"))
 class UMG_API UTextBlock : public UWidget
 {
-	GENERATED_UCLASS_BODY()
+	GENERATED_BODY()
+public:
+	UTextBlock(const FObjectInitializer& ObjectInitializer = FObjectInitializer::Get());
 
 public:
 	/**  
diff --git a/Engine/Source/Runtime/UMG/Public/Components/Throbber.h b/Engine/Source/Runtime/UMG/Public/Components/Throbber.h
--- a/Engine/Source/Runtime/UMG/Public/Components/Throbber.h
+++ b/Engine/Source/Runtime/UMG/Public/Components/Throbber.h
@@ -13,7 +13,9 @@ class USlateBrushAsset;
 UCLASS(ClassGroup=UserInterface)
 class UMG_API UThrobber : public UWidget
 {
-	GENERATED_UCLASS_BODY()
+	GENERATED_BODY()
+public:
+	UThrobber(const FObjectInitializer& ObjectInitializer = FObjectInitializer::Get());
 
 public:
 	
diff --git a/Engine/Source/Runtime/UMG/Public/Components/TileView.h b/Engine/Source/Runtime/UMG/Public/Components/TileView.h
--- a/Engine/Source/Runtime/UMG/Public/Components/TileView.h
+++ b/Engine/Source/Runtime/UMG/Public/Components/TileView.h
@@ -10,7 +10,9 @@
 UCLASS(Experimental, ClassGroup=UserInterface)
 class UMG_API UTileView : public UTableViewBase
 {
-	GENERATED_UCLASS_BODY()
+	GENERATED_BODY()
+public:
+	UTileView(const FObjectInitializer& ObjectInitializer = FObjectInitializer::Get());
 
 public:
 
diff --git a/Engine/Source/Runtime/UMG/Public/Components/UniformGridPanel.h b/Engine/Source/Runtime/UMG/Public/Components/UniformGridPanel.h
--- a/Engine/Source/Runtime/UMG/Public/Components/UniformGridPanel.h
+++ b/Engine/Source/Runtime/UMG/Public/Components/UniformGridPanel.h
@@ -10,7 +10,9 @@
 UCLASS()
 class UMG_API UUniformGridPanel : public UPanelWidget
 {
-	GENERATED_UCLASS_BODY()
+	GENERATED_BODY()
+public:
+	UUniformGridPanel(const FObjectInitializer& ObjectInitializer = FObjectInitializer::Get());
 
 public:
 
diff --git a/Engine/Source/Runtime/UMG/Public/Components/UniformGridSlot.h b/Engine/Source/Runtime/UMG/Public/Components/UniformGridSlot.h
--- a/Engine/Source/Runtime/UMG/Public/Components/UniformGridSlot.h
+++ b/Engine/Source/Runtime/UMG/Public/Components/UniformGridSlot.h
@@ -13,7 +13,9 @@
 UCLASS()
 class UMG_API UUniformGridSlot : public UPanelSlot
 {
-	GENERATED_UCLASS_BODY()
+	GENERATED_BODY()
+public:
+	UUniformGridSlot(const FObjectInitializer& ObjectInitializer = FObjectInitializer::Get());
 
 public:
 
@@ -67,4 +69,4 @@ private:
 
 	/** A raw pointer to the slot to allow us to adjust the size, padding...etc at runtime. */
 	SUniformGridPanel::FSlot* Slot;
-};
\ No newline at end of file
+};
diff --git a/Engine/Source/Runtime/UMG/Public/Components/VerticalBox.h b/Engine/Source/Runtime/UMG/Public/Components/VerticalBox.h
--- a/Engine/Source/Runtime/UMG/Public/Components/VerticalBox.h
+++ b/Engine/Source/Runtime/UMG/Public/Components/VerticalBox.h
@@ -1,4 +1,4 @@
-﻿// Copyright 1998-2015 Epic Games, Inc. All Rights Reserved.
+// Copyright 1998-2015 Epic Games, Inc. All Rights Reserved.
 
 #pragma once
 #include "Components/PanelWidget.h"
@@ -16,7 +16,9 @@ class UVerticalBoxSlot;
 UCLASS()
 class UMG_API UVerticalBox : public UPanelWidget
 {
-	GENERATED_UCLASS_BODY()
+	GENERATED_BODY()
+public:
+	UVerticalBox(const FObjectInitializer& ObjectInitializer = FObjectInitializer::Get());
 
 public:
 
diff --git a/Engine/Source/Runtime/UMG/Public/Components/VerticalBoxSlot.h b/Engine/Source/Runtime/UMG/Public/Components/VerticalBoxSlot.h
--- a/Engine/Source/Runtime/UMG/Public/Components/VerticalBoxSlot.h
+++ b/Engine/Source/Runtime/UMG/Public/Components/VerticalBoxSlot.h
@@ -10,7 +10,9 @@
 UCLASS()
 class UMG_API UVerticalBoxSlot : public UPanelSlot
 {
-	GENERATED_UCLASS_BODY()
+	GENERATED_BODY()
+public:
+	UVerticalBoxSlot(const FObjectInitializer& ObjectInitializer = FObjectInitializer::Get());
 
 public:
 	
diff --git a/Engine/Source/Runtime/UMG/Public/Components/Viewport.h b/Engine/Source/Runtime/UMG/Public/Components/Viewport.h
--- a/Engine/Source/Runtime/UMG/Public/Components/Viewport.h
+++ b/Engine/Source/Runtime/UMG/Public/Components/Viewport.h
@@ -219,7 +219,9 @@ protected:
 UCLASS(Experimental, ClassGroup=UserInterface)
 class UMG_API UViewport : public UContentWidget
 {
-	GENERATED_UCLASS_BODY()
+	GENERATED_BODY()
+public:
+	UViewport(const FObjectInitializer& ObjectInitializer = FObjectInitializer::Get());
 
 	UPROPERTY(EditDefaultsOnly, Category=Appearance)
 	FLinearColor BackgroundColor;
diff --git a/Engine/Source/Runtime/UMG/Public/Components/Visual.h b/Engine/Source/Runtime/UMG/Public/Components/Visual.h
--- a/Engine/Source/Runtime/UMG/Public/Components/Visual.h
+++ b/Engine/Source/Runtime/UMG/Public/Components/Visual.h
@@ -8,7 +8,9 @@
 UCLASS()
 class UMG_API UVisual : public UObject
 {
-	GENERATED_UCLASS_BODY()
+	GENERATED_BODY()
+public:
+	UVisual(const FObjectInitializer& ObjectInitializer = FObjectInitializer::Get());
 
 public:
 	virtual void ReleaseSlateResources(bool bReleaseChildren);
diff --git a/Engine/Source/Runtime/UMG/Public/Components/Widget.h b/Engine/Source/Runtime/UMG/Public/Components/Widget.h
--- a/Engine/Source/Runtime/UMG/Public/Components/Widget.h
+++ b/Engine/Source/Runtime/UMG/Public/Components/Widget.h
@@ -73,7 +73,9 @@ class UPanelSlot;
 UCLASS(Abstract, BlueprintType)
 class UMG_API UWidget : public UVisual
 {
-	GENERATED_UCLASS_BODY()
+	GENERATED_BODY()
+public:
+	UWidget(const FObjectInitializer& ObjectInitializer = FObjectInitializer::Get());
 
 public:
 
diff --git a/Engine/Source/Runtime/UMG/Public/Components/WidgetComponent.h b/Engine/Source/Runtime/UMG/Public/Components/WidgetComponent.h
--- a/Engine/Source/Runtime/UMG/Public/Components/WidgetComponent.h
+++ b/Engine/Source/Runtime/UMG/Public/Components/WidgetComponent.h
@@ -22,7 +22,9 @@ enum class EWidgetSpace : uint8
 UCLASS(ClassGroup=Experimental, hidecategories=(Object,Activation,"Components|Activation",Sockets,Base,Lighting,LOD,Mesh), editinlinenew, meta=(BlueprintSpawnableComponent,  DevelopmentStatus=Experimental) )
 class UMG_API UWidgetComponent : public UPrimitiveComponent
 {
-	GENERATED_UCLASS_BODY()
+	GENERATED_BODY()
+public:
+	UWidgetComponent(const FObjectInitializer& ObjectInitializer = FObjectInitializer::Get());
 
 public:
 	/* UPrimitiveComponent Interface */
diff --git a/Engine/Source/Runtime/UMG/Public/Components/WidgetSwitcher.h b/Engine/Source/Runtime/UMG/Public/Components/WidgetSwitcher.h
--- a/Engine/Source/Runtime/UMG/Public/Components/WidgetSwitcher.h
+++ b/Engine/Source/Runtime/UMG/Public/Components/WidgetSwitcher.h
@@ -10,7 +10,9 @@
 UCLASS()
 class UMG_API UWidgetSwitcher : public UPanelWidget
 {
-	GENERATED_UCLASS_BODY()
+	GENERATED_BODY()
+public:
+	UWidgetSwitcher(const FObjectInitializer& ObjectInitializer = FObjectInitializer::Get());
 
 public:
 	/** The slot index to display */
diff --git a/Engine/Source/Runtime/UMG/Public/Components/WidgetSwitcherSlot.h b/Engine/Source/Runtime/UMG/Public/Components/WidgetSwitcherSlot.h
--- a/Engine/Source/Runtime/UMG/Public/Components/WidgetSwitcherSlot.h
+++ b/Engine/Source/Runtime/UMG/Public/Components/WidgetSwitcherSlot.h
@@ -11,7 +11,9 @@
 UCLASS()
 class UMG_API UWidgetSwitcherSlot : public UPanelSlot
 {
-	GENERATED_UCLASS_BODY()
+	GENERATED_BODY()
+public:
+	UWidgetSwitcherSlot(const FObjectInitializer& ObjectInitializer = FObjectInitializer::Get());
 
 public:
 	
diff --git a/Engine/Source/Runtime/UMG/Public/Components/WrapBox.h b/Engine/Source/Runtime/UMG/Public/Components/WrapBox.h
--- a/Engine/Source/Runtime/UMG/Public/Components/WrapBox.h
+++ b/Engine/Source/Runtime/UMG/Public/Components/WrapBox.h
@@ -1,4 +1,4 @@
-﻿// Copyright 1998-2015 Epic Games, Inc. All Rights Reserved.
+// Copyright 1998-2015 Epic Games, Inc. All Rights Reserved.
 
 #pragma once
 
@@ -14,7 +14,9 @@
 UCLASS()
 class UMG_API UWrapBox : public UPanelWidget
 {
-	GENERATED_UCLASS_BODY()
+	GENERATED_BODY()
+public:
+	UWrapBox(const FObjectInitializer& ObjectInitializer = FObjectInitializer::Get());
 
 public:
 	/** The inner slot padding goes between slots sharing borders */
diff --git a/Engine/Source/Runtime/UMG/Public/Components/WrapBoxSlot.h b/Engine/Source/Runtime/UMG/Public/Components/WrapBoxSlot.h
--- a/Engine/Source/Runtime/UMG/Public/Components/WrapBoxSlot.h
+++ b/Engine/Source/Runtime/UMG/Public/Components/WrapBoxSlot.h
@@ -10,7 +10,9 @@
 UCLASS()
 class UMG_API UWrapBoxSlot : public UPanelSlot
 {
-	GENERATED_UCLASS_BODY()
+	GENERATED_BODY()
+public:
+	UWrapBoxSlot(const FObjectInitializer& ObjectInitializer = FObjectInitializer::Get());
 
 public:
 
