Allow libcurl to be included in Rocket builds; it already is included for Android.

diff --git a/Engine/Source/Runtime/Online/HTTP/HTTP.Build.cs b/Engine/Source/Runtime/Online/HTTP/HTTP.Build.cs
--- a/Engine/Source/Runtime/Online/HTTP/HTTP.Build.cs
+++ b/Engine/Source/Runtime/Online/HTTP/HTTP.Build.cs
@@ -24,10 +24,7 @@ public class HTTP : ModuleRules
             Target.Platform == UnrealTargetPlatform.Win64)
         {
             AddThirdPartyPrivateStaticDependencies(Target, "WinInet");
-            if (!UnrealBuildTool.UnrealBuildTool.BuildingRocket() && !UnrealBuildTool.UnrealBuildTool.RunningRocket())
-            {
-                AddThirdPartyPrivateStaticDependencies(Target, "libcurl");
-            }
+            AddThirdPartyPrivateStaticDependencies(Target, "libcurl");
         }
 
         if (Target.Platform == UnrealTargetPlatform.Linux ||
