Update format of BuildPatchTool patch data enumeration mode when ran with -includesizes.

To avoid possible ambiguity between a filename and the output including the size, we now split the filename from the size with a tab.

diff --git a/Engine/Source/Programs/BuildPatchTool/Private/BuildPatchToolMain.cpp b/Engine/Source/Programs/BuildPatchTool/Private/BuildPatchToolMain.cpp
--- a/Engine/Source/Programs/BuildPatchTool/Private/BuildPatchToolMain.cpp
+++ b/Engine/Source/Programs/BuildPatchTool/Private/BuildPatchToolMain.cpp
@@ -75,7 +75,7 @@
 	Required arguments:
 	-ManifestFile=""	Specifies in quotes the file path to the manifest to enumerate from
 	-OutputFile=""		Specifies in quotes the file path to a file where the list will be saved out, \r\n separated cloud relative paths.
-	-includesizes		When specified, the size of each file will be output following the filename in the format " (xxx bytes)".
+	-includesizes		When specified, the size of each file (in bytes) will be output following the filename and a tab.  I.e. /path/to/chunk\t1233
 	-dataenumerate		Must be specified to launch the tool in cloud seeder mode
 
 =============================================================================*/
diff --git a/Engine/Source/Runtime/Online/BuildPatchServices/Private/BuildPatchDataEnumeration.cpp b/Engine/Source/Runtime/Online/BuildPatchServices/Private/BuildPatchDataEnumeration.cpp
--- a/Engine/Source/Runtime/Online/BuildPatchServices/Private/BuildPatchDataEnumeration.cpp
+++ b/Engine/Source/Runtime/Online/BuildPatchServices/Private/BuildPatchDataEnumeration.cpp
@@ -24,15 +24,14 @@ bool FBuildDataEnumeration::EnumerateManifestData(FString ManifestFilePath, FStr
 		FString FullList;
 		for (FGuid& DataGuid : DataList)
 		{
-			FString Path = FBuildPatchUtils::GetDataFilename(AppManifest, FString(), DataGuid);
-			FullList += Path;
+			FString OutputLine = FBuildPatchUtils::GetDataFilename(AppManifest, FString(), DataGuid);
 			if (bIncludeSizes)
 			{
 				uint64 FileSize = AppManifest->GetDataSize(DataGuid);
-				FullList += FString::Printf(TEXT(" (%u bytes)"), FileSize);
+				OutputLine += FString::Printf(TEXT("\t%u"), FileSize);
 			}
-			FullList += TEXT("\r\n");
-			GLog->Log(ELogVerbosity::Log, Path);
+			GLog->Log(ELogVerbosity::Log, OutputLine);
+			FullList += OutputLine + TEXT("\r\n");
 		}
 		if (FFileHelper::SaveStringToFile(FullList, *OutputFile))
 		{
