#952 Fixed sample rate conversion and increased BufferSize for VoiceModuleMac.
- Merge pull request #952 from devel.bmad to fix the sample rate conversion path for CoreAudio input.
- Attempt to set the input device sample rate to the desired rate, falling back to sample rate conversion only on failure.

diff --git a/Engine/Source/Runtime/Online/Voice/Private/Mac/VoiceModuleMac.cpp b/Engine/Source/Runtime/Online/Voice/Private/Mac/VoiceModuleMac.cpp
--- a/Engine/Source/Runtime/Online/Voice/Private/Mac/VoiceModuleMac.cpp
+++ b/Engine/Source/Runtime/Online/Voice/Private/Mac/VoiceModuleMac.cpp
@@ -11,8 +11,11 @@
 /** Maximum buffer size for storing raw uncompressed audio from the system */
 #define MAX_UNCOMPRESSED_VOICE_BUFFER_SIZE 30 * 1024
 
+/** Number of hardware buffers per uncompressed CoreAudio buffer */
+#define NUM_HARDWARE_BUFFERS_PER_UNCOMPRESSED 6
+
 /**
- * Implementation of voice capture using OpenAL
+ * Implementation of voice capture using CoreAudio
  */
 class FVoiceCaptureCoreAudio : public IVoiceCapture
 {
@@ -168,7 +171,6 @@ public:
 			return false;
 		}
 		
-		
 		AURenderCallbackStruct InputCb;
 		InputCb.inputProc = &FVoiceCaptureCoreAudio::InputProc;
 		InputCb.inputProcRefCon = this;
@@ -183,14 +185,18 @@ public:
 		}
 		
 		OutputDesc = StreamDesc;
-		OutputDesc.mSampleRate = NativeDesc.mSampleRate;
 		if ( AudioUnitSetProperty(StreamComponent, kAudioUnitProperty_StreamFormat, kAudioUnitScope_Output, 1, (void *)&OutputDesc, sizeof(OutputDesc)) != 0 )
 		{
-			UE_LOG(LogVoiceCapture, Warning, TEXT("Couldn't configure CoreAudio input component format"));
-			AudioComponentInstanceDispose(StreamComponent);
-			return false;
+			OutputDesc.mSampleRate = NativeDesc.mSampleRate;
+			if ( AudioUnitSetProperty(StreamComponent, kAudioUnitProperty_StreamFormat, kAudioUnitScope_Output, 1, (void *)&OutputDesc, sizeof(OutputDesc)) != 0 )
+			{
+				UE_LOG(LogVoiceCapture, Warning, TEXT("Couldn't configure CoreAudio input component format"));
+				AudioComponentInstanceDispose(StreamComponent);
+				return false;
+			}
 		}
 		
+		// create an audio converter
 		if ( StreamDesc.mSampleRate != OutputDesc.mSampleRate )
 		{
 			if ( AudioConverterNew(&OutputDesc, &StreamDesc, &StreamConverter) != 0 || !StreamConverter )
@@ -231,7 +237,9 @@ public:
 			}
 			return false;
 		}
-		BufferSize *= StreamDesc.mBytesPerFrame;
+		
+		// Make the buffer size sufficiently large to prevent underflow when decoding.
+		BufferSize *= (StreamDesc.mBytesPerFrame * NUM_HARDWARE_BUFFERS_PER_UNCOMPRESSED);
 		
 		InputLatency = SafetyOffset + BufferSize;
 		check(InputLatency < MAX_UNCOMPRESSED_VOICE_BUFFER_SIZE);
@@ -360,50 +368,90 @@ public:
 		return VoiceCaptureState;
 	}
 	
+	struct FAudioFileIO
+	{
+		uint32 SrcReadFrames;
+		AudioBufferList* SrcBuffer;
+		uint32 SrcBufferSize;
+		uint32 SrcBytesPerFrame;
+	};
+	
 	static OSStatus ConvertInputFormat( AudioConverterRef AudioConverter,
 										UInt32* IONumberDataPackets,
 										AudioBufferList* IOData,
 										AudioStreamPacketDescription** OutDataPacketDescription,
-										void* InUserData)
+										void* InUserData )
 	{
-		AudioBufferList* Buffer = static_cast<AudioBufferList*>(InUserData);
-		if ( Buffer )
+		FAudioFileIO* Input = (FAudioFileIO*)InUserData;
+		
+		if ( Input && Input->SrcBuffer )
 		{
-			*IOData = *Buffer;
+			*IOData = *Input->SrcBuffer;
+			
+			int32 Num = *IONumberDataPackets;
+			if ( ( Input->SrcReadFrames * (Input->SrcBytesPerFrame)) < Input->SrcBufferSize )
+			{
+				if ( (Num * Input->SrcBytesPerFrame) < IOData->mBuffers[0].mDataByteSize )
+				{
+					IOData->mBuffers[0].mDataByteSize = Num * Input->SrcBytesPerFrame;
+				}
+				else
+				{
+					*IONumberDataPackets = IOData->mBuffers[0].mDataByteSize / Input->SrcBytesPerFrame;
+				}
+				Input->SrcReadFrames += *IONumberDataPackets;
+			}
+			else
+			{
+				*IONumberDataPackets = 0;
+				IOData->mBuffers[0].mDataByteSize = 0;
+				return eofErr;
+			}
 		}
 		
 		return noErr;
 	}
 	
-	EVoiceCaptureState::Type CopyBuffer(uint8 ReadBuffer, uint8* OutVoiceBuffer, uint32 ReadOffset, uint32 ReadLen)
+	EVoiceCaptureState::Type CopyBuffer(uint8 ReadBuffer, uint8* OutVoiceBuffer, uint32 ReadOffset, uint32* ReadLen, uint32* WriteLen)
 	{
 		EVoiceCaptureState::Type State = VoiceCaptureState;
-		if ( ReadLen )
+		if ( *ReadLen )
 		{
 			uint8* Data = (uint8*)(BufferList[ReadBuffer].mData) + ReadOffset;
+			
 			if ( StreamDesc.mSampleRate == OutputDesc.mSampleRate )
 			{
-				FMemory::Memcpy(OutVoiceBuffer, Data, ReadLen);
+				FMemory::Memcpy(OutVoiceBuffer, Data, *ReadLen);
+				*WriteLen = *ReadLen;
 				State = EVoiceCaptureState::Ok;
 			}
 			else
 			{
-				UInt32 FramesToCopy = ReadLen / OutputDesc.mBytesPerFrame;
+				UInt32 FramesToCopy = (*ReadLen / OutputDesc.mBytesPerFrame);
 				AudioBufferList OutputBuffer;
 				OutputBuffer.mNumberBuffers = 1;
 				OutputBuffer.mBuffers[0].mNumberChannels = StreamDesc.mChannelsPerFrame;
-				OutputBuffer.mBuffers[0].mDataByteSize = ReadLen;
+				OutputBuffer.mBuffers[0].mDataByteSize = *ReadLen;
 				OutputBuffer.mBuffers[0].mData = OutVoiceBuffer;
 				
 				AudioBufferList InputBuffer;
 				InputBuffer.mNumberBuffers = 1;
-				InputBuffer.mBuffers[0].mNumberChannels = StreamDesc.mChannelsPerFrame;
+				InputBuffer.mBuffers[0].mNumberChannels =  StreamDesc.mChannelsPerFrame;
 				InputBuffer.mBuffers[0].mDataByteSize = BufferList[ReadBuffer].mDataByteSize;
 				InputBuffer.mBuffers[0].mData = Data;
 				
-				if ( AudioConverterFillComplexBuffer(StreamConverter, &FVoiceCaptureCoreAudio::ConvertInputFormat, &InputBuffer, &FramesToCopy, &OutputBuffer, nullptr) == 0 )
+				FAudioFileIO inUserData;
+				inUserData.SrcBuffer = &InputBuffer;
+				inUserData.SrcBytesPerFrame = StreamDesc.mBytesPerFrame;
+				inUserData.SrcBufferSize = *ReadLen;
+				inUserData.SrcReadFrames = 0;
+				
+				OSStatus result = AudioConverterFillComplexBuffer(StreamConverter, &FVoiceCaptureCoreAudio::ConvertInputFormat, &inUserData, &FramesToCopy, &OutputBuffer, nullptr);
+				if ( result == 0 || result == eofErr )
 				{
 					State = EVoiceCaptureState::Ok;
+					*ReadLen = inUserData.SrcReadFrames * StreamDesc.mBytesPerFrame;
+					*WriteLen = FramesToCopy * OutputDesc.mBytesPerFrame;
 				}
 				else
 				{
@@ -423,8 +471,10 @@ public:
 			VoiceCaptureState != EVoiceCaptureState::NotCapturing)
 		{
 			GetCaptureState(OutAvailableVoiceData);
+			
 			if ( OutAvailableVoiceData > 0 )
 			{
+				uint32 ConvertedVoiceData = 0;
 				if ( InVoiceBufferSize >= OutAvailableVoiceData )
 				{
 					check(ReadOffset <= BufferSize);
@@ -433,11 +483,14 @@ public:
 					while (BytesToRead)
 					{
 						uint32 CurrentRead = FMath::Min(FMath::Min(BufferList[ReadBuffer].mDataByteSize - ReadOffset, BytesToRead), BufferSize);
-						State = CopyBuffer(ReadBuffer, Data, ReadOffset, CurrentRead);
+						uint32 CurrentWrite = 0;
+						State = CopyBuffer(ReadBuffer, Data, ReadOffset, &CurrentRead, &CurrentWrite);
+						
 						if ( State == EVoiceCaptureState::Ok )
 						{
-							Data += CurrentRead;
-							BytesToRead -= CurrentRead;
+							Data += CurrentWrite;
+							ConvertedVoiceData += CurrentWrite;
+							BytesToRead = ( CurrentRead != 0 ) ? BytesToRead - CurrentRead : 0;
 							ReadOffset += CurrentRead;
 							if( ReadOffset == BufferList[ReadBuffer].mDataByteSize )
 							{
@@ -455,6 +508,8 @@ public:
 					}
 					int32 BytesRead = OutAvailableVoiceData;
 					FPlatformAtomics::InterlockedAdd(&ReadableBytes, -BytesRead);
+					
+					OutAvailableVoiceData  = ConvertedVoiceData;
 				}
 				else
 				{
@@ -471,7 +526,6 @@ public:
 		{
 			OutAvailableVoiceData = 0;
 		}
-		
 		return State;
 	}
 	
