fix for API deprecation warnings in iOS
updated virtual keyboard for iOS 8 APIs
UEPLAT-544
#ios

diff --git a/Engine/Source/Programs/UnrealBuildTool/IOS/IOSToolChain.cs b/Engine/Source/Programs/UnrealBuildTool/IOS/IOSToolChain.cs
--- a/Engine/Source/Programs/UnrealBuildTool/IOS/IOSToolChain.cs
+++ b/Engine/Source/Programs/UnrealBuildTool/IOS/IOSToolChain.cs
@@ -43,7 +43,7 @@ namespace UnrealBuildTool
 
 		/** Which version of the iOS to allow at build time */
 		[XmlConfig]
-		public static string BuildIOSVersion = "6.1";
+		public static string BuildIOSVersion = "7.0";
 
 		/** Which developer directory to root from */
 		[XmlConfig]
@@ -210,6 +210,9 @@ namespace UnrealBuildTool
 		
 		string GetCompileArguments_Global(CPPEnvironment CompileEnvironment)
 		{
+			IOSPlatform BuildPlat = UEBuildPlatform.GetBuildPlatform(UnrealTargetPlatform.IOS) as IOSPlatform;
+			BuildPlat.SetUpProjectEnvironment(UnrealTargetPlatform.IOS);
+
 			string Result = "";
 
 			Result += " -fmessage-length=0";
@@ -253,7 +256,7 @@ namespace UnrealBuildTool
 				Result += " -isysroot " + BaseSDKDir + "/iPhoneOS" + IOSSDKVersion + ".sdk";
 			}
 
-			Result += " -miphoneos-version-min=" + BuildIOSVersion;
+			Result += " -miphoneos-version-min=" + BuildPlat.GetRunTimeVersion();
 
 			// Optimize non- debug builds.
 			if (CompileEnvironment.Config.Target.Configuration != CPPTargetConfiguration.Debug)
@@ -406,7 +409,10 @@ namespace UnrealBuildTool
 
 		string GetLinkArguments_Global( LinkEnvironment LinkEnvironment )
 		{
-			string Result = "";
+			IOSPlatform BuildPlat = UEBuildPlatform.GetBuildPlatform(UnrealTargetPlatform.IOS) as IOSPlatform;
+			BuildPlat.SetUpProjectEnvironment(UnrealTargetPlatform.IOS);
+
+						string Result = "";
 			if (LinkEnvironment.Config.Target.Architecture == "-simulator")
 			{
 				Result += " -arch i386";
@@ -418,7 +424,7 @@ namespace UnrealBuildTool
 				Result += " -isysroot " + XcodeDeveloperDir + "Platforms/iPhoneOS.platform/Developer/SDKs/iPhoneOS" + IOSSDKVersion + ".sdk";
 			}
 			Result += " -dead_strip";
-			Result += " -miphoneos-version-min=" + BuildIOSVersion;
+			Result += " -miphoneos-version-min=" + BuildPlat.GetRunTimeVersion();
 			Result += " -Wl,-no_pie";
 			//			Result += " -v";
 
diff --git a/Engine/Source/Programs/UnrealBuildTool/IOS/UEBuildIOS.cs b/Engine/Source/Programs/UnrealBuildTool/IOS/UEBuildIOS.cs
--- a/Engine/Source/Programs/UnrealBuildTool/IOS/UEBuildIOS.cs
+++ b/Engine/Source/Programs/UnrealBuildTool/IOS/UEBuildIOS.cs
@@ -16,7 +16,7 @@ namespace UnrealBuildTool
         public static string IOSArchitecture = "";
 
 		/** Which version of the iOS to allow at run time */
-		public static string RunTimeIOSVersion = "6.1";
+		public static string RunTimeIOSVersion = "7.0";
 
 		/** which devices the game is allowed to run on */
 		public static string RunTimeIOSDevices = "1,2";
@@ -297,6 +297,11 @@ namespace UnrealBuildTool
 
             InBuildTarget.GlobalLinkEnvironment.Config.AdditionalFrameworks.Add( new UEBuildFramework( "GameKit" ) );
             InBuildTarget.GlobalLinkEnvironment.Config.AdditionalFrameworks.Add( new UEBuildFramework( "StoreKit" ) );
+
+			if (RunTimeIOSVersion != "6.1")
+			{
+				InBuildTarget.GlobalLinkEnvironment.Config.AdditionalFrameworks.Add (new UEBuildFramework ("MultipeerConnectivity"));
+			}
         }
 
         /**
diff --git a/Engine/Source/Runtime/Core/Private/IOS/IOSAppDelegate.cpp b/Engine/Source/Runtime/Core/Private/IOS/IOSAppDelegate.cpp
--- a/Engine/Source/Runtime/Core/Private/IOS/IOSAppDelegate.cpp
+++ b/Engine/Source/Runtime/Core/Private/IOS/IOSAppDelegate.cpp
@@ -60,6 +60,9 @@ void InstallSignalHandlers()
 
 #if !UE_BUILD_SHIPPING
 	@synthesize ConsoleAlert;
+#ifdef __IPHONE_8_0
+    @synthesize ConsoleAlertController;
+#endif
 	@synthesize ConsoleHistoryValues;
 	@synthesize ConsoleHistoryValuesIndex;
 #endif
@@ -654,7 +657,18 @@ void InstallSignalHandlers()
 	// create the leaderboard display object 
 	GKGameCenterViewController* GameCenterDisplay = [[[GKGameCenterViewController alloc] init] autorelease];
 	GameCenterDisplay.viewState = GKGameCenterViewControllerStateLeaderboards;
-	GameCenterDisplay.leaderboardCategory = Category;
+#ifdef __IPHONE_7_0
+    if ([GameCenterDisplay respondsToSelector:@selector(leaderboardIdentifier)] == YES)
+    {
+        GameCenterDisplay.leaderboardIdentifier = Category;
+    }
+    else
+#endif
+    {
+#if __IPHONE_OS_VERSION_MIN_REQUIRED < __IPHONE_7_0
+        GameCenterDisplay.leaderboardCategory = Category;
+#endif
+    }
 	GameCenterDisplay.gameCenterDelegate = self;
 
 	// show it 
diff --git a/Engine/Source/Runtime/Core/Private/IOS/IOSInputInterface.cpp b/Engine/Source/Runtime/Core/Private/IOS/IOSInputInterface.cpp
--- a/Engine/Source/Runtime/Core/Private/IOS/IOSInputInterface.cpp
+++ b/Engine/Source/Runtime/Core/Private/IOS/IOSInputInterface.cpp
@@ -33,7 +33,20 @@ void FIOSInputInterface::Tick( float DeltaTime )
 
 void ModifyVectorByOrientation(FVector& Vec, bool bIsRotation)
 {
-	UIInterfaceOrientation Orientation = [[IOSAppDelegate GetDelegate].IOSController interfaceOrientation];
+    UIInterfaceOrientation Orientation = UIInterfaceOrientationPortrait;
+#ifdef __IPHONE_8_0
+    if ([[IOSAppDelegate GetDelegate].IOSController respondsToSelector:@selector(interfaceOrientation)] == NO)
+    {
+        Orientation = [[UIApplication sharedApplication] statusBarOrientation];
+    }
+    else
+#endif
+    {
+#if __IPHONE_OS_VERSION_MIN_REQUIRED < __IPHONE_8_0
+        UIInterfaceOrientation Orientation = [[IOSAppDelegate GetDelegate].IOSController interfaceOrientation];
+#endif
+    }
+
 	switch (Orientation)
 	{
 	case UIInterfaceOrientationPortrait:
diff --git a/Engine/Source/Runtime/Core/Private/IOS/IOSWindow.cpp b/Engine/Source/Runtime/Core/Private/IOS/IOSWindow.cpp
--- a/Engine/Source/Runtime/Core/Private/IOS/IOSWindow.cpp
+++ b/Engine/Source/Runtime/Core/Private/IOS/IOSWindow.cpp
@@ -30,14 +30,36 @@ void FIOSWindow::Initialize( class FIOSApplication* const Application, const TSh
 	if(InParent.Get() != NULL)
 	{
 		dispatch_async(dispatch_get_main_queue(),^ {
-			UIAlertView* AlertView = [[UIAlertView alloc] initWithTitle:@""
-															message:@"Error: Only one UIWindow may be created on iOS."
-															delegate:nil
-															cancelButtonTitle:NSLocalizedString(@"Ok", nil)
-															otherButtonTitles: nil];
-
-			[AlertView show];
-			[AlertView release];
+#ifdef __IPHONE_8_0
+			if ([UIAlertController class])
+			{
+				UIAlertController* AlertController = [UIAlertController alertControllerWithTitle:@""
+														message:@"Error: Only one UIWindow may be created on iOS."
+														preferredStyle:UIAlertControllerStyleAlert];
+				UIAlertAction* okAction = [UIAlertAction
+											actionWithTitle:NSLocalizedString(@"OK", nil)
+											style:UIAlertActionStyleDefault
+											handler:^(UIAlertAction* action)
+											{
+												[AlertController dismissViewControllerAnimated : YES completion : nil];
+											}
+				];
+
+				[AlertController addAction : okAction];
+				[[IOSAppDelegate GetDelegate].IOSController presentViewController : AlertController animated : YES completion : nil];
+			}
+			else
+#endif
+			{
+				UIAlertView* AlertView = [[UIAlertView alloc] initWithTitle:@""
+											message:@"Error: Only one UIWindow may be created on iOS."
+											delegate:nil
+											cancelButtonTitle:NSLocalizedString(@"Ok", nil)
+											otherButtonTitles:nil];
+
+				[AlertView show];
+				[AlertView release];
+			}
 		} );
 	}
 }
diff --git a/Engine/Source/Runtime/Core/Public/IOS/IOSAppDelegate.h b/Engine/Source/Runtime/Core/Public/IOS/IOSAppDelegate.h
--- a/Engine/Source/Runtime/Core/Public/IOS/IOSAppDelegate.h
+++ b/Engine/Source/Runtime/Core/Public/IOS/IOSAppDelegate.h
@@ -60,6 +60,9 @@ UITextFieldDelegate>
 #if !UE_BUILD_SHIPPING
 	/** Properties for managing the console */
 	@property (nonatomic, retain) UIAlertView*		ConsoleAlert;
+#ifdef __IPHONE_8_0
+	@property (nonatomic, assign) UIAlertController* ConsoleAlertController;
+#endif
 	@property (nonatomic, retain) NSMutableArray*	ConsoleHistoryValues;
 	@property (nonatomic, assign) int				ConsoleHistoryValuesIndex;
 #endif
diff --git a/Engine/Source/Runtime/Launch/Private/IOS/IOSAppDelegateConsoleHandling.cpp b/Engine/Source/Runtime/Launch/Private/IOS/IOSAppDelegateConsoleHandling.cpp
--- a/Engine/Source/Runtime/Launch/Private/IOS/IOSAppDelegateConsoleHandling.cpp
+++ b/Engine/Source/Runtime/Launch/Private/IOS/IOSAppDelegateConsoleHandling.cpp
@@ -16,35 +16,94 @@
 	self.ConsoleHistoryValuesIndex = [self.ConsoleHistoryValues count];
 
 	// Set up a containing alert message and buttons
-	self.ConsoleAlert = [[UIAlertView alloc] initWithTitle:@"Type a console command"
-													message:@""
-													delegate:self
-													cancelButtonTitle:NSLocalizedString(@"Cancel", nil)
-													otherButtonTitles:NSLocalizedString(@"OK", nil), nil];
-
-	self.ConsoleAlert.alertViewStyle = UIAlertViewStylePlainTextInput;
-
-	// The property is now the owner
-	[self.ConsoleAlert release];
-
-	UITextField* TextField = [self.ConsoleAlert textFieldAtIndex:0];
-	TextField.clearsOnBeginEditing = NO;
-	TextField.autocorrectionType = UITextAutocorrectionTypeNo;
-	TextField.autocapitalizationType = UITextAutocapitalizationTypeNone;
-	TextField.placeholder = @"or swipe for history";
-	TextField.clearButtonMode = UITextFieldViewModeWhileEditing;
-	TextField.delegate = self;
-
-	// Add gesture recognizers
-	UISwipeGestureRecognizer* SwipeLeftGesture = [[UISwipeGestureRecognizer alloc] initWithTarget:self action:@selector(SwipeLeftAction:)];
-	SwipeLeftGesture.direction = UISwipeGestureRecognizerDirectionLeft;
-	[TextField addGestureRecognizer:SwipeLeftGesture];
-
-	UISwipeGestureRecognizer* SwipeRightGesture = [[UISwipeGestureRecognizer alloc] initWithTarget:self action:@selector(SwipeRightAction:)];
-	SwipeRightGesture.direction = UISwipeGestureRecognizerDirectionRight;
-	[TextField addGestureRecognizer:SwipeRightGesture];
-
-	[self.ConsoleAlert show];
+#ifdef __IPHONE_8_0
+	if ([UIAlertController class])
+	{
+		self.ConsoleAlertController = [UIAlertController alertControllerWithTitle : @""
+												message:@"Type a console command"
+												preferredStyle:UIAlertControllerStyleAlert];
+
+		UIAlertAction* okAction = [UIAlertAction 
+									actionWithTitle:NSLocalizedString(@"OK", nil)
+									style:UIAlertActionStyleDefault
+									handler:^(UIAlertAction* action)
+									{
+										[self.ConsoleAlertController dismissViewControllerAnimated : YES completion : nil];
+
+										// we clicked Ok (not Cancel at index 0), submit the console command
+                                        UITextField* AlertTextField = self.ConsoleAlertController.textFields.firstObject;
+										[self HandleConsoleCommand:AlertTextField.text];
+									}
+		];
+		UIAlertAction* cancelAction = [UIAlertAction
+										actionWithTitle : NSLocalizedString(@"Cancel", nil)
+										style : UIAlertActionStyleDefault
+										handler : ^ (UIAlertAction* action)
+										{
+											[self.ConsoleAlertController dismissViewControllerAnimated : YES completion : nil];
+										}
+		];
+
+		[self.ConsoleAlertController addAction : okAction];
+		[self.ConsoleAlertController addAction : cancelAction];
+		[self.ConsoleAlertController
+			addTextFieldWithConfigurationHandler : ^ (UITextField* AlertTextField)
+			{
+				AlertTextField.clearsOnBeginEditing = NO;
+				AlertTextField.autocorrectionType = UITextAutocorrectionTypeNo;
+				AlertTextField.autocapitalizationType = UITextAutocapitalizationTypeNone;
+				AlertTextField.placeholder = @"or swipe for history";
+				AlertTextField.clearButtonMode = UITextFieldViewModeWhileEditing;
+				AlertTextField.delegate = self;
+				AlertTextField.clearsOnInsertion = NO;
+				AlertTextField.keyboardType = UIKeyboardTypeDefault;
+
+				// Add gesture recognizers
+				UISwipeGestureRecognizer* SwipeLeftGesture = [[UISwipeGestureRecognizer alloc] initWithTarget:self action : @selector(SwipeLeftAction:)];
+				SwipeLeftGesture.direction = UISwipeGestureRecognizerDirectionLeft;
+				[AlertTextField addGestureRecognizer : SwipeLeftGesture];
+
+				UISwipeGestureRecognizer* SwipeRightGesture = [[UISwipeGestureRecognizer alloc] initWithTarget:self action : @selector(SwipeRightAction:)];
+				SwipeRightGesture.direction = UISwipeGestureRecognizerDirectionRight;
+				[AlertTextField addGestureRecognizer : SwipeRightGesture];
+			}
+		];
+//		[self.ConsoleAlertController release];
+		[[IOSAppDelegate GetDelegate].IOSController presentViewController : self.ConsoleAlertController animated : YES completion : nil];
+	}
+	else
+#endif
+	{
+		self.ConsoleAlert = [[UIAlertView alloc] initWithTitle:@"Type a console command"
+								message:@""
+								delegate:self
+								cancelButtonTitle:NSLocalizedString(@"Cancel", nil)
+								otherButtonTitles:NSLocalizedString(@"OK", nil), nil];
+
+		self.ConsoleAlert.alertViewStyle = UIAlertViewStylePlainTextInput;
+
+		// The property is now the owner
+		[self.ConsoleAlert release];
+
+		UITextField* TextField = [self.ConsoleAlert textFieldAtIndex : 0];
+		TextField.clearsOnBeginEditing = NO;
+		TextField.autocorrectionType = UITextAutocorrectionTypeNo;
+		TextField.autocapitalizationType = UITextAutocapitalizationTypeNone;
+		TextField.placeholder = @"or swipe for history";
+		TextField.clearButtonMode = UITextFieldViewModeWhileEditing;
+		TextField.delegate = self;
+
+		// Add gesture recognizers
+		UISwipeGestureRecognizer* SwipeLeftGesture = [[UISwipeGestureRecognizer alloc] initWithTarget:self action : @selector(SwipeLeftAction:)];
+		SwipeLeftGesture.direction = UISwipeGestureRecognizerDirectionLeft;
+		[TextField addGestureRecognizer : SwipeLeftGesture];
+
+		UISwipeGestureRecognizer* SwipeRightGesture = [[UISwipeGestureRecognizer alloc] initWithTarget:self action : @selector(SwipeRightAction:)];
+		SwipeRightGesture.direction = UISwipeGestureRecognizerDirectionRight;
+		[TextField addGestureRecognizer : SwipeRightGesture];
+
+		[self.ConsoleAlert show];
+	}
 }
 
 /**
@@ -89,23 +148,50 @@
  */
 - (void)ShowAlert:(NSMutableArray*)StringArray
 {
-	// set up the alert message and buttons
-	UIAlertView* Alert = [[[UIAlertView alloc] initWithTitle:[StringArray objectAtIndex:0]
-													 message:[StringArray objectAtIndex:1]
-													delegate:self // use ourself to handle the button clicks 
-										   cancelButtonTitle:[StringArray objectAtIndex:2]
-										   otherButtonTitles:nil] autorelease];
+#ifdef __IPHONE_8_0
+	if ([UIAlertController class])
+	{
+		UIAlertController* AlertController = [UIAlertController alertControllerWithTitle:[StringArray objectAtIndex : 0]
+											message : [StringArray objectAtIndex : 1]
+											preferredStyle:UIAlertControllerStyleAlert];
 
-	Alert.alertViewStyle = UIAlertViewStyleDefault;
+		// add any extra buttons
+		for (int OptionalButtonIndex = 2; OptionalButtonIndex < [StringArray count]; OptionalButtonIndex++)
+		{
+			UIAlertAction* alertAction = [UIAlertAction
+											actionWithTitle : [StringArray objectAtIndex : OptionalButtonIndex]
+											style : UIAlertActionStyleDefault
+											handler : ^ (UIAlertAction* action)
+											{
+												[AlertController dismissViewControllerAnimated : YES completion : nil];
+											}
+			];
+			[AlertController addAction : alertAction];
+		}
 
-	// add any extra buttons
-	for (int OptionalButtonIndex = 3; OptionalButtonIndex < [StringArray count]; OptionalButtonIndex++)
-	{
-		[Alert addButtonWithTitle:[StringArray objectAtIndex:OptionalButtonIndex]];
+		[[IOSAppDelegate GetDelegate].IOSController presentViewController : AlertController animated : YES completion : nil];
 	}
+	else
+#endif
+	{
+		// set up the alert message and buttons
+		UIAlertView* Alert = [[[UIAlertView alloc] initWithTitle:[StringArray objectAtIndex : 0]
+									message : [StringArray objectAtIndex : 1]
+									delegate : self // use ourself to handle the button clicks 
+									cancelButtonTitle : [StringArray objectAtIndex : 2]
+									otherButtonTitles : nil] autorelease];
 
-	// show it!
-	[Alert show];
+		Alert.alertViewStyle = UIAlertViewStyleDefault;
+
+		// add any extra buttons
+		for (int OptionalButtonIndex = 3; OptionalButtonIndex < [StringArray count]; OptionalButtonIndex++)
+		{
+			[Alert addButtonWithTitle : [StringArray objectAtIndex : OptionalButtonIndex]];
+		}
+
+		// show it!
+		[Alert show];
+	}
 }
 
 - (BOOL)textFieldShouldReturn:(UITextField *)alertTextField 
@@ -141,7 +227,17 @@
 		self.ConsoleHistoryValuesIndex + 1 < self.ConsoleHistoryValues.count )
 	{
 		self.ConsoleHistoryValuesIndex++;
-		UITextField* TextField = [self.ConsoleAlert textFieldAtIndex:0];
+        UITextField* TextField = nil;
+#ifdef __IPHONE_8_0
+        if ([UIAlertController class])
+        {
+            TextField = self.ConsoleAlertController.textFields.firstObject;
+        }
+        else
+#endif
+        {
+            TextField = [self.ConsoleAlert textFieldAtIndex:0];
+        }
 		TextField.text = [self.ConsoleHistoryValues objectAtIndex:self.ConsoleHistoryValuesIndex];
 	}
 }
@@ -153,7 +249,17 @@
 		self.ConsoleHistoryValuesIndex > 0 )
 	{
 		self.ConsoleHistoryValuesIndex--;
-		UITextField* TextField = [self.ConsoleAlert textFieldAtIndex:0];
+        UITextField* TextField = nil;
+#ifdef __IPHONE_8_0
+        if ([UIAlertController class])
+        {
+            TextField = self.ConsoleAlertController.textFields.firstObject;
+        }
+        else
+#endif
+        {
+            TextField = [self.ConsoleAlert textFieldAtIndex:0];
+        }
 		TextField.text = [self.ConsoleHistoryValues objectAtIndex:self.ConsoleHistoryValuesIndex];
 	}
 }
diff --git a/Engine/Source/Runtime/Online/IOS/OnlineSubsystemIOS/Private/OnlineFriendsInterfaceIOS.cpp b/Engine/Source/Runtime/Online/IOS/OnlineSubsystemIOS/Private/OnlineFriendsInterfaceIOS.cpp
--- a/Engine/Source/Runtime/Online/IOS/OnlineSubsystemIOS/Private/OnlineFriendsInterfaceIOS.cpp
+++ b/Engine/Source/Runtime/Online/IOS/OnlineSubsystemIOS/Private/OnlineFriendsInterfaceIOS.cpp
@@ -59,13 +59,80 @@ bool FOnlineFriendsIOS::ReadFriendsList(int32 LocalUserNum, const FString& ListN
 		dispatch_async(dispatch_get_main_queue(), ^
 		{
 		// Get the friends list for the local player from the server
-		[[GKLocalPlayer localPlayer] loadFriendsWithCompletionHandler:
-			^(NSArray* Friends, NSError* Error) 
-			{
-				if( Error )
-				{
-					FString ErrorStr(FString::Printf(TEXT("FOnlineFriendsIOS::ReadFriendsList() - Failed to read friends list with error: [%i]"), [Error code]));
-					UE_LOG(LogOnline, Verbose, TEXT("%s"), *ErrorStr);
+#ifdef __IPHONE_8_0
+            if ([[GKLocalPlayer localPlayer] respondsToSelector:@selector(loadFriendPlayersWithCompletionHandler)] == YES)
+            {
+                [[GKLocalPlayer localPlayer] loadFriendPlayersWithCompletionHandler:
+                 ^(NSArray* Friends, NSError* Error)
+                 {
+                    if( Error )
+                    {
+                        FString ErrorStr(FString::Printf(TEXT("FOnlineFriendsIOS::ReadFriendsList() - Failed to read friends list with error: [%i]"), [Error code]));
+                        UE_LOG(LogOnline, Verbose, TEXT("%s"), *ErrorStr);
+                 
+                        // Report back to the game thread whether this succeeded.
+                        [FIOSAsyncTask CreateTaskWithBlock : ^ bool(void)
+                         {
+                            Delegate.ExecuteIfBound(0, false, ListName, ErrorStr);
+                            return true;
+                         }];
+                    }
+                    else
+                    {
+                        [GKPlayer loadPlayersForIdentifiers:Friends withCompletionHandler:
+                         ^(NSArray* Players, NSError* Error2)
+                         {
+                            FString ErrorStr;
+                            bool bWasSuccessful = false;
+                            if( Error2 )
+                            {
+                                ErrorStr = FString::Printf(TEXT("FOnlineFriendsIOS::ReadFriendsList() - Failed to loadPlayersForIdentifiers with error: [%i]"), [Error2 code]);
+                            }
+                            else
+                            {
+                                // Clear our previosly cached friends before we repopulate the cache.
+                                CachedFriends.Empty();
+                                for( int32 FriendIdx = 0; FriendIdx < [Players count]; FriendIdx++ )
+                                {
+                                    GKPlayer* Friend = Players[ FriendIdx ];
+                  
+                                    // Add new friend entry to list
+                                    TSharedRef<FOnlineFriendIOS> FriendEntry(
+                                            new FOnlineFriendIOS(*FString(Friend.playerID))
+                                    );
+                                    FriendEntry->AccountData.Add(
+                                            TEXT("nickname"), *FString(Friend.alias)
+                                    );
+                                    CachedFriends.Add(FriendEntry);
+                  
+                                    UE_LOG(LogOnline, Verbose, TEXT("GCFriend - Id:%s Alias:%s"), *FString(Friend.playerID),
+                                        *FriendEntry->GetDisplayName() );
+                                }
+                  
+                                bWasSuccessful = true;
+                            }
+                  
+                            // Report back to the game thread whether this succeeded.
+                            [FIOSAsyncTask CreateTaskWithBlock : ^ bool(void)
+                             {
+                                Delegate.ExecuteIfBound(0, bWasSuccessful, ListName, ErrorStr);
+                                return true;
+                             }];
+                         }];
+                    }
+                 }];
+            }
+            else
+#endif
+            {
+#if __IPHONE_OS_VERSION_MIN_REQUIRED < __IPHONE_8_0
+                [[GKLocalPlayer localPlayer] loadFriendsWithCompletionHandler:
+                 ^(NSArray* Friends, NSError* Error)
+                 {
+                    if( Error )
+                    {
+                        FString ErrorStr(FString::Printf(TEXT("FOnlineFriendsIOS::ReadFriendsList() - Failed to read friends list with error: [%i]"), [Error code]));
+                        UE_LOG(LogOnline, Verbose, TEXT("%s"), *ErrorStr);
 
 						// Report back to the game thread whether this succeeded.
 						[FIOSAsyncTask CreateTaskWithBlock : ^ bool(void)
@@ -73,12 +140,12 @@ bool FOnlineFriendsIOS::ReadFriendsList(int32 LocalUserNum, const FString& ListN
 							Delegate.ExecuteIfBound(0, false, ListName, ErrorStr);
 							return true;
 						}];
-				}
-				else
-				{
-					[GKPlayer loadPlayersForIdentifiers:Friends withCompletionHandler:
-						^(NSArray* Players, NSError* Error2)
-						{
+                    }
+                    else
+                    {
+                        [GKPlayer loadPlayersForIdentifiers:Friends withCompletionHandler:
+                         ^(NSArray* Players, NSError* Error2)
+                         {
 								FString ErrorStr;
 								bool bWasSuccessful = false;
 							if( Error2 )
@@ -106,20 +173,20 @@ bool FOnlineFriendsIOS::ReadFriendsList(int32 LocalUserNum, const FString& ListN
 										*FriendEntry->GetDisplayName() );
 								}
 
-									bWasSuccessful = true;
-							}
+                                bWasSuccessful = true;
+                            }
 								
-								// Report back to the game thread whether this succeeded.
-								[FIOSAsyncTask CreateTaskWithBlock : ^ bool(void)
-								{
-									Delegate.ExecuteIfBound(0, bWasSuccessful, ListName, ErrorStr);
-									return true;
-								}];
-						}
-					];
-				}
+							// Report back to the game thread whether this succeeded.
+							[FIOSAsyncTask CreateTaskWithBlock : ^ bool(void)
+							{
+								Delegate.ExecuteIfBound(0, bWasSuccessful, ListName, ErrorStr);
+								return true;
+							}];
+						}];
+                    }
+                }];
+#endif
 			}
-		];
 		});
 	}
 	else
diff --git a/Engine/Source/Runtime/Online/IOS/OnlineSubsystemIOS/Private/OnlineLeaderboardsInterfaceIOS.cpp b/Engine/Source/Runtime/Online/IOS/OnlineSubsystemIOS/Private/OnlineLeaderboardsInterfaceIOS.cpp
--- a/Engine/Source/Runtime/Online/IOS/OnlineSubsystemIOS/Private/OnlineLeaderboardsInterfaceIOS.cpp
+++ b/Engine/Source/Runtime/Online/IOS/OnlineSubsystemIOS/Private/OnlineLeaderboardsInterfaceIOS.cpp
@@ -32,12 +32,147 @@ FOnlineLeaderboardsIOS::~FOnlineLeaderboardsIOS()
 	}
 }
 
+bool FOnlineLeaderboardsIOS::ReadLeaderboardCompletionDelegate(NSArray* players, FOnlineLeaderboardReadRef& InReadObject)
+{
+    auto ReadObject = InReadObject;
+    bool bTriggeredReadRequest = false;
+
+    GKLeaderboard* LeaderboardRequest = nil;
+#ifdef __IPHONE_8_0
+    if ([GKLeaderboard respondsToSelector:@selector(initWithPlayers)] == YES)
+    {
+        LeaderboardRequest = [[GKLeaderboard alloc] initWithPlayers:players];
+    }
+    else
+#endif
+    {
+#if __IPHONE_OS_VERSION_MIN_REQUIRED < __IPHONE_8_0
+        LeaderboardRequest = [[GKLeaderboard alloc] initWithPlayerIDs:players];
+#endif
+    }
+    if (LeaderboardRequest != nil)
+    {
+        const FString LeaderboardName = ReadObject->LeaderboardName.ToString();
+        
+        NSString* Category = [NSString stringWithFString:LeaderboardName];
+        UE_LOG(LogOnline, Display, TEXT("Attempting to read leaderboard: %s"), *LeaderboardName);
+        
+        LeaderboardRequest.playerScope = GKLeaderboardPlayerScopeGlobal;
+        LeaderboardRequest.timeScope = GKLeaderboardTimeScopeToday;
+#ifdef __IPHONE_7_0
+        if ([LeaderboardRequest respondsToSelector:@selector(identifier)] == YES)
+        {
+            LeaderboardRequest.identifier = Category;
+        }
+        else
+#endif
+        {
+#if __IPHONE_OS_VERSION_MIN_REQUIRED < __IPHONE_7_0
+            LeaderboardRequest.category = Category;
+#endif
+        }
+        LeaderboardRequest.range = NSMakeRange(1,10);
+        
+        bTriggeredReadRequest = true;
+        dispatch_async(dispatch_get_main_queue(), ^
+        {
+            [LeaderboardRequest loadScoresWithCompletionHandler: ^(NSArray *scores, NSError *Error)
+             {
+                bool bWasSuccessful = (Error == nil) && [scores count] > 0;
+                            
+                if (bWasSuccessful)
+                {
+                    bool bWasSuccessful = [scores count] > 0;
+                    UE_LOG(LogOnline, Display, TEXT("FOnlineLeaderboardsIOS::loadScoresWithCompletionHandler() - %s"), (bWasSuccessful ? TEXT("Success!") : TEXT("Failed!, no scores retrieved")));
+                    for (GKScore* score in scores)
+                    {
+                        FString PlayerIDString;
+                            
+#ifdef __IPHONE_8_0
+                        if ([GKScore respondsToSelector:@selector(player)] == YES)
+                        {
+                            PlayerIDString = FString(score.player.playerID);
+                        }
+                        else
+#endif
+                        {
+#if __IPHONE_OS_VERSION_MIN_REQUIRED < __IPHONE_8_0
+                            PlayerIDString = FString(score.playerID);
+#endif
+                        }
+                            
+                        UE_LOG(LogOnline, Display, TEXT("----------------------------------------------------------------"));
+                        UE_LOG(LogOnline, Display, TEXT("PlayerId: %s"), *PlayerIDString);
+                        UE_LOG(LogOnline, Display, TEXT("Value: %d"), score.value);
+                        UE_LOG(LogOnline, Display, TEXT("----------------------------------------------------------------"));
+                            
+                        TSharedRef<FUniqueNetId> UserId = MakeShareable(new FUniqueNetIdString(PlayerIDString));
+                            
+                        FOnlineStatsRow* UserRow = ReadObject.Get().FindPlayerRecord(UserId.Get());
+                        if (UserRow == NULL)
+                        {
+                            UserRow = new (ReadObject->Rows) FOnlineStatsRow(PlayerIDString, UserId);
+                        }
+                            
+                        for (int32 StatIdx = 0; StatIdx < ReadObject->ColumnMetadata.Num(); StatIdx++)
+                        {
+                            const FColumnMetaData& ColumnMeta = ReadObject->ColumnMetadata[StatIdx];
+                            
+                            FVariantData* LastColumn = NULL;
+                            switch (ColumnMeta.DataType)
+                            {
+                                case EOnlineKeyValuePairDataType::Int32:
+                                {
+                                    int32 Value = score.value;
+                                    LastColumn = UserRow->Columns.Add(ColumnMeta.ColumnName, FVariantData(Value));
+                                    bWasSuccessful = true;
+                                    break;
+                                }
+                            
+                                default:
+                                {
+                                    UE_LOG_ONLINE(Warning, TEXT("Unsupported key value pair during retrieval from GameCenter %s"), ColumnMeta.ColumnName);
+                                    break;
+                                }
+                            }
+                        }
+                    }
+                }
+                else if (Error)
+                {
+                    // if we have failed to read the leaderboard then report this
+                    NSDictionary *userInfo = [Error userInfo];
+                    NSString *errstr = [[userInfo objectForKey : NSUnderlyingErrorKey] localizedDescription];
+                    UE_LOG(LogOnline, Display, TEXT("FOnlineLeaderboardsIOS::loadScoresWithCompletionHandler() - Failed to read leaderboard with error: [%s]"), *FString(errstr));
+                    UE_LOG(LogOnline, Warning, TEXT("You should check that the leaderboard name matches that of one in ITunesConnect"));
+                }
+                            
+                // Report back to the game thread whether this succeeded.
+                [FIOSAsyncTask CreateTaskWithBlock : ^ bool(void)
+                 {
+                    ReadObject->ReadState = bWasSuccessful ? EOnlineAsyncTaskState::Done : EOnlineAsyncTaskState::Failed;
+                    TriggerOnLeaderboardReadCompleteDelegates(bWasSuccessful);
+                    return true;
+                 }];
+            }];
+        });
+    }
+    
+    // If we have failed to kick off a read request, we should still tell whoever is listening.
+    if (!bTriggeredReadRequest)
+    {
+        UE_LOG(LogOnline, Display, TEXT("FOnlineLeaderboardsIOS::loadScoresWithCompletionHandler() - Failed!"));
+        TriggerOnLeaderboardReadCompleteDelegates(false);
+    }
+    
+    return bTriggeredReadRequest;
+}
+
 bool FOnlineLeaderboardsIOS::ReadLeaderboards(const TArray< TSharedRef<FUniqueNetId> >& Players, FOnlineLeaderboardReadRef& InReadObject)
 {
 	auto ReadObject = InReadObject;
 
 	UE_LOG(LogOnline, Display, TEXT("FOnlineLeaderboardsIOS::ReadLeaderboards()"));
-	bool bTriggeredReadRequest = false;
 
 	ReadObject->ReadState = EOnlineAsyncTaskState::Failed;
 	ReadObject->Rows.Empty();
@@ -62,101 +197,30 @@ bool FOnlineLeaderboardsIOS::ReadLeaderboards(const TArray< TSharedRef<FUniqueNe
 		}
 
 		// Kick off a game center read request for the list of users
-		GKLeaderboard* LeaderboardRequest = [[GKLeaderboard alloc] initWithPlayerIDs:FriendIds];
-		if (LeaderboardRequest != nil)
-		{
-			const FString LeaderboardName = ReadObject->LeaderboardName.ToString();
-
-			NSString* Category = [NSString stringWithFString:LeaderboardName];
-			UE_LOG(LogOnline, Display, TEXT("Attempting to read leaderboard: %s"), *LeaderboardName);
-
-			LeaderboardRequest.playerScope = GKLeaderboardPlayerScopeGlobal;
-			LeaderboardRequest.timeScope = GKLeaderboardTimeScopeToday;
-			LeaderboardRequest.category = Category;
-			LeaderboardRequest.range = NSMakeRange(1,10);
-
-			bTriggeredReadRequest = true;
-			dispatch_async(dispatch_get_main_queue(), ^
-			{
-				[LeaderboardRequest loadScoresWithCompletionHandler: ^(NSArray *scores, NSError *Error) 
-				{
-					bool bWasSuccessful = (Error == nil) && [scores count] > 0;
-
-					if (bWasSuccessful)
-					{
-						bool bWasSuccessful = [scores count] > 0;
-						UE_LOG(LogOnline, Display, TEXT("FOnlineLeaderboardsIOS::loadScoresWithCompletionHandler() - %s"), (bWasSuccessful ? TEXT("Success!") : TEXT("Failed!, no scores retrieved")));
-						for (GKScore* score in scores)
-						{
-							const FString PlayerIDString(score.playerID);
-
-							UE_LOG(LogOnline, Display, TEXT("----------------------------------------------------------------"));
-							UE_LOG(LogOnline, Display, TEXT("PlayerId: %s"), *PlayerIDString);
-							UE_LOG(LogOnline, Display, TEXT("Value: %d"), score.value);
-							UE_LOG(LogOnline, Display, TEXT("----------------------------------------------------------------"));
-
-							TSharedRef<FUniqueNetId> UserId = MakeShareable(new FUniqueNetIdString(PlayerIDString));
-
-							FOnlineStatsRow* UserRow = ReadObject.Get().FindPlayerRecord(UserId.Get());
-							if (UserRow == NULL)
-							{
-								UserRow = new (ReadObject->Rows) FOnlineStatsRow(PlayerIDString, UserId);
-							}
-
-							for (int32 StatIdx = 0; StatIdx < ReadObject->ColumnMetadata.Num(); StatIdx++)
-							{
-								const FColumnMetaData& ColumnMeta = ReadObject->ColumnMetadata[StatIdx];
-
-								FVariantData* LastColumn = NULL;
-								switch (ColumnMeta.DataType)
-								{
-								case EOnlineKeyValuePairDataType::Int32:
-									{
-										int32 Value = score.value;
-										LastColumn = UserRow->Columns.Add(ColumnMeta.ColumnName, FVariantData(Value));
-										bWasSuccessful = true;
-										break;
-									}
-
-								default:
-									{
-										UE_LOG_ONLINE(Warning, TEXT("Unsupported key value pair during retrieval from GameCenter %s"), ColumnMeta.ColumnName);
-										break;
-									}
-								}
-							}
-						}
-					}
-					else if (Error)
-					{
-						// if we have failed to read the leaderboard then report this
-						NSDictionary *userInfo = [Error userInfo];
-						NSString *errstr = [[userInfo objectForKey : NSUnderlyingErrorKey] localizedDescription];
-						UE_LOG(LogOnline, Display, TEXT("FOnlineLeaderboardsIOS::loadScoresWithCompletionHandler() - Failed to read leaderboard with error: [%s]"), *FString(errstr));
-						UE_LOG(LogOnline, Warning, TEXT("You should check that the leaderboard name matches that of one in ITunesConnect"));
-					}
-
-					// Report back to the game thread whether this succeeded.
-					[FIOSAsyncTask CreateTaskWithBlock : ^ bool(void)
-					{
-						ReadObject->ReadState = bWasSuccessful ? EOnlineAsyncTaskState::Done : EOnlineAsyncTaskState::Failed;
-						TriggerOnLeaderboardReadCompleteDelegates(bWasSuccessful);
-						return true;
-					}];
-				}
-			];
-			});
-		}
-	}
-	
-	// If we have failed to kick off a read request, we should still tell whoever is listening.
-	if (!bTriggeredReadRequest)
-	{
-		UE_LOG(LogOnline, Display, TEXT("FOnlineLeaderboardsIOS::loadScoresWithCompletionHandler() - Failed!"));
-		TriggerOnLeaderboardReadCompleteDelegates(false);
+#ifdef __IPHONE_8_0
+        if ([GKLeaderboard respondsToSelector:@selector(initWithPlayers)] == YES)
+        {
+            [GKPlayer loadPlayersForIdentifiers:FriendIds withCompletionHandler:^(NSArray *players, NSError *Error)
+             {
+                bool bWasSuccessful = (Error == nil) && [players count] > 0;
+             
+                if (bWasSuccessful)
+                {
+                    bool bWasSuccessful = [players count] > 0;
+                    ReadLeaderboardCompletionDelegate(players, InReadObject);
+                }
+             }];
+        }
+        else
+#endif
+        {
+#if __IPHONE_OS_VERSION_MIN_REQUIRED < __IPHONE_8_0
+            return ReadLeaderboardCompletionDelegate(FriendIds, InReadObject);
+#endif
+        }
 	}
 
-	return bTriggeredReadRequest;
+	return true;
 }
 
 
@@ -212,8 +276,20 @@ bool FOnlineLeaderboardsIOS::WriteLeaderboards(const FName& SessionName, const F
 		NSString* Category = [NSString stringWithFString:LeaderboardName];
 
 		// Create a leaderboard score object which should be posted to the [Category] leaderboard.
-		GKScore* Score = [[GKScore alloc] initWithCategory:Category];
-		
+        GKScore* Score = nil;
+#ifdef __IPHONE_7_0
+        if ([GKScore respondsToSelector:@selector(initWithLeaderboardIdentifier)] == YES)
+        {
+            Score = [[GKScore alloc] initWithLeaderboardIdentifier:Category];
+        }
+        else
+#endif
+        {
+#if __IPHONE_OS_VERSION_MIN_REQUIRED < __IPHONE_7_0
+            Score = [[GKScore alloc] initWithCategory:Category];
+#endif
+        }
+        
 		Score.context = 0;
 
 		bool bIsValidScore = false;
diff --git a/Engine/Source/Runtime/Online/IOS/OnlineSubsystemIOS/Private/OnlineLeaderboardsInterfaceIOS.h b/Engine/Source/Runtime/Online/IOS/OnlineSubsystemIOS/Private/OnlineLeaderboardsInterfaceIOS.h
--- a/Engine/Source/Runtime/Online/IOS/OnlineSubsystemIOS/Private/OnlineLeaderboardsInterfaceIOS.h
+++ b/Engine/Source/Runtime/Online/IOS/OnlineSubsystemIOS/Private/OnlineLeaderboardsInterfaceIOS.h
@@ -14,6 +14,8 @@ private:
 
 	NSMutableArray* UnreportedScores;
 
+    bool ReadLeaderboardCompletionDelegate(NSArray* players, FOnlineLeaderboardReadRef& ReadObject);
+    
 PACKAGE_SCOPE:
 
 	FOnlineLeaderboardsIOS(FOnlineSubsystemIOS* InSubsystem);
diff --git a/Engine/Source/Runtime/Online/IOS/OnlineSubsystemIOS/Private/OnlineSessionInterfaceIOS.cpp b/Engine/Source/Runtime/Online/IOS/OnlineSubsystemIOS/Private/OnlineSessionInterfaceIOS.cpp
--- a/Engine/Source/Runtime/Online/IOS/OnlineSubsystemIOS/Private/OnlineSessionInterfaceIOS.cpp
+++ b/Engine/Source/Runtime/Online/IOS/OnlineSubsystemIOS/Private/OnlineSessionInterfaceIOS.cpp
@@ -2,7 +2,8 @@
 
 #include "OnlineSubsystemIOSPrivatePCH.h"
 
-@implementation FGameCenterSessionDelegate
+#if __IPHONE_OS_VERSION_MIN_REQUIRED < __IPHONE_7_0
+@implementation FGameCenterSessionDelegateGK
 @synthesize Session;
 
 - (void)initSessionWithName:(NSString*) sessionName
@@ -22,6 +23,18 @@
 	self.Session.delegate = nil;
 }
 
+- (bool)sessionsAvailable
+{
+    NSArray* availablePeers = [self.Session peersWithConnectionState:GKPeerStateAvailable];
+    return [availablePeers count] > 0;
+}
+
+-(void)joinSession
+{
+    NSString* PeerID = [self.Session peerID];
+    [self connectToPeer:PeerID];
+}
+
 -(void)session:(GKSession *)session didReceiveConnectionRequestFromPeer:(NSString *)peerID
 {
 	UE_LOG(LogOnline, Display, TEXT("-(void)session:(GKSession *)session didReceiveConnectionRequestFromPeer:(NSString *)peerID - %s"), *FString([session displayNameForPeer:peerID]));
@@ -88,6 +101,159 @@
 	UE_LOG(LogOnline, Display, TEXT("%s"), *ConvertedErrorStr);
 }
 
+@end
+#endif
+
+#ifdef __IPHONE_7_0
+@implementation FGameCenterSessionDelegateMC
+@synthesize Session;
+@synthesize PeerID;
+
+- (void)initSessionWithName:(NSString*) sessionName
+{
+    UE_LOG(LogOnline, Display, TEXT("- (void)initSessionWithName:(NSString*) sessionName"));
+    self.PeerID = [[MCPeerID alloc] initWithDisplayName:nil];
+    self.Session = [[MCSession alloc] initWithPeer:self.PeerID];
+    self.Session.delegate = self;
+}
+
+- (void)shutdownSession
+{
+    UE_LOG(LogOnline, Display, TEXT("- (void)shutdownSession"));
+    [self.Session disconnect];
+    self.Session.delegate = nil;
+}
+
+- (bool)sessionsAvailable
+{
+    return NO;
+}
+
+-(void)joinSession
+{
+}
+
+- (void)session:(MCSession *)session didReceiveData:(NSData *)data fromPeer:(MCPeerID *)peerID
+{
+    
+}
+
+- (void)session:(MCSession *)session didStartReceivingResourceWithName:(NSString *)resourceName fromPeer:(MCPeerID *)peerID withProgress:(NSProgress *)progress
+{
+}
+
+- (void)session:(MCSession *)session didFinishReceivingResourceWithName:(NSString *)resourceName fromPeer:(MCPeerID *)peerID atURL:(NSURL *)localURL withError:(NSError *)error
+{
+    
+}
+
+- (void)session:(MCSession *)session didReceiveStream:(NSInputStream *)stream withName:(NSString *)streamName fromPeer:(MCPeerID *)peerID
+{
+    
+}
+
+- (void)session:(MCSession *)session peer:(MCPeerID *)peerID didChangeState:(MCSessionState)state
+{
+    const FString PeerId(PeerID.displayName);
+    switch (state)
+    {
+        case MCSessionStateConnected:
+        {
+            UE_LOG(LogOnline, Display, TEXT("Peer connected: %s"), *PeerId);
+            break;
+        }
+            
+        case MCSessionStateNotConnected:
+        {
+            UE_LOG(LogOnline, Display, TEXT("Peer not connected: %s"), *PeerId);
+            break;
+        }
+    }
+}
+
+@end
+#endif
+
+@implementation FGameCenterSessionDelegate
+#if __IPHONE_OS_VERSION_MIN_REQUIRED < __IPHONE_7_0
+@synthesize SessionGK;
+#endif
+#ifdef __IPHONE_7_0
+@synthesize SessionMC;
+#endif
+
+- (void)initSessionWithName:(NSString*) sessionName
+{
+    UE_LOG(LogOnline, Display, TEXT("- (void)initSessionWithName:(NSString*) sessionName"));
+    // Create the session object
+#ifdef __IPHONE_7_0
+    if ([MCSession class])
+    {
+        self.SessionMC = [FGameCenterSessionDelegateMC alloc];
+        [self.SessionMC initSessionWithName:sessionName];
+    }
+    else
+#endif
+    {
+#if __IPHONE_OS_VERSION_MIN_REQUIRED < __IPHONE_7_0
+        self.SessionGK = [FGameCenterSessionDelegateGK alloc];
+        [self.SessionGK initSessionWithName:sessionName];
+#endif
+    }
+}
+
+- (void)shutdownSession
+{
+    UE_LOG(LogOnline, Display, TEXT("- (void)shutdownSession"));
+#ifdef __IPHONE_7_0
+    if ([MCSession class])
+    {
+        [self.SessionMC shutdownSession];
+    }
+    else
+#endif
+    {
+#if __IPHONE_OS_VERSION_MIN_REQUIRED < __IPHONE_7_0
+        [self.SessionGK shutdownSession];
+#endif
+    }
+}
+
+- (bool)sessionsAvailable
+{
+    UE_LOG(LogOnline, Display, TEXT("- (void)shutdownSession"));
+#ifdef __IPHONE_7_0
+    if ([MCSession class])
+    {
+        return [self.SessionMC sessionsAvailable];
+    }
+    else
+#endif
+    {
+#if __IPHONE_OS_VERSION_MIN_REQUIRED < __IPHONE_7_0
+        return [self.SessionGK sessionsAvailable];
+#endif
+    }
+    return NO;
+}
+
+- (void)joinSession
+{
+    UE_LOG(LogOnline, Display, TEXT("- (void)shutdownSession"));
+#ifdef __IPHONE_7_0
+    if ([MCSession class])
+    {
+        [self.SessionMC joinSession];
+    }
+    else
+#endif
+    {
+#if __IPHONE_OS_VERSION_MIN_REQUIRED < __IPHONE_7_0
+        [self.SessionGK joinSession];
+#endif
+    }
+}
+
 @end
 
 
@@ -205,7 +371,7 @@ bool FOnlineSessionIOS::CreateSession(int32 HostingPlayerNum, FName SessionName,
 		UE_LOG(LogOnline, Display, TEXT("Creating new session."));
 
 		// Create the session object
-		FGameCenterSessionDelegate* NewGKSession = [FGameCenterSessionDelegate alloc];
+        FGameCenterSessionDelegate* NewGKSession = [FGameCenterSessionDelegate alloc];
 		if( NewGKSession != NULL )
 		{
 			UE_LOG(LogOnline, Display, TEXT("Created session delegate"));
@@ -248,9 +414,9 @@ bool FOnlineSessionIOS::StartSession(FName SessionName)
 	if (Session != NULL)
 	{
 		// Find the linked GK session and start it.
-		FGameCenterSessionDelegate* LinkedGKSession = *GKSessions.Find( SessionName );
+        FGameCenterSessionDelegate* LinkedGKSession = *GKSessions.Find( SessionName );
 		NSString* SafeSessionName = [NSString stringWithFString:SessionName.ToString()];
-		[LinkedGKSession  initSessionWithName:SafeSessionName];
+		[LinkedGKSession initSessionWithName:SafeSessionName];
 
 		// Update the session state as we are now running.
 		Session->SessionState = EOnlineSessionState::InProgress;
@@ -352,8 +518,7 @@ bool FOnlineSessionIOS::FindSessions(int32 SearchingPlayerNum, const TSharedRef<
 		for (TMap< FName, FGameCenterSessionDelegate* >::TConstIterator SessionIt(GKSessions); SessionIt; ++SessionIt)
 		{
 			FGameCenterSessionDelegate* GKSession = SessionIt.Value();
-			NSArray* availablePeers = [[GKSession Session] peersWithConnectionState:GKPeerStateAvailable];
-			bSuccessfullyFoundSessions = [availablePeers count] > 0;
+            bSuccessfullyFoundSessions = [GKSession sessionsAvailable];
 		}
 	}
 
@@ -406,9 +571,7 @@ bool FOnlineSessionIOS::JoinSession(int32 PlayerNum, FName SessionName, const FO
 	FGameCenterSessionDelegate* SessionDelegate = *GKSessions.Find( SessionName );
 	if( SessionDelegate != NULL )
 	{
-		NSString* PeerID = [[SessionDelegate Session] peerID];
-		[SessionDelegate connectToPeer:PeerID];
-
+        [SessionDelegate joinSession];
 		JoinSessionResult = EOnJoinSessionCompleteResult::Success;
 	}
 
diff --git a/Engine/Source/Runtime/Online/IOS/OnlineSubsystemIOS/Private/OnlineSessionInterfaceIOS.h b/Engine/Source/Runtime/Online/IOS/OnlineSubsystemIOS/Private/OnlineSessionInterfaceIOS.h
--- a/Engine/Source/Runtime/Online/IOS/OnlineSubsystemIOS/Private/OnlineSessionInterfaceIOS.h
+++ b/Engine/Source/Runtime/Online/IOS/OnlineSubsystemIOS/Private/OnlineSessionInterfaceIOS.h
@@ -1,15 +1,15 @@
 // Copyright 1998-2015 Epic Games, Inc. All Rights Reserved.
 
 #pragma once
+#include <GameKit/GKLocalPlayer.h>
 
 #include "OnlineSessionInterface.h"
 #include "OnlineSubsystemIOSTypes.h"
 
+#if __IPHONE_OS_VERSION_MIN_REQUIRED < __IPHONE_7_0
 #include <GameKit/GKSession.h>
-#include <GameKit/GKLocalPlayer.h>
-
 
-@interface FGameCenterSessionDelegate : UIViewController<GKSessionDelegate>
+@interface FGameCenterSessionDelegateGK : UIViewController<GKSessionDelegate>
 {
 };
 
@@ -17,6 +17,45 @@
 
 -(void) initSessionWithName:(NSString*) sessionName;
 -(void) shutdownSession;
+-(bool) sessionsAvailable;
+-(void)joinSession;
+
+@end
+#endif
+
+#ifdef __IPHONE_7_0
+#include <MultipeerConnectivity/MultipeerConnectivity.h>
+
+@interface FGameCenterSessionDelegateMC : UIViewController<MCSessionDelegate>
+{
+};
+
+@property (nonatomic, strong) MCPeerID *PeerID;
+@property (nonatomic, strong) MCSession *Session;
+
+-(void) initSessionWithName:(NSString*) sessionName;
+-(void) shutdownSession;
+-(bool) sessionsAvailable;
+-(void)joinSession;
+
+@end
+#endif
+
+@interface FGameCenterSessionDelegate : NSObject
+{
+};
+
+#if __IPHONE_OS_VERSION_MIN_REQUIRED < __IPHONE_7_0
+@property (nonatomic, strong) FGameCenterSessionDelegateGK *SessionGK;
+#endif
+#ifdef __IPHONE_7_0
+@property (nonatomic, strong) FGameCenterSessionDelegateMC *SessionMC;
+#endif
+
+-(void) initSessionWithName:(NSString*) sessionName;
+-(void) shutdownSession;
+-(bool) sessionsAvailable;
+-(void)joinSession;
 
 @end
 
diff --git a/Engine/Source/Runtime/Online/IOS/OnlineSubsystemIOS/Private/OnlineStoreInterfaceIOS.cpp b/Engine/Source/Runtime/Online/IOS/OnlineSubsystemIOS/Private/OnlineStoreInterfaceIOS.cpp
--- a/Engine/Source/Runtime/Online/IOS/OnlineSubsystemIOS/Private/OnlineStoreInterfaceIOS.cpp
+++ b/Engine/Source/Runtime/Online/IOS/OnlineSubsystemIOS/Private/OnlineStoreInterfaceIOS.cpp
@@ -60,6 +60,7 @@
 		{
 			FString ReceiptData;
 
+#ifdef __IPHONE_7_0
 			if( [IOSAppDelegate GetDelegate].OSVersion >= 7.0 )
 			{
 				NSURL* nsReceiptUrl = [[NSBundle mainBundle] appStoreReceiptURL];
@@ -69,10 +70,13 @@
 				ReceiptData = nsEncodedReceiptData;
 			}
 			else
+#endif
 			{
+#if __IPHONE_OS_VERSION_MIN_REQUIRED < __IPHONE_7_0
 				// If earlier than IOS 7, we will need to use the transactionReceipt
 				NSString* nsEncodedReceiptData = [transaction.transactionReceipt base64EncodedStringWithOptions:NSDataBase64EncodingEndLineWithLineFeed];
 				ReceiptData = nsEncodedReceiptData;
+#endif
 			}
 
 			StoreInterface->CachedPurchaseStateObject->ProvidedProductInformation.ReceiptData = ReceiptData;
diff --git a/Engine/Source/Runtime/Online/IOS/OnlineSubsystemIOS/Private/OnlineTurnBasedInterfaceIOS.cpp b/Engine/Source/Runtime/Online/IOS/OnlineSubsystemIOS/Private/OnlineTurnBasedInterfaceIOS.cpp
--- a/Engine/Source/Runtime/Online/IOS/OnlineSubsystemIOS/Private/OnlineTurnBasedInterfaceIOS.cpp
+++ b/Engine/Source/Runtime/Online/IOS/OnlineSubsystemIOS/Private/OnlineTurnBasedInterfaceIOS.cpp
@@ -114,7 +114,20 @@ int32 FTurnBasedMatchIOS::GetLocalPlayerIndex() const
 	NSArray* participantArray = Match.participants;
 	for (GKTurnBasedParticipant* participant in participantArray)
 	{
-		if ([playerID isEqualToString : participant.playerID])
+        NSString* PlayerIDString = nil;
+#ifdef __IPHONE_8_0
+        if ([GKTurnBasedParticipant respondsToSelector:@selector(player)] == YES)
+        {
+            PlayerIDString = participant.player.playerID;
+        }
+        else
+#endif
+        {
+#if __IPHONE_OS_VERSION_MIN_REQUIRED < __IPHONE_8_0
+            PlayerIDString = participant.playerID;
+#endif
+        }
+		if ([playerID isEqualToString : PlayerIDString])
 		{
 			return PlayerIndex;
 		}
@@ -151,7 +164,20 @@ int32 FTurnBasedMatchIOS::GetPlayerIndexForPlayer(NSString* PlayerID) const
 	int32 playerIndex = 0;
 	for (GKTurnBasedParticipant* participant in Match.participants)
 	{
-		if ([participant.playerID isEqualToString : PlayerID])
+        NSString* PlayerIDString = nil;
+#ifdef __IPHONE_8_0
+        if ([GKTurnBasedParticipant respondsToSelector:@selector(player)] == YES)
+        {
+            PlayerIDString = participant.player.playerID;
+        }
+        else
+#endif
+        {
+#if __IPHONE_OS_VERSION_MIN_REQUIRED < __IPHONE_8_0
+            PlayerIDString = participant.playerID;
+#endif
+        }
+		if ([PlayerIDString isEqualToString : PlayerID])
 		{
 			return playerIndex;
 		}
@@ -538,11 +564,24 @@ NSArray* FOnlineTurnBasedIOS::GetPlayerIdentifierArrayForMatch(GKTurnBasedMatch*
 	NSMutableArray* result = [NSMutableArray array];
 	for (GKTurnBasedParticipant* participant in match.participants)
 	{
-		if (!participant.playerID)
+        NSString* PlayerIDString = nil;
+#ifdef __IPHONE_8_0
+        if ([GKTurnBasedParticipant respondsToSelector:@selector(player)] == YES)
+        {
+            PlayerIDString = participant.player.playerID;
+        }
+        else
+#endif
+        {
+#if __IPHONE_OS_VERSION_MIN_REQUIRED < __IPHONE_8_0
+            PlayerIDString = participant.playerID;
+#endif
+        }
+		if (!PlayerIDString)
 		{
 			break;
 		}
-		[result addObject : participant.playerID];
+		[result addObject : PlayerIDString];
 	}
 	return result;
 }
diff --git a/Engine/Source/Runtime/Online/IOS/OnlineSubsystemIOS/Private/TurnBasedMatchmakerDelegateIOS.cpp b/Engine/Source/Runtime/Online/IOS/OnlineSubsystemIOS/Private/TurnBasedMatchmakerDelegateIOS.cpp
--- a/Engine/Source/Runtime/Online/IOS/OnlineSubsystemIOS/Private/TurnBasedMatchmakerDelegateIOS.cpp
+++ b/Engine/Source/Runtime/Online/IOS/OnlineSubsystemIOS/Private/TurnBasedMatchmakerDelegateIOS.cpp
@@ -41,11 +41,24 @@
 
 			NSMutableArray* playerIdentifierArray = [NSMutableArray array];
 			for (GKTurnBasedParticipant* participant in match.participants) {
-				if (!participant.playerID)
+                NSString* PlayerIDString = nil;
+#ifdef __IPHONE_8_0
+                if ([GKTurnBasedParticipant respondsToSelector:@selector(player)] == YES)
+                {
+                    PlayerIDString = participant.player.playerID;
+                }
+                else
+#endif
+                {
+#if __IPHONE_OS_VERSION_MIN_REQUIRED < __IPHONE_8_0
+                    PlayerIDString = participant.playerID;
+#endif
+                }
+				if (!PlayerIDString)
 				{
 					break;
 				}
-				[playerIdentifierArray addObject : participant.playerID];
+				[playerIdentifierArray addObject : PlayerIDString];
 			}
 			[GKPlayer loadPlayersForIdentifiers : playerIdentifierArray withCompletionHandler : ^ (NSArray *players, NSError *nameLoadError) {
 				if (nameLoadError) {
diff --git a/Engine/Source/Runtime/Slate/Private/Framework/Text/Android/AndroidPlatformTextField.cpp b/Engine/Source/Runtime/Slate/Private/Framework/Text/Android/AndroidPlatformTextField.cpp
--- a/Engine/Source/Runtime/Slate/Private/Framework/Text/Android/AndroidPlatformTextField.cpp
+++ b/Engine/Source/Runtime/Slate/Private/Framework/Text/Android/AndroidPlatformTextField.cpp
@@ -37,6 +37,7 @@ void FAndroidPlatformTextField::ShowVirtualKeyboard(bool bShow, int32 UserIndex,
 		case EKeyboardType::Keyboard_Password:
 			InputType = TYPE_CLASS_TEXT | TYPE_TEXT_VARIATION_PASSWORD;
 			break;
+		case EKeyboardType::Keyboard_AlphaNumeric:
 		case EKeyboardType::Keyboard_Default:
 		default:
 			InputType = TYPE_CLASS_TEXT | TYPE_TEXT_VARIATION_NORMAL;
diff --git a/Engine/Source/Runtime/Slate/Private/Framework/Text/IOS/IOSPlatformTextField.cpp b/Engine/Source/Runtime/Slate/Private/Framework/Text/IOS/IOSPlatformTextField.cpp
--- a/Engine/Source/Runtime/Slate/Private/Framework/Text/IOS/IOSPlatformTextField.cpp
+++ b/Engine/Source/Runtime/Slate/Private/Framework/Text/IOS/IOSPlatformTextField.cpp
@@ -1,6 +1,7 @@
 // Copyright 1998-2015 Epic Games, Inc. All Rights Reserved.
 
 #include "SlatePrivatePCH.h"
+#include "IOSAppDelegate.h"
 
 
 FIOSPlatformTextField::FIOSPlatformTextField()
@@ -39,19 +40,102 @@ void FIOSPlatformTextField::ShowVirtualKeyboard(bool bShow, int32 UserIndex, TSh
 -(void)show:(TSharedPtr<IVirtualKeyboardEntry>)InTextWidget
 {
 	TextWidget = InTextWidget;
-
-	UIAlertView* AlertView = [[UIAlertView alloc] initWithTitle:@""
-													message:@""
-													delegate:self
-													cancelButtonTitle:NSLocalizedString(@"Cancel", nil)
-													otherButtonTitles:NSLocalizedString(@"OK", nil), nil];
-
-	// give the UIAlertView a style so a UITextField is created
-	switch(TextWidget->GetVirtualKeyboardType())
+#ifdef __IPHONE_8_0
+	if ([UIAlertController class])
 	{
+		UIAlertController* AlertController = [UIAlertController alertControllerWithTitle : @"" message:@"" preferredStyle:UIAlertControllerStyleAlert];
+		UIAlertAction* okAction = [UIAlertAction 
+										actionWithTitle:NSLocalizedString(@"OK", nil)
+										style:UIAlertActionStyleDefault
+										handler:^(UIAlertAction* action)
+										{
+											[AlertController dismissViewControllerAnimated : YES completion : nil];
+
+											UITextField* AlertTextField = AlertController.textFields.firstObject;
+											TextWidget->SetTextFromVirtualKeyboard(FText::FromString(AlertTextField.text));
+
+											FIOSAsyncTask* AsyncTask = [[FIOSAsyncTask alloc] init];
+											AsyncTask.GameThreadCallback = ^ bool(void)
+											{
+												// clear the TextWidget
+												TextWidget = nullptr;
+												return true;
+											};
+											[AsyncTask FinishedTask];
+										}
+		];
+		UIAlertAction* cancelAction = [UIAlertAction
+										actionWithTitle: NSLocalizedString(@"Cancel", nil)
+										style:UIAlertActionStyleDefault
+										handler:^(UIAlertAction* action)
+										{
+											[AlertController dismissViewControllerAnimated : YES completion : nil];
+
+											FIOSAsyncTask* AsyncTask = [[FIOSAsyncTask alloc] init];
+											AsyncTask.GameThreadCallback = ^ bool(void)
+											{
+												// clear the TextWidget
+												TextWidget = nullptr;
+												return true;
+											};
+											[AsyncTask FinishedTask];
+										}
+		];
+
+		[AlertController addAction: okAction];
+		[AlertController addAction: cancelAction];
+		[AlertController
+						addTextFieldWithConfigurationHandler:^(UITextField* AlertTextField)
+						{
+							AlertTextField.clearsOnBeginEditing = NO;
+							AlertTextField.clearsOnInsertion = NO;
+							AlertTextField.autocorrectionType = UITextAutocorrectionTypeNo;
+							AlertTextField.autocapitalizationType = UITextAutocapitalizationTypeNone;
+							AlertTextField.text = [NSString stringWithFString : TextWidget->GetText().ToString()];
+							AlertTextField.placeholder = [NSString stringWithFString : TextWidget->GetHintText().ToString()];
+
+							// set up the keyboard styles not supported in the AlertViewStyle styles
+							switch (TextWidget->GetVirtualKeyboardType())
+							{
+							case EKeyboardType::Keyboard_Email:
+								AlertTextField.keyboardType = UIKeyboardTypeEmailAddress;
+								break;
+							case EKeyboardType::Keyboard_Number:
+								AlertTextField.keyboardType = UIKeyboardTypeDecimalPad;
+								break;
+							case EKeyboardType::Keyboard_Web:
+								AlertTextField.keyboardType = UIKeyboardTypeURL;
+								break;
+							case EKeyboardType::Keyboard_AlphaNumeric:
+								AlertTextField.keyboardType = UIKeyboardTypeASCIICapable;
+								break;
+							case EKeyboardType::Keyboard_Password:
+								AlertTextField.secureTextEntry = YES;
+							case EKeyboardType::Keyboard_Default:
+							default:
+								AlertTextField.keyboardType = UIKeyboardTypeDefault;
+								break;
+							}
+						}
+		];
+		[[IOSAppDelegate GetDelegate].IOSController presentViewController : AlertController animated : YES completion : nil];
+	}
+	else
+#endif
+	{
+		UIAlertView* AlertView = [[UIAlertView alloc] initWithTitle:@""
+									message:@""
+									delegate:self
+									cancelButtonTitle:NSLocalizedString(@"Cancel", nil)
+									otherButtonTitles:NSLocalizedString(@"OK", nil), nil];
+
+		// give the UIAlertView a style so a UITextField is created
+		switch (TextWidget->GetVirtualKeyboardType())
+		{
 		case EKeyboardType::Keyboard_Password:
 			AlertView.alertViewStyle = UIAlertViewStyleSecureTextInput;
 			break;
+		case EKeyboardType::Keyboard_AlphaNumeric:
 		case EKeyboardType::Keyboard_Default:
 		case EKeyboardType::Keyboard_Email:
 		case EKeyboardType::Keyboard_Number:
@@ -59,38 +143,42 @@ void FIOSPlatformTextField::ShowVirtualKeyboard(bool bShow, int32 UserIndex, TSh
 		default:
 			AlertView.alertViewStyle = UIAlertViewStylePlainTextInput;
 			break;
-	}
+		}
+
+		UITextField* AlertTextField = [AlertView textFieldAtIndex : 0];
+		AlertTextField.clearsOnBeginEditing = NO;
+		AlertTextField.clearsOnInsertion = NO;
+		AlertTextField.autocorrectionType = UITextAutocorrectionTypeNo;
+		AlertTextField.autocapitalizationType = UITextAutocapitalizationTypeNone;
+		AlertTextField.text = [NSString stringWithFString : TextWidget->GetText().ToString()];
+		AlertTextField.placeholder = [NSString stringWithFString : TextWidget->GetHintText().ToString()];
+
+		// set up the keyboard styles not supported in the AlertViewStyle styles
+		switch (TextWidget->GetVirtualKeyboardType())
+		{
+		case EKeyboardType::Keyboard_Email:
+			AlertTextField.keyboardType = UIKeyboardTypeEmailAddress;
+			break;
+		case EKeyboardType::Keyboard_Number:
+			AlertTextField.keyboardType = UIKeyboardTypeDecimalPad;
+			break;
+		case EKeyboardType::Keyboard_Web:
+			AlertTextField.keyboardType = UIKeyboardTypeURL;
+			break;
+		case EKeyboardType::Keyboard_AlphaNumeric:
+			AlertTextField.keyboardType = UIKeyboardTypeASCIICapable;
+			break;
+		case EKeyboardType::Keyboard_Default:
+		case EKeyboardType::Keyboard_Password:
+		default:
+			// nothing to do, UIAlertView style handles these keyboard types
+			break;
+		}
 
-	UITextField* AlertTextField = [AlertView textFieldAtIndex: 0];
-	AlertTextField.clearsOnBeginEditing = NO;
-	AlertTextField.clearsOnInsertion = NO;
-	AlertTextField.autocorrectionType = UITextAutocorrectionTypeNo;
-	AlertTextField.autocapitalizationType = UITextAutocapitalizationTypeNone;
-	AlertTextField.text = [NSString stringWithFString : TextWidget->GetText().ToString()];
-	AlertTextField.placeholder = [NSString stringWithFString : TextWidget->GetHintText().ToString()];
+		[AlertView show];
 
-	// set up the keyboard styles not supported in the AlertViewStyle styles
-	switch (TextWidget->GetVirtualKeyboardType())
-	{
-	case EKeyboardType::Keyboard_Email:
-		AlertTextField.keyboardType = UIKeyboardTypeEmailAddress;
-		break;
-	case EKeyboardType::Keyboard_Number:
-		AlertTextField.keyboardType = UIKeyboardTypeDecimalPad;
-		break;
-	case EKeyboardType::Keyboard_Web:
-		AlertTextField.keyboardType = UIKeyboardTypeURL;
-		break;
-	case EKeyboardType::Keyboard_Default:
-	case EKeyboardType::Keyboard_Password:
-	default:
-		// nothing to do, UIAlertView style handles these keyboard types
-		break;
+		[AlertView release];
 	}
-
-	[AlertView show];
-
-	[AlertView release];
 }
 
 - (void)alertView:(UIAlertView *)alertView clickedButtonAtIndex:(NSInteger)buttonIndex
diff --git a/Engine/Source/Runtime/Slate/Public/Widgets/Input/IVirtualKeyboardEntry.h b/Engine/Source/Runtime/Slate/Public/Widgets/Input/IVirtualKeyboardEntry.h
--- a/Engine/Source/Runtime/Slate/Public/Widgets/Input/IVirtualKeyboardEntry.h
+++ b/Engine/Source/Runtime/Slate/Public/Widgets/Input/IVirtualKeyboardEntry.h
@@ -10,6 +10,7 @@ enum EKeyboardType
 	Keyboard_Web,
 	Keyboard_Email,
 	Keyboard_Password,
+	Keyboard_AlphaNumeric,
 };
 
 class SLATE_API IVirtualKeyboardEntry
