CoreAudio requires that we use the native sampling frequency for the input device, otherwise it won't work.

diff --git a/Engine/Source/Runtime/Online/Voice/Private/Mac/VoiceModuleMac.cpp b/Engine/Source/Runtime/Online/Voice/Private/Mac/VoiceModuleMac.cpp
--- a/Engine/Source/Runtime/Online/Voice/Private/Mac/VoiceModuleMac.cpp
+++ b/Engine/Source/Runtime/Online/Voice/Private/Mac/VoiceModuleMac.cpp
@@ -185,15 +185,12 @@ public:
 		}
 		
 		OutputDesc = StreamDesc;
+		OutputDesc.mSampleRate = NativeDesc.mSampleRate;
 		if ( AudioUnitSetProperty(StreamComponent, kAudioUnitProperty_StreamFormat, kAudioUnitScope_Output, 1, (void *)&OutputDesc, sizeof(OutputDesc)) != 0 )
 		{
-			OutputDesc.mSampleRate = NativeDesc.mSampleRate;
-			if ( AudioUnitSetProperty(StreamComponent, kAudioUnitProperty_StreamFormat, kAudioUnitScope_Output, 1, (void *)&OutputDesc, sizeof(OutputDesc)) != 0 )
-			{
-				UE_LOG(LogVoiceCapture, Warning, TEXT("Couldn't configure CoreAudio input component format"));
-				AudioComponentInstanceDispose(StreamComponent);
-				return false;
-			}
+			UE_LOG(LogVoiceCapture, Warning, TEXT("Couldn't configure CoreAudio input component format"));
+			AudioComponentInstanceDispose(StreamComponent);
+			return false;
 		}
 		
 		// create an audio converter
