Add end match functionality for turn based games, also some small refactoring.

diff --git a/Engine/Source/Runtime/Online/IOS/OnlineSubsystemIOS/Private/OnlineTurnBasedInterfaceIOS.cpp b/Engine/Source/Runtime/Online/IOS/OnlineSubsystemIOS/Private/OnlineTurnBasedInterfaceIOS.cpp
--- a/Engine/Source/Runtime/Online/IOS/OnlineSubsystemIOS/Private/OnlineTurnBasedInterfaceIOS.cpp
+++ b/Engine/Source/Runtime/Online/IOS/OnlineSubsystemIOS/Private/OnlineTurnBasedInterfaceIOS.cpp
@@ -68,7 +68,7 @@ bool FTurnBasedMatchIOS::HasMatchData() const
 
 bool FTurnBasedMatchIOS::GetMatchData(TArray<uint8>& OutMatchData) const
 {
-	if (!Match.matchData || Match.matchData.length == 0)
+	if (!HasMatchData())
 	{
 		return false;
 	}
@@ -321,13 +321,45 @@ void FTurnBasedMatchIOS::SetGKMatch(GKTurnBasedMatch* GKMatch)
 	[Match release];
 	Match = [GKMatch retain];
 }
+
+void FTurnBasedMatchIOS::EndMatch(FEndMatchSignature EndMatchCallback, EMPMatchOutcome::Outcome LocalPlayerOutcome, EMPMatchOutcome::Outcome OtherPlayersOutcome)
+{
+	FString MatchID = GetMatchID();
+
+	for (int32 i = 0; i < Match.participants.count; ++i)
+	{
+		GKTurnBasedParticipant* participant = Match.participants[i];
+		GKTurnBasedMatchOutcome GKOutcome = participant.matchOutcome;
+		if (GKOutcome == GKTurnBasedMatchOutcomeNone)
+		{
+            if (GetLocalPlayerIndex() == i)
+            {
+				participant.matchOutcome = GetGKTurnBasedMatchOutcomeFromMatchOutcome(LocalPlayerOutcome);
+            }
+            else
+            {
+				participant.matchOutcome = GetGKTurnBasedMatchOutcomeFromMatchOutcome(OtherPlayersOutcome);
+            }
+			
+		}
+	}
+
+	[Match endMatchInTurnWithMatchData : Match.matchData completionHandler : ^ (NSError* error)
+	{
+		if (EndMatchCallback.IsBound())
+		{
+			EndMatchCallback.Execute(MatchID, error == nil);
+		}
+	}];
+}
+
 FOnlineTurnBasedIOS::FOnlineTurnBasedIOS()
 	: IOnlineTurnBased()
 	, FTurnBasedMatchmakerDelegate()
 	, Matchmaker(*this)
-	, MatchmakerDelegate(NULL)
+	, MatchmakerDelegate(nullptr)
 	, EventListener(nil)
-	, EventDelegate(NULL)
+	, EventDelegate(nullptr)
 	, NumberOfMatchesBeingLoaded(0)
 {
 	EventListener = [[FTurnBasedEventListenerIOS alloc] initWithOwner:*this];
@@ -452,12 +484,11 @@ FTurnBasedMatchPtr FOnlineTurnBasedIOS::GetMatchWithID(FString MatchID) const
 			return *It;
 		}
 	}
-	return NULL;
+	return nullptr;
 }
 
 void FOnlineTurnBasedIOS::RemoveMatch(FTurnBasedMatchRef Match, FRemoveMatchSignature RemoveMatchCallback)
 {
-	// AGC: There's almost certainly a better way to do this cast
 	FTurnBasedMatchIOS& MatchIOS = static_cast<FTurnBasedMatchIOS&>(Match.Get());
 
 	FString MatchID = MatchIOS.GetMatchID();
@@ -473,20 +504,18 @@ void FOnlineTurnBasedIOS::RemoveMatch(FTurnBasedMatchRef Match, FRemoveMatchSign
 
 void FOnlineTurnBasedIOS::OnMatchmakerCancelled()
 {
-	if (!MatchmakerDelegate.IsValid())
+	if (MatchmakerDelegate.IsValid())
 	{
-		return;
+		MatchmakerDelegate.Pin()->OnMatchmakerCancelled();
 	}
-	MatchmakerDelegate.Pin()->OnMatchmakerCancelled();
 }
 
 void FOnlineTurnBasedIOS::OnMatchmakerFailed()
 {
 	if (!MatchmakerDelegate.IsValid())
 	{
-		return;
+		MatchmakerDelegate.Pin()->OnMatchmakerFailed();
 	}
-	MatchmakerDelegate.Pin()->OnMatchmakerFailed();
 }
 
 void FOnlineTurnBasedIOS::OnMatchFound(FTurnBasedMatchRef Match)
@@ -509,10 +538,21 @@ void FOnlineTurnBasedIOS::OnMatchFound(FTurnBasedMatchRef Match)
 
 void FOnlineTurnBasedIOS::OnMatchEnded(FString MatchID)
 {
-	if (EventDelegate.IsValid())
+	dispatch_async(dispatch_get_main_queue(), ^
 	{
-		EventDelegate.Pin()->OnMatchEnded(MatchID);
-	}
+        [FIOSAsyncTask CreateTaskWithBlock : ^ bool(void)
+        {
+            if (TurnBasedMatchInterfaceObject)
+            {
+                   ITurnBasedMatchInterface::Execute_OnMatchEnded(TurnBasedMatchInterfaceObject, MatchID);
+            }
+            if (EventDelegate.IsValid())
+            {
+                EventDelegate.Pin()->OnMatchEnded(MatchID);
+            }
+            return true;
+         }];
+	});
 }
 
 void FOnlineTurnBasedIOS::OnMatchReceivedTurnEvent(FString MatchID, bool BecameActive, void* Match)
diff --git a/Engine/Source/Runtime/Online/IOS/OnlineSubsystemIOS/Private/OnlineTurnBasedInterfaceIOS.h b/Engine/Source/Runtime/Online/IOS/OnlineSubsystemIOS/Private/OnlineTurnBasedInterfaceIOS.h
--- a/Engine/Source/Runtime/Online/IOS/OnlineSubsystemIOS/Private/OnlineTurnBasedInterfaceIOS.h
+++ b/Engine/Source/Runtime/Online/IOS/OnlineSubsystemIOS/Private/OnlineTurnBasedInterfaceIOS.h
@@ -61,6 +61,9 @@ public:
 	GKTurnBasedMatch* GetGKMatch() const { return Match; }
 	void SetGKMatch(GKTurnBasedMatch* GKMatch);
 
+	// Ends the match
+	virtual void EndMatch(FEndMatchSignature QuitMatchCallback, EMPMatchOutcome::Outcome LocalPlayerOutcome, EMPMatchOutcome::Outcome OtherPlayersOutcome);
+
 private:
 
 	// Helper functions for QuitMatch, depending on whether it is the current player's turn or not.
diff --git a/Engine/Source/Runtime/Online/OnlineSubsystem/Public/Interfaces/OnlineTurnBasedInterface.h b/Engine/Source/Runtime/Online/OnlineSubsystem/Public/Interfaces/OnlineTurnBasedInterface.h
--- a/Engine/Source/Runtime/Online/OnlineSubsystem/Public/Interfaces/OnlineTurnBasedInterface.h
+++ b/Engine/Source/Runtime/Online/OnlineSubsystem/Public/Interfaces/OnlineTurnBasedInterface.h
@@ -57,6 +57,7 @@ DECLARE_DELEGATE_TwoParams(FQuitMatchSignature, FString, bool);   // bool Succes
 DECLARE_DELEGATE_TwoParams(FRemoveMatchSignature, FString, bool);   // bool Success
 DECLARE_DELEGATE_TwoParams(FUploadMatchDataSignature, FString, bool);   // bool Success
 DECLARE_DELEGATE_TwoParams(FDownloadMatchDataSignature, FString, bool);  // bool Success
+DECLARE_DELEGATE_TwoParams(FEndMatchSignature, FString, bool);   // bool Success
 
 /**
 * EMPMatchOutcome represents all the possible outcomes for this player in a match
@@ -122,6 +123,9 @@ public:
 
     /* Leave the match, providing an outcome for the player */
     virtual void QuitMatch(EMPMatchOutcome::Outcome Outcome, int32 TurnTimeoutInSeconds, FQuitMatchSignature QuitMatchCallback) {}
+
+	/* Ends the match while setting the match outcome (win/loss/tie) for all players */
+	virtual void EndMatch(FEndMatchSignature QuitMatchCallback, EMPMatchOutcome::Outcome LocalPlayerOutcome, EMPMatchOutcome::Outcome OtherPlayersOutcome) {}
 };
 
 typedef TSharedRef<FTurnBasedMatch, ESPMode::ThreadSafe> FTurnBasedMatchRef;
diff --git a/Engine/Source/Runtime/Online/OnlineSubsystem/Public/Interfaces/TurnBasedMatchInterface.h b/Engine/Source/Runtime/Online/OnlineSubsystem/Public/Interfaces/TurnBasedMatchInterface.h
--- a/Engine/Source/Runtime/Online/OnlineSubsystem/Public/Interfaces/TurnBasedMatchInterface.h
+++ b/Engine/Source/Runtime/Online/OnlineSubsystem/Public/Interfaces/TurnBasedMatchInterface.h
@@ -19,4 +19,7 @@ class ITurnBasedMatchInterface
 public:
 	UFUNCTION(BlueprintImplementableEvent, Category = "Online|TurnBased")
 	void OnMatchReceivedTurn(const FString& Match, bool bDidBecomeActive);
+
+	UFUNCTION(BlueprintImplementableEvent, Category = "Online|TurnBased")
+	void OnMatchEnded(const FString& Match);
 };
diff --git a/Engine/Source/Runtime/Online/OnlineSubsystemUtils/Classes/EndMatchCallbackProxy.h b/Engine/Source/Runtime/Online/OnlineSubsystemUtils/Classes/EndMatchCallbackProxy.h
new file mode 100644
--- /dev/null
+++ b/Engine/Source/Runtime/Online/OnlineSubsystemUtils/Classes/EndMatchCallbackProxy.h
@@ -0,0 +1,56 @@
+// Copyright 1998-2015 Epic Games, Inc. All Rights Reserved.
+
+#pragma once
+
+#include "OnlineTurnBasedInterface.h"
+#include "EndMatchCallbackProxy.generated.h"
+
+class UFindTurnBasedMatchCallbackProxy;
+class UTurnBasedMatchInterface;
+
+UCLASS(MinimalAPI)
+class UEndMatchCallbackProxy : public UOnlineBlueprintCallProxyBase
+{
+	GENERATED_UCLASS_BODY()
+
+	virtual ~UEndMatchCallbackProxy();
+
+	// Called when the match ends successfully
+	UPROPERTY(BlueprintAssignable)
+	FEmptyOnlineDelegate OnSuccess;
+
+	// Called when ending the match fails
+	UPROPERTY(BlueprintAssignable)
+	FEmptyOnlineDelegate OnFailure;
+
+	// End a match that is in progress while it is the current player's turn
+	UFUNCTION(BlueprintCallable, meta = (BlueprintInternalUseOnly = "true", WorldContext = "WorldContextObject"), Category = "Online|TurnBased")
+	static UEndMatchCallbackProxy* EndMatch(UObject* WorldContextObject, class APlayerController* PlayerController, TScriptInterface<ITurnBasedMatchInterface> MatchActor, FString MatchID, EMPMatchOutcome::Outcome LocalPlayerOutcome, EMPMatchOutcome::Outcome OtherPlayersOutcome);
+
+	// UOnlineBlueprintCallProxyBase interface
+	virtual void Activate() override;
+	// End of UOnlineBlueprintCallProxyBase interface
+
+	UTurnBasedMatchInterface* GetTurnBasedMatchInterfaceObject() { return TurnBasedMatchInterface; }
+
+	void EndMatchDelegate(FString MatchID, bool Succeeded);
+
+private:
+	// The player controller triggering things
+	TWeakObjectPtr<APlayerController> PlayerControllerWeakPtr;
+
+	// The world context object in which this call is taking place
+	UObject* WorldContextObject;
+
+	// TurnBasedMatchInterface object, used to set the match data after a match is found
+	UTurnBasedMatchInterface* TurnBasedMatchInterface;
+
+	// ID of the match to end
+	FString MatchID;
+	
+	// Match outcome for the current player (win/loss/tie)
+	EMPMatchOutcome::Outcome LocalPlayerOutcome;
+	
+	// Match outcome for all other players (win/loss/tie)
+	EMPMatchOutcome::Outcome OtherPlayersOutcome;
+};
diff --git a/Engine/Source/Runtime/Online/OnlineSubsystemUtils/Private/EndMatchCallbackProxy.cpp b/Engine/Source/Runtime/Online/OnlineSubsystemUtils/Private/EndMatchCallbackProxy.cpp
new file mode 100644
--- /dev/null
+++ b/Engine/Source/Runtime/Online/OnlineSubsystemUtils/Private/EndMatchCallbackProxy.cpp
@@ -0,0 +1,71 @@
+// Copyright 1998-2015 Epic Games, Inc. All Rights Reserved.
+
+#include "OnlineSubsystemUtilsPrivatePCH.h"
+#include "RepLayout.h"
+#include "TurnBasedMatchInterface.h"
+
+//////////////////////////////////////////////////////////////////////////
+// UEndMatchCallbackProxy
+
+UEndMatchCallbackProxy::UEndMatchCallbackProxy(const FObjectInitializer& ObjectInitializer)
+	: Super(ObjectInitializer)
+	, WorldContextObject(nullptr)
+	, TurnBasedMatchInterface(nullptr)
+{
+}
+
+UEndMatchCallbackProxy::~UEndMatchCallbackProxy()
+{
+}
+
+UEndMatchCallbackProxy* UEndMatchCallbackProxy::EndMatch(UObject* WorldContextObject, class APlayerController* PlayerController, TScriptInterface<ITurnBasedMatchInterface> MatchActor, FString MatchID, EMPMatchOutcome::Outcome LocalPlayerOutcome, EMPMatchOutcome::Outcome OtherPlayersOutcome)
+{
+	UEndMatchCallbackProxy* Proxy = NewObject<UEndMatchCallbackProxy>();
+	Proxy->PlayerControllerWeakPtr = PlayerController;
+	Proxy->WorldContextObject = WorldContextObject;
+	Proxy->MatchID = MatchID;
+	Proxy->LocalPlayerOutcome = LocalPlayerOutcome;
+	Proxy->OtherPlayersOutcome = OtherPlayersOutcome;
+	return Proxy;
+}
+
+void UEndMatchCallbackProxy::Activate()
+{
+	FOnlineSubsystemBPCallHelper Helper(TEXT("ConnectToService"), GEngine->GetWorldFromContextObject(WorldContextObject));
+	Helper.QueryIDFromPlayerController(PlayerControllerWeakPtr.Get());
+
+	if (Helper.IsValid())
+	{
+		IOnlineTurnBasedPtr TurnBasedInterface = Helper.OnlineSub->GetTurnBasedInterface();
+		if (TurnBasedInterface.IsValid())
+		{
+			FTurnBasedMatchPtr Match = TurnBasedInterface->GetMatchWithID(MatchID);
+			if (Match.IsValid())
+			{
+				FEndMatchSignature EndMatchDelegate;
+				EndMatchDelegate.BindUObject(this, &UEndMatchCallbackProxy::EndMatchDelegate);
+				Match->EndMatch(EndMatchDelegate, LocalPlayerOutcome, OtherPlayersOutcome);
+				return;
+			}
+		}
+		else
+		{
+			FFrame::KismetExecutionMessage(TEXT("Turn based games not supported by online subsystem"), ELogVerbosity::Warning);
+		}
+	}
+
+	// Fail immediately
+	OnFailure.Broadcast();
+}
+
+void UEndMatchCallbackProxy::EndMatchDelegate(FString MatchID, bool Succeeded)
+{
+	if (Succeeded)
+	{ 
+		OnSuccess.Broadcast();
+	}
+	else
+	{
+		OnFailure.Broadcast();
+	}
+}
\ No newline at end of file
