Add GetMyPlayerIndex function for turn-based games (references an index in a player array provided by apple/google play that doesn't change)

diff --git a/Engine/Source/Runtime/Online/OnlineSubsystemUtils/Classes/TurnBasedBlueprintLibrary.h b/Engine/Source/Runtime/Online/OnlineSubsystemUtils/Classes/TurnBasedBlueprintLibrary.h
--- a/Engine/Source/Runtime/Online/OnlineSubsystemUtils/Classes/TurnBasedBlueprintLibrary.h
+++ b/Engine/Source/Runtime/Online/OnlineSubsystemUtils/Classes/TurnBasedBlueprintLibrary.h
@@ -14,6 +14,9 @@ class ONLINESUBSYSTEMUTILS_API UTurnBasedBlueprintLibrary : public UBlueprintFun
 	UFUNCTION(BlueprintCallable, meta = (WorldContext = "WorldContextObject"), Category = "Online|TurnBased")
 	static void GetIsMyTurn(UObject* WorldContextObject, APlayerController* PlayerController, FString MatchID, /*out*/ bool& bIsMyTurn);
 
+	UFUNCTION(BlueprintCallable, meta = (WorldContext = "WorldContextObject"), Category = "Online|TurnBased")
+    static void GetMyPlayerIndex(UObject* WorldContextObject, APlayerController* PlayerController, FString MatchID, /*out*/ int32& PlayerIndex);
+
 	UFUNCTION(BlueprintCallable, meta = (WorldContext = "WorldContextObject"), Category = "Online|TurnBased")
 	static void RegisterTurnBasedMatchInterfaceObject(UObject* WorldContextObject, APlayerController* PlayerController, UObject* Object);
 };
diff --git a/Engine/Source/Runtime/Online/OnlineSubsystemUtils/Private/TurnBasedBlueprintLibrary.cpp b/Engine/Source/Runtime/Online/OnlineSubsystemUtils/Private/TurnBasedBlueprintLibrary.cpp
--- a/Engine/Source/Runtime/Online/OnlineSubsystemUtils/Private/TurnBasedBlueprintLibrary.cpp
+++ b/Engine/Source/Runtime/Online/OnlineSubsystemUtils/Private/TurnBasedBlueprintLibrary.cpp
@@ -39,6 +39,36 @@ void UTurnBasedBlueprintLibrary::GetIsMyTurn(UObject* WorldContextObject, APlaye
 	}
 }
 
+void UTurnBasedBlueprintLibrary::GetMyPlayerIndex(UObject* WorldContextObject, APlayerController* PlayerController, FString MatchID, /*out*/ int32& PlayerIndex)
+{
+	PlayerIndex = -1;
+
+	FOnlineSubsystemBPCallHelper Helper(TEXT("GetMyPlayerIndex"), GEngine->GetWorldFromContextObject(WorldContextObject));
+	Helper.QueryIDFromPlayerController(PlayerController);
+
+	if (Helper.IsValid())
+	{
+		IOnlineTurnBasedPtr OnlineTurnBasedPtr = Helper.OnlineSub->GetTurnBasedInterface();
+		if (OnlineTurnBasedPtr.IsValid())
+		{
+			FTurnBasedMatchPtr Match = OnlineTurnBasedPtr->GetMatchWithID(MatchID);
+			if (Match.IsValid())
+			{
+				PlayerIndex = Match->GetLocalPlayerIndex();
+			}
+			else
+			{
+				FString Message = FString::Printf(TEXT("Match ID %s not found"), *MatchID);
+				FFrame::KismetExecutionMessage(*Message, ELogVerbosity::Warning);
+			}
+		}
+		else
+		{
+			FFrame::KismetExecutionMessage(TEXT("Turn Based Matches not supported by Online Subsystem"), ELogVerbosity::Warning);
+		}
+	}
+}
+
 void UTurnBasedBlueprintLibrary::RegisterTurnBasedMatchInterfaceObject(UObject* WorldContextObject, APlayerController* PlayerController, UObject* Object)
 {
 	FOnlineSubsystemBPCallHelper Helper(TEXT("GetIsMyTurn"), GEngine->GetWorldFromContextObject(WorldContextObject));
