Correctly set response code on error for HTML5 HttpRequests.

diff --git a/Engine/Source/Runtime/Online/HTTP/Private/HTML5/HTML5HTTP.cpp b/Engine/Source/Runtime/Online/HTTP/Private/HTML5/HTML5HTTP.cpp
--- a/Engine/Source/Runtime/Online/HTTP/Private/HTML5/HTML5HTTP.cpp
+++ b/Engine/Source/Runtime/Online/HTTP/Private/HTML5/HTML5HTTP.cpp
@@ -239,7 +239,18 @@ void FHTML5HttpRequest::StaticErrorCallback(void* arg, int httpStatusCode, const
 }
 
 void FHTML5HttpRequest::ErrorCallback(void* arg, int httpStatusCode, const char* httpStatusText) {
-	UE_LOG(LogHttp, Verbose, TEXT("FHTML5HttpRequest::ErrorDataCallback()"));
+ 	UE_LOG(LogHttp, Verbose, TEXT("FHTML5HttpRequest::ErrorDataCallback() HttpStatusCode: %d"), httpStatusCode);
+
+	check(Response.IsValid());
+
+	if (Response.IsValid())
+	{
+		Response->Payload.Empty();
+		Response->TotalBytesRead = 0;
+		Response->HttpCode = httpStatusCode;
+		MarkAsCompleted();
+	}
+
 }
 
 void FHTML5HttpRequest::StaticProgressCallback(void* arg, int Loaded, int Total) {
