Fix for BuildPatchManifest being GC'd when still in use.

diff --git a/Engine/Source/Runtime/Online/BuildPatchServices/Private/BuildPatchManifest.cpp b/Engine/Source/Runtime/Online/BuildPatchServices/Private/BuildPatchManifest.cpp
--- a/Engine/Source/Runtime/Online/BuildPatchServices/Private/BuildPatchManifest.cpp
+++ b/Engine/Source/Runtime/Online/BuildPatchServices/Private/BuildPatchManifest.cpp
@@ -480,6 +480,7 @@ FBuildPatchAppManifest::FBuildPatchAppManifest()
 	, bNeedsResaving(false)
 {
 	Data = NewObject<UBuildPatchManifest>();
+	Data->AddToRoot();
 }
 
 FBuildPatchAppManifest::FBuildPatchAppManifest(const uint32& InAppID, const FString& AppName)
@@ -488,6 +489,7 @@ FBuildPatchAppManifest::FBuildPatchAppManifest(const uint32& InAppID, const FStr
 	, bNeedsResaving(false)
 {
 	Data = NewObject<UBuildPatchManifest>();
+	Data->AddToRoot();
 	Data->AppID = InAppID;
 	Data->AppName = AppName;
 }
@@ -495,12 +497,14 @@ FBuildPatchAppManifest::FBuildPatchAppManifest(const uint32& InAppID, const FStr
 FBuildPatchAppManifest::FBuildPatchAppManifest(const FBuildPatchAppManifest& Other)
 {
 	Data = DuplicateObject<UBuildPatchManifest>(Other.Data, Other.Data->GetOuter());
+	Data->AddToRoot();
 	InitLookups();
 	bNeedsResaving = Other.bNeedsResaving;
 }
 
 FBuildPatchAppManifest::~FBuildPatchAppManifest()
 {
+	Data->RemoveFromRoot();
 //	Data->MarkPendingKill(); ??? Mark for destory and run GC?
 }
 
