Fix FOnlineSessionSettings serialization for EOnlineKeyValuePairDataType::Bool
Thanks to GitHub user [REDACTED_PERSON]!

#github https://github.com/EpicGames/UnrealEngine/pull/835

diff --git a/Engine/Source/Runtime/Online/OnlineSubsystem/Public/NboSerializer.h b/Engine/Source/Runtime/Online/OnlineSubsystem/Public/NboSerializer.h
--- a/Engine/Source/Runtime/Online/OnlineSubsystem/Public/NboSerializer.h
+++ b/Engine/Source/Runtime/Online/OnlineSubsystem/Public/NboSerializer.h
@@ -365,8 +365,16 @@ public:
 				break;
 			}
 		case EOnlineKeyValuePairDataType::Bool:
+			{
+				bool Value;
+				KeyValuePair.GetValue(Value);
+				Ar << (uint8)(Value ? 1 : 0);
+				break;
+			}
 		case EOnlineKeyValuePairDataType::Empty:
 			break;
+		default:
+			checkfSlow(false, TEXT("Unsupported EOnlineKeyValuePairDataType: %d"), Type);
 		}
 		return Ar;
 	}
@@ -742,8 +750,16 @@ public:
 					break;
 				}
 			case EOnlineKeyValuePairDataType::Bool:
+				{
+					uint8 Value;
+					Ar >> Value;
+					KeyValuePair.SetValue(Value != 0);
+					break;
+				}
 			case EOnlineKeyValuePairDataType::Empty:
 				break;
+			default:
+				checkfSlow(false, TEXT("Unsupported EOnlineKeyValuePairDataType: %d"), Type);
 			}
 		}
 
