fix for crash when starting up Tappy Chicken on iOS 6
UE-11804
#ios

diff --git a/Engine/Source/Runtime/Online/IOS/OnlineSubsystemIOS/Private/TurnBasedEventListener.cpp b/Engine/Source/Runtime/Online/IOS/OnlineSubsystemIOS/Private/TurnBasedEventListener.cpp
--- a/Engine/Source/Runtime/Online/IOS/OnlineSubsystemIOS/Private/TurnBasedEventListener.cpp
+++ b/Engine/Source/Runtime/Online/IOS/OnlineSubsystemIOS/Private/TurnBasedEventListener.cpp
@@ -16,14 +16,24 @@
 {
     if (self = [super init]) {
         _owner = &owner;
-        [[GKLocalPlayer localPlayer] registerListener:self];
+#if __IPHONE_OS_VERSION_MIN_REQUIRED < __IPHONE_7_0
+		if ([GKPLocalPlayer respondsToSelector : @selector(registerListener)] == YES)
+#endif
+		{
+			[[GKLocalPlayer localPlayer] registerListener:self];
+		}
     }
     return self;
 }
 
 - (void)dealloc
 {
-    [[GKLocalPlayer localPlayer] unregisterListener:self];
+#if __IPHONE_OS_VERSION_MIN_REQUIRED < __IPHONE_7_0
+	if ([GKPLocalPlayer respondsToSelector : @selector(unregisterListener)] == YES)
+#endif
+	{
+		[[GKLocalPlayer localPlayer] unregisterListener:self];
+	}
     [super dealloc];
 }
 
