compile fix
#ios

diff --git a/Engine/Source/Runtime/Online/IOS/OnlineSubsystemIOS/Private/TurnBasedEventListener.cpp b/Engine/Source/Runtime/Online/IOS/OnlineSubsystemIOS/Private/TurnBasedEventListener.cpp
--- a/Engine/Source/Runtime/Online/IOS/OnlineSubsystemIOS/Private/TurnBasedEventListener.cpp
+++ b/Engine/Source/Runtime/Online/IOS/OnlineSubsystemIOS/Private/TurnBasedEventListener.cpp
@@ -17,7 +17,7 @@
     if (self = [super init]) {
         _owner = &owner;
 #if __IPHONE_OS_VERSION_MIN_REQUIRED < __IPHONE_7_0
-		if ([GKPLocalPlayer respondsToSelector : @selector(registerListener)] == YES)
+		if ([GKLocalPlayer respondsToSelector : @selector(registerListener)] == YES)
 #endif
 		{
 			[[GKLocalPlayer localPlayer] registerListener:self];
@@ -29,7 +29,7 @@
 - (void)dealloc
 {
 #if __IPHONE_OS_VERSION_MIN_REQUIRED < __IPHONE_7_0
-	if ([GKPLocalPlayer respondsToSelector : @selector(unregisterListener)] == YES)
+	if ([GKLocalPlayer respondsToSelector : @selector(unregisterListener)] == YES)
 #endif
 	{
 		[[GKLocalPlayer localPlayer] unregisterListener:self];
