HTML5 - UE-13069 - wrong http headers and UE-13431 (unexpected characters at end of POST string)

diff --git a/Engine/Source/Runtime/HTML5/HTML5JS/Private/HTML5JavaScriptFx.js b/Engine/Source/Runtime/HTML5/HTML5JS/Private/HTML5JavaScriptFx.js
--- a/Engine/Source/Runtime/HTML5/HTML5JS/Private/HTML5JavaScriptFx.js
+++ b/Engine/Source/Runtime/HTML5/HTML5JS/Private/HTML5JavaScriptFx.js
@@ -112,15 +112,22 @@ var UE_JavaScriptLibary =
         return 1;
     },
 
-    UE_MakeHTTPDataRequest: function (ctx, url, verb, payload, freeBuffer, onload, onerror, onprogress) {
+    UE_MakeHTTPDataRequest: function (ctx, url, verb, payload, payloadsize, headers, freeBuffer, onload, onerror, onprogress) {
 	    var _url = Pointer_stringify(url);
 	    var _verb = Pointer_stringify(verb);
-	    var _payload = Pointer_stringify(payload);
+	    var _headers = Pointer_stringify(headers);
 
 	    var xhr = new XMLHttpRequest();
 	    xhr.open(_verb, _url, true);
 	    xhr.responseType = 'arraybuffer';
 
+        // set all headers. 
+	    var _headerArray = _headers.split("%");
+	    for(var headerArrayidx = 0; headerArrayidx < _headerArray.length; headerArrayidx++){
+	        var header = _headerArray[headerArrayidx].split(":");
+	        xhr.setRequestHeader(header[0],header[1]);
+	    }
+
 	    // Onload event handler
 	    xhr.addEventListener('load', function (e) {
    	        if (xhr.status === 200 || _url.substr(0, 4).toLowerCase() !== "http") {
@@ -154,11 +161,9 @@ var UE_JavaScriptLibary =
 	    } catch (ex) { }
 
 	    if (_verb === "POST") {
-	        //Send the proper header information along with the request
-	        xhr.setRequestHeader("Content-type", "application/x-www-form-urlencoded");
-	        xhr.setRequestHeader("Content-length", _payload.length);
+	        var postData = Module.HEAP8.subarray(payload, payload + payloadsize);
 	        xhr.setRequestHeader("Connection", "close");
-	        xhr.send(_payload);
+	        xhr.send(postData);
 	    } else {
 	        xhr.send(null);
 	    }
diff --git a/Engine/Source/Runtime/HTML5/HTML5JS/Public/HTML5JavaScriptFx.h b/Engine/Source/Runtime/HTML5/HTML5JS/Public/HTML5JavaScriptFx.h
--- a/Engine/Source/Runtime/HTML5/HTML5JS/Public/HTML5JavaScriptFx.h
+++ b/Engine/Source/Runtime/HTML5/HTML5JS/Public/HTML5JavaScriptFx.h
@@ -8,6 +8,6 @@
 		 bool UE_LoadGame(const char* Name, const int UserIndex, char **outdata, int* outsize);
 		 int  UE_MessageBox( int MsgType, const char* Text, const char* Caption);
 		 int  UE_GetCurrentCultureName(const char* OutName, int outsize);
-		 void UE_MakeHTTPDataRequest(void *ctx, const char* url, const char* verb, const char* payload, int freeBuffer, void(*onload)(void*, void*, unsigned), void(*onerror)(void*, int, const char*), void(*onprogress)(void*, int, int));
+		 void UE_MakeHTTPDataRequest(void *ctx, const char* url, const char* verb, const char* payload, int payloadsize, const char* headers, int freeBuffer, void(*onload)(void*, void*, unsigned), void(*onerror)(void*, int, const char*), void(*onprogress)(void*, int, int));
 	}
 
diff --git a/Engine/Source/Runtime/Online/HTTP/Private/HTML5/HTML5HTTP.cpp b/Engine/Source/Runtime/Online/HTTP/Private/HTML5/HTML5HTTP.cpp
--- a/Engine/Source/Runtime/Online/HTTP/Private/HTML5/HTML5HTTP.cpp
+++ b/Engine/Source/Runtime/Online/HTTP/Private/HTML5/HTML5HTTP.cpp
@@ -306,11 +306,11 @@ bool FHTML5HttpRequest::StartRequest()
 		return false;
 	}
 
-	// set up headers
-	if (GetHeader("User-Agent").IsEmpty())
-	{
-		SetHeader(TEXT("User-Agent"), FString::Printf(TEXT("game=%s, engine=UE4, version=%s"), FApp::GetGameName(), *GEngineVersion.ToString()));
-	}
+
+	//"User-Agent" && "Content-Length" are automatically set by the browser xhr request. We can't do much.
+
+	// make a fake header, so server has some idea this is UE
+    SetHeader(TEXT("X-UnrealEngine-Agent"), FString::Printf(TEXT("game=%s, engine=UE4, version=%s"), FApp::GetGameName(), *GEngineVersion.ToString()));
 
 	// Add "Pragma: no-cache" to mimic WinInet behavior
 	if (GetHeader("Pragma").IsEmpty())
@@ -319,11 +319,10 @@ bool FHTML5HttpRequest::StartRequest()
 	}
 
 	TArray<FString> AllHeaders = GetAllHeaders();
-	const int32 NumAllHeaders = AllHeaders.Num();
-
-	//TODO: Add headers into UE_MakeHTTPDataRequest
-
 
+	// Create a String which emscripten can understand. 
+	FString Headers = FString::Join(AllHeaders, TEXT("%"));
+	
 	// set up verb (note that Verb is expected to be uppercase only)
 	if (Verb == TEXT("POST"))
 	{
@@ -332,7 +331,7 @@ bool FHTML5HttpRequest::StartRequest()
 		check(!GetHeader("Content-Type").IsEmpty() || IsURLEncoded(RequestPayload));
 
 #if PLATFORM_HTML5_BROWSER
-		UE_MakeHTTPDataRequest(this, TCHAR_TO_ANSI(*URL), "POST", (char*)RequestPayload.GetData(), 0, StaticReceiveCallback, StaticErrorCallback, StaticProgressCallback);
+		UE_MakeHTTPDataRequest(this, TCHAR_TO_ANSI(*URL), "POST", (char*)RequestPayload.GetData(), RequestPayload.Num(),TCHAR_TO_ANSI(*Headers), 0, StaticReceiveCallback, StaticErrorCallback, StaticProgressCallback);
 #else
 		return false;
 #endif
@@ -351,7 +350,7 @@ bool FHTML5HttpRequest::StartRequest()
 	else if (Verb == TEXT("GET"))
 	{
 #if PLATFORM_HTML5_BROWSER
-		UE_MakeHTTPDataRequest(this, TCHAR_TO_ANSI(*URL), "GET", NULL, 1, StaticReceiveCallback, StaticErrorCallback, StaticProgressCallback);
+		UE_MakeHTTPDataRequest(this, TCHAR_TO_ANSI(*URL), "GET", NULL, 0,TCHAR_TO_ANSI(*Headers), 1, StaticReceiveCallback, StaticErrorCallback, StaticProgressCallback);
 #else
 		return false;
 #endif
