[GitHub] 1005: Steam sockets now discard packets from dead connections. (Contributed by [REDACTED_PERSON])
https://github.com/EpicGames/UnrealEngine/pull/1005

diff --git a/Engine/Source/Runtime/Online/OnlineSubsystemSteam/Private/SocketSubsystemSteam.cpp b/Engine/Source/Runtime/Online/OnlineSubsystemSteam/Private/SocketSubsystemSteam.cpp
--- a/Engine/Source/Runtime/Online/OnlineSubsystemSteam/Private/SocketSubsystemSteam.cpp
+++ b/Engine/Source/Runtime/Online/OnlineSubsystemSteam/Private/SocketSubsystemSteam.cpp
@@ -430,16 +430,21 @@ bool FSocketSubsystemSteam::AcceptP2PConnection(ISteamNetworking* SteamNetworkin
 /**
  * Add/update a Steam P2P connection as being recently accessed
  *
- * @param SteanNetworkingPtr proper networking interface that this session is communicating on
+ * @param SteamNetworkingPtr proper networking interface that this session is communicating on
  * @param SessionId P2P session recently heard from
+ *
+ * @return true if the connection is active, false if this is in the dead connections list
  */
-void FSocketSubsystemSteam::P2PTouch(ISteamNetworking* SteamNetworkingPtr, FUniqueNetIdSteam& SessionId)
+bool FSocketSubsystemSteam::P2PTouch(ISteamNetworking* SteamNetworkingPtr, FUniqueNetIdSteam& SessionId)
 {
     // Don't update any sessions coming from pending disconnects
 	if (DeadConnections.Find(SessionId) == NULL)
 	{
 		AcceptedConnections.Add(SessionId, FSteamP2PConnectionInfo(SteamNetworkingPtr, FPlatformTime::Seconds()));
+		return true;
 	}
+
+	return false;
 }
 	
 /**
diff --git a/Engine/Source/Runtime/Online/OnlineSubsystemSteam/Private/SocketSubsystemSteam.h b/Engine/Source/Runtime/Online/OnlineSubsystemSteam/Private/SocketSubsystemSteam.h
--- a/Engine/Source/Runtime/Online/OnlineSubsystemSteam/Private/SocketSubsystemSteam.h
+++ b/Engine/Source/Runtime/Online/OnlineSubsystemSteam/Private/SocketSubsystemSteam.h
@@ -152,10 +152,12 @@ PACKAGE_SCOPE:
 	/**
 	 * Add/update a Steam P2P connection as being recently accessed
 	 *
-	 * @param SteanNetworkingPtr proper networking interface that this session is communicating on
+	 * @param SteamNetworkingPtr proper networking interface that this session is communicating on
 	 * @param SessionId P2P session recently heard from
+     *
+     * @return true if the connection is active, false if this is in the dead connections list
 	 */
-	void P2PTouch(ISteamNetworking* SteamNetworkingPtr, FUniqueNetIdSteam& SessionId);
+	bool P2PTouch(ISteamNetworking* SteamNetworkingPtr, FUniqueNetIdSteam& SessionId);
 
 	/**
 	 * Remove a Steam P2P session from tracking and close the connection
diff --git a/Engine/Source/Runtime/Online/OnlineSubsystemSteam/Private/SocketsSteam.cpp b/Engine/Source/Runtime/Online/OnlineSubsystemSteam/Private/SocketsSteam.cpp
--- a/Engine/Source/Runtime/Online/OnlineSubsystemSteam/Private/SocketsSteam.cpp
+++ b/Engine/Source/Runtime/Online/OnlineSubsystemSteam/Private/SocketsSteam.cpp
@@ -176,8 +176,16 @@ bool FSocketSteam::RecvFrom(uint8* Data, int32 BufferSize, int32& BytesRead, FIn
 	}
 	else
 	{
-		SocketSubsystem->P2PTouch(SteamNetworkingPtr, SteamAddr.SteamId);
-		SocketSubsystem->LastSocketError = SE_NO_ERROR;
+		if (SocketSubsystem->P2PTouch( SteamNetworkingPtr, SteamAddr.SteamId))
+		{
+			SocketSubsystem->LastSocketError = SE_NO_ERROR;
+		}
+		else
+		{
+			MessageSize = 0;
+			SocketSubsystem->LastSocketError = SE_UDP_ERR_PORT_UNREACH;
+			bSuccess = false;
+		}
 	}
 
 	// Steam always sends/receives on the same channel both sides
