Add constructed files into Manifest for Android
#android

diff --git a/Engine/Source/Runtime/Online/BuildPatchServices/Private/BuildPatchFileConstructor.cpp b/Engine/Source/Runtime/Online/BuildPatchServices/Private/BuildPatchFileConstructor.cpp
--- a/Engine/Source/Runtime/Online/BuildPatchServices/Private/BuildPatchFileConstructor.cpp
+++ b/Engine/Source/Runtime/Online/BuildPatchServices/Private/BuildPatchFileConstructor.cpp
@@ -502,6 +502,13 @@ bool FBuildPatchFileConstructor::ConstructFileFromChunks( const FString& Filenam
 	}
 #endif
 	
+#if PLATFORM_ANDROID || PLATFORM_ANDROIDGL4 || PLATFORM_ANDROIDES31
+	if (bSuccess)
+	{
+		IFileManager::Get().SetTimeStamp(*NewFilename, FDateTime::UtcNow());
+	}
+#endif
+
 	// Delete the staging file if unsuccessful by means of construction fail (i.e. keep if canceled or download issue)
 	if( !bSuccess && FBuildPatchInstallError::GetErrorState() == EBuildPatchInstallError::FileConstructionFail )
 	{
