CIS fix for iOS
#ios

diff --git a/Engine/Source/Runtime/Online/IOS/OnlineSubsystemIOS/Private/OnlineLeaderboardsInterfaceIOS.cpp b/Engine/Source/Runtime/Online/IOS/OnlineSubsystemIOS/Private/OnlineLeaderboardsInterfaceIOS.cpp
--- a/Engine/Source/Runtime/Online/IOS/OnlineSubsystemIOS/Private/OnlineLeaderboardsInterfaceIOS.cpp
+++ b/Engine/Source/Runtime/Online/IOS/OnlineSubsystemIOS/Private/OnlineLeaderboardsInterfaceIOS.cpp
@@ -118,13 +118,12 @@ bool FOnlineLeaderboardsIOS::ReadLeaderboardCompletionDelegate(NSArray* players,
                         {
                             const FColumnMetaData& ColumnMeta = ReadObject->ColumnMetadata[StatIdx];
                             
-                            FVariantData* LastColumn = NULL;
                             switch (ColumnMeta.DataType)
                             {
                                 case EOnlineKeyValuePairDataType::Int32:
                                 {
                                     int32 Value = score.value;
-                                    LastColumn = UserRow->Columns.Add(ColumnMeta.ColumnName, FVariantData(Value));
+                                    UserRow->Columns.Add(ColumnMeta.ColumnName, FVariantData(Value));
                                     bWasSuccessful = true;
                                     break;
                                 }
