HTML5: match HttpCurl functionality. Remove explicit encoding for url, browser already does that for you.

diff --git a/Engine/Source/Runtime/Online/HTTP/Private/HTML5/HTML5HTTP.cpp b/Engine/Source/Runtime/Online/HTTP/Private/HTML5/HTML5HTTP.cpp
--- a/Engine/Source/Runtime/Online/HTTP/Private/HTML5/HTML5HTTP.cpp
+++ b/Engine/Source/Runtime/Online/HTTP/Private/HTML5/HTML5HTTP.cpp
@@ -39,22 +39,7 @@ FString FHTML5HttpRequest::GetURL()
 void FHTML5HttpRequest::SetURL(const FString& InURL)
 {
 	UE_LOG(LogHttp, Verbose, TEXT("FHTML5HttpRequest::SetURL() - %s"), *InURL);
-
-	URL = InURL.Replace(TEXT("%"), TEXT("%25"));
-	URL = URL.Replace(TEXT(" "), TEXT("%20"));
-	URL = URL.Replace(TEXT("\""), TEXT("%22"));
-	URL = URL.Replace(TEXT("<"), TEXT("%3C"));
-	URL = URL.Replace(TEXT(">"), TEXT("%3E"));
-
-	URL = URL.Replace(TEXT("["), TEXT("%5B"));
-	URL = URL.Replace(TEXT("]"), TEXT("%5D"));
-	URL = URL.Replace(TEXT("\\"), TEXT("%5C"));
-	URL = URL.Replace(TEXT("^"), TEXT("%5E"));
-
-	URL = URL.Replace(TEXT("`"), TEXT("%60"));
-	URL = URL.Replace(TEXT("{"), TEXT("%7B"));
-	URL = URL.Replace(TEXT("}"), TEXT("%7D"));
-	URL = URL.Replace(TEXT("|"), TEXT("%7C"));
+	URL = InURL; 
 }
 
 
