Start of xmpp party system implementation
Merging //depot/UE4-LauncherDev/Engine/Source/Runtime/Online/NotForLicensees/ to //depot/UE4/Engine/Source/Runtime/Online/NotForLicensees/

diff --git a/Engine/Source/Runtime/Online/OnlineSubsystem/Private/OnlineSessionSettings.cpp b/Engine/Source/Runtime/Online/OnlineSubsystem/Private/OnlineSessionSettings.cpp
--- a/Engine/Source/Runtime/Online/OnlineSubsystem/Private/OnlineSessionSettings.cpp
+++ b/Engine/Source/Runtime/Online/OnlineSubsystem/Private/OnlineSessionSettings.cpp
@@ -133,7 +133,7 @@ template ONLINESUBSYSTEM_API bool FOnlineSessionSettings::Get(FName Key, TArray<
 
 bool FOnlineSessionSettings::Remove(FName Key)
 {
-	return Settings.Remove(Key);
+	return Settings.Remove(Key) > 0;
 }
 
 EOnlineDataAdvertisementType::Type FOnlineSessionSettings::GetAdvertisementType(FName Key) const
diff --git a/Engine/Source/Runtime/Online/OnlineSubsystem/Public/Interfaces/OnlineChatInterface.h b/Engine/Source/Runtime/Online/OnlineSubsystem/Public/Interfaces/OnlineChatInterface.h
--- a/Engine/Source/Runtime/Online/OnlineSubsystem/Public/Interfaces/OnlineChatInterface.h
+++ b/Engine/Source/Runtime/Online/OnlineSubsystem/Public/Interfaces/OnlineChatInterface.h
@@ -160,13 +160,14 @@ public:
 	/**
 	 * Kick off request for creating a chat room with a provided configuration
 	 * 
-	 * @param UserId id of user that is joining
+	 * @param UserId id of user that is creating the room
+	 * @param RoomId name of room to create
 	 * @param RoomConfig configuration for the room
 	 * @param Nickname display name for the chat room. Name must be unique and is reserved for duration of join
 	 *
 	 * @return if successfully started the async operation
 	 */
-	virtual bool CreateRoom(const FUniqueNetId& UserId, const FChatRoomConfig& RoomConfig, const FString& Nickname) = 0;
+	virtual bool CreateRoom(const FUniqueNetId& UserId, const FChatRoomId& RoomId, const FString& Nickname) = 0;
 	
 	/**
 	 * Kick off request for joining a public chat room
diff --git a/Engine/Source/Runtime/Online/OnlineSubsystem/Public/Interfaces/OnlinePartyInterface.h b/Engine/Source/Runtime/Online/OnlineSubsystem/Public/Interfaces/OnlinePartyInterface.h
--- a/Engine/Source/Runtime/Online/OnlineSubsystem/Public/Interfaces/OnlinePartyInterface.h
+++ b/Engine/Source/Runtime/Online/OnlineSubsystem/Public/Interfaces/OnlinePartyInterface.h
@@ -10,6 +10,11 @@
  */
 typedef FUniqueNetId FOnlinePartyId;
 
+/**
+ * Key value data that can be sent to party members
+ */
+typedef FOnlineKeyValuePairs<FString, FVariantData> FOnlinePartyAttrs;
+
 /**
  * Party member user info returned by IOnlineParty interface
  */
@@ -30,7 +35,7 @@ public:
 	/**
 	 * @return true if this is the party owner/leader
 	 */
-	virtual bool IsOwner() const = 0;
+	virtual bool IsLeader() const = 0;
 };
 
 /**
@@ -43,7 +48,7 @@ public:
 	/**
 	 * @return id of the party this data is associated with
 	 */
-	virtual const FOnlinePartyId& GetPartyId() const = 0;
+	virtual TSharedRef<FOnlinePartyId> GetPartyId() const = 0;
 
 	/**
 	 * Get an attribute from the party data
@@ -51,17 +56,16 @@ public:
 	 * @param AttrName - key for the attribute
 	 * @param OutAttrValue - [out] value for the attribute if found
 	 *
-	 * @returnm true if the attribute was vound
+	 * @return true if the attribute was vound
 	 */
 	virtual bool GetAttribute(const FString& AttrName, FVariantData& OutAttrValue) const = 0;
 
 	/**
-	 * Set an attribute on the party data
+	 * Get all attributes from the party data
 	 *
-	 * @param AttrName - key for the attribute
-	 * @param AttrValue - value for the attribute to set
+	 * @return key/val attributes
 	 */
-	virtual void SetAttribute(const FString& AttrName, const FVariantData& AttrValue) const = 0;
+	virtual const FOnlinePartyAttrs& GetAttributes() const = 0;
 };
 
 /**
@@ -74,7 +78,7 @@ public:
 	/**
 	 * @return id of the party member this data is associated with
 	 */
-	virtual const FUniqueNetId& GetPartyMemberId() = 0;
+	virtual TSharedRef<FUniqueNetId> GetPartyMemberId() const = 0;
 };
 
 /**
@@ -104,11 +108,35 @@ struct FPartyConfiguration
 	int32 MaxMembers;
 };
 
+namespace EPartyState
+{
+	enum Type
+	{
+		None,
+		CreatePending,
+		DestroyPending,
+		JoinPending,
+		LeavePending,
+		Success,
+		Failed
+	};
+}
+
 /**
  * Current state associated with a party
  */
-struct FOnlinePartyInfo
+class FOnlinePartyInfo
 {
+public:
+	FOnlinePartyInfo(const TSharedPtr<FOnlinePartyId>& PartyId = TSharedPtr<FOnlinePartyId>())
+		: State(EPartyState::None)
+		, PartyId(PartyId)
+	{}
+
+	virtual ~FOnlinePartyInfo() 
+	{}
+
+	EPartyState::Type State;
 	/** Current state of configuration */
 	FPartyConfiguration Config;
 	/** unique id of hte party */
@@ -209,6 +237,14 @@ DECLARE_DELEGATE_FiveParams(FOnPromotePartyMemberComplete, const FUniqueNetId& /
 DECLARE_MULTICAST_DELEGATE_ThreeParams(FOnPartyInviteReceived, const FUniqueNetId& /*RecipientId*/, const FUniqueNetId& /*SenderId*/, const FOnlinePartyId& /*PartyId*/);
 typedef FOnPartyInviteReceived::FDelegate FOnPartyInviteReceivedDelegate;
 
+/**
+ * Notification when a party has been disbanded 
+ *
+ * @param PartyId - id associated with the party
+ */
+DECLARE_MULTICAST_DELEGATE_OneParam(FOnPartyDisbanded, const FOnlinePartyId& /*PartyId*/);
+typedef FOnPartyDisbanded::FDelegate FOnPartyDisbandedDelegate;
+
 /**
  * Notification when a member joins the party
  *
@@ -273,10 +309,11 @@ public:
 	 * @param LocalUserId - user making the request
 	 * @param PartyConfig - configuration for the party (can be updated later)
 	 * @param Delegate - called on completion
+	 * @param UserRoomId - this forces the name of the room to be this value
 	 *
 	 * @return true if task was started
 	 */
-	virtual bool CreateParty(const FUniqueNetId& LocalUserId, const FPartyConfiguration& PartyConfig, const FOnCreatePartyComplete& Delegate = FOnCreatePartyComplete()) = 0;
+	virtual bool CreateParty(const FUniqueNetId& LocalUserId, const FPartyConfiguration& PartyConfig, const FOnCreatePartyComplete& Delegate = FOnCreatePartyComplete(), const FString& UserRoomId = TEXT("")) = 0;
 
 	/**
 	 * Update an existing party with new configuration
@@ -377,7 +414,7 @@ public:
 	 *
 	 * @return true if task was started
 	 */
-	virtual bool UpdatePartyData(const FUniqueNetId& LocalUserId, const FOnlinePartyId& PartyId, const IOnlinePartyData& PartyData) = 0;
+	virtual bool UpdatePartyData(const FUniqueNetId& LocalUserId, const FOnlinePartyId& PartyId, const FOnlinePartyAttrs& PartyData) = 0;
 
 	/**
 	 * Set party data for a single party member and broadcast to all members
@@ -391,7 +428,7 @@ public:
 	 *
 	 * @return true if task was started
 	 */
-	virtual bool UpdatePartyMemberData(const FUniqueNetId& LocalUserId, const FOnlinePartyId& PartyId, const IOnlinePartyMemberData& PartyMemberData) = 0;
+	virtual bool UpdatePartyMemberData(const FUniqueNetId& LocalUserId, const FOnlinePartyId& PartyId, const FOnlinePartyAttrs& PartyMemberData) = 0;
 
 	/**
 	 * Get a list of currently joined parties for the user
@@ -403,6 +440,16 @@ public:
 	 */
 	virtual bool GetJoinedParties(const FUniqueNetId& LocalUserId, TArray< TSharedRef<FOnlinePartyId> >& OutPartyIds) const = 0;
 
+	/**
+	 * Get a list of parties the user has been invited to
+	 *
+	 * @param LocalUserId - user making the request
+	 * @param OutPartyIds - list of party ids joined by the current user
+	 *
+	 * @return true if entries found
+	 */
+	virtual bool GetPartyInvites(const FUniqueNetId& LocalUserId, TArray< TSharedRef<FOnlinePartyId> >& OutPartyIds) const = 0;
+
 	/**
 	 * Get info associated with a party
 	 *
@@ -468,6 +515,63 @@ public:
 	 * @return party member data or nullptr if not found
 	 */
 	virtual TSharedPtr<IOnlinePartyMemberData> GetPartyMemberData(const FUniqueNetId& LocalUserId, const FOnlinePartyId& PartyId, const FUniqueNetId& MemberId) const = 0;
+
+	/**
+	 * Notification when a new invite is received
+	 *
+	 * @param RecipientId - id of user that this invite is for
+	 * @param SenderId - id of member that sent the invite
+	 * @param PartyId - id associated with the party
+	 */
+	DEFINE_ONLINE_DELEGATE_THREE_PARAM(OnPartyInviteReceived, const FUniqueNetId& /*RecipientId*/, const FUniqueNetId& /*SenderId*/, const FOnlinePartyId& /*PartyId*/);
+
+	/**
+	 * Notification when a party has been disbanded 
+	 *
+	 * @param PartyId - id associated with the party
+	 */
+	DEFINE_ONLINE_DELEGATE_ONE_PARAM(OnPartyDisbanded, const FOnlinePartyId& /*PartyId*/);
+
+	/**
+	 * Notification when a member joins the party
+	 *
+	 * @param PartyId - id associated with the party
+	 * @param MemberId - id of member that joined
+	 */
+	DEFINE_ONLINE_DELEGATE_TWO_PARAM(OnPartyMemberJoined, const FOnlinePartyId& /*PartyId*/, const FUniqueNetId& /*MemberId*/);
+
+	/**
+	 * Notification when a member leaves the party
+	 *
+	 * @param PartyId - id associated with the party
+	 * @param MemberId - id of member that left
+	 */
+	DEFINE_ONLINE_DELEGATE_TWO_PARAM(OnPartyMemberLeft, const FOnlinePartyId& /*PartyId*/, const FUniqueNetId& /*MemberId*/);
+
+	/**
+	 * Notification when a member is promoted to admin in the party
+	 *
+	 * @param PartyId - id associated with the party
+	 * @param MemberId - id of member that was promoted to admin
+	 */
+	DEFINE_ONLINE_DELEGATE_TWO_PARAM(OnPartyMemberPromoted, const FOnlinePartyId& /*PartyId*/, const FUniqueNetId& /*MemberId*/);
+
+	/**
+	 * Notification when party data is updated
+	 *
+	 * @param PartyId - id associated with the party
+	 * @param PartyData - party data that was updated
+	 */
+	DEFINE_ONLINE_DELEGATE_TWO_PARAM(OnPartyDataReceived, const FOnlinePartyId& /*PartyId*/, const TSharedRef<IOnlinePartyData>& /*PartyData*/);
+
+	/**
+	 * Notification when party member data is updated
+	 *
+	 * @param PartyId - id associated with the party
+	 * @param MemberId - id of member that had updated data
+	 * @param PartyMemberData - party member data that was updated
+	 */
+	DEFINE_ONLINE_DELEGATE_THREE_PARAM(OnPartyMemberDataReceived, const FOnlinePartyId& /*PartyId*/, const FUniqueNetId& /*MemberId*/, const TSharedRef<IOnlinePartyMemberData>& /*PartyMemberData*/);
 };
 
 typedef TSharedPtr<IOnlineParty, ESPMode::ThreadSafe> IOnlinePartyPtr;
diff --git a/Engine/Source/Runtime/Online/OnlineSubsystem/Public/OnlineKeyValuePair.h b/Engine/Source/Runtime/Online/OnlineSubsystem/Public/OnlineKeyValuePair.h
--- a/Engine/Source/Runtime/Online/OnlineSubsystem/Public/OnlineKeyValuePair.h
+++ b/Engine/Source/Runtime/Online/OnlineSubsystem/Public/OnlineKeyValuePair.h
@@ -57,214 +57,18 @@ namespace EOnlineKeyValuePairDataType
  *	Associative container for key value pairs
  */
 template<class KeyType, class ValueType>
-class FOnlineKeyValuePairs
+class FOnlineKeyValuePairs : public TMap<KeyType,ValueType>
 {
-
-private:
-
-	struct FKeyValuePairInternal 
-	{
-		KeyType Key;
-		ValueType Value;
-
-		FKeyValuePairInternal(KeyType InKey, const ValueType& InValue) :
-			Key(InKey),
-			Value(InValue)
-		{
-		}
-
-		/**
-		 *	Comparison operator
-		 */
-		bool operator==(const FKeyValuePairInternal& Other) const
-		{
-			// Only the key matters for equality
-			return Key == Other.Key;
-		}
-	};
-	
-	/** Array of all key value pairs */
-	TArray<FKeyValuePairInternal> KeyValuePairs;
+	typedef TMap<KeyType, ValueType> Super;
 
 public:
-	/**
-	 *	Finds a key value pair by key name
-	 * @param Key key to search for
-	 * @return pointer to value paired with key
-	 */
-	const ValueType* Find(KeyType Key) const
-	{
-		// Search for the individual element
-		for (int32 PairIdx = 0; PairIdx < KeyValuePairs.Num(); ++PairIdx)
-		{
-			const FKeyValuePairInternal& Setting = KeyValuePairs[PairIdx];
-			if (Setting.Key == Key)
-			{
-				return &Setting.Value;
-			}
-		}
 
-		return NULL;
-	}
+	FORCEINLINE FOnlineKeyValuePairs() {}
+	FORCEINLINE FOnlineKeyValuePairs(FOnlineKeyValuePairs&& Other) : Super(MoveTemp(Other)) {}
+	FORCEINLINE FOnlineKeyValuePairs(const FOnlineKeyValuePairs&  Other) : Super(Other) {}
+	FORCEINLINE FOnlineKeyValuePairs& operator=(FOnlineKeyValuePairs&& Other) { Super::operator=(MoveTemp(Other)); return *this; }
+	FORCEINLINE FOnlineKeyValuePairs& operator=(const FOnlineKeyValuePairs&  Other) { Super::operator=(Other); return *this; }
 
-	/**
-	 *	Finds a key value pair by key name
-	 * @param Key key to search for
-	 * @return pointer to value paired with key
-	 */
-	ValueType* Find(KeyType Key)
-	{
-		// Search for the individual element
-		for (int32 PairIdx = 0; PairIdx < KeyValuePairs.Num(); ++PairIdx)
-		{
-			FKeyValuePairInternal& Setting = KeyValuePairs[PairIdx];
-			if (Setting.Key == Key)
-			{
-				return &Setting.Value;
-			}
-		}
-
-		return NULL;
-	}
-
-	/**
-	 *	Add an element to the grouping, overwriting existing values
-	 *
-	 * @param Key key to assign to the value
-	 * @param Value value to store
-	 */
-	ValueType* Add(KeyType Key, const ValueType& Value)
-	{
-		ValueType* ExistingElement = Find(Key);
-		if (ExistingElement)
-		{
-			*ExistingElement = Value;
-			return ExistingElement;
-		}
-		else
-		{
-			int32 Idx = KeyValuePairs.Add(FKeyValuePairInternal(Key, Value));
-			return &KeyValuePairs[Idx].Value;
-		}
-	}
-
-	/**
-	 *	Remove an element from the grouping
-	 *
-	 * @param Key key to assign to the value
-	 *
-	 * @return true if found and removed, false otherwise
-	 */
-	bool Remove(KeyType Key)
-	{
-		for (int32 PairIdx = 0; PairIdx < KeyValuePairs.Num(); ++PairIdx)
-		{
-			FKeyValuePairInternal& Setting = KeyValuePairs[PairIdx];
-			if (Setting.Key == Key)
-			{
-				KeyValuePairs.RemoveAtSwap(PairIdx);
-				return true;
-			}
-		}
-
-		return false;
-	}
-
-	/**
-	 *	Append another FOnlineKeyValuePairs 
-	 *
-	 * @param Source key value pairs to append
-	 */
-	void Append(const FOnlineKeyValuePairs& Source)
-	{
-		KeyValuePairs.Append(Source.KeyValuePairs);
-	}
-
-	/**
-	 *	Remove all elements from the grouping
-	 */
-	void Empty()
-	{
-		KeyValuePairs.Empty();
-	}
-
-	/** @return number of elements in grouping */
-	int32 Num() const
-	{
-		return KeyValuePairs.Num();
-	}
-
-	/** Base class for const/nonconst iterators */
-	template<bool bConst>
-	class TBaseIterator
-	{
-	private:
-
-		/** Define the data type to iterate on as const/nonconst */
-		typedef typename TChooseClass<bConst, const ValueType, ValueType>::Result PairsType;
-		/** Define the internal iterator as const/nonconst */
-		typedef typename TChooseClass<bConst, typename TArray<FKeyValuePairInternal>::TConstIterator, typename TArray<FKeyValuePairInternal>::TIterator>::Result PairsItType;
-
-		/** Internal iterator over array */
-		PairsItType PairsIt;
-
-	public:
-
-		/** Initialization constructor. */
-		TBaseIterator(const FOnlineKeyValuePairs& InPair, int32 StartIndex) :
-			PairsIt(InPair.KeyValuePairs)
-		{
-			PairsIt += StartIndex;
-		}
-
-		/** Initialization constructor. */
-		TBaseIterator(FOnlineKeyValuePairs& InPair, int32 StartIndex) :
-			PairsIt(InPair.KeyValuePairs)
-		{
-			PairsIt += StartIndex;
-		}
-
-		/** Prefix op */
-		TBaseIterator& operator++()	{ ++PairsIt; return *this; }
-
-		SAFE_BOOL_OPERATORS(TBaseIterator<bConst>)
-
-		/** Conversion to "bool" returning true if the iterator is valid. */
-		FORCEINLINE_EXPLICIT_OPERATOR_BOOL() const
-		{ 
-			return !!PairsIt; 
-		}
-		/** inverse of the "bool" operator */
-		FORCEINLINE bool operator !() const 
-		{
-			return !(bool)*this;
-		}
-
-		/** @return the current key */
-		KeyType Key() const		{ return PairsIt->Key; }
-		/** @return the current value */
-		PairsType& Value() const	{ return PairsIt->Value; }
-	};
-
-	/** Used to iterate over the elements of a const array of pairs. */
-	class TConstIterator : public TBaseIterator<true>
-	{
-	public:
-		TConstIterator(const FOnlineKeyValuePairs& InPairs,int32 StartIndex = 0) :
-		  TBaseIterator<true>(InPairs, StartIndex)
-		{
-		}
-	};
-
-	/** Used to iterate over the elements of an array of pairs. */
-	class TIterator : public TBaseIterator<false>
-	{
-	public:
-		TIterator(FOnlineKeyValuePairs& InPairs,int32 StartIndex = 0) :
-			TBaseIterator<false>(InPairs, StartIndex)
-		{
-		}
-	};
 };
 
 /**
