Fixed blocking compilation error in OnlineLeaderboardInterfaceSteam.cpp.

diff --git a/Engine/Source/Runtime/Online/OnlineSubsystemSteam/Private/OnlineLeaderboardInterfaceSteam.cpp b/Engine/Source/Runtime/Online/OnlineSubsystemSteam/Private/OnlineLeaderboardInterfaceSteam.cpp
--- a/Engine/Source/Runtime/Online/OnlineSubsystemSteam/Private/OnlineLeaderboardInterfaceSteam.cpp
+++ b/Engine/Source/Runtime/Online/OnlineSubsystemSteam/Private/OnlineLeaderboardInterfaceSteam.cpp
@@ -464,7 +464,7 @@ public:
 						{
 							int32 Value;
 							bSuccess = SteamUserStatsPtr->GetUserStat(SteamUserId, TCHAR_TO_UTF8(*StatName), &Value) ? true : false;
-							LastColumn = UserRow->Columns.Add(ColumnMeta.ColumnName, FVariantData(Value));
+							LastColumn = &(UserRow->Columns.Add(ColumnMeta.ColumnName, FVariantData(Value)));
 							break;
 						}
 
@@ -472,12 +472,12 @@ public:
 						{
 							float Value;
 							bSuccess = SteamUserStatsPtr->GetUserStat(SteamUserId, TCHAR_TO_UTF8(*StatName), &Value) ? true : false;
-							LastColumn = UserRow->Columns.Add(ColumnMeta.ColumnName, FVariantData(Value));
+							LastColumn = &(UserRow->Columns.Add(ColumnMeta.ColumnName, FVariantData(Value)));
 							break;
 						}
 					default:
 						UE_LOG_ONLINE(Warning, TEXT("Unsupported key value pair during retrieval from Steam %s"), *StatName);
-						LastColumn = UserRow->Columns.Add(ColumnMeta.ColumnName, FVariantData());
+						LastColumn = &(UserRow->Columns.Add(ColumnMeta.ColumnName, FVariantData()));
 						break;
 					}
 
