Compile fix for android leaderboard from FVariantData change

diff --git a/Engine/Source/Runtime/Online/Android/OnlineSubsystemGooglePlay/Private/OnlineAsyncTaskGooglePlayReadLeaderboard.cpp b/Engine/Source/Runtime/Online/Android/OnlineSubsystemGooglePlay/Private/OnlineAsyncTaskGooglePlayReadLeaderboard.cpp
--- a/Engine/Source/Runtime/Online/Android/OnlineSubsystemGooglePlay/Private/OnlineAsyncTaskGooglePlayReadLeaderboard.cpp
+++ b/Engine/Source/Runtime/Online/Android/OnlineSubsystemGooglePlay/Private/OnlineAsyncTaskGooglePlayReadLeaderboard.cpp
@@ -73,7 +73,7 @@ void FOnlineAsyncTaskGooglePlayReadLeaderboard::Tick()
 			case EOnlineKeyValuePairDataType::Int32:
 			{
 				int32 Value = Response.data.CurrentPlayerScore().Value();
-				LastColumn = UserRow->Columns.Add(ColumnMeta.ColumnName, FVariantData(Value));
+				LastColumn = &(UserRow->Columns.Add(ColumnMeta.ColumnName, FVariantData(Value)));
 				bWasSuccessful = true;
 				break;
 			}
