- Remove OnSessionInviteAccepted delegate
- Moved SessionUserInviteAccepted delegate handling to GameInstance (users can override OnSessionUserInviteAccepted for custom behavior)

diff --git a/Engine/Source/Runtime/Engine/Classes/Engine/GameInstance.h b/Engine/Source/Runtime/Engine/Classes/Engine/GameInstance.h
--- a/Engine/Source/Runtime/Engine/Classes/Engine/GameInstance.h
+++ b/Engine/Source/Runtime/Engine/Classes/Engine/GameInstance.h
@@ -5,7 +5,6 @@
 #include "EngineBaseTypes.h"
 #include "GameInstance.generated.h"
 
-
 // 
 // 	EWelcomeScreen, 	//initial screen.  Used for platforms where we may not have a signed in user yet.
 // 	EMessageScreen, 	//message screen.  Used to display a message - EG unable to connect to game.
@@ -22,7 +21,6 @@ namespace GameInstanceState
 
 class FOnlineSessionSearchResult;
 
-
 /**
  * GameInstance: high-level manager object for an instance of the running game.
  * Spawned at game creation and not destroyed until game instance is shut down.
@@ -44,6 +42,10 @@ protected:
 
 	UPROPERTY()
 	TArray<ULocalPlayer*> LocalPlayers;		// List of locally participating players in this game instance
+	
+	// Delegate handle that stores delegate for when an invite is accepted by a user
+	FDelegateHandle OnSessionUserInviteAcceptedDelegateHandle;
+	
 public:
 
 	FString PIEMapName;
@@ -157,4 +159,11 @@ public:
 	virtual void			HandleDemoPlaybackFailure( EDemoPlayFailure::Type FailureType, const FString& ErrorString = TEXT("") ) { }
 
 	static void				AddReferencedObjects(UObject* InThis, FReferenceCollector& Collector);
+
+    /** Delegate that is called when a user has accepted an invite. */
+	void HandleSessionUserInviteAccepted(const bool bWasSuccess, const int32 ControllerId, TSharedPtr< FUniqueNetId > UserId, const FOnlineSessionSearchResult & InviteResult);
+	
+	/** Overridable implementation of HandleSessionUserInviteAccepted, which does nothing but call this function */
+	virtual void OnSessionUserInviteAccepted(const bool bWasSuccess, const int32 ControllerId, TSharedPtr< FUniqueNetId > UserId, const FOnlineSessionSearchResult & InviteResult);
+
 };
\ No newline at end of file
diff --git a/Engine/Source/Runtime/Engine/Classes/GameFramework/OnlineSession.h b/Engine/Source/Runtime/Engine/Classes/GameFramework/OnlineSession.h
--- a/Engine/Source/Runtime/Engine/Classes/GameFramework/OnlineSession.h
+++ b/Engine/Source/Runtime/Engine/Classes/GameFramework/OnlineSession.h
@@ -7,6 +7,8 @@
 #pragma once
 #include "OnlineSession.generated.h"
 
+class FOnlineSessionSearchResult;
+
 UCLASS(config=Game)
 class ENGINE_API UOnlineSession : public UObject
 {
@@ -28,6 +30,9 @@ public:
 
 	/** End the online session specified */
 	virtual void EndOnlineSession(FName SessionName) {};
+
+	/** Called when a user accepts an invite */
+	virtual void OnSessionUserInviteAccepted(const bool bWasSuccess, const int32 ControllerId, TSharedPtr< FUniqueNetId > UserId, const FOnlineSessionSearchResult & InviteResult) {};
 };
 
 
diff --git a/Engine/Source/Runtime/Engine/Private/GameInstance.cpp b/Engine/Source/Runtime/Engine/Private/GameInstance.cpp
--- a/Engine/Source/Runtime/Engine/Private/GameInstance.cpp
+++ b/Engine/Source/Runtime/Engine/Private/GameInstance.cpp
@@ -8,6 +8,9 @@
 #include "EnginePrivate.h"
 #include "Engine/GameInstance.h"
 #include "Engine/Engine.h"
+#include "OnlineSubsystem.h"
+#include "OnlineSessionInterface.h"
+#include "GameFramework/OnlineSession.h"
 
 #if WITH_EDITOR
 #include "UnrealEd.h"
@@ -33,12 +36,32 @@ UEngine* UGameInstance::GetEngine() const
 void UGameInstance::Init()
 {
 	ReceiveInit();
+
+	const auto OnlineSub = IOnlineSubsystem::Get();
+	if (OnlineSub != nullptr)
+	{
+		IOnlineSessionPtr SessionInt = OnlineSub->GetSessionInterface();
+		if (SessionInt.IsValid())
+		{
+			SessionInt->AddOnSessionUserInviteAcceptedDelegate_Handle(FOnSessionUserInviteAcceptedDelegate::CreateUObject(this, &UGameInstance::HandleSessionUserInviteAccepted));
+		}
+	}
 }
 
 void UGameInstance::Shutdown()
 {
 	ReceiveShutdown();
 
+	const auto OnlineSub = IOnlineSubsystem::Get();
+	if (OnlineSub != nullptr)
+	{
+		IOnlineSessionPtr SessionInt = OnlineSub->GetSessionInterface();
+		if (SessionInt.IsValid())
+		{
+			SessionInt->ClearOnSessionUserInviteAcceptedDelegate_Handle(OnSessionUserInviteAcceptedDelegateHandle);
+		}
+	}
+
 	// Clear the world context pointer to prevent further access.
 	WorldContext = nullptr;
 }
@@ -593,3 +616,42 @@ void UGameInstance::AddReferencedObjects(UObject* InThis, FReferenceCollector& C
 
 	Super::AddReferencedObjects(This, Collector);
 }
+
+void UGameInstance::HandleSessionUserInviteAccepted(const bool bWasSuccess, const int32 ControllerId, TSharedPtr< FUniqueNetId > UserId, const FOnlineSessionSearchResult &	InviteResult)
+{
+	OnSessionUserInviteAccepted(bWasSuccess, ControllerId, UserId, InviteResult);
+}
+
+void UGameInstance::OnSessionUserInviteAccepted(const bool bWasSuccess, const int32 ControllerId, TSharedPtr< FUniqueNetId > UserId, const FOnlineSessionSearchResult &	InviteResult)
+{
+	UE_LOG(LogPlayerManagement, Verbose, TEXT("OnSessionUserInviteAccepted LocalUserNum: %d bSuccess: %d"), ControllerId, bWasSuccess);
+	// Don't clear invite accept delegate
+
+	if (bWasSuccess)
+	{
+		if (InviteResult.IsValid())
+		{
+			for (ULocalPlayer* LocalPlayer : LocalPlayers)
+			{
+				// Route the call to the actual user that accepted the invite
+				if (LocalPlayer->GetCachedUniqueNetId() == UserId)
+				{
+					LocalPlayer->GetOnlineSession()->OnSessionUserInviteAccepted(bWasSuccess, ControllerId, UserId, InviteResult);
+					return;
+				}
+			}
+
+			// Go ahead and have the active local player handle accepting the invite. A game can detect that the user id is different and handle
+			// it how it needs to.
+			ULocalPlayer* LocalPlayer = GetFirstGamePlayer();
+			if (LocalPlayer)
+			{
+				LocalPlayer->GetOnlineSession()->OnSessionUserInviteAccepted(bWasSuccess, ControllerId, UserId, InviteResult);
+			}
+		}
+		else
+		{
+			UE_LOG(LogPlayerManagement, Warning, TEXT("Invite accept returned invalid search result."));
+		}
+	}
+}
diff --git a/Engine/Source/Runtime/Online/OnlineSubsystem/Public/Interfaces/OnlineSessionInterface.h b/Engine/Source/Runtime/Online/OnlineSubsystem/Public/Interfaces/OnlineSessionInterface.h
--- a/Engine/Source/Runtime/Online/OnlineSubsystem/Public/Interfaces/OnlineSessionInterface.h
+++ b/Engine/Source/Runtime/Online/OnlineSubsystem/Public/Interfaces/OnlineSessionInterface.h
@@ -148,18 +148,6 @@ typedef FOnFindFriendSessionComplete::FDelegate FOnFindFriendSessionCompleteDele
 DECLARE_MULTICAST_DELEGATE_OneParam(FOnPingSearchResultsComplete, bool);
 typedef FOnPingSearchResultsComplete::FDelegate FOnPingSearchResultsCompleteDelegate;
 
-/**
- * Called when a user accepts a session invitation. Allows the game code a chance
- * to clean up any existing state before accepting the invite. The invite must be
- * accepted by calling JoinSession() after clean up has completed
- *
- * @param LocalUserNum the controller number of the accepting user
- * @param bWasSuccessful the session was found and is joinable, false otherwise
- * @param InviteResult the search/settings for the session we're joining via invite
- */
-typedef FOnSingleSessionResultComplete FOnSessionInviteAccepted;
-typedef FOnSessionInviteAccepted::FDelegate FOnSessionInviteAcceptedDelegate;
-
 /**
  * Called when a user accepts a session invitation. Allows the game code a chance
  * to clean up any existing state before accepting the invite. The invite must be
@@ -632,17 +620,6 @@ public:
 	 */
 	virtual bool SendSessionInviteToFriends(const FUniqueNetId& LocalUserId, FName SessionName, const TArray< TSharedRef<FUniqueNetId> >& Friends) = 0;
 
-	/**
-	 * Called when a user accepts a session invitation. Allows the game code a chance
-	 * to clean up any existing state before accepting the invite. The invite must be
-	 * accepted by calling JoinSession() after clean up has completed
-	 *
-	 * @param LocalUserNum the controller number of the accepting user
-	 * @param bWasSuccessful true if the async action completed without error, false if there was an error
-	 * @param InviteResult the search/settings for the session we're joining via invite
-	 */
-	DEFINE_ONLINE_PLAYER_DELEGATE_TWO_PARAM(MAX_LOCAL_PLAYERS, OnSessionInviteAccepted, bool, const FOnlineSessionSearchResult&);
-
 	/**
 	 * Called when a user accepts a session invitation. Allows the game code a chance
 	 * to clean up any existing state before accepting the invite. The invite must be
diff --git a/Engine/Source/Runtime/Online/OnlineSubsystemSteam/Private/OnlineSessionAsyncLobbySteam.cpp b/Engine/Source/Runtime/Online/OnlineSubsystemSteam/Private/OnlineSessionAsyncLobbySteam.cpp
--- a/Engine/Source/Runtime/Online/OnlineSubsystemSteam/Private/OnlineSessionAsyncLobbySteam.cpp
+++ b/Engine/Source/Runtime/Online/OnlineSubsystemSteam/Private/OnlineSessionAsyncLobbySteam.cpp
@@ -1140,12 +1140,28 @@ void FOnlineAsyncTaskSteamFindLobby::TriggerDelegates()
 {
 	if (bWasSuccessful && SearchSettings->SearchResults.Num() > 0)
 	{
-		OnFindLobbyCompleteDelegates.Broadcast(LocalUserNum, bWasSuccessful, SearchSettings->SearchResults[0]);
+		if (bIsUsingNetIdDelegate)
+		{
+			OnFindLobbyCompleteWithNetIdDelegate.Broadcast(bWasSuccessful, LocalUserNum, MakeShareable<FUniqueNetId>(new FUniqueNetIdSteam(SteamUser()->GetSteamID())), SearchSettings->SearchResults[0]);
+			
+		}
+		else
+		{
+			OnFindLobbyCompleteDelegates.Broadcast(LocalUserNum, bWasSuccessful, SearchSettings->SearchResults[0]);
+		}
 	}
 	else
 	{
 		FOnlineSessionSearchResult EmptyResult;
-		OnFindLobbyCompleteDelegates.Broadcast(LocalUserNum, bWasSuccessful, EmptyResult);
+		if (bIsUsingNetIdDelegate)
+		{
+			OnFindLobbyCompleteWithNetIdDelegate.Broadcast(bWasSuccessful, LocalUserNum, MakeShareable<FUniqueNetId>(new FUniqueNetIdSteam(SteamUser()->GetSteamID())), EmptyResult);
+		}
+		else
+		{
+			OnFindLobbyCompleteDelegates.Broadcast(LocalUserNum, bWasSuccessful, EmptyResult);
+		}
+
 	}
 }
 
@@ -1172,7 +1188,7 @@ void FOnlineAsyncEventSteamLobbyInviteAccepted::Finalize()
 		SessionInt->CurrentSessionSearch = MakeShareable(new FOnlineSessionSearch());
 		SessionInt->CurrentSessionSearch->SearchState = EOnlineAsyncTaskState::InProgress;
 
-		FOnlineAsyncTaskSteamFindLobby* NewTask = new FOnlineAsyncTaskSteamFindLobby(Subsystem, LobbyId, SessionInt->CurrentSessionSearch, LocalUserNum, SessionInt->OnSessionInviteAcceptedDelegates[LocalUserNum]);
+		FOnlineAsyncTaskSteamFindLobby* NewTask = new FOnlineAsyncTaskSteamFindLobby(Subsystem, LobbyId, SessionInt->CurrentSessionSearch, LocalUserNum, SessionInt->OnSessionUserInviteAcceptedDelegates);
 		Subsystem->QueueAsyncTask(NewTask);
 	}
 	else
diff --git a/Engine/Source/Runtime/Online/OnlineSubsystemSteam/Private/OnlineSessionAsyncLobbySteam.h b/Engine/Source/Runtime/Online/OnlineSubsystemSteam/Private/OnlineSessionAsyncLobbySteam.h
--- a/Engine/Source/Runtime/Online/OnlineSubsystemSteam/Private/OnlineSessionAsyncLobbySteam.h
+++ b/Engine/Source/Runtime/Online/OnlineSubsystemSteam/Private/OnlineSessionAsyncLobbySteam.h
@@ -297,6 +297,9 @@ public:
 	virtual void TriggerDelegates() override;
 };
 
+DECLARE_MULTICAST_DELEGATE_FourParams(FOnAsyncFindLobbyCompleteWithNetId, const bool, const int32, TSharedPtr< FUniqueNetId >, const class FOnlineSessionSearchResult&);
+typedef FOnAsyncFindLobbyCompleteWithNetId::FDelegate FOnAsyncFindLobbyCompleteDelegateWithNetId;
+
 /** 
  *  Async task for any search query to find Steam lobbies based on search criteria
  */
@@ -312,8 +315,12 @@ private:
 	TSharedPtr<class FOnlineSessionSearch> SearchSettings;
 	/** User that initiated the request */
 	int32 LocalUserNum;
-	/** Delegate to fire when the search is complete */
+	/** Delegates to fire when the search is complete. Only one of these will be called based on which constructor was used. */
 	FOnSingleSessionResultComplete OnFindLobbyCompleteDelegates;
+	FOnAsyncFindLobbyCompleteWithNetId OnFindLobbyCompleteWithNetIdDelegate;
+	/** This is true if the delegate passed in to the constructor was a FOnAsyncFindLobbyCompleteWithNetId delegate, false otherwise. */
+	bool bIsUsingNetIdDelegate;
+
 	/** Cached instance of Steam interface */
 	ISteamMatchmaking* SteamMatchmakingPtr;
 
@@ -334,7 +341,20 @@ public:
 		SearchSettings(InSearchSettings),
 		LocalUserNum(InLocalUserNum),
 		OnFindLobbyCompleteDelegates(InOnFindLobbyCompleteDelegates),
-		SteamMatchmakingPtr(SteamMatchmaking())
+        bIsUsingNetIdDelegate(false),
+        SteamMatchmakingPtr(SteamMatchmaking())
+	{
+	}
+
+	FOnlineAsyncTaskSteamFindLobby(class FOnlineSubsystemSteam* InSubsystem, const FUniqueNetIdSteam& InLobbyId, const TSharedPtr<FOnlineSessionSearch>& InSearchSettings, int32 InLocalUserNum, const FOnAsyncFindLobbyCompleteWithNetId& InOnFindLobbyCompleteDelegates) :
+		FOnlineAsyncTaskSteam(InSubsystem, k_uAPICallInvalid),
+		bInit(false),
+		LobbyId(InLobbyId),
+		SearchSettings(InSearchSettings),
+		LocalUserNum(InLocalUserNum),
+		OnFindLobbyCompleteWithNetIdDelegate(InOnFindLobbyCompleteDelegates),
+		bIsUsingNetIdDelegate(true),
+        SteamMatchmakingPtr(SteamMatchmaking())
 	{
 	}
 
diff --git a/Engine/Source/Runtime/Online/OnlineSubsystemSteam/Private/OnlineSessionAsyncServerSteam.cpp b/Engine/Source/Runtime/Online/OnlineSubsystemSteam/Private/OnlineSessionAsyncServerSteam.cpp
--- a/Engine/Source/Runtime/Online/OnlineSubsystemSteam/Private/OnlineSessionAsyncServerSteam.cpp
+++ b/Engine/Source/Runtime/Online/OnlineSubsystemSteam/Private/OnlineSessionAsyncServerSteam.cpp
@@ -1155,12 +1155,27 @@ void FOnlineAsyncTaskSteamFindServer::TriggerDelegates()
 	{
 		if (bWasSuccessful && SearchSettings->SearchResults.Num() > 0)
 		{
-			FindServerInviteCompleteDelegates.Broadcast(LocalUserNum, bWasSuccessful, SearchSettings->SearchResults[0]);
+			if (bIsUsingNetIdDelegate)
+			{
+				FindServerInviteCompleteWithUserIdDelegates.Broadcast(bWasSuccessful, LocalUserNum, MakeShareable<FUniqueNetId>(new FUniqueNetIdSteam(SteamUser()->GetSteamID())), SearchSettings->SearchResults[0]);
+			}
+			else
+			{
+				FindServerInviteCompleteDelegates.Broadcast(LocalUserNum, bWasSuccessful, SearchSettings->SearchResults[0]);
+			}
+			
 		}
 		else
 		{
 			FOnlineSessionSearchResult EmptyResult;
-			FindServerInviteCompleteDelegates.Broadcast(LocalUserNum, bWasSuccessful, EmptyResult);
+			if (bIsUsingNetIdDelegate)
+			{
+				FindServerInviteCompleteWithUserIdDelegates.Broadcast(bWasSuccessful, LocalUserNum, MakeShareable<FUniqueNetId>(new FUniqueNetIdSteam(SteamUser()->GetSteamID())), EmptyResult);
+			}
+			else
+			{
+				FindServerInviteCompleteDelegates.Broadcast(LocalUserNum, bWasSuccessful, EmptyResult);
+			}
 		}
 	}
 }
@@ -1223,7 +1238,7 @@ void FOnlineAsyncEventSteamInviteAccepted::Finalize()
 		if (bIsValid)
 		{
 			SessionInt->CurrentSessionSearch->QuerySettings.Set(FName(SEARCH_STEAM_HOSTIP), IpAddr->ToString(false), EOnlineComparisonOp::Equals);
-			FOnlineAsyncTaskSteamFindServer* NewTask = new FOnlineAsyncTaskSteamFindServer(Subsystem, SearchSettings, LocalUserNum, SessionInt->OnSessionInviteAcceptedDelegates[LocalUserNum]);
+			FOnlineAsyncTaskSteamFindServer* NewTask = new FOnlineAsyncTaskSteamFindServer(Subsystem, SearchSettings, LocalUserNum, SessionInt->OnSessionUserInviteAcceptedDelegates);
 			Subsystem->QueueAsyncTask(NewTask);
 		}
 	}
diff --git a/Engine/Source/Runtime/Online/OnlineSubsystemSteam/Private/OnlineSessionAsyncServerSteam.h b/Engine/Source/Runtime/Online/OnlineSubsystemSteam/Private/OnlineSessionAsyncServerSteam.h
--- a/Engine/Source/Runtime/Online/OnlineSubsystemSteam/Private/OnlineSessionAsyncServerSteam.h
+++ b/Engine/Source/Runtime/Online/OnlineSubsystemSteam/Private/OnlineSessionAsyncServerSteam.h
@@ -322,6 +322,10 @@ public:
 	virtual void Finalize() override;
 };
 
+
+DECLARE_MULTICAST_DELEGATE_FourParams(FOnAsyncFindServerInviteCompleteWithNetId, const bool, const int32, TSharedPtr< FUniqueNetId >, const class FOnlineSessionSearchResult&);
+typedef FOnAsyncFindServerInviteCompleteWithNetId::FDelegate FOnAsyncFindServerInviteCompleteWithNetIdDelegate;
+
 /**
  *	Async task for finding a single server and returning that search result to a properly defined delegate
  */
@@ -330,18 +334,32 @@ class FOnlineAsyncTaskSteamFindServer : public FOnlineAsyncTaskSteamFindServerBa
 protected:
 	/** User initiating the request */
 	int32 LocalUserNum;
-	/** User initiated session invite delegates */
+
+	/** Delegates that are called when the find server request completes. Only one is used at a time, depending on which constructor is used.*/
 	FOnSingleSessionResultComplete FindServerInviteCompleteDelegates;
+	FOnAsyncFindServerInviteCompleteWithNetId FindServerInviteCompleteWithUserIdDelegates;
+
+	/** Set to true if the constructor with the FOnAsyncFindServerInviteCompleteWithNetId delegate is called, false if not */
+	bool bIsUsingNetIdDelegate;
 
 public:
 
 	FOnlineAsyncTaskSteamFindServer(class FOnlineSubsystemSteam* InSubsystem, const TSharedPtr<class FOnlineSessionSearch>& InSearchSettings, int32 InLocalUserNum, FOnSingleSessionResultComplete& InDelegates) :
 		FOnlineAsyncTaskSteamFindServerBase(InSubsystem, InSearchSettings),
 		LocalUserNum(InLocalUserNum),
-		FindServerInviteCompleteDelegates(InDelegates)
+		FindServerInviteCompleteDelegates(InDelegates),
+		bIsUsingNetIdDelegate(false)
 	{	
 	}
 
+	FOnlineAsyncTaskSteamFindServer(class FOnlineSubsystemSteam* InSubsystem, const TSharedPtr<class FOnlineSessionSearch>& InSearchSettings, int32 InLocalUserNum, FOnAsyncFindServerInviteCompleteWithNetId& InDelegates) :
+		FOnlineAsyncTaskSteamFindServerBase(InSubsystem, InSearchSettings),
+		LocalUserNum(InLocalUserNum),
+		FindServerInviteCompleteWithUserIdDelegates(InDelegates),
+		bIsUsingNetIdDelegate(true)
+	{
+	}
+
 	/**
 	 *	Get a human readable description of task
 	 */
diff --git a/Engine/Source/Runtime/Online/OnlineSubsystemSteam/Private/OnlineSessionInterfaceSteam.cpp b/Engine/Source/Runtime/Online/OnlineSubsystemSteam/Private/OnlineSessionInterfaceSteam.cpp
--- a/Engine/Source/Runtime/Online/OnlineSubsystemSteam/Private/OnlineSessionInterfaceSteam.cpp
+++ b/Engine/Source/Runtime/Online/OnlineSubsystemSteam/Private/OnlineSessionInterfaceSteam.cpp
@@ -1420,30 +1420,27 @@ void FOnlineSessionSteam::TickPendingInvites(float DeltaTime)
 {
 	if (PendingInvite.PendingInviteType != ESteamSession::None)
 	{
-		for (int32 UserIndex = 0; UserIndex < MAX_LOCAL_PLAYERS; UserIndex++)
+		if (OnSessionUserInviteAcceptedDelegates.IsBound())
 		{
-			if (OnSessionInviteAcceptedDelegates[UserIndex].IsBound())
+			FOnlineAsyncItem* NewEvent = NULL;
+			FUniqueNetIdSteam FriendId(0);
+			if (PendingInvite.PendingInviteType == ESteamSession::LobbySession)
 			{
-				FOnlineAsyncItem* NewEvent = NULL;
-				FUniqueNetIdSteam FriendId(0);
-				if (PendingInvite.PendingInviteType == ESteamSession::LobbySession)
-				{
-					NewEvent = new FOnlineAsyncEventSteamLobbyInviteAccepted(SteamSubsystem, FriendId, PendingInvite.LobbyId);
-				}
-				else
-				{
-					NewEvent = new FOnlineAsyncEventSteamInviteAccepted(SteamSubsystem, FriendId, PendingInvite.ServerIp);
-				}
-
-				if (NewEvent)
-				{
-					UE_LOG_ONLINE(Verbose, TEXT("%s"), *NewEvent->ToString());
-					SteamSubsystem->QueueAsyncOutgoingItem(NewEvent);
-				}
+				NewEvent = new FOnlineAsyncEventSteamLobbyInviteAccepted(SteamSubsystem, FriendId, PendingInvite.LobbyId);
+			}
+			else
+			{
+				NewEvent = new FOnlineAsyncEventSteamInviteAccepted(SteamSubsystem, FriendId, PendingInvite.ServerIp);
+			}
 
-				// Clear the invite
-				PendingInvite.PendingInviteType = ESteamSession::None;
+			if (NewEvent)
+			{
+				UE_LOG_ONLINE(Verbose, TEXT("%s"), *NewEvent->ToString());
+				SteamSubsystem->QueueAsyncOutgoingItem(NewEvent);
 			}
+
+			// Clear the invite
+			PendingInvite.PendingInviteType = ESteamSession::None;
 		}
 	}
 }
diff --git a/Engine/Source/Runtime/Online/OnlineSubsystemUtils/Classes/OnlineSessionClient.h b/Engine/Source/Runtime/Online/OnlineSubsystemUtils/Classes/OnlineSessionClient.h
--- a/Engine/Source/Runtime/Online/OnlineSubsystemUtils/Classes/OnlineSessionClient.h
+++ b/Engine/Source/Runtime/Online/OnlineSubsystemUtils/Classes/OnlineSessionClient.h
@@ -23,8 +23,6 @@ protected:
 	/** Reference to the online sessions interface */
 	IOnlineSessionPtr SessionInt;
 
-	/** Delegate when an invite has been accepted from an external source */
-	FOnSessionInviteAcceptedDelegate OnSessionInviteAcceptedDelegate;
 	/** Delegate for destroying a session after previously ending it */
 	FOnEndSessionCompleteDelegate OnEndForJoinSessionCompleteDelegate;
 	/** Delegate for joining a new session after previously destroying it */
@@ -159,6 +157,17 @@ protected:
 	void DestroyExistingSession_Impl(FDelegateHandle& OutResult, FName SessionName, FOnDestroySessionCompleteDelegate& Delegate);
 
 protected:
+
+	/**
+	* Called from GameInstance when the user accepts an invite
+	*
+	* @param bWasSuccess true if invite was accepted successfully
+	* @param ControllerId the controller index of the user being invited
+	* @param UserId the user being invited
+	* @param InviteResult the search/settings result for the session we're joining via invite
+	*/
+	void OnSessionUserInviteAccepted(const bool bWasSuccess, const int32 ControllerId, TSharedPtr<FUniqueNetId> UserId, const FOnlineSessionSearchResult& InviteResult) override;
+
 	/**
 	 * Delegate fired when the joining process for an online session has completed
 	 *
@@ -174,16 +183,7 @@ protected:
 	 * @param SessionName name of session to join
 	 * @param SearchResult the session to join
 	 */
-	virtual void JoinSession(int32 LocalUserNum, FName SessionName, const FOnlineSessionSearchResult& SearchResult);
-
-	/**
-	 * Delegate fired when an invite request has been accepted (via external UI)
-	 *
-	 * @param LocalUserNum local user accepting invite
-	 * @param bWasSuccessful true if the async action completed without error, false if there was an error
-	 * @param SearchResult search result containing the invite data
-	 */
-	virtual void OnSessionInviteAccepted(int32 LocalUserNum, bool bWasSuccessful, const class FOnlineSessionSearchResult& SearchResult);
+	virtual void JoinSession(FName SessionName, const FOnlineSessionSearchResult& SearchResult);
 
 public:
 
diff --git a/Engine/Source/Runtime/Online/OnlineSubsystemUtils/Private/OnlineSessionClient.cpp b/Engine/Source/Runtime/Online/OnlineSubsystemUtils/Private/OnlineSessionClient.cpp
--- a/Engine/Source/Runtime/Online/OnlineSubsystemUtils/Private/OnlineSessionClient.cpp
+++ b/Engine/Source/Runtime/Online/OnlineSubsystemUtils/Private/OnlineSessionClient.cpp
@@ -19,19 +19,8 @@ UOnlineSessionClient::UOnlineSessionClient(const FObjectInitializer& ObjectIniti
 UWorld* UOnlineSessionClient::GetWorld() const
 {
 	ULocalPlayer* LP = Cast<ULocalPlayer>(GetOuter());
-	if (LP != nullptr)
-	{
-		if (LP->PlayerController != nullptr)
-		{
-			return LP->PlayerController->GetWorld();
-		}
-		else if (LP->ViewportClient != nullptr)
-		{
-			return LP->ViewportClient->GetWorld();
-		}
-	}
-
-	return nullptr;
+	check(LP);
+	return LP->GetWorld();
 }
 
 APlayerController* UOnlineSessionClient::GetPlayerController()
@@ -72,18 +61,6 @@ void UOnlineSessionClient::RegisterOnlineDelegates(UWorld* InWorld)
 	IOnlineSubsystem* OnlineSub = Online::GetSubsystem(InWorld);
 	if (OnlineSub)
 	{
-		SessionInt = OnlineSub->GetSessionInterface();
-		if (SessionInt.IsValid())
-		{
-			int32 ControllerId = GetControllerId();
-			if (ControllerId != INVALID_CONTROLLERID)
-			{
-				// Always on the lookout for invite acceptance (via actual invite or join from external ui)
-				OnSessionInviteAcceptedDelegate       = FOnSessionInviteAcceptedDelegate::CreateUObject(this, &UOnlineSessionClient::OnSessionInviteAccepted);
-				OnSessionInviteAcceptedDelegateHandle = SessionInt->AddOnSessionInviteAcceptedDelegate_Handle(ControllerId, OnSessionInviteAcceptedDelegate);
-			}
-		}
-
 		OnJoinSessionCompleteDelegate           = FOnJoinSessionCompleteDelegate   ::CreateUObject(this, &UOnlineSessionClient::OnJoinSessionComplete);
 		OnEndForJoinSessionCompleteDelegate     = FOnEndSessionCompleteDelegate    ::CreateUObject(this, &UOnlineSessionClient::OnEndForJoinSessionComplete);
 		OnDestroyForJoinSessionCompleteDelegate = FOnDestroySessionCompleteDelegate::CreateUObject(this, &UOnlineSessionClient::OnDestroyForJoinSessionComplete);
@@ -95,25 +72,12 @@ void UOnlineSessionClient::RegisterOnlineDelegates(UWorld* InWorld)
 
 void UOnlineSessionClient::ClearOnlineDelegates(UWorld* InWorld)
 {
-	IOnlineSubsystem* OnlineSub = Online::GetSubsystem(InWorld);
-	if (OnlineSub)
-	{
-		if (SessionInt.IsValid())
-		{
-			int32 ControllerId = GetControllerId();
-			if (ControllerId != INVALID_CONTROLLERID)
-			{
-				SessionInt->ClearOnSessionInviteAcceptedDelegate_Handle(ControllerId, OnSessionInviteAcceptedDelegateHandle);
-			}
-		}
-	}
-
 	SessionInt = NULL;
 }
 
-void UOnlineSessionClient::OnSessionInviteAccepted(int32 LocalUserNum, bool bWasSuccessful, const FOnlineSessionSearchResult& SearchResult)
+void UOnlineSessionClient::OnSessionUserInviteAccepted(bool bWasSuccessful, int32 ControllerId, TSharedPtr<FUniqueNetId> UserId, const FOnlineSessionSearchResult& SearchResult)
 {
-	UE_LOG(LogOnline, Verbose, TEXT("OnSessionInviteAccepted LocalUserNum: %d bSuccess: %d"), LocalUserNum, bWasSuccessful);
+	UE_LOG(LogOnline, Verbose, TEXT("OnSessionInviteAccepted LocalUserNum: %d bSuccess: %d"), ControllerId, bWasSuccessful);
 	// Don't clear invite accept delegate
 
 	if (bWasSuccessful)
@@ -121,8 +85,7 @@ void UOnlineSessionClient::OnSessionInviteAccepted(int32 LocalUserNum, bool bWas
 		if (SearchResult.IsValid())
 		{
 			bIsFromInvite = true;
-			check(GetControllerId() == LocalUserNum);
-			JoinSession(LocalUserNum, GameSessionName, SearchResult);
+			JoinSession(GameSessionName, SearchResult);
 		}
 		else
 		{
@@ -177,11 +140,7 @@ void UOnlineSessionClient::OnDestroyForJoinSessionComplete(FName SessionName, bo
 
 	if (bWasSuccessful)
 	{
-		int32 ControllerId = GetControllerId();
-		if (ControllerId != INVALID_CONTROLLERID)
-		{
-			JoinSession(ControllerId, SessionName, CachedSessionResult);
-		}
+		JoinSession(SessionName, CachedSessionResult);
 	}
 
 	bHandlingDisconnect = false;
@@ -195,12 +154,8 @@ void UOnlineSessionClient::OnDestroyForMainMenuComplete(FName SessionName, bool
 		SessionInt->ClearOnDestroySessionCompleteDelegate_Handle(OnDestroyForMainMenuCompleteDelegateHandle);
 	}	
 
-	APlayerController* PC = GetPlayerController();
-	if (PC)
-	{
-		// Call disconnect to force us back to the menu level
-		GEngine->HandleDisconnect(PC->GetWorld(), PC->GetWorld()->GetNetDriver());
-	}
+	// Call disconnect to force us back to the menu level
+	GEngine->HandleDisconnect(GetWorld(), GetWorld()->GetNetDriver());
 
 	bHandlingDisconnect = false;
 }
@@ -260,7 +215,7 @@ void UOnlineSessionClient::OnJoinSessionComplete(FName SessionName, EOnJoinSessi
 	}
 }
 
-void UOnlineSessionClient::JoinSession(int32 LocalUserNum, FName SessionName, const FOnlineSessionSearchResult& SearchResult)
+void UOnlineSessionClient::JoinSession(FName SessionName, const FOnlineSessionSearchResult& SearchResult)
 {
 	// Clean up existing sessions if applicable
 	EOnlineSessionState::Type SessionState = SessionInt->GetSessionState(SessionName);
@@ -271,10 +226,12 @@ void UOnlineSessionClient::JoinSession(int32 LocalUserNum, FName SessionName, co
 	}
 	else
 	{
-		UGameInstance * const GameInstance = GetPlayerController()->GetGameInstance();
-		GameInstance->JoinSession(static_cast<ULocalPlayer*>(GetPlayerController()->Player), SearchResult);
-		/*OnJoinSessionCompleteDelegateHandle = SessionInt->AddOnJoinSessionCompleteDelegate_Handle(OnJoinSessionCompleteDelegate);
-		SessionInt->JoinSession(LocalUserNum, SessionName, SearchResult);*/
+		ULocalPlayer* LocalPlayer = static_cast<ULocalPlayer*>(GetOuter());
+		check(LocalPlayer);
+		UGameInstance * const GameInstance = LocalPlayer->GetGameInstance();
+		check(GameInstance);
+
+		GameInstance->JoinSession(LocalPlayer, SearchResult);
 	}
 }
 
@@ -292,21 +249,17 @@ void UOnlineSessionClient::HandleDisconnect(UWorld *World, UNetDriver *NetDriver
 
 bool UOnlineSessionClient::HandleDisconnectInternal(UWorld* World, UNetDriver* NetDriver)
 {
-	APlayerController* PC = GetPlayerController();
-	if (PC)
+	// This was a disconnect for our active world, we will handle it
+	if (GetWorld() == World)
 	{
-		// This was a disconnect for our active world, we will handle it
-		if (PC->GetWorld() == World)
+		// Prevent multiple calls to this async flow
+		if (!bHandlingDisconnect)
 		{
-			// Prevent multiple calls to this async flow
-			if (!bHandlingDisconnect)
-			{
-				bHandlingDisconnect = true;
-				DestroyExistingSession_Impl(OnDestroyForMainMenuCompleteDelegateHandle, GameSessionName, OnDestroyForMainMenuCompleteDelegate);
-			}
-
-			return true;
+			bHandlingDisconnect = true;
+			DestroyExistingSession_Impl(OnDestroyForMainMenuCompleteDelegateHandle, GameSessionName, OnDestroyForMainMenuCompleteDelegate);
 		}
+
+		return true;
 	}
 
 	return false;
@@ -314,20 +267,15 @@ bool UOnlineSessionClient::HandleDisconnectInternal(UWorld* World, UNetDriver* N
 
 void UOnlineSessionClient::StartOnlineSession(FName SessionName)
 {
-	APlayerController* PC = GetPlayerController();
-	if (PC)
+	IOnlineSessionPtr SessionInterface = Online::GetSessionInterface(GetWorld());
+	if (SessionInterface.IsValid())
 	{
-		UWorld* World = PC->GetWorld();
-		IOnlineSessionPtr SessionInterface = Online::GetSessionInterface(World);
-		if (SessionInterface.IsValid())
+		FNamedOnlineSession* Session = SessionInterface->GetNamedSession(SessionName);
+		if (Session &&
+			(Session->SessionState == EOnlineSessionState::Pending || Session->SessionState == EOnlineSessionState::Ended))
 		{
-			FNamedOnlineSession* Session = SessionInterface->GetNamedSession(SessionName);
-			if (Session &&
-				(Session->SessionState == EOnlineSessionState::Pending || Session->SessionState == EOnlineSessionState::Ended))
-			{
-				StartSessionCompleteHandle = SessionInterface->AddOnStartSessionCompleteDelegate_Handle(FOnStartSessionCompleteDelegate::CreateUObject(this, &UOnlineSessionClient::OnStartSessionComplete));
-				SessionInterface->StartSession(SessionName);
-			}
+			StartSessionCompleteHandle = SessionInterface->AddOnStartSessionCompleteDelegate_Handle(FOnStartSessionCompleteDelegate::CreateUObject(this, &UOnlineSessionClient::OnStartSessionComplete));
+			SessionInterface->StartSession(SessionName);
 		}
 	}
 }
@@ -335,34 +283,25 @@ void UOnlineSessionClient::StartOnlineSession(FName SessionName)
 void UOnlineSessionClient::OnStartSessionComplete(FName SessionName, bool bWasSuccessful)
 {
 	UE_LOG(LogOnline, Verbose, TEXT("OnStartSessionComplete %s bSuccess: %d"), *SessionName.ToString(), bWasSuccessful);
-	APlayerController* PC = GetPlayerController();
-	if (PC)
+
+	IOnlineSessionPtr SessionInterface = Online::GetSessionInterface(GetWorld());
+	if (SessionInterface.IsValid())
 	{
-		UWorld* World = PC->GetWorld();
-		IOnlineSessionPtr SessionInterface = Online::GetSessionInterface(World);
-		if (SessionInterface.IsValid())
-		{
-			SessionInterface->ClearOnStartSessionCompleteDelegate_Handle(StartSessionCompleteHandle);
-		}
+		SessionInterface->ClearOnStartSessionCompleteDelegate_Handle(StartSessionCompleteHandle);
 	}
 }
 
 void UOnlineSessionClient::EndOnlineSession(FName SessionName)
 {
-	APlayerController* PC = GetPlayerController();
-	if (PC)
+	IOnlineSessionPtr SessionInterface = Online::GetSessionInterface(GetWorld());
+	if (SessionInterface.IsValid())
 	{
-		UWorld* World = PC->GetWorld();
-		IOnlineSessionPtr SessionInterface = Online::GetSessionInterface(World);
-		if (SessionInterface.IsValid())
+		FNamedOnlineSession* Session = SessionInterface->GetNamedSession(SessionName);
+		if (Session &&
+			Session->SessionState == EOnlineSessionState::InProgress)
 		{
-			FNamedOnlineSession* Session = SessionInterface->GetNamedSession(SessionName);
-			if (Session &&
-				Session->SessionState == EOnlineSessionState::InProgress)
-			{
-				EndSessionCompleteHandle = SessionInterface->AddOnEndSessionCompleteDelegate_Handle(FOnStartSessionCompleteDelegate::CreateUObject(this, &UOnlineSessionClient::OnEndSessionComplete));
-				SessionInterface->EndSession(SessionName);
-			}
+			EndSessionCompleteHandle = SessionInterface->AddOnEndSessionCompleteDelegate_Handle(FOnStartSessionCompleteDelegate::CreateUObject(this, &UOnlineSessionClient::OnEndSessionComplete));
+			SessionInterface->EndSession(SessionName);
 		}
 	}
 }
@@ -370,14 +309,10 @@ void UOnlineSessionClient::EndOnlineSession(FName SessionName)
 void UOnlineSessionClient::OnEndSessionComplete(FName SessionName, bool bWasSuccessful)
 {
 	UE_LOG(LogOnline, Verbose, TEXT("OnEndSessionComplete %s bSuccess: %d"), *SessionName.ToString(), bWasSuccessful);
-	APlayerController* PC = GetPlayerController();
-	if (PC)
+
+	IOnlineSessionPtr SessionInterface = Online::GetSessionInterface(GetWorld());
+	if (SessionInterface.IsValid())
 	{
-		UWorld* World = PC->GetWorld();
-		IOnlineSessionPtr SessionInterface = Online::GetSessionInterface(World);
-		if (SessionInterface.IsValid())
-		{
-			SessionInterface->ClearOnEndSessionCompleteDelegate_Handle(EndSessionCompleteHandle);
-		}
+		SessionInterface->ClearOnEndSessionCompleteDelegate_Handle(EndSessionCompleteHandle);
 	}
 }
\ No newline at end of file
diff --git a/Engine/Source/Runtime/Online/OnlineSubsystemUtils/Private/Tests/TestSessionInterface.cpp b/Engine/Source/Runtime/Online/OnlineSubsystemUtils/Private/Tests/TestSessionInterface.cpp
--- a/Engine/Source/Runtime/Online/OnlineSubsystemUtils/Private/Tests/TestSessionInterface.cpp
+++ b/Engine/Source/Runtime/Online/OnlineSubsystemUtils/Private/Tests/TestSessionInterface.cpp
@@ -153,13 +153,13 @@ void FTestSessionInterface::Test(UWorld* InWorld, bool bTestLAN, bool bIsPresenc
 	}
 
 	// Always on the lookout for invite acceptance (via actual invite or join in an external client)
-	OnSessionInviteAcceptedDelegate = FOnSessionInviteAcceptedDelegate::CreateRaw(this, &FTestSessionInterface::OnSessionInviteAccepted);
-	OnSessionInviteAcceptedDelegateHandle = SessionInt->AddOnSessionInviteAcceptedDelegate_Handle(0, OnSessionInviteAcceptedDelegate);
+	OnSessionUserInviteAcceptedDelegate = FOnSessionUserInviteAcceptedDelegate::CreateRaw(this, &FTestSessionInterface::OnSessionUserInviteAccepted);
+	OnSessionUserInviteAcceptedDelegateHandle = SessionInt->AddOnSessionUserInviteAcceptedDelegate_Handle(OnSessionUserInviteAcceptedDelegate);
 }
 
 void FTestSessionInterface::ClearDelegates()
 {
-	SessionInt->ClearOnSessionInviteAcceptedDelegate_Handle(0, OnSessionInviteAcceptedDelegateHandle);
+	SessionInt->ClearOnSessionUserInviteAcceptedDelegate_Handle(OnSessionUserInviteAcceptedDelegateHandle);
 	GEngine->OnWorldDestroyed().RemoveAll(this);
 }
 
@@ -180,14 +180,14 @@ void FTestSessionInterface::OnReadFriendsListComplete(int32 LocalUserNum, bool b
 	}
 }
 
-void FTestSessionInterface::OnSessionInviteAccepted(int32 LocalUserNum, bool bWasSuccessful, const FOnlineSessionSearchResult& SearchResult)
+void FTestSessionInterface::OnSessionUserInviteAccepted(bool bWasSuccessful, int32 ControllerId, TSharedPtr<FUniqueNetId> UserId, const FOnlineSessionSearchResult& SearchResult)
 {
-	UE_LOG(LogOnline, Verbose, TEXT("OnSessionInviteAccepted LocalUserNum: %d bSuccess: %d"), LocalUserNum, bWasSuccessful);
+	UE_LOG(LogOnline, Verbose, TEXT("OnSessionInviteAccepted ControllerId: %d bSuccess: %d"), ControllerId, bWasSuccessful);
 	// Don't clear invite accept delegate
 
 	if (bWasSuccessful)
 	{
-		JoinSession(LocalUserNum, GameSessionName, SearchResult);
+		JoinSession(ControllerId, GameSessionName, SearchResult);
 	}
 }
 
@@ -229,7 +229,7 @@ void FTestSessionInterface::OnUnregisterPlayerComplete(FName SessionName, const
 	SessionInt->ClearOnUnregisterPlayersCompleteDelegate_Handle(OnUnregisterPlayersCompleteDelegateHandle);
 }
 
-void FTestSessionInterface::JoinSession(int32 LocalUserNum, FName SessionName, const FOnlineSessionSearchResult& SearchResult)
+void FTestSessionInterface::JoinSession(int32 ControllerId, FName SessionName, const FOnlineSessionSearchResult& SearchResult)
 {
 	// Clean up existing sessions if applicable
 	EOnlineSessionState::Type SessionState = SessionInt->GetSessionState(SessionName);
@@ -241,7 +241,7 @@ void FTestSessionInterface::JoinSession(int32 LocalUserNum, FName SessionName, c
 	else
 	{
 		OnJoinSessionCompleteDelegateHandle = SessionInt->AddOnJoinSessionCompleteDelegate_Handle(OnJoinSessionCompleteDelegate);
-		SessionInt->JoinSession(LocalUserNum, SessionName, SearchResult);
+		SessionInt->JoinSession(ControllerId, SessionName, SearchResult);
 	}
 }
 
diff --git a/Engine/Source/Runtime/Online/OnlineSubsystemUtils/Private/Tests/TestSessionInterface.h b/Engine/Source/Runtime/Online/OnlineSubsystemUtils/Private/Tests/TestSessionInterface.h
--- a/Engine/Source/Runtime/Online/OnlineSubsystemUtils/Private/Tests/TestSessionInterface.h
+++ b/Engine/Source/Runtime/Online/OnlineSubsystemUtils/Private/Tests/TestSessionInterface.h
@@ -37,7 +37,7 @@
 	IOnlineFriendsPtr Friends;
 
 	/** Delegate for handling an accepted invite */
-	FOnSessionInviteAcceptedDelegate OnSessionInviteAcceptedDelegate;
+	FOnSessionUserInviteAcceptedDelegate OnSessionUserInviteAcceptedDelegate;
 
 	/** Delegate for creating a new session */
 	FOnCreateSessionCompleteDelegate OnCreateSessionCompleteDelegate;
@@ -69,7 +69,7 @@
 
 	/** Handles for the above delegates */
 	FDelegateHandle OnReadFriendsListCompleteDelegateHandle;
-	FDelegateHandle OnSessionInviteAcceptedDelegateHandle;
+	FDelegateHandle OnSessionUserInviteAcceptedDelegateHandle;
 	FDelegateHandle OnCreateSessionCompleteDelegateHandle;
 	FDelegateHandle OnStartSessionCompleteDelegateHandle;
 	FDelegateHandle OnEndSessionCompleteDelegateHandle;
@@ -120,7 +120,7 @@
 	 * @param SessionName name of session to join
 	 * @param SearchResult the session to join
 	 */
-	void JoinSession(int32 LocalUserNum, FName SessionName, const FOnlineSessionSearchResult& SearchResult);
+	void JoinSession(int32 ControllerId, FName SessionName, const FOnlineSessionSearchResult& SearchResult);
 
 	/**
 	 * Transition from ending a session to destroying a session
@@ -151,11 +151,12 @@
     /**
      * Delegate fired when an invite request has been accepted (via external client)
      *
-     * @param LocalUserNum local user accepting invite
-     * @param bWasSuccessful true if the async action completed without error, false if there was an error
-     * @param SearchResult search result containing the invite data
+	 * @param bWasSuccessful true if the async action completed without error, false if there was an error
+	 * @param ControllerId the controller number of the accepting user
+	 * @param UserId the user being invited
+	 * @param InviteResult the search/settings for the session we're joining via invite
      */
-	void OnSessionInviteAccepted(int32 LocalUserNum, bool bWasSuccessful, const FOnlineSessionSearchResult& SearchResult);
+	void OnSessionUserInviteAccepted(const bool bWasSuccessful, const int32 LocalUserNum, TSharedPtr< FUniqueNetId > UserId, const FOnlineSessionSearchResult& SearchResult);
 
 	/**
 	 * Delegate fired when a session create request has completed
