Creating parties fully functional first pass.
Group information published to presence.
#online, party, mcp
Merging //depot/UE4-LauncherDev/Engine/Source/Runtime/Online/... to //depot/UE4/Engine/Source/Runtime/Online/...

diff --git a/Engine/Source/Runtime/Online/OnlineSubsystem/Public/Interfaces/OnlinePartyInterface.h b/Engine/Source/Runtime/Online/OnlineSubsystem/Public/Interfaces/OnlinePartyInterface.h
--- a/Engine/Source/Runtime/Online/OnlineSubsystem/Public/Interfaces/OnlinePartyInterface.h
+++ b/Engine/Source/Runtime/Online/OnlineSubsystem/Public/Interfaces/OnlinePartyInterface.h
@@ -5,11 +5,6 @@
 #include "OnlineSubsystemTypes.h"
 #include "OnlineChatInterface.h"
 
-/**
- * Id of a party instance
- */
-typedef FUniqueNetId FOnlinePartyId;
-
 /**
  * Key value data that can be sent to party members
  */
@@ -102,10 +97,19 @@ namespace EPartyPermissions
  */
 struct FPartyConfiguration
 {
+	FPartyConfiguration()
+		: JoinPerms(EPartyPermissions::FriendsOnly)
+		, MaxMembers(0)
+	{}
+
 	/** Permission for configuring party */
 	EPartyPermissions::Type JoinPerms;
 	/** Maximum active members allowed */
 	int32 MaxMembers;
+	/** PartyId to use - largely for debug purposes*/
+	FString PartyId;
+	/** Password to use */
+	FString Password;
 };
 
 namespace EPartyState
@@ -313,7 +317,7 @@ public:
 	 *
 	 * @return true if task was started
 	 */
-	virtual bool CreateParty(const FUniqueNetId& LocalUserId, const FPartyConfiguration& PartyConfig, const FOnCreatePartyComplete& Delegate = FOnCreatePartyComplete(), const FString& UserRoomId = TEXT("")) = 0;
+	virtual bool CreateParty(const FUniqueNetId& LocalUserId, const FPartyConfiguration& PartyConfig, const FOnCreatePartyComplete& Delegate = FOnCreatePartyComplete()) = 0;
 
 	/**
 	 * Update an existing party with new configuration
diff --git a/Engine/Source/Runtime/Online/OnlineSubsystem/Public/OnlineSubsystemTypes.h b/Engine/Source/Runtime/Online/OnlineSubsystem/Public/OnlineSubsystemTypes.h
--- a/Engine/Source/Runtime/Online/OnlineSubsystem/Public/OnlineSubsystemTypes.h
+++ b/Engine/Source/Runtime/Online/OnlineSubsystem/Public/OnlineSubsystemTypes.h
@@ -1330,3 +1330,7 @@ struct FOnlineNotification
 */
 typedef FString FNotificationTransportId;
 
+/**
+* Id of a party instance
+*/
+typedef FUniqueNetId FOnlinePartyId;
