Multi-user chat configuration for use with parties
Muc rooms are created and with privacy and password setting & defaults for other fields,  then fire room creation delegates for ChatMcp/PartyMcp
Merging //depot/UE4-LauncherDev/Engine/Source/Runtime/Online/... to //depot/UE4/Engine/Source/Runtime/Online/...

diff --git a/Engine/Source/Runtime/Online/OnlineSubsystem/Public/Interfaces/OnlineChatInterface.h b/Engine/Source/Runtime/Online/OnlineSubsystem/Public/Interfaces/OnlineChatInterface.h
--- a/Engine/Source/Runtime/Online/OnlineSubsystem/Public/Interfaces/OnlineChatInterface.h
+++ b/Engine/Source/Runtime/Online/OnlineSubsystem/Public/Interfaces/OnlineChatInterface.h
@@ -32,19 +32,17 @@ class FChatRoomConfig
 {
 public:
 	FChatRoomConfig()
-		: bMembersOnly(false)
-		, bHidden(false)
-		, bPasswordRequired(false)
-		, bPersistent(false)
-		, bAllowMemberInvites(false)
-		, bLoggingEnabled(false)
-		, MessageHistory(0)
-		, MaxMembers(0)
+		: bPasswordRequired(false)
+		, Password(TEXT(""))
 	{}
 
+	bool bPasswordRequired;
+	FString Password;
+
+private:
+	// Below are unused, move to public when hooking up to functionality
 	bool bMembersOnly;
 	bool bHidden;
-	bool bPasswordRequired;
 	bool bPersistent;
 	bool bAllowMemberInvites;
 	bool bLoggingEnabled;
@@ -167,7 +165,7 @@ public:
 	 *
 	 * @return if successfully started the async operation
 	 */
-	virtual bool CreateRoom(const FUniqueNetId& UserId, const FChatRoomId& RoomId, const FString& Nickname) = 0;
+	virtual bool CreateRoom(const FUniqueNetId& UserId, const FChatRoomId& RoomId, const FString& Nickname, const FChatRoomConfig& ChatRoomConfig) = 0;
 	
 	/**
 	 * Kick off request for joining a public chat room
