Merge 2518312 from Orion branch.

Fix issue with not all multicast RPCs getting replicated to network replays. Fixes a number of problems with certain sounds not playing during replays (and likely other issues too).

diff --git a/Engine/Source/Runtime/Engine/Private/Actor.cpp b/Engine/Source/Runtime/Engine/Private/Actor.cpp
--- a/Engine/Source/Runtime/Engine/Private/Actor.cpp
+++ b/Engine/Source/Runtime/Engine/Private/Actor.cpp
@@ -3354,15 +3354,6 @@ bool AActor::CallRemoteFunction( UFunction* Function, void* Parameters, FOutParm
 	if (NetDriver)
 	{
 		NetDriver->ProcessRemoteFunction(this, Function, Parameters, OutParms, Stack, NULL);
-		if (NetDriver->NetDriverName == NAME_GameNetDriver)
-		{
-			// Replicate any RPCs to the replay net driver so that they can get saved in network replays
-			NetDriver = GEngine->FindNamedNetDriver(GetWorld(), NAME_DemoNetDriver);
-			if (NetDriver)
-			{
-				NetDriver->ProcessRemoteFunction(this, Function, Parameters, OutParms, Stack, NULL);
-			}
-		}
 		return true;
 	}
 
diff --git a/Engine/Source/Runtime/Online/OnlineSubsystemUtils/Private/IpNetDriver.cpp b/Engine/Source/Runtime/Online/OnlineSubsystemUtils/Private/IpNetDriver.cpp
--- a/Engine/Source/Runtime/Online/OnlineSubsystemUtils/Private/IpNetDriver.cpp
+++ b/Engine/Source/Runtime/Online/OnlineSubsystemUtils/Private/IpNetDriver.cpp
@@ -323,6 +323,12 @@ void UIpNetDriver::ProcessRemoteFunction(class AActor* Actor, UFunction* Functio
 				}
 			}			
 
+			// Replicate any RPCs to the replay net driver so that they can get saved in network replays
+			UNetDriver* NetDriver = GEngine->FindNamedNetDriver(GetWorld(), NAME_DemoNetDriver);
+			if (NetDriver)
+			{
+				NetDriver->ProcessRemoteFunction(Actor, Function, Parameters, OutParms, Stack, SubObject);
+			}
 			// Return here so we don't call InternalProcessRemoteFunction again at the bottom of this function
 			return;
 		}
