Merging using REVERSE UE4-To-UE4-LauncherDev

Latest BuildPatchServices module, implementing stage only installation

diff --git a/Engine/Source/Runtime/Online/BuildPatchServices/Private/BuildPatchDownloader.cpp b/Engine/Source/Runtime/Online/BuildPatchServices/Private/BuildPatchDownloader.cpp
--- a/Engine/Source/Runtime/Online/BuildPatchServices/Private/BuildPatchDownloader.cpp
+++ b/Engine/Source/Runtime/Online/BuildPatchServices/Private/BuildPatchDownloader.cpp
@@ -304,14 +304,14 @@ uint32 FBuildPatchDownloader::Run()
 				FileDataArray.Reset();
 				FileDataArray.AddUninitialized( FileSize );
 				int64 BytesRead = 0;
-				while ( BytesRead < FileSize )
+				while ( BytesRead < FileSize && !FBuildPatchInstallError::HasFatalError() )
 				{
 					const int64 ReadLen = FMath::Min<int64>( BytesPerCall, FileSize - BytesRead );
 					Reader->Serialize( FileDataArray.GetData() + BytesRead, ReadLen );
 					BytesRead += ReadLen;
 					OnDownloadProgress( NextGuid, BytesRead );
 				}
-				bSuccess = Reader->Close();
+				bSuccess = Reader->Close() && !FBuildPatchInstallError::HasFatalError();
 				delete Reader;
 			}
 			OnDownloadComplete( NextGuid, DownloadUrl, FileDataArray, bSuccess, INDEX_NONE );
diff --git a/Engine/Source/Runtime/Online/BuildPatchServices/Private/BuildPatchFileAttributes.cpp b/Engine/Source/Runtime/Online/BuildPatchServices/Private/BuildPatchFileAttributes.cpp
new file mode 100644
--- /dev/null
+++ b/Engine/Source/Runtime/Online/BuildPatchServices/Private/BuildPatchFileAttributes.cpp
@@ -0,0 +1,220 @@
+// Copyright 1998-2015 Epic Games, Inc. All Rights Reserved.
+
+#include "BuildPatchServicesPrivatePCH.h"
+
+class FBuildPatchFileAttributesImpl : public FBuildPatchFileAttributes
+{
+public:
+	FBuildPatchFileAttributesImpl(const FBuildPatchAppManifestRef& NewManifest, const FBuildPatchAppManifestPtr& OldManifest, const FString& VerifyDirectory, const FString& StagedFileDirectory, bool bUseStageDirectory);
+
+	virtual ~FBuildPatchFileAttributesImpl(){}
+
+	virtual bool ApplyAttributes(bool bForce) override;
+
+private:
+	FBuildPatchAppManifestRef NewManifest;
+	FBuildPatchAppManifestPtr OldManifest;
+	const FString VerifyDirectory;
+	const FString StagedFileDirectory;
+	const bool bUseStageDirectory;
+
+	FString SelectFullFilePath(const FString& BuildFile);
+	void SetupFileAttributes(const FString& FilePath, const FFileManifestData& FileManifest);
+	bool SetFileCompressionFlag(const FString& Filepath, bool bIsCompressed);
+	bool SetExecutableFlag(const FString& Filepath);
+};
+
+FBuildPatchFileAttributesImpl::FBuildPatchFileAttributesImpl(const FBuildPatchAppManifestRef& NewManifest, const FBuildPatchAppManifestPtr& OldManifest, const FString& VerifyDirectory, const FString& StagedFileDirectory, bool bUseStageDirectory)
+	: NewManifest(NewManifest)
+	, OldManifest(OldManifest)
+	, VerifyDirectory(VerifyDirectory)
+	, StagedFileDirectory(StagedFileDirectory)
+	, bUseStageDirectory(bUseStageDirectory)
+{}
+
+bool FBuildPatchFileAttributesImpl::ApplyAttributes(bool bForce)
+{
+	// We need to set attributes for all files in the new build that require it
+	TArray<FString> BuildFiles;
+	NewManifest->GetFileList(BuildFiles);
+	for (const FString& BuildFile : BuildFiles)
+	{
+		// Break if quitting
+		if (FBuildPatchInstallError::HasFatalError())
+		{
+			break;
+		}
+
+		const FFileManifestData* FileManifest = NewManifest->GetFileManifest(BuildFile);
+		if (FileManifest != nullptr)
+		{
+			// Apply
+			const bool bHasAttrib = FileManifest->bIsReadOnly || FileManifest->bIsCompressed || FileManifest->bIsUnixExecutable;
+			if (bHasAttrib || bForce)
+			{
+				SetupFileAttributes(SelectFullFilePath(BuildFile), *FileManifest);
+			}
+		}
+	}
+
+	// We also need to check if any attributes have been removed, unless we forced anyway
+	if (OldManifest.IsValid() && !bForce)
+	{
+		TArray<FString> BuildFiles;
+		OldManifest->GetFileList(BuildFiles);
+		for (const FString& BuildFile : BuildFiles)
+		{
+			// Break if quitting
+			if (FBuildPatchInstallError::HasFatalError())
+			{
+				break;
+			}
+
+			const FFileManifestData* NewFileManifest = NewManifest->GetFileManifest(BuildFile);
+			const FFileManifestData* OldFileManifest = OldManifest->GetFileManifest(BuildFile);
+			if (NewFileManifest != nullptr)
+			{
+				const bool bAttribChanged = (OldFileManifest->bIsReadOnly && !NewFileManifest->bIsReadOnly) || (OldFileManifest->bIsCompressed && !NewFileManifest->bIsCompressed);
+				if (bAttribChanged)
+				{
+					SetupFileAttributes(SelectFullFilePath(BuildFile), *NewFileManifest);
+				}
+			}
+		}
+	}
+
+	// We don't fail on this step currently
+	return true;
+}
+
+FString FBuildPatchFileAttributesImpl::SelectFullFilePath(const FString& BuildFile)
+{
+	if (bUseStageDirectory)
+	{
+		FString StagedFilename = StagedFileDirectory / BuildFile;
+		if (IFileManager::Get().FileSize(*StagedFilename) != -1)
+		{
+			return MoveTemp(StagedFilename);
+		}
+	}
+	FString InstallFilename = VerifyDirectory / BuildFile;
+	return MoveTemp(InstallFilename);
+}
+
+void FBuildPatchFileAttributesImpl::SetupFileAttributes(const FString& FilePath, const FFileManifestData& FileManifest)
+{
+	// File must not be readonly to be able to set attributes
+	IPlatformFile::GetPlatformPhysical().SetReadOnly(*FilePath, false);
+	// Set correct attributes
+	SetFileCompressionFlag(FilePath, FileManifest.bIsCompressed);
+	if (!IPlatformFile::GetPlatformPhysical().SetReadOnly(*FilePath, FileManifest.bIsReadOnly))
+	{
+		GLog->Logf(TEXT("BuildPatchServices: WARNING: Could not set readonly flag %s"), *FilePath);
+	}
+	if (FileManifest.bIsUnixExecutable && !SetExecutableFlag(FilePath))
+	{
+		GLog->Logf(TEXT("BuildPatchServices: WARNING: Could not set executable flag %s"), *FilePath);
+	}
+}
+
+#if PLATFORM_WINDOWS
+// Start of region that uses windows types
+#include "AllowWindowsPlatformTypes.h"
+#include <wtypes.h>
+#include <ioapiset.h>
+#include <winbase.h>
+#include <fileapi.h>
+#include <winioctl.h>
+
+bool FBuildPatchFileAttributesImpl::SetFileCompressionFlag(const FString& Filepath, bool bIsCompressed)
+{
+	// Get the file handle
+	HANDLE FileHandle = ::CreateFile(
+		*Filepath,								// Path to file
+		GENERIC_READ | GENERIC_WRITE,			// Read and write access
+		FILE_SHARE_READ | FILE_SHARE_WRITE,		// Share access for DeviceIoControl
+		NULL,									// Default security
+		OPEN_EXISTING,							// Open existing file
+		FILE_ATTRIBUTE_NORMAL,					// No specific attributes
+		NULL									// No template file handle
+		);
+	uint32 Error = ::GetLastError();
+	if (FileHandle == NULL || FileHandle == INVALID_HANDLE_VALUE)
+	{
+		GLog->Logf(TEXT("BuildPatchServices: WARNING: Could not open file to set compression flag %d Error:%d File:%s"), bIsCompressed, Error, *Filepath);
+		return false;
+	}
+
+	// Send the compression control code to the device
+	uint16 Message = bIsCompressed ? COMPRESSION_FORMAT_DEFAULT : COMPRESSION_FORMAT_NONE;
+	uint16* pMessage = &Message;
+	DWORD Dummy = 0;
+	LPDWORD pDummy = &Dummy;
+	BOOL bSuccess = ::DeviceIoControl(
+		FileHandle,								// The file handle
+		FSCTL_SET_COMPRESSION,					// Control code
+		pMessage,								// Our message
+		sizeof(uint16),							//
+		NULL,									// Not used
+		0,										// Not used
+		pDummy,									// The value returned will be meaningless, but is required
+		NULL									// No overlap structure, we a running this synchronously
+		);
+	Error = ::GetLastError();
+	const bool bFileSystemUnsupported = Error == ERROR_INVALID_FUNCTION;
+	if (bSuccess == FALSE && bFileSystemUnsupported == false)
+	{
+		GLog->Logf(TEXT("BuildPatchServices: WARNING: Could not set compression flag %d Error:%d File:%s"), bIsCompressed, Error, *Filepath);
+	}
+
+	// Close the open file handle
+	::CloseHandle(FileHandle);
+
+	// We treat unsupported as not being a failure
+	return bSuccess == TRUE || bFileSystemUnsupported;
+}
+
+bool FBuildPatchFileAttributesImpl::SetExecutableFlag(const FString& Filepath)
+{
+	// Not implemented
+	return true;
+}
+// End of region that uses windows types
+#include "HideWindowsPlatformTypes.h"
+
+#elif PLATFORM_MAC
+bool FBuildPatchFileAttributesImpl::SetFileCompressionFlag(const FString& Filepath, bool bIsCompressed)
+{
+	// Not implemented
+	return true;
+}
+
+bool FBuildPatchFileAttributesImpl::SetExecutableFlag(const FString& Filepath)
+{
+	bool bSuccess = false;
+	// Enable executable permission bit
+	struct stat FileInfo;
+	if (stat(TCHAR_TO_UTF8(*Filepath), &FileInfo) == 0)
+	{
+		bSuccess = chmod(TCHAR_TO_UTF8(*Filepath), FileInfo.st_mode | S_IXUSR | S_IXGRP | S_IXOTH) == 0;
+	}
+	return bSuccess;
+}
+#elif PLATFORM_LINUX
+bool FBuildPatchFileAttributesImpl::SetFileCompressionFlag(const FString& Filepath, bool bIsCompressed)
+{
+	// Not implemented
+	return true;
+}
+
+bool FBuildPatchFileAttributesImpl::SetExecutableFlag(const FString& Filepath)
+{
+	// Not implemented
+	return true;
+}
+#endif
+
+TSharedRef<FBuildPatchFileAttributes> FBuildPatchFileAttributesFactory::Create(const FBuildPatchAppManifestRef& NewManifest, const FBuildPatchAppManifestPtr& OldManifest, const FString& VerifyDirectory, const FString& StagedFileDirectory)
+{
+	return MakeShareable(new FBuildPatchFileAttributesImpl(NewManifest, OldManifest, VerifyDirectory, StagedFileDirectory, !StagedFileDirectory.IsEmpty()));
+}
diff --git a/Engine/Source/Runtime/Online/BuildPatchServices/Private/BuildPatchFileAttributes.h b/Engine/Source/Runtime/Online/BuildPatchServices/Private/BuildPatchFileAttributes.h
new file mode 100644
--- /dev/null
+++ b/Engine/Source/Runtime/Online/BuildPatchServices/Private/BuildPatchFileAttributes.h
@@ -0,0 +1,34 @@
+// Copyright 1998-2015 Epic Games, Inc. All Rights Reserved.
+
+#pragma once
+
+class FBuildPatchFileAttributes
+{
+public:
+	/**
+	 * Applies the attributes to the files in the Staging directory, or installation directory. A force param can be set true in the
+	 * case that all attributes should be set regardless of changing since previous build. This is helpful for fixing file attributes
+	 * when they are suspected incorrect.
+	 * @param bForce   Default: False. Pass in true to always set up attributes even if patch does not change them.
+	 */
+	virtual bool ApplyAttributes(bool bForce = false) = 0;
+};
+
+class FBuildPatchFileAttributesFactory
+{
+public:
+	/**
+	 * Creates a file attribute class that can be used to setup file attributes that are contained in the build manifest, optionally
+	 * taking account of a staging directory where alternative files are.
+	 * @param NewManifest           The manifest describing the build data.
+	 * @param OldManifest           The manifest describing the build data.
+	 * @param VerifyDirectory       The directory to analyze.
+	 * @param StagedFileDirectory   A stage directory for updated files, ignored if empty string. If a file exists here, it will be
+	 *                              checked instead of the one in VerifyDirectory.
+	 * @return     Ref of an object that can be used to perform the operation.
+	 */
+	static TSharedRef<FBuildPatchFileAttributes> Create(const FBuildPatchAppManifestRef& NewManifest,
+	                                                    const FBuildPatchAppManifestPtr& OldManifest,
+	                                                    const FString& VerifyDirectory,
+	                                                    const FString& StagedFileDirectory);
+};
diff --git a/Engine/Source/Runtime/Online/BuildPatchServices/Private/BuildPatchInstaller.cpp b/Engine/Source/Runtime/Online/BuildPatchServices/Private/BuildPatchInstaller.cpp
--- a/Engine/Source/Runtime/Online/BuildPatchServices/Private/BuildPatchInstaller.cpp
+++ b/Engine/Source/Runtime/Online/BuildPatchServices/Private/BuildPatchInstaller.cpp
@@ -7,105 +7,6 @@
 
 #include "BuildPatchServicesPrivatePCH.h"
 
-// Required for setting file compression flag
-#if PLATFORM_WINDOWS
-// Start of region that uses windows types
-#include "AllowWindowsPlatformTypes.h"
-
-#include <wtypes.h>
-#include <ioapiset.h>
-#include <winbase.h>
-#include <fileapi.h>
-#include <winioctl.h>
-
-bool SetFileCompressionFlag(const FString& Filepath, bool bIsCompressed)
-{
-	// Get the file handle
-	HANDLE FileHandle = ::CreateFile(
-		*Filepath,								// Path to file
-		GENERIC_READ | GENERIC_WRITE,			// Read and write access
-		FILE_SHARE_READ | FILE_SHARE_WRITE,		// Share access for DeviceIoControl
-		NULL,									// Default security
-		OPEN_EXISTING,							// Open existing file
-		FILE_ATTRIBUTE_NORMAL,					// No specific attributes
-		NULL									// No template file handle
-		);
-	uint32 Error = ::GetLastError();
-	if (FileHandle == NULL || FileHandle == INVALID_HANDLE_VALUE)
-	{
-		GLog->Logf(TEXT("BuildPatchServices: WARNING: Could not open file to set compression flag %d Error:%d File:%s"), bIsCompressed, Error, *Filepath);
-		return false;
-	}
-
-	// Send the compression control code to the device
-	uint16 Message = bIsCompressed ? COMPRESSION_FORMAT_DEFAULT : COMPRESSION_FORMAT_NONE;
-	uint16* pMessage = &Message;
-	DWORD Dummy = 0;
-	LPDWORD pDummy = &Dummy;
-	BOOL bSuccess = ::DeviceIoControl(
-		FileHandle,								// The file handle
-		FSCTL_SET_COMPRESSION,					// Control code
-		pMessage,								// Our message
-		sizeof(uint16),							//
-		NULL,									// Not used
-		0,										// Not used
-		pDummy,									// The value returned will be meaningless, but is required
-		NULL									// No overlap structure, we a running this synchronously
-		);
-	Error = ::GetLastError();
-	const bool bFileSystemUnsupported = Error == ERROR_INVALID_FUNCTION;
-	if (bSuccess == FALSE && bFileSystemUnsupported == false)
-	{
-		GLog->Logf(TEXT("BuildPatchServices: WARNING: Could not set compression flag %d Error:%d File:%s"), bIsCompressed, Error, *Filepath);
-	}
-
-	// Close the open file handle
-	::CloseHandle(FileHandle);
-
-	// We treat unsupported as not being a failure
-	return bSuccess == TRUE || bFileSystemUnsupported;
-}
-
-bool SetExecutableFlag(const FString& Filepath)
-{
-	// Not implemented
-	return true;
-}
-// End of region that uses windows types
-#include "HideWindowsPlatformTypes.h"
-
-#elif PLATFORM_MAC
-bool SetFileCompressionFlag(const FString& Filepath, bool bIsCompressed) 
-{
-	// Not implemented
-	return true;
-}
-
-bool SetExecutableFlag(const FString& Filepath)
-{
-	bool bSuccess = false;
-	// Enable executable permission bit
-	struct stat FileInfo;
-	if (stat(TCHAR_TO_UTF8(*Filepath), &FileInfo) == 0)
-	{
-		bSuccess = chmod(TCHAR_TO_UTF8(*Filepath), FileInfo.st_mode | S_IXUSR | S_IXGRP | S_IXOTH) == 0;
-	}
-	return bSuccess;
-}
-#elif PLATFORM_LINUX
-bool SetFileCompressionFlag(const FString& Filepath, bool bIsCompressed)
-{
-	// Not implemented
-	return true;
-}
-
-bool SetExecutableFlag(const FString& Filepath)
-{
-	// Not implemented
-	return true;
-}
-#endif
-
 #define LOCTEXT_NAMESPACE "BuildPatchInstaller"
 
 #define NUM_DOWNLOAD_READINGS	5
@@ -120,7 +21,7 @@ bool SetExecutableFlag(const FString& Filepath)
 
 /* FBuildPatchInstaller implementation
 *****************************************************************************/
-FBuildPatchInstaller::FBuildPatchInstaller( FBuildPatchBoolManifestDelegate InOnCompleteDelegate, FBuildPatchAppManifestPtr CurrentManifest, FBuildPatchAppManifestRef InstallManifest, const FString& InInstallDirectory, const FString& InStagingDirectory, FBuildPatchInstallationInfo& InstallationInfoRef )
+FBuildPatchInstaller::FBuildPatchInstaller(FBuildPatchBoolManifestDelegate InOnCompleteDelegate, FBuildPatchAppManifestPtr CurrentManifest, FBuildPatchAppManifestRef InstallManifest, const FString& InInstallDirectory, const FString& InStagingDirectory, FBuildPatchInstallationInfo& InstallationInfoRef, bool ShouldStageOnly)
 	: Thread( NULL )
 	, OnCompleteDelegate( InOnCompleteDelegate )
 	, CurrentBuildManifest( CurrentManifest )
@@ -133,6 +34,8 @@ FBuildPatchInstaller::FBuildPatchInstaller( FBuildPatchBoolManifestDelegate InOn
 	, ThreadLock()
 	, bIsFileData( InstallManifest->IsFileDataManifest() )
 	, bIsChunkData( !bIsFileData )
+	, bIsRepairing(CurrentManifest.IsValid() && CurrentManifest->IsSameAs(InstallManifest))
+	, bShouldStageOnly(ShouldStageOnly)
 	, bSuccess( true )
 	, bIsRunning( false )
 	, bIsInited( false )
@@ -146,7 +49,6 @@ FBuildPatchInstaller::FBuildPatchInstaller( FBuildPatchBoolManifestDelegate InOn
 	, TimePausedAt( 0.0 )
 	, InstallationInfo( InstallationInfoRef )
 {
-	bIsRepairing = CurrentManifest.IsValid() && CurrentManifest->IsSameAs(InstallManifest);
 	// Start thread!
 	const TCHAR* ThreadName = TEXT( "BuildPatchInstallerThread" );
 	Thread = FRunnableThread::Create(this, ThreadName);
@@ -234,13 +136,22 @@ uint32 FBuildPatchInstaller::Run()
 		bProcessSuccess = bInstallSuccess && RunVerification(CorruptFiles);
 
 		// Clean staging if INSTALL success
+		CleanUpTime = FPlatformTime::Seconds();
 		if (bInstallSuccess)
 		{
-			GLog->Logf(TEXT("BuildPatchServices: Deleting staging area"));
-			CleanUpTime = FPlatformTime::Seconds();
-			IFileManager::Get().DeleteDirectory(*StagingDirectory, false, true);
-			CleanUpTime = FPlatformTime::Seconds() - CleanUpTime;
+			if (bShouldStageOnly)
+			{
+				GLog->Logf(TEXT("BuildPatchServices: Deleting litter from staging area"));
+				IFileManager::Get().DeleteDirectory(*DataStagingDir, false, true);
+				IFileManager::Get().Delete(*(InstallStagingDir / TEXT("$resumeData")), false, true);
+			}
+			else
+			{
+				GLog->Logf(TEXT("BuildPatchServices: Deleting staging area"));
+				IFileManager::Get().DeleteDirectory(*StagingDirectory, false, true);
+			}
 		}
+		CleanUpTime = FPlatformTime::Seconds() - CleanUpTime;
 		BuildProgress.SetStateProgress(EBuildPatchProgress::CleanUp, 1.0f);
 
 		// Set if we can retry
@@ -551,157 +462,135 @@ void FBuildPatchInstaller::CleanupEmptyDirectories(const FString& RootDirectory)
 
 bool FBuildPatchInstaller::RunBackupAndMove()
 {
-	GLog->Logf(TEXT("BuildPatchServices: Running backup and stage relocation"));
-	// If there's no error, move all complete files
-	bool bMoveSuccess = FBuildPatchInstallError::HasFatalError() == false;
-	if (bMoveSuccess)
+	// We skip this step if performing stage only
+	bool bMoveSuccess = true;
+	if (bShouldStageOnly)
 	{
-		// First handle files that should be removed for patching
-		TArray< FString > FilesToRemove;
-		if (CurrentBuildManifest.IsValid())
-		{
-			FBuildPatchAppManifest::GetRemovableFiles(CurrentBuildManifest.ToSharedRef(), NewBuildManifest, FilesToRemove);
-		}
-		// Add to build stats
-		ThreadLock.Lock();
-		BuildStats.NumFilesToRemove = FilesToRemove.Num();
-		ThreadLock.Unlock();
-		for (const FString& OldFilename : FilesToRemove)
-		{
-			BackupFileIfNecessary(OldFilename);
-			const bool bDeleteSuccess = IFileManager::Get().Delete(*(InstallDirectory / OldFilename), false, true, true);
-			const uint32 LastError = FPlatformMisc::GetLastError();
-			GLog->Logf(TEXT("BuildPatchServices: Removed (%u,%u) %s"), bDeleteSuccess ? 1 : 0, LastError, *OldFilename);
-		}
-		
-		// Now handle files that have been constructed
-		bool bSavedMoveMarkerFile = false;
-		TArray< FString > ConstructionFiles;
-		NewBuildManifest->GetFileList(ConstructionFiles);
-		BuildProgress.SetStateProgress(EBuildPatchProgress::MovingToInstall, 0.0f);
-		const float NumConstructionFilesFloat = ConstructionFiles.Num();
-		for (auto ConstructionFilesIt = ConstructionFiles.CreateConstIterator(); ConstructionFilesIt && bMoveSuccess && !FBuildPatchInstallError::HasFatalError(); ++ConstructionFilesIt)
+		GLog->Logf(TEXT("BuildPatchServices: Skipping backup and stage relocation"));
+		BuildProgress.SetStateProgress(EBuildPatchProgress::MovingToInstall, 1.0f);
+	}
+	else
+	{
+		GLog->Logf(TEXT("BuildPatchServices: Running backup and stage relocation"));
+		// If there's no error, move all complete files
+		bMoveSuccess = FBuildPatchInstallError::HasFatalError() == false;
+		if (bMoveSuccess)
 		{
-			const FString& ConstructionFile = *ConstructionFilesIt;
-			const FString SrcFilename = InstallStagingDir / ConstructionFile;
-			const FString DestFilename = InstallDirectory / ConstructionFile;
-			const float FileIndexFloat = ConstructionFilesIt.GetIndex();
-			// Skip files not constructed
-			if (!FPaths::FileExists(SrcFilename))
+			// First handle files that should be removed for patching
+			TArray< FString > FilesToRemove;
+			if (CurrentBuildManifest.IsValid())
 			{
-				BuildProgress.SetStateProgress(EBuildPatchProgress::MovingToInstall, FileIndexFloat / NumConstructionFilesFloat);
-				continue;
+				FBuildPatchAppManifest::GetRemovableFiles(CurrentBuildManifest.ToSharedRef(), NewBuildManifest, FilesToRemove);
 			}
-			// Create the move marker file
-			if (!bSavedMoveMarkerFile)
+			// Add to build stats
+			ThreadLock.Lock();
+			BuildStats.NumFilesToRemove = FilesToRemove.Num();
+			ThreadLock.Unlock();
+			for (const FString& OldFilename : FilesToRemove)
 			{
-				bSavedMoveMarkerFile = true;
-				GLog->Logf(TEXT("BuildPatchServices: Create MM"));
-				FArchive* MoveMarkerFile = IFileManager::Get().CreateFileWriter(*PreviousMoveMarker, FILEWRITE_EvenIfReadOnly);
-				if (MoveMarkerFile != NULL)
+				BackupFileIfNecessary(OldFilename);
+				const bool bDeleteSuccess = IFileManager::Get().Delete(*(InstallDirectory / OldFilename), false, true, true);
+				const uint32 LastError = FPlatformMisc::GetLastError();
+				GLog->Logf(TEXT("BuildPatchServices: Removed (%u,%u) %s"), bDeleteSuccess ? 1 : 0, LastError, *OldFilename);
+			}
+
+			// Now handle files that have been constructed
+			bool bSavedMoveMarkerFile = false;
+			TArray< FString > ConstructionFiles;
+			NewBuildManifest->GetFileList(ConstructionFiles);
+			BuildProgress.SetStateProgress(EBuildPatchProgress::MovingToInstall, 0.0f);
+			const float NumConstructionFilesFloat = ConstructionFiles.Num();
+			for (auto ConstructionFilesIt = ConstructionFiles.CreateConstIterator(); ConstructionFilesIt && bMoveSuccess && !FBuildPatchInstallError::HasFatalError(); ++ConstructionFilesIt)
+			{
+				const FString& ConstructionFile = *ConstructionFilesIt;
+				const FString SrcFilename = InstallStagingDir / ConstructionFile;
+				const FString DestFilename = InstallDirectory / ConstructionFile;
+				const float FileIndexFloat = ConstructionFilesIt.GetIndex();
+				// Skip files not constructed
+				if (!FPaths::FileExists(SrcFilename))
 				{
-					MoveMarkerFile->Close();
-					delete MoveMarkerFile;
+					BuildProgress.SetStateProgress(EBuildPatchProgress::MovingToInstall, FileIndexFloat / NumConstructionFilesFloat);
+					continue;
 				}
-				// Make sure we have some progress if we do some work
-				if (BuildProgress.GetStateWeight(EBuildPatchProgress::MovingToInstall) <= 0.0f)
+				// Create the move marker file
+				if (!bSavedMoveMarkerFile)
 				{
-					BuildProgress.SetStateWeight(EBuildPatchProgress::MovingToInstall, 0.1f);
+					bSavedMoveMarkerFile = true;
+					GLog->Logf(TEXT("BuildPatchServices: Create MM"));
+					FArchive* MoveMarkerFile = IFileManager::Get().CreateFileWriter(*PreviousMoveMarker, FILEWRITE_EvenIfReadOnly);
+					if (MoveMarkerFile != NULL)
+					{
+						MoveMarkerFile->Close();
+						delete MoveMarkerFile;
+					}
+					// Make sure we have some progress if we do some work
+					if (BuildProgress.GetStateWeight(EBuildPatchProgress::MovingToInstall) <= 0.0f)
+					{
+						BuildProgress.SetStateWeight(EBuildPatchProgress::MovingToInstall, 0.1f);
+					}
 				}
-			}
-			// Backup file if need be
-			BackupFileIfNecessary(ConstructionFile);
-			// Move the file to the installation directory
-			int32 MoveRetries = NUM_FILE_MOVE_RETRIES;
-			bMoveSuccess = IFileManager::Get().Move(*DestFilename, *SrcFilename, true, true);
-			uint32 ErrorCode = FPlatformMisc::GetLastError();
-			while (!bMoveSuccess && MoveRetries > 0)
-			{
-				--MoveRetries;
-				FBuildPatchAnalytics::RecordConstructionError(ConstructionFile, ErrorCode, TEXT("Failed To Move"));
-				GWarn->Logf(TEXT("BuildPatchServices: ERROR: Failed to move file %s (%d), trying copy"), *ConstructionFile, ErrorCode);
-				bMoveSuccess = IFileManager::Get().Copy(*DestFilename, *SrcFilename, true, true) == COPY_OK;
-				ErrorCode = FPlatformMisc::GetLastError();
-				if (!bMoveSuccess)
+				// Backup file if need be
+				BackupFileIfNecessary(ConstructionFile);
+				// Delete existing if there
+				IFileManager::Get().Delete(*DestFilename, false, true, false);
+				// Move the file to the installation directory
+				int32 MoveRetries = NUM_FILE_MOVE_RETRIES;
+				bMoveSuccess = IFileManager::Get().Move(*DestFilename, *SrcFilename, true, true, true, true);
+				uint32 ErrorCode = FPlatformMisc::GetLastError();
+				while (!bMoveSuccess && MoveRetries > 0)
 				{
-					GWarn->Logf(TEXT("BuildPatchServices: ERROR: Failed to copy file %s (%d), retying after 0.5 sec"), *ConstructionFile, ErrorCode);
-					FPlatformProcess::Sleep(0.5f);
 					--MoveRetries;
-					bMoveSuccess = IFileManager::Get().Move(*DestFilename, *SrcFilename, true, true);
+					FBuildPatchAnalytics::RecordConstructionError(ConstructionFile, ErrorCode, TEXT("Failed To Move"));
+					GWarn->Logf(TEXT("BuildPatchServices: ERROR: Failed to move file %s (%d), trying copy"), *ConstructionFile, ErrorCode);
+					bMoveSuccess = IFileManager::Get().Copy(*DestFilename, *SrcFilename, true, true, true) == COPY_OK;
 					ErrorCode = FPlatformMisc::GetLastError();
+					if (!bMoveSuccess)
+					{
+						GWarn->Logf(TEXT("BuildPatchServices: ERROR: Failed to copy file %s (%d), retying after 0.5 sec"), *ConstructionFile, ErrorCode);
+						FPlatformProcess::Sleep(0.5f);
+						bMoveSuccess = IFileManager::Get().Move(*DestFilename, *SrcFilename, true, true, true, true);
+						ErrorCode = FPlatformMisc::GetLastError();
+					}
+					else
+					{
+						IFileManager::Get().Delete(*SrcFilename, false, true, false);
+					}
+				}
+				if (!bMoveSuccess)
+				{
+					GWarn->Logf(TEXT("BuildPatchServices: ERROR: Failed to move file %s"), *FPaths::GetCleanFilename(ConstructionFile));
+					FBuildPatchInstallError::SetFatalError(EBuildPatchInstallError::MoveFileToInstall);
 				}
 				else
 				{
-					IFileManager::Get().Delete(*SrcFilename, false, true, false);
+					FilesInstalled.Add(ConstructionFile);
+					BuildProgress.SetStateProgress(EBuildPatchProgress::MovingToInstall, FileIndexFloat / NumConstructionFilesFloat);
 				}
 			}
-			if (!bMoveSuccess)
-			{
-				GWarn->Logf(TEXT("BuildPatchServices: ERROR: Failed to move file %s"), *FPaths::GetCleanFilename(ConstructionFile));
-				FBuildPatchInstallError::SetFatalError(EBuildPatchInstallError::MoveFileToInstall);
-			}
-			else
+
+			// After we've completed deleting/moving patch files to the install directory, clean up any empty directories left over
+			CleanupEmptyDirectories(InstallDirectory);
+
+			bMoveSuccess = bMoveSuccess && (FBuildPatchInstallError::HasFatalError() == false);
+			if (bMoveSuccess)
 			{
-				FilesInstalled.Add(ConstructionFile);
-				BuildProgress.SetStateProgress(EBuildPatchProgress::MovingToInstall, FileIndexFloat / NumConstructionFilesFloat);
+				BuildProgress.SetStateProgress(EBuildPatchProgress::MovingToInstall, 1.0f);
 			}
 		}
-		
-		// After we've completed deleting/moving patch files to the install directory, clean up any empty directories left over
-		CleanupEmptyDirectories(InstallDirectory);
-
-		bMoveSuccess = bMoveSuccess && (FBuildPatchInstallError::HasFatalError() == false);
-		if (bMoveSuccess)
-		{
-			BuildProgress.SetStateProgress(EBuildPatchProgress::MovingToInstall, 1.0f);
-		}
+		GLog->Logf(TEXT("BuildPatchServices: Relocation complete %d"), bMoveSuccess?1:0);
 	}
-	GLog->Logf(TEXT("BuildPatchServices: Relocation complete %d"), bMoveSuccess?1:0);
 	return bMoveSuccess;
 }
 
 bool FBuildPatchInstaller::RunFileAttributes(bool bForce)
 {
-	// We need to set attributes for all files in the new build that require it
-	for (const FFileManifestData& FileManifest : NewBuildManifest->Data->FileManifestList)
-	{
-		// Break if quitting
-		if (FBuildPatchInstallError::HasFatalError())
-		{
-			break;
-		}
-		// Apply
-		const bool bHasAttrib = FileManifest.bIsReadOnly || FileManifest.bIsCompressed || FileManifest.bIsUnixExecutable;
-		if (bHasAttrib || bForce)
-		{
-			FString DestFilename = InstallDirectory / FileManifest.Filename;
-			SetupFileAttributes(DestFilename, FileManifest);
-		}
-	}
+	// Only provide stage directory if stage-only mode
+	FString EmptyString;
+	FString& OptionalStageDirectory = bShouldStageOnly ? InstallStagingDir : EmptyString;
 
-	// We also need to check if any attributes have been removed, unless we forced anyway
-	if (CurrentBuildManifest.IsValid() && !bForce)
-	{
-		for (const FFileManifestData& OldFileManifest : CurrentBuildManifest->Data->FileManifestList)
-		{
-			// Break if quitting
-			if (FBuildPatchInstallError::HasFatalError())
-			{
-				break;
-			}
-			const FFileManifestData* const* NewFileManifestPtr = NewBuildManifest->FileManifestLookup.Find(OldFileManifest.Filename);
-			if (NewFileManifestPtr != nullptr)
-			{
-				const FFileManifestData& NewFileManifest = **NewFileManifestPtr;
-				const bool bAttribChanged = (OldFileManifest.bIsReadOnly && !NewFileManifest.bIsReadOnly) || (OldFileManifest.bIsCompressed && !NewFileManifest.bIsCompressed);
-				if (bAttribChanged)
-				{
-					FString DestFilename = InstallDirectory / OldFileManifest.Filename;
-					SetupFileAttributes(DestFilename, NewFileManifest);
-				}
-			}
-		}
-	}
+	// Construct the attributes class
+	auto Attributes = FBuildPatchFileAttributesFactory::Create(NewBuildManifest, CurrentBuildManifest, InstallDirectory, OptionalStageDirectory);
+	Attributes->ApplyAttributes(bForce);
 
 	// We don't fail on this step currently
 	return true;
@@ -721,7 +610,20 @@ bool FBuildPatchInstaller::RunVerification(TArray< FString >& CorruptFiles)
 	GLog->Logf(TEXT("BuildPatchServices: Verifying install"));
 	CorruptFiles.Empty();
 	VerifyTime = FPlatformTime::Seconds();
-	bool bVerifySuccess = NewBuildManifest->VerifyAgainstDirectory(InstallDirectory, CorruptFiles, FBuildPatchFloatDelegate::CreateRaw(this, &FBuildPatchInstaller::UpdateVerificationProgress), FBuildPatchBoolRetDelegate::CreateRaw(this, &FBuildPatchInstaller::IsPaused), VerifyPauseTime);
+
+	// Setup the verify delegates
+	auto ProgressDelegate = FBuildPatchFloatDelegate::CreateRaw(this, &FBuildPatchInstaller::UpdateVerificationProgress);
+	auto IsPausedDelegate = FBuildPatchBoolRetDelegate::CreateRaw(this, &FBuildPatchInstaller::IsPaused);
+
+	// Only provide stage directory if stage-only mode
+	FString EmptyString;
+	FString& OptionalStageDirectory = bShouldStageOnly ? InstallStagingDir : EmptyString;
+
+	// Construct the verifier
+	auto Verifier = FBuildPatchVerificationFactory::Create(NewBuildManifest, ProgressDelegate, IsPausedDelegate, InstallDirectory, OptionalStageDirectory);
+
+	// Verify the build
+	bool bVerifySuccess = Verifier->VerifyAgainstDirectory(CorruptFiles, VerifyPauseTime);
 	VerifyTime = FPlatformTime::Seconds() - VerifyTime - VerifyPauseTime;
 	if (!bVerifySuccess)
 	{
@@ -743,7 +645,10 @@ bool FBuildPatchInstaller::RunVerification(TArray< FString >& CorruptFiles)
 		for (const FString& CorruptFile : CorruptFiles)
 		{
 			BackupFileIfNecessary(CorruptFile, true);
-			IFileManager::Get().Delete(*(InstallDirectory / CorruptFile), false, true);
+			if (!bShouldStageOnly)
+			{
+				IFileManager::Get().Delete(*(InstallDirectory / CorruptFile), false, true);
+			}
 			IFileManager::Get().Delete(*(InstallStagingDir / CorruptFile), false, true);
 		}
 	}
@@ -1123,22 +1028,6 @@ void FBuildPatchInstaller::UpdateVerificationProgress( float Percent )
 	BuildProgress.SetStateProgress( EBuildPatchProgress::BuildVerification, Percent );
 }
 
-void FBuildPatchInstaller::SetupFileAttributes( const FString& FilePath, const FFileManifestData& FileManifest )
-{
-	// File must not be readonly to be able to set attributes
-	IPlatformFile::GetPlatformPhysical().SetReadOnly(*FilePath, false);
-	// Set correct attributes
-	SetFileCompressionFlag(FilePath, FileManifest.bIsCompressed);
-	if (!IPlatformFile::GetPlatformPhysical().SetReadOnly(*FilePath, FileManifest.bIsReadOnly))
-	{
-		GLog->Logf(TEXT("BuildPatchServices: WARNING: Could not set readonly flag %s"), *FilePath);
-	}
-	if (FileManifest.bIsUnixExecutable && !SetExecutableFlag(FilePath))
-	{
-		GLog->Logf(TEXT("BuildPatchServices: WARNING: Could not set executable flag %s"), *FilePath);
-	}
-}
-
 void FBuildPatchInstaller::ExecuteCompleteDelegate()
 {
 	// Should be executed in main thread, and already be complete
diff --git a/Engine/Source/Runtime/Online/BuildPatchServices/Private/BuildPatchInstaller.h b/Engine/Source/Runtime/Online/BuildPatchServices/Private/BuildPatchInstaller.h
--- a/Engine/Source/Runtime/Online/BuildPatchServices/Private/BuildPatchInstaller.h
+++ b/Engine/Source/Runtime/Online/BuildPatchServices/Private/BuildPatchInstaller.h
@@ -57,6 +57,9 @@ private:
 	// A flag to store if we are performing a repair
 	bool bIsRepairing;
 
+	// A flag to store if we should only be staging, and not moving to install location
+	bool bShouldStageOnly;
+
 	// A flag storing whether the process was a success
 	bool bSuccess;
 
@@ -105,8 +108,9 @@ public:
 	 * @param InstallDirectory		The directory where the build will be installed
 	 * @param StagingDirectory		The directory for storing the intermediate files
 	 * @param InstallationInfoRef	Reference to the module's installation info that keeps record of locally installed apps for use as chunk sources
+	 * @param ShouldStageOnly		Whether the installer should only stage the required files, and skip moving them to the install directory
 	 */
-	FBuildPatchInstaller( FBuildPatchBoolManifestDelegate OnCompleteDelegate, FBuildPatchAppManifestPtr CurrentManifest, FBuildPatchAppManifestRef InstallManifest, const FString& InstallDirectory, const FString& StagingDirectory, FBuildPatchInstallationInfo& InstallationInfoRef );
+	FBuildPatchInstaller(FBuildPatchBoolManifestDelegate OnCompleteDelegate, FBuildPatchAppManifestPtr CurrentManifest, FBuildPatchAppManifestRef InstallManifest, const FString& InstallDirectory, const FString& StagingDirectory, FBuildPatchInstallationInfo& InstallationInfoRef, bool ShouldStageOnly);
 
 	/**
 	 * Default Destructor, will delete the allocated Thread
@@ -228,13 +232,6 @@ private:
 	 */
 	void UpdateVerificationProgress( float Percent );
 
-	/**
-	 * Setup file attributes
-	 * @param FilePath		The path to the file to apply attributes on
-	 * @param FileManifest	The manifest for the file
-	 */
-	void SetupFileAttributes( const FString& FilePath, const FFileManifestData& FileManifest );
-
 	/**
 	 * Delete empty directories from an installation
 	 * @param RootDirectory	 Root Directory for search
diff --git a/Engine/Source/Runtime/Online/BuildPatchServices/Private/BuildPatchManifest.cpp b/Engine/Source/Runtime/Online/BuildPatchServices/Private/BuildPatchManifest.cpp
--- a/Engine/Source/Runtime/Online/BuildPatchServices/Private/BuildPatchManifest.cpp
+++ b/Engine/Source/Runtime/Online/BuildPatchServices/Private/BuildPatchManifest.cpp
@@ -55,7 +55,7 @@ bool BufferIsJsonManifest(const TArray<uint8>& DataInput)
 
 /* EBuildPatchAppManifestVersion implementation
 *****************************************************************************/
-const EBuildPatchAppManifestVersion::Type EBuildPatchAppManifestVersion::GetLatestVersion()
+EBuildPatchAppManifestVersion::Type EBuildPatchAppManifestVersion::GetLatestVersion()
 {
 #if ENABLE_NOCHUNKS_COMPRESSION
 	return static_cast<EBuildPatchAppManifestVersion::Type>(LatestPlusOne - 1);
@@ -64,7 +64,7 @@ const EBuildPatchAppManifestVersion::Type EBuildPatchAppManifestVersion::GetLate
 #endif
 }
 
-const EBuildPatchAppManifestVersion::Type EBuildPatchAppManifestVersion::GetLatestJsonVersion()
+EBuildPatchAppManifestVersion::Type EBuildPatchAppManifestVersion::GetLatestJsonVersion()
 {
 	return static_cast<EBuildPatchAppManifestVersion::Type>(StoredAsCompressedUClass - 1);
 }
@@ -164,7 +164,7 @@ FFileManifestData::FFileManifestData()
 {
 }
 
-const int64& FFileManifestData::GetFileSize() const
+int64 FFileManifestData::GetFileSize() const
 {
 	return FileSize;
 }
@@ -214,13 +214,12 @@ public:
 
 	virtual FArchive& operator<<(FName& N) override
 	{
-//		const FName NameZero(N, 0);
-		if (!FNameIndexLookup.Contains(N/*ameZero*/))
+		if (!FNameIndexLookup.Contains(N))
 		{
 			const int32 ArNameIndex = FNameIndexLookup.Num();
-			FNameIndexLookup.Add(N/*ameZero*/, ArNameIndex);
+			FNameIndexLookup.Add(N, ArNameIndex);
 		}
-		*this << FNameIndexLookup[N/*ameZero*/];
+		*this << FNameIndexLookup[N];
 		return *this;
 	}
 
@@ -444,12 +443,12 @@ FBuildPatchCustomField::FBuildPatchCustomField(const FString& Value)
 {
 }
 
-const FString FBuildPatchCustomField::AsString() const
+FString FBuildPatchCustomField::AsString() const
 {
 	return CustomValue;
 }
 
-const double FBuildPatchCustomField::AsDouble() const
+double FBuildPatchCustomField::AsDouble() const
 {
 	// The Json parser currently only supports float so we have to decode string blob instead
 	double Rtn;
@@ -460,7 +459,7 @@ const double FBuildPatchCustomField::AsDouble() const
 	return 0;
 }
 
-const int64 FBuildPatchCustomField::AsInteger() const
+int64 FBuildPatchCustomField::AsInteger() const
 {
 	// The Json parser currently only supports float so we have to decode string blob instead
 	int64 Rtn;
@@ -1117,7 +1116,7 @@ bool FBuildPatchAppManifest::DeserializeFromJSON( const FString& JSONInput )
 	return bSuccess;
 }
 
-const EBuildPatchAppManifestVersion::Type FBuildPatchAppManifest::GetManifestVersion() const
+EBuildPatchAppManifestVersion::Type FBuildPatchAppManifest::GetManifestVersion() const
 {
 	return static_cast<EBuildPatchAppManifestVersion::Type>(Data->ManifestFileVersion);
 }
@@ -1144,17 +1143,17 @@ void FBuildPatchAppManifest::GetChunksRequiredForFiles(const TArray<FString>& Fi
 	}
 }
 
-const int64 FBuildPatchAppManifest::GetDownloadSize() const
+int64 FBuildPatchAppManifest::GetDownloadSize() const
 {
 	return TotalDownloadSize;
 }
 
-const int64 FBuildPatchAppManifest::GetBuildSize() const
+int64 FBuildPatchAppManifest::GetBuildSize() const
 {
 	return TotalBuildSize;
 }
 
-const int64 FBuildPatchAppManifest::GetFileSize(const TArray<FString>& Filenames) const
+int64 FBuildPatchAppManifest::GetFileSize(const TArray<FString>& Filenames) const
 {
 	int64 TotalSize = 0;
 	for (const auto& Filename : Filenames)
@@ -1164,7 +1163,7 @@ const int64 FBuildPatchAppManifest::GetFileSize(const TArray<FString>& Filenames
 	return TotalSize;
 }
 
-const int64 FBuildPatchAppManifest::GetFileSize(const FString& Filename) const
+int64 FBuildPatchAppManifest::GetFileSize(const FString& Filename) const
 {
 	FFileManifestData* const * FileManifest = FileManifestLookup.Find(Filename);
 	if (FileManifest)
@@ -1174,7 +1173,7 @@ const int64 FBuildPatchAppManifest::GetFileSize(const FString& Filename) const
 	return 0;
 }
 
-const int64 FBuildPatchAppManifest::GetDataSize(const FGuid& DataGuid) const
+int64 FBuildPatchAppManifest::GetDataSize(const FGuid& DataGuid) const
 {
 	if (ChunkInfoLookup.Contains(DataGuid))
 	{
@@ -1194,7 +1193,7 @@ const int64 FBuildPatchAppManifest::GetDataSize(const FGuid& DataGuid) const
 	}
 }
 
-const int64 FBuildPatchAppManifest::GetDataSize(const TArray<FGuid>& DataGuids) const
+int64 FBuildPatchAppManifest::GetDataSize(const TArray<FGuid>& DataGuids) const
 {
 	int64 TotalSize = 0;
 	for (const auto& DataGuid : DataGuids)
@@ -1204,7 +1203,7 @@ const int64 FBuildPatchAppManifest::GetDataSize(const TArray<FGuid>& DataGuids)
 	return TotalSize;
 }
 
-const uint32 FBuildPatchAppManifest::GetNumFiles() const
+uint32 FBuildPatchAppManifest::GetNumFiles() const
 {
 	return Data->FileManifestList.Num();
 }
@@ -1225,12 +1224,12 @@ const FFileManifestData* FBuildPatchAppManifest::GetFileManifest(const FString&
 	return (FileManifest) ? (*FileManifest) : nullptr;
 }
 
-const bool FBuildPatchAppManifest::IsFileDataManifest() const
+bool FBuildPatchAppManifest::IsFileDataManifest() const
 {
 	return Data->bIsFileData;
 }
 
-const bool FBuildPatchAppManifest::GetChunkHash(const FGuid& ChunkGuid, uint64& OutHash) const
+bool FBuildPatchAppManifest::GetChunkHash(const FGuid& ChunkGuid, uint64& OutHash) const
 {
 	const FChunkInfoData* const * ChunkInfo = ChunkInfoLookup.Find(ChunkGuid);
 	if (ChunkInfo)
@@ -1241,7 +1240,7 @@ const bool FBuildPatchAppManifest::GetChunkHash(const FGuid& ChunkGuid, uint64&
 	return false;
 }
 
-const bool FBuildPatchAppManifest::GetFileHash(const FGuid& FileGuid, FSHAHashData& OutHash) const
+bool FBuildPatchAppManifest::GetFileHash(const FGuid& FileGuid, FSHAHashData& OutHash) const
 {
 	const FString* const * FoundFilename = FileNameLookup.Find(FileGuid);
 	if (FoundFilename)
@@ -1251,7 +1250,7 @@ const bool FBuildPatchAppManifest::GetFileHash(const FGuid& FileGuid, FSHAHashDa
 	return false;
 }
 
-const bool FBuildPatchAppManifest::GetFileHash(const FString& Filename, FSHAHashData& OutHash) const
+bool FBuildPatchAppManifest::GetFileHash(const FString& Filename, FSHAHashData& OutHash) const
 {
 	const FFileManifestData* const * FoundFileManifest = FileManifestLookup.Find(Filename);
 	if (FoundFileManifest)
@@ -1262,7 +1261,7 @@ const bool FBuildPatchAppManifest::GetFileHash(const FString& Filename, FSHAHash
 	return false;
 }
 
-const bool FBuildPatchAppManifest::GetFilePartHash(const FGuid& FilePartGuid, uint64& OutHash) const
+bool FBuildPatchAppManifest::GetFilePartHash(const FGuid& FilePartGuid, uint64& OutHash) const
 {
 	const FChunkInfoData* const * FilePartInfo = ChunkInfoLookup.Find(FilePartGuid);
 	if (FilePartInfo)
@@ -1273,7 +1272,7 @@ const bool FBuildPatchAppManifest::GetFilePartHash(const FGuid& FilePartGuid, ui
 	return false;
 }
 
-const uint32& FBuildPatchAppManifest::GetAppID() const
+uint32 FBuildPatchAppManifest::GetAppID() const
 {
 	return Data->AppID;
 }
@@ -1463,62 +1462,6 @@ void FBuildPatchAppManifest::EnumerateProducibleChunks( const FString& InstallDi
 	}
 }
 
-bool FBuildPatchAppManifest::VerifyAgainstDirectory( const FString& VerifyDirectory, TArray < FString >& OutDatedFiles, FBuildPatchFloatDelegate ProgressDelegate, FBuildPatchBoolRetDelegate ShouldPauseDelegate, double& TimeSpentPaused )
-{
-	bool bAllCorrect = true;
-	OutDatedFiles.Empty();
-	TimeSpentPaused = 0;
-
-	// Setup progress tracking
-	double TotalBuildSizeDouble = GetBuildSize();
-	double ProcessedBytes = 0;
-	struct FLocalProgress
-	{
-		double CurrentBuildPercentage;
-		double CurrentFileWeight;
-		FBuildPatchFloatDelegate ProgressDelegate;
-		void PerFileProgress( float Progress )
-		{
-			const double PercentComplete = CurrentBuildPercentage + ( Progress * CurrentFileWeight );
-			ProgressDelegate.ExecuteIfBound( PercentComplete );
-		}
-	} LocalProgress;
-	LocalProgress.CurrentBuildPercentage = 0;
-	LocalProgress.ProgressDelegate = ProgressDelegate;
-
-	// For all files in FileManifestList, check that they produce the correct SHA1 hash, adding any that don't to the list
-	for( auto FileManifestIt = Data->FileManifestList.CreateConstIterator(); FileManifestIt && !FBuildPatchInstallError::HasFatalError(); ++FileManifestIt )
-	{
-		// Pause if necessary
-		const double PrePauseTime = FPlatformTime::Seconds();
-		double PostPauseTime = PrePauseTime;
-		bool bShouldPause = ShouldPauseDelegate.IsBound() && ShouldPauseDelegate.Execute();
-		while( bShouldPause )
-		{
-			FPlatformProcess::Sleep( 0.1f );
-			bShouldPause = ShouldPauseDelegate.Execute();
-			PostPauseTime = FPlatformTime::Seconds();
-		}
-		// Count up pause time
-		TimeSpentPaused += PostPauseTime - PrePauseTime;
-		// Get file details and construct
-		const FFileManifestData& FileManifest = *FileManifestIt;
-		FString FullFilename = VerifyDirectory / FileManifest.Filename;
-		const int64 CurrentFileSize = FileManifest.GetFileSize();
-		LocalProgress.CurrentFileWeight = CurrentFileSize / TotalBuildSizeDouble;
-		const bool bSameFile = FBuildPatchUtils::VerifyFile(FullFilename, FileManifest.FileHash, FileManifest.FileHash, FBuildPatchFloatDelegate::CreateRaw(&LocalProgress, &FLocalProgress::PerFileProgress), ShouldPauseDelegate, TimeSpentPaused) != 0;
-		ProcessedBytes += CurrentFileSize;
-		LocalProgress.CurrentBuildPercentage = ProcessedBytes / TotalBuildSizeDouble;
-		bAllCorrect &= bSameFile;
-		if( !bSameFile )
-		{
-			OutDatedFiles.Add( FileManifest.Filename );
-		}
-	}
-
-	return bAllCorrect && !FBuildPatchInstallError::HasFatalError();
-}
-
 void FBuildPatchAppManifest::EnumerateChunkPartInventory(const TArray<FGuid>& ChunksRequired, TMap<FGuid, TArray<FFileChunkPart>>& ChunkPartsAvailable) const
 {
 	ChunkPartsAvailable.Empty();
@@ -1653,7 +1596,7 @@ bool FBuildPatchAppManifest::IsFileOutdated(FBuildPatchAppManifestRef OldManifes
 	return false;
 }
 
-const uint32 FBuildPatchAppManifest::GetNumberOfChunkReferences(const FGuid& ChunkGuid) const
+uint32 FBuildPatchAppManifest::GetNumberOfChunkReferences(const FGuid& ChunkGuid) const
 {
 	uint32 RefCount = 0;
 	// For each file in the manifest, check if it has references to this chunk
diff --git a/Engine/Source/Runtime/Online/BuildPatchServices/Private/BuildPatchManifest.h b/Engine/Source/Runtime/Online/BuildPatchServices/Private/BuildPatchManifest.h
--- a/Engine/Source/Runtime/Online/BuildPatchServices/Private/BuildPatchManifest.h
+++ b/Engine/Source/Runtime/Online/BuildPatchServices/Private/BuildPatchManifest.h
@@ -10,9 +10,6 @@
 
 #pragma once
 
-// A delegate returning a bool. Used to pass a paused state
-DECLARE_DELEGATE_RetVal( bool, FBuildPatchBoolRetDelegate );
-
 typedef TSharedPtr< class FBuildPatchCustomField, ESPMode::ThreadSafe > FBuildPatchCustomFieldPtr;
 typedef TSharedRef< class FBuildPatchCustomField, ESPMode::ThreadSafe > FBuildPatchCustomFieldRef;
 typedef TSharedPtr< class FBuildPatchAppManifest, ESPMode::ThreadSafe > FBuildPatchAppManifestPtr;
@@ -61,10 +58,10 @@ namespace EBuildPatchAppManifestVersion
 namespace EBuildPatchAppManifestVersion
 {
 	/** @return The last known manifest feature of this code base. Handy for manifest constructor */
-	const Type GetLatestVersion();
+	Type GetLatestVersion();
 
 	/** @return The latest version of a manifest support by JSON serialization */
-	const Type GetLatestJsonVersion();
+	Type GetLatestJsonVersion();
 
 	/**
 	 * Get the chunk subdirectory for used for a specific manifest version, e.g. Chunks, ChunksV2 etc
@@ -193,7 +190,7 @@ struct FFileManifestData
 
 	void Init();
 
-	const int64& GetFileSize() const;
+	int64 GetFileSize() const;
 
 	FORCEINLINE bool operator<(const FFileManifestData& Other) const;
 
@@ -293,9 +290,9 @@ public:
 	FBuildPatchCustomField(const FString& Value);
 
 	// START IBuildManifest Interface
-	virtual const FString AsString() const override;
-	virtual const double AsDouble() const override;
-	virtual const int64 AsInteger() const override;
+	virtual FString AsString() const override;
+	virtual double AsDouble() const override;
+	virtual int64 AsInteger() const override;
 	// END IBuildManifest Interface
 };
 
@@ -355,7 +352,7 @@ public:
 	~FBuildPatchAppManifest();
 
 	// START IBuildManifest Interface
-	virtual const uint32&  GetAppID() const override;
+	virtual uint32  GetAppID() const override;
 	virtual const FString& GetAppName() const override;
 	virtual const FString& GetVersionString() const override;
 	virtual const FString& GetLaunchExe() const override;
@@ -363,8 +360,8 @@ public:
 	virtual const FString& GetPrereqName() const override;
 	virtual const FString& GetPrereqPath() const override;
 	virtual const FString& GetPrereqArgs() const override;
-	virtual const int64 GetDownloadSize() const override;
-	virtual const int64 GetBuildSize() const override;
+	virtual int64 GetDownloadSize() const override;
+	virtual int64 GetBuildSize() const override;
 	virtual void GetRemovableFiles(IBuildManifestRef OldManifest, TArray< FString >& RemovableFiles) const override;
 	virtual void GetRemovableFiles(const TCHAR* InstallPath, TArray< FString >& RemovableFiles) const override;
 	virtual bool NeedsResaving() const override;
@@ -415,7 +412,7 @@ public:
 	 * Gets the version for this manifest. Useful for manifests that were loaded from JSON.
 	 * @return		The highest available feature support
 	 */
-	const EBuildPatchAppManifestVersion::Type GetManifestVersion() const;
+	EBuildPatchAppManifestVersion::Type GetManifestVersion() const;
 
 	/**
 	 * Provides a list of chunks required to produce the list of given files.
@@ -429,21 +426,21 @@ public:
 	 * @param ChunkGuid		The chunk GUID
 	 * @return	The number of references to this chunk
 	 */
-	const uint32 GetNumberOfChunkReferences(const FGuid& ChunkGuid) const;
+	uint32 GetNumberOfChunkReferences(const FGuid& ChunkGuid) const;
 
 	/**
 	 * Returns the size of a particular data file by it's GUID.
 	 * @param DataGuid		The GUID for the data
 	 * @return		File size.
 	 */
-	const int64 GetDataSize(const FGuid& DataGuid) const;
+	int64 GetDataSize(const FGuid& DataGuid) const;
 
 	/**
 	 * Returns the total size of all data files in it's list.
 	 * @param DataGuids		The GUID array for the data
 	 * @return		File size.
 	 */
-	const int64 GetDataSize(const TArray< FGuid >& DataGuids) const;
+	int64 GetDataSize(const TArray< FGuid >& DataGuids) const;
 
 	/**
 	 * Returns the size of a particular file in the build
@@ -451,7 +448,7 @@ public:
 	 * @param Filename		The file.
 	 * @return		File size.
 	 */
-	const int64 GetFileSize(const FString& Filename) const;
+	int64 GetFileSize(const FString& Filename) const;
 
 	/**
 	 * Returns the total size of all files in the array
@@ -459,13 +456,13 @@ public:
 	 * @param Filenames		The array of files.
 	 * @return		Total size of files in array.
 	 */
-	const int64 GetFileSize(const TArray< FString >& Filenames) const;
+	int64 GetFileSize(const TArray< FString >& Filenames) const;
 
 	/**
 	 * Returns the number of files in this build.
 	 * @return		The number of files.
 	 */
-	const uint32 GetNumFiles() const;
+	uint32 GetNumFiles() const;
 
 	/**
 	 * Get the list of files described by this manifest
@@ -490,7 +487,7 @@ public:
 	 * Gets whether this manifest is made up of file data instead of chunk data
 	 * @return	True if the build is made from file data. False if the build is constructed from chunk data.
 	 */
-	const bool IsFileDataManifest() const;
+	bool IsFileDataManifest() const;
 
 	/**
 	 * Gets the chunk hash for a given chunk
@@ -498,7 +495,7 @@ public:
 	 * @param OutHash		OUT		Receives the hash value if found
 	 * @return	true if we had the hash for this chunk
 	 */
-	const bool GetChunkHash(const FGuid& ChunkGuid, uint64& OutHash) const;
+	bool GetChunkHash(const FGuid& ChunkGuid, uint64& OutHash) const;
 
 	/**
 	 * Gets the file hash for given file data
@@ -506,7 +503,7 @@ public:
 	 * @param OutHash		OUT		Receives the hash value if found
 	 * @return	true if we had the hash for this file
 	 */
-	const bool GetFileHash(const FGuid& FileGuid, FSHAHashData& OutHash) const; // DEPRECATE ME
+	bool GetFileHash(const FGuid& FileGuid, FSHAHashData& OutHash) const; // DEPRECATE ME
 
 	/**
 	 * Gets the file hash for a given file
@@ -514,7 +511,7 @@ public:
 	 * @param OutHash		OUT		Receives the hash value if found
 	 * @return	true if we had the hash for this file
 	 */
-	const bool GetFileHash(const FString& Filename, FSHAHashData& OutHash) const;
+	bool GetFileHash(const FString& Filename, FSHAHashData& OutHash) const;
 
 	/**
 	 * Gets the file hash for given file data. Valid for non-chunked manifest
@@ -522,7 +519,7 @@ public:
 	 * @param OutHash		OUT		Receives the hash value if found
 	 * @return	true if we had the hash for this file
 	 */
-	const bool GetFilePartHash(const FGuid& FilePartGuid, uint64& OutHash) const;
+	bool GetFilePartHash(const FGuid& FilePartGuid, uint64& OutHash) const;
 
 	/**
 	 * Populates an array of chunks that should be producible from this local build, given the list of chunks needed. Also checks source files exist and match size.
@@ -532,18 +529,6 @@ public:
 	 */
 	void EnumerateProducibleChunks(const FString& InstallDirectory, const TArray<FGuid>& ChunksRequired, TArray<FGuid>& ChunksAvailable) const;
 
-	/**
-	 * Verifies a local directory structure against a given manifest.
-	 * NOTE: This function is blocking and will not return until finished. Don't run on main thread.
-	 * @param VerifyDirectory		IN		The directory to analyze.
-	 * @param OutDatedFiles			OUT		The array of files that do not match or are locally missing.
-	 * @param ProgressDelegate		IN		Delegate to receive progress updates in the form of a float range 0.0f to 1.0f
-	 * @param ShouldPauseDelegate	IN		Delegate that returns a bool, which if true will pause the process
-	 * @param TimeSpentPaused		OUT		The amount of time we were paused for, in seconds.
-	 * @return		true if no file errors occurred AND the verification was successful
-	 */
-	bool VerifyAgainstDirectory(const FString& VerifyDirectory, TArray < FString >& OutDatedFiles, FBuildPatchFloatDelegate ProgressDelegate, FBuildPatchBoolRetDelegate ShouldPauseDelegate, double& TimeSpentPaused);
-
 	/**
 	 * Gets a list of files that were installed with the Old Manifest, but no longer required by the New Manifest.
 	 * @param OldManifest		IN		The Build Manifest that is currently installed
diff --git a/Engine/Source/Runtime/Online/BuildPatchServices/Private/BuildPatchServicesModule.cpp b/Engine/Source/Runtime/Online/BuildPatchServices/Private/BuildPatchServicesModule.cpp
--- a/Engine/Source/Runtime/Online/BuildPatchServices/Private/BuildPatchServicesModule.cpp
+++ b/Engine/Source/Runtime/Online/BuildPatchServices/Private/BuildPatchServicesModule.cpp
@@ -144,7 +144,27 @@ IBuildInstallerPtr FBuildPatchServicesModule::StartBuildInstall( IBuildManifestP
 	// Make sure the http wrapper is already created
 	FBuildPatchHTTP::Initialize();
 	// Run the install thread
-	BuildPatchInstallers.Add( MakeShareable( new FBuildPatchInstaller( OnCompleteDelegate, CurrentManifestInternal, InstallManifestInternal.ToSharedRef(), InstallDirectory, GetStagingDirectory(), InstallationInfo ) ) );
+	BuildPatchInstallers.Add( MakeShareable( new FBuildPatchInstaller( OnCompleteDelegate, CurrentManifestInternal, InstallManifestInternal.ToSharedRef(), InstallDirectory, GetStagingDirectory(), InstallationInfo, false ) ) );
+	return BuildPatchInstallers.Top();
+}
+
+IBuildInstallerPtr FBuildPatchServicesModule::StartBuildInstallStageOnly(IBuildManifestPtr CurrentManifest, IBuildManifestPtr InstallManifest, const FString& InstallDirectory, FBuildPatchBoolManifestDelegate OnCompleteDelegate)
+{
+	// Using a local bool for this check will improve the assert message that gets displayed
+	const bool bIsCalledFromMainThread = IsInGameThread();
+	check( bIsCalledFromMainThread );
+	// Cast manifest parameters
+	FBuildPatchAppManifestPtr CurrentManifestInternal = StaticCastSharedPtr< FBuildPatchAppManifest >( CurrentManifest );
+	FBuildPatchAppManifestPtr InstallManifestInternal = StaticCastSharedPtr< FBuildPatchAppManifest >( InstallManifest );
+	if( !InstallManifestInternal.IsValid() )
+	{
+		// We must have an install manifest to continue
+		return NULL;
+	}
+	// Make sure the http wrapper is already created
+	FBuildPatchHTTP::Initialize();
+	// Run the install thread
+	BuildPatchInstallers.Add( MakeShareable( new FBuildPatchInstaller( OnCompleteDelegate, CurrentManifestInternal, InstallManifestInternal.ToSharedRef(), InstallDirectory, GetStagingDirectory(), InstallationInfo, true ) ) );
 	return BuildPatchInstallers.Top();
 }
 
diff --git a/Engine/Source/Runtime/Online/BuildPatchServices/Private/BuildPatchServicesModule.h b/Engine/Source/Runtime/Online/BuildPatchServices/Private/BuildPatchServicesModule.h
--- a/Engine/Source/Runtime/Online/BuildPatchServices/Private/BuildPatchServicesModule.h
+++ b/Engine/Source/Runtime/Online/BuildPatchServices/Private/BuildPatchServicesModule.h
@@ -75,7 +75,8 @@ public:
 	virtual IBuildManifestPtr LoadManifestFromFile( const FString& Filename ) override;
 	virtual IBuildManifestPtr MakeManifestFromData( const TArray<uint8>& ManifestData ) override;
 	virtual bool SaveManifestToFile( const FString& Filename, IBuildManifestRef Manifest, bool bUseBinary = true ) override;
-	virtual IBuildInstallerPtr StartBuildInstall( IBuildManifestPtr CurrentManifest, IBuildManifestPtr InstallManifest, const FString& InstallDirectory, FBuildPatchBoolManifestDelegate OnCompleteDelegate ) override;
+	virtual IBuildInstallerPtr StartBuildInstall(IBuildManifestPtr CurrentManifest, IBuildManifestPtr InstallManifest, const FString& InstallDirectory, FBuildPatchBoolManifestDelegate OnCompleteDelegate) override;
+	virtual IBuildInstallerPtr StartBuildInstallStageOnly(IBuildManifestPtr CurrentManifest, IBuildManifestPtr InstallManifest, const FString& InstallDirectory, FBuildPatchBoolManifestDelegate OnCompleteDelegate) override;
 	virtual void SetStagingDirectory( const FString& StagingDir ) override;
 	virtual void SetCloudDirectory( const FString& CloudDir ) override;
 	virtual void SetBackupDirectory( const FString& BackupDir ) override;
diff --git a/Engine/Source/Runtime/Online/BuildPatchServices/Private/BuildPatchServicesPrivatePCH.h b/Engine/Source/Runtime/Online/BuildPatchServices/Private/BuildPatchServicesPrivatePCH.h
--- a/Engine/Source/Runtime/Online/BuildPatchServices/Private/BuildPatchServicesPrivatePCH.h
+++ b/Engine/Source/Runtime/Online/BuildPatchServices/Private/BuildPatchServicesPrivatePCH.h
@@ -40,6 +40,8 @@
 #include "BuildPatchHTTP.h"
 #include "BuildPatchChunk.h"
 #include "BuildPatchManifest.h"
+#include "BuildPatchVerification.h"
+#include "BuildPatchFileAttributes.h"
 #include "BuildPatchChunkCache.h"
 #include "BuildPatchCompactifier.h"
 #include "BuildPatchDataEnumeration.h"
diff --git a/Engine/Source/Runtime/Online/BuildPatchServices/Private/BuildPatchVerification.cpp b/Engine/Source/Runtime/Online/BuildPatchServices/Private/BuildPatchVerification.cpp
new file mode 100644
--- /dev/null
+++ b/Engine/Source/Runtime/Online/BuildPatchServices/Private/BuildPatchVerification.cpp
@@ -0,0 +1,100 @@
+// Copyright 1998-2015 Epic Games, Inc. All Rights Reserved.
+
+#include "BuildPatchServicesPrivatePCH.h"
+
+class FBuildPatchVerificationImpl : public FBuildPatchVerification
+{
+public:
+	FBuildPatchVerificationImpl(const FBuildPatchAppManifestRef& Manifest, const FBuildPatchFloatDelegate& ProgressDelegate, const FBuildPatchBoolRetDelegate& ShouldPauseDelegate, const FString& VerifyDirectory, const FString& StagedFileDirectory, bool bUseStageDirectory);
+
+	virtual ~FBuildPatchVerificationImpl(){}
+
+	virtual bool VerifyAgainstDirectory(TArray<FString>& OutDatedFiles, double& TimeSpentPaused) override;
+
+private:
+	FBuildPatchAppManifestRef Manifest;
+	FBuildPatchFloatDelegate ProgressDelegate;
+	FBuildPatchBoolRetDelegate ShouldPauseDelegate;
+	const FString VerifyDirectory;
+	const FString StagedFileDirectory;
+	const bool bUseStageDirectory;
+	double CurrentBuildPercentage;
+	double CurrentFileWeight;
+
+	void PerFileProgress(float Progress);
+	FString SelectFullFilePath(const FString& BuildFile);
+};
+
+FBuildPatchVerificationImpl::FBuildPatchVerificationImpl(const FBuildPatchAppManifestRef& Manifest, const FBuildPatchFloatDelegate& ProgressDelegate, const FBuildPatchBoolRetDelegate& ShouldPauseDelegate, const FString& VerifyDirectory, const FString& StagedFileDirectory, bool bUseStageDirectory)
+	: Manifest(Manifest)
+	, ProgressDelegate(ProgressDelegate)
+	, ShouldPauseDelegate(ShouldPauseDelegate)
+	, VerifyDirectory(VerifyDirectory)
+	, StagedFileDirectory(StagedFileDirectory)
+	, bUseStageDirectory(bUseStageDirectory)
+	, CurrentBuildPercentage(0)
+	, CurrentFileWeight(0)
+{}
+
+bool FBuildPatchVerificationImpl::VerifyAgainstDirectory(TArray<FString>& OutDatedFiles, double& TimeSpentPaused)
+{
+	bool bAllCorrect = true;
+	OutDatedFiles.Empty();
+	TimeSpentPaused = 0;
+
+	// Setup progress tracking
+	double TotalBuildSizeDouble = Manifest->GetBuildSize();
+	double ProcessedBytes = 0;
+	CurrentBuildPercentage = 0;
+
+	// For all files in the manifest, check that they produce the correct SHA1 hash, adding any that don't to the list
+	TArray<FString> BuildFiles;
+	Manifest->GetFileList(BuildFiles);
+	for (const FString& BuildFile : BuildFiles)
+	{
+		// Get file details
+		int64 BuildFileSize = Manifest->GetFileSize(BuildFile);
+		FSHAHashData BuildFileHash;
+		bool bFoundHash = Manifest->GetFileHash(BuildFile, BuildFileHash);
+		check(bFoundHash);
+
+		// Chose the file to check
+		FString FullFilename = SelectFullFilePath(BuildFile);
+
+		// Verify the file
+		CurrentFileWeight = BuildFileSize / TotalBuildSizeDouble;
+		if (FBuildPatchUtils::VerifyFile(FullFilename, BuildFileHash, BuildFileHash, FBuildPatchFloatDelegate::CreateRaw(this, &FBuildPatchVerificationImpl::PerFileProgress), ShouldPauseDelegate, TimeSpentPaused) == 0)
+		{
+			bAllCorrect = false;
+			OutDatedFiles.Add(BuildFile);
+		}
+		ProcessedBytes += BuildFileSize;
+		CurrentBuildPercentage = ProcessedBytes / TotalBuildSizeDouble;
+	}
+
+	return bAllCorrect && !FBuildPatchInstallError::HasFatalError();
+}
+
+void FBuildPatchVerificationImpl::PerFileProgress(float Progress)
+{
+	ProgressDelegate.ExecuteIfBound(CurrentBuildPercentage + (Progress * CurrentFileWeight));
+}
+
+FString FBuildPatchVerificationImpl::SelectFullFilePath(const FString& BuildFile)
+{
+	if (bUseStageDirectory)
+	{
+		FString StagedFilename = StagedFileDirectory / BuildFile;
+		if (IFileManager::Get().FileSize(*StagedFilename) != -1)
+		{
+			return MoveTemp(StagedFilename);
+		}
+	}
+	FString InstallFilename = VerifyDirectory / BuildFile;
+	return MoveTemp(InstallFilename);
+}
+
+TSharedRef<FBuildPatchVerification> FBuildPatchVerificationFactory::Create(const FBuildPatchAppManifestRef& Manifest, const FBuildPatchFloatDelegate& ProgressDelegate, const FBuildPatchBoolRetDelegate& ShouldPauseDelegate, const FString& VerifyDirectory, const FString& StagedFileDirectory)
+{
+	return MakeShareable(new FBuildPatchVerificationImpl(Manifest, ProgressDelegate, ShouldPauseDelegate, VerifyDirectory, StagedFileDirectory, !StagedFileDirectory.IsEmpty()));
+}
diff --git a/Engine/Source/Runtime/Online/BuildPatchServices/Private/BuildPatchVerification.h b/Engine/Source/Runtime/Online/BuildPatchServices/Private/BuildPatchVerification.h
new file mode 100644
--- /dev/null
+++ b/Engine/Source/Runtime/Online/BuildPatchServices/Private/BuildPatchVerification.h
@@ -0,0 +1,41 @@
+// Copyright 1998-2015 Epic Games, Inc. All Rights Reserved.
+
+#pragma once
+
+// A delegate returning a bool. Used to pass a paused state
+DECLARE_DELEGATE_RetVal(bool, FBuildPatchBoolRetDelegate);
+
+class FBuildPatchVerification
+{
+public:
+	/**
+	 * Verifies a local directory structure against a given manifest.
+	 * NOTE: This function is blocking and will not return until finished. Don't run on main thread.
+	 * @param OutDatedFiles    OUT  The array of files that do not match or are locally missing.
+	 * @param TimeSpentPaused  OUT  The amount of time we were paused for, in seconds.
+	 * @return    true if no file errors occurred AND the verification was successful
+	 */
+	virtual bool VerifyAgainstDirectory(TArray<FString>& OutDatedFiles, double& TimeSpentPaused) = 0;
+};
+
+class FBuildPatchVerificationFactory
+{
+public:
+	/**
+	 * Creates a verification class that will verify a local directory structure against a given manifest, optionally 
+	 * taking account of a staging directory where alternative files are used.
+	 * NOTE: This function is blocking and will not return until finished. Don't run on main thread.
+	 * @param Manifest              The manifest describing the build data.
+	 * @param ProgressDelegate      Delegate to receive progress updates in the form of a float range 0.0f to 1.0f
+	 * @param ShouldPauseDelegate   Delegate that returns a bool, which if true will pause the process
+	 * @param VerifyDirectory       The directory to analyze.
+	 * @param StagedFileDirectory   A stage directory for updated files, ignored if empty string. If a file exists here, it will be
+	 *                              checked instead of the one in VerifyDirectory.
+	 * @return     Ref of an object that can be used to perform the operation.
+	 */
+	static TSharedRef<FBuildPatchVerification> Create(const FBuildPatchAppManifestRef& Manifest,
+	                                                  const FBuildPatchFloatDelegate& ProgressDelegate,
+	                                                  const FBuildPatchBoolRetDelegate& ShouldPauseDelegate,
+	                                                  const FString& VerifyDirectory,
+	                                                  const FString& StagedFileDirectory);
+};
diff --git a/Engine/Source/Runtime/Online/BuildPatchServices/Public/Interfaces/IBuildManifest.h b/Engine/Source/Runtime/Online/BuildPatchServices/Public/Interfaces/IBuildManifest.h
--- a/Engine/Source/Runtime/Online/BuildPatchServices/Public/Interfaces/IBuildManifest.h
+++ b/Engine/Source/Runtime/Online/BuildPatchServices/Public/Interfaces/IBuildManifest.h
@@ -6,9 +6,6 @@
 	App, and also the interface to creating and reading custom fields.
 =============================================================================*/
 
-#ifndef __IBuildManifest_h__
-#define __IBuildManifest_h__
-
 #pragma once
 
 typedef TSharedPtr< class IManifestField, ESPMode::ThreadSafe > IManifestFieldPtr;
@@ -31,19 +28,19 @@ public:
 	 * Get the fields value as an FString.
 	 * @return	The string value.
 	 */
-	virtual const FString AsString() const = 0;
+	virtual FString AsString() const = 0;
 	
 	/**
 	 * Get the fields value as a double.
 	 * @return	The double value.
 	 */
-	virtual const double AsDouble() const = 0;
+	virtual double AsDouble() const = 0;
 	
 	/**
 	 * Get the fields value as an int.
 	 * @return	The int value.
 	 */
-	virtual const int64 AsInteger() const = 0;
+	virtual int64 AsInteger() const = 0;
 };
 
 /**
@@ -61,7 +58,7 @@ public:
 	 * Get the App ID that this manifest belongs to
 	 * @return		the app ID
 	 */
-	virtual const uint32& GetAppID() const = 0;
+	virtual uint32 GetAppID() const = 0;
 	
 	/**
 	 * Get the name of the App that this manifest belongs to
@@ -109,13 +106,13 @@ public:
 	 * Get the size of this download, assuming fresh install
 	 * @return		the total download size in bytes
 	 */
-	virtual const int64 GetDownloadSize() const = 0;
+	virtual int64 GetDownloadSize() const = 0;
 
 	/**
 	 * Get the size of this build
 	 * @return		the build size in bytes
 	 */
-	virtual const int64 GetBuildSize() const = 0;
+	virtual int64 GetBuildSize() const = 0;
 
 	/**
 	 * Gets a list of files that were installed with the Old Manifest, but no longer required by this Manifest.
@@ -168,5 +165,3 @@ public:
 	 */
 	virtual IBuildManifestRef Duplicate() const = 0;
 };
-
-#endif // __IBuildManifest_h__
diff --git a/Engine/Source/Runtime/Online/BuildPatchServices/Public/Interfaces/IBuildPatchServicesModule.h b/Engine/Source/Runtime/Online/BuildPatchServices/Public/Interfaces/IBuildPatchServicesModule.h
--- a/Engine/Source/Runtime/Online/BuildPatchServices/Public/Interfaces/IBuildPatchServicesModule.h
+++ b/Engine/Source/Runtime/Online/BuildPatchServices/Public/Interfaces/IBuildPatchServicesModule.h
@@ -72,6 +72,17 @@ public:
 	 */
 	virtual IBuildInstallerPtr StartBuildInstall( IBuildManifestPtr CurrentManifest, IBuildManifestPtr InstallManifest, const FString& InstallDirectory, FBuildPatchBoolManifestDelegate OnCompleteDelegate ) = 0;
 
+	/**
+	 * Starts an installer thread for the provided manifests, only producing the necessary stage. Useful for handling specific install directory write access requirements yourself.
+	 * The staged files will be in the provided staging directory as StagingDirectory/Install/
+	 * @param	CurrentManifest			The manifest that the current install was generated from (if applicable)
+	 * @param	InstallManifest			The manifest to be installed
+	 * @param	InstallDirectory		The directory to install the App to - this should still be the real install directory. It may be read from for patching.
+	 * @param	OnCompleteDelegate		The delegate to call on completion
+	 * @return		An interface to the created installer. Will be an invalid ptr if error.
+	 */
+	virtual IBuildInstallerPtr StartBuildInstallStageOnly( IBuildManifestPtr CurrentManifest, IBuildManifestPtr InstallManifest, const FString& InstallDirectory, FBuildPatchBoolManifestDelegate OnCompleteDelegate ) = 0;
+
 	/**
 	 * Sets the directory used for staging intermediate files.
 	 * @param StagingDir	The staging directory
