Implemented ability to configure changes to a room independent of room creation
Added xmpp command to refresh room info from server
Merging //depot/UE4-LauncherDev/Engine/Source/Runtime/Online/... to //depot/UE4/Engine/Source/Runtime/Online/...

diff --git a/Engine/Source/Runtime/Online/OnlineSubsystem/Public/Interfaces/OnlineChatInterface.h b/Engine/Source/Runtime/Online/OnlineSubsystem/Public/Interfaces/OnlineChatInterface.h
--- a/Engine/Source/Runtime/Online/OnlineSubsystem/Public/Interfaces/OnlineChatInterface.h
+++ b/Engine/Source/Runtime/Online/OnlineSubsystem/Public/Interfaces/OnlineChatInterface.h
@@ -160,12 +160,23 @@ public:
 	 * 
 	 * @param UserId id of user that is creating the room
 	 * @param RoomId name of room to create
-	 * @param RoomConfig configuration for the room
 	 * @param Nickname display name for the chat room. Name must be unique and is reserved for duration of join
+	 * @param RoomConfig configuration for the room
 	 *
 	 * @return if successfully started the async operation
 	 */
 	virtual bool CreateRoom(const FUniqueNetId& UserId, const FChatRoomId& RoomId, const FString& Nickname, const FChatRoomConfig& ChatRoomConfig) = 0;
+
+	/**
+	* Kick off request for configuring a chat room with a provided configuration
+	*
+	* @param UserId id of user that is creating the room
+	* @param RoomId name of room to create
+	* @param RoomConfig configuration for the room
+	*
+	* @return if successfully started the async operation
+	*/
+	virtual bool ConfigureRoom(const FUniqueNetId& UserId, const FChatRoomId& RoomId, const FChatRoomConfig& ChatRoomConfig) = 0;
 	
 	/**
 	 * Kick off request for joining a public chat room
