Corrected the facebook include path usage for IOS

diff --git a/Engine/Source/Runtime/Online/OnlineSubsystemFacebook/Private/IOS/OnlineFriendsFacebook.cpp b/Engine/Source/Runtime/Online/OnlineSubsystemFacebook/Private/IOS/OnlineFriendsFacebook.cpp
--- a/Engine/Source/Runtime/Online/OnlineSubsystemFacebook/Private/IOS/OnlineFriendsFacebook.cpp
+++ b/Engine/Source/Runtime/Online/OnlineSubsystemFacebook/Private/IOS/OnlineFriendsFacebook.cpp
@@ -6,6 +6,8 @@
 #include "OnlineSharingFacebook.h"
 #include "OnlineFriendsFacebook.h"
 
+#import <FBSDKCoreKit/FBSDKCoreKit.h>
+
 // FOnlineFriendFacebook
 
 TSharedRef<FUniqueNetId> FOnlineFriendFacebook::GetUserId() const
diff --git a/Engine/Source/Runtime/Online/OnlineSubsystemFacebook/Private/IOS/OnlineIdentityFacebook.cpp b/Engine/Source/Runtime/Online/OnlineSubsystemFacebook/Private/IOS/OnlineIdentityFacebook.cpp
--- a/Engine/Source/Runtime/Online/OnlineSubsystemFacebook/Private/IOS/OnlineIdentityFacebook.cpp
+++ b/Engine/Source/Runtime/Online/OnlineSubsystemFacebook/Private/IOS/OnlineIdentityFacebook.cpp
@@ -5,6 +5,9 @@
 #include "OnlineSubsystemFacebookPrivatePCH.h"
 #include "OnlineIdentityFacebook.h"
 
+#import <FBSDKCoreKit/FBSDKCoreKit.h>
+#import <FBSDKLoginKit/FBSDKLoginKit.h>
+
 // Other UE4 includes
 #include "IOSAppDelegate.h"
 
diff --git a/Engine/Source/Runtime/Online/OnlineSubsystemFacebook/Private/IOS/OnlineSharingFacebook.cpp b/Engine/Source/Runtime/Online/OnlineSubsystemFacebook/Private/IOS/OnlineSharingFacebook.cpp
--- a/Engine/Source/Runtime/Online/OnlineSubsystemFacebook/Private/IOS/OnlineSharingFacebook.cpp
+++ b/Engine/Source/Runtime/Online/OnlineSubsystemFacebook/Private/IOS/OnlineSharingFacebook.cpp
@@ -6,6 +6,7 @@
 #include "OnlineSharingFacebook.h"
 #include "ImageCore.h"
 
+#import <FBSDKCoreKit/FBSDKCoreKit.h>
 #import <FBSDKShareKit/FBSDKShareKit.h>
 
 FOnlineSharingFacebook::FOnlineSharingFacebook(class FOnlineSubsystemFacebook* InSubsystem)
diff --git a/Engine/Source/Runtime/Online/OnlineSubsystemFacebook/Private/OnlineSubsystemFacebookPrivatePCH.h b/Engine/Source/Runtime/Online/OnlineSubsystemFacebook/Private/OnlineSubsystemFacebookPrivatePCH.h
--- a/Engine/Source/Runtime/Online/OnlineSubsystemFacebook/Private/OnlineSubsystemFacebookPrivatePCH.h
+++ b/Engine/Source/Runtime/Online/OnlineSubsystemFacebook/Private/OnlineSubsystemFacebookPrivatePCH.h
@@ -14,42 +14,6 @@
 #include "OnlineIdentityFacebook.h"
 #include "OnlineFriendsFacebook.h"
 
-
-// Begin contents of Facebook core kit
-#import <FBSDKCoreKit/FBSDKCoreKit.h>
-/*
-#include "FBSDKAccessToken.h"
-#include "FBSDKAppEvents.h"
-#include "FBSDKAppLinkUtility.h"
-#include "FBSDKApplicationDelegate.h"
-#include "FBSDKConstants.h"
-#include "FBSDKCopying.h"
-#include "FBSDKGraphRequest.h"
-#include "FBSDKGraphRequestConnection.h"
-#include "FBSDKMacros.h"
-#include "FBSDKMutableCopying.h"
-#include "FBSDKProfile.h"
-#include "FBSDKProfilePictureView.h"
-#include "FBSDKSettings.h"
-#include "FBSDKUtility.h"
-
-#define FBSDK_VERSION_STRING @"4.0.1"
-#define FBSDK_TARGET_PLATFORM_VERSION @"v2.3"
-*/
-// End contents of Facebook core kit
-
-
-// Begin contents of Facebook login kit
-#import <FBSDKLoginKit/FBSDKLoginKit.h>
-/*
-#include "FBSDKLoginButton.h"
-#include "FBSDKLoginConstants.h"
-#include "FBSDKLoginManager.h"
-#include "FBSDKLoginManagerLoginResult.h"
-#include "FBSDKLoginTooltipView.h"
-*/
-// End contents of Facebook login kit
-
 /** FName declaration of Facebook subsystem */
 #define FACEBOOK_SUBSYSTEM FName(TEXT("Facebook"))
 
