Merging using UE4-MS_to_UE4, CLs 2506150 and 2520191:
Lobby/Party based matchmaking and some bug fixes

diff --git a/Engine/Source/Runtime/Online/OnlineSubsystem/Public/OnlineSubsystemSessionSettings.h b/Engine/Source/Runtime/Online/OnlineSubsystem/Public/OnlineSubsystemSessionSettings.h
--- a/Engine/Source/Runtime/Online/OnlineSubsystem/Public/OnlineSubsystemSessionSettings.h
+++ b/Engine/Source/Runtime/Online/OnlineSubsystem/Public/OnlineSubsystemSessionSettings.h
@@ -37,3 +37,10 @@
 #define SETTING_CHANGE_NUMBER FName(TEXT("CHANGENUMBER"))
 /** Session can attempt arbiter migration */
 #define SETTING_ALLOW_ARBITER_MIGRATION FName(TEXT("ALLOWARBITERMIGRATION"))
+/** Session preservation type for matchmaking */
+#define SETTING_MATCHING_PRESERVE_SESSION FName(TEXT("PRESERVESESSIONALWAYS"))
+/** Matchmade game session URI for join in progress/invites */
+#define  SETTING_GAME_SESSION_URI FName(TEXT("GAMESESSIONURI"))
+/** Session member group identifier (value is FString) - This is to be used with FString::Printf, populated with user xuid - Field required with Team Based matchmaking*/
+#define SETTING_GROUP_NAME FName(TEXT("USERGROUPNAME%s"))
+
