Fix precision loss when using the OnlineJsonSerializer with large int32 and uint32 values

diff --git a/Engine/Source/Runtime/Online/OnlineSubsystem/Public/OnlineJsonSerializer.h b/Engine/Source/Runtime/Online/OnlineSubsystem/Public/OnlineJsonSerializer.h
--- a/Engine/Source/Runtime/Online/OnlineSubsystem/Public/OnlineJsonSerializer.h
+++ b/Engine/Source/Runtime/Online/OnlineSubsystem/Public/OnlineJsonSerializer.h
@@ -214,7 +214,7 @@ public:
 	 */
 	virtual void Serialize(const TCHAR* Name, int32& Value) override
 	{
-		JsonWriter->WriteValue(Name, (const float)Value);
+		JsonWriter->WriteValue(Name, Value);
 	}
 	/**
 	 * Writes the field name and the corresponding value to the JSON data
@@ -224,7 +224,7 @@ public:
 	 */
 	virtual void Serialize(const TCHAR* Name, uint32& Value) override
 	{
-		JsonWriter->WriteValue(Name, (const float)Value);
+		JsonWriter->WriteValue(Name, static_cast<int64>(Value));
 	}
 	/**
 	 * Writes the field name and the corresponding value to the JSON data
@@ -234,7 +234,7 @@ public:
 	 */
 	virtual void Serialize(const TCHAR* Name, bool& Value) override
 	{
-		JsonWriter->WriteValue(Name, (const bool)Value);
+		JsonWriter->WriteValue(Name, Value);
 	}
 	/**
 	 * Writes the field name and the corresponding value to the JSON data
@@ -264,7 +264,7 @@ public:
 	 */
 	virtual void Serialize(const TCHAR* Name, float& Value) override
 	{
-		JsonWriter->WriteValue(Name, (const float)Value);
+		JsonWriter->WriteValue(Name, Value);
 	}
 	/**
 	* Writes the field name and the corresponding value to the JSON data
@@ -274,7 +274,7 @@ public:
 	*/
 	virtual void Serialize(const TCHAR* Name, double& Value) override
 	{
-		JsonWriter->WriteValue(Name, (const double)Value);
+		JsonWriter->WriteValue(Name, Value);
 	}
 	/**
 	* Writes the field name and the corresponding value to the JSON data
@@ -422,7 +422,7 @@ public:
 	{
 		if (JsonObject->HasTypedField<EJson::Number>(Name))
 		{
-			Value = FMath::TruncToInt(JsonObject->GetNumberField(Name));
+			JsonObject->TryGetNumberField(Name, Value);
 		}
 	}
 	/**
@@ -435,7 +435,7 @@ public:
 	{
 		if (JsonObject->HasTypedField<EJson::Number>(Name))
 		{
-			Value = FMath::TruncToInt(JsonObject->GetNumberField(Name));
+			JsonObject->TryGetNumberField(Name, Value);
 		}
 	}
 	/**
