Fix shadow variables in BuildPatchServices

diff --git a/Engine/Source/Programs/BuildPatchTool/Private/BuildPatchToolMain.cpp b/Engine/Source/Programs/BuildPatchTool/Private/BuildPatchToolMain.cpp
--- a/Engine/Source/Programs/BuildPatchTool/Private/BuildPatchToolMain.cpp
+++ b/Engine/Source/Programs/BuildPatchTool/Private/BuildPatchToolMain.cpp
@@ -98,8 +98,8 @@ struct FCommandLineMatcher
 {
 	const FString Command;
 
-	FCommandLineMatcher(FString Command)
-		: Command(MoveTemp(Command))
+	FCommandLineMatcher(FString InCommand)
+		: Command(MoveTemp(InCommand))
 	{}
 
 	FORCEINLINE bool operator()(const FString& ToMatch) const
diff --git a/Engine/Source/Runtime/Online/BuildPatchServices/Private/BuildPatchCompactifier.cpp b/Engine/Source/Runtime/Online/BuildPatchServices/Private/BuildPatchCompactifier.cpp
--- a/Engine/Source/Runtime/Online/BuildPatchServices/Private/BuildPatchCompactifier.cpp
+++ b/Engine/Source/Runtime/Online/BuildPatchServices/Private/BuildPatchCompactifier.cpp
@@ -18,10 +18,10 @@ namespace BuildDataCompactifierDefs
 /* Constructors
 *****************************************************************************/
 
-FBuildDataCompactifier::FBuildDataCompactifier(const FString& CloudDir, const bool bPreview, const bool bNoPatchDelete)
-	: CloudDir(CloudDir)
-	, bPreview(bPreview)
-	, bNoPatchDelete(bNoPatchDelete)
+FBuildDataCompactifier::FBuildDataCompactifier(const FString& InCloudDir, const bool bInPreview, const bool bInNoPatchDelete)
+	: CloudDir(InCloudDir)
+	, bPreview(bInPreview)
+	, bNoPatchDelete(bInNoPatchDelete)
 {
 
 }
diff --git a/Engine/Source/Runtime/Online/BuildPatchServices/Private/BuildPatchDownloader.cpp b/Engine/Source/Runtime/Online/BuildPatchServices/Private/BuildPatchDownloader.cpp
--- a/Engine/Source/Runtime/Online/BuildPatchServices/Private/BuildPatchDownloader.cpp
+++ b/Engine/Source/Runtime/Online/BuildPatchServices/Private/BuildPatchDownloader.cpp
@@ -112,7 +112,6 @@ bool FBuildPatchDownloader::Init()
 
 uint32 FBuildPatchDownloader::Run()
 {
-	FString ErrorString;
 	SetRunning( true );
 	SetInited( true );
 
diff --git a/Engine/Source/Runtime/Online/BuildPatchServices/Private/BuildPatchFileAttributes.cpp b/Engine/Source/Runtime/Online/BuildPatchServices/Private/BuildPatchFileAttributes.cpp
--- a/Engine/Source/Runtime/Online/BuildPatchServices/Private/BuildPatchFileAttributes.cpp
+++ b/Engine/Source/Runtime/Online/BuildPatchServices/Private/BuildPatchFileAttributes.cpp
@@ -24,38 +24,40 @@ private:
 	bool SetExecutableFlag(const FString& Filepath);
 };
 
-FBuildPatchFileAttributesImpl::FBuildPatchFileAttributesImpl(const FBuildPatchAppManifestRef& NewManifest, const FBuildPatchAppManifestPtr& OldManifest, const FString& VerifyDirectory, const FString& StagedFileDirectory, bool bUseStageDirectory)
-	: NewManifest(NewManifest)
-	, OldManifest(OldManifest)
-	, VerifyDirectory(VerifyDirectory)
-	, StagedFileDirectory(StagedFileDirectory)
-	, bUseStageDirectory(bUseStageDirectory)
+FBuildPatchFileAttributesImpl::FBuildPatchFileAttributesImpl(const FBuildPatchAppManifestRef& InNewManifest, const FBuildPatchAppManifestPtr& InOldManifest, const FString& InVerifyDirectory, const FString& InStagedFileDirectory, bool bInUseStageDirectory)
+	: NewManifest(InNewManifest)
+	, OldManifest(InOldManifest)
+	, VerifyDirectory(InVerifyDirectory)
+	, StagedFileDirectory(InStagedFileDirectory)
+	, bUseStageDirectory(bInUseStageDirectory)
 {}
 
 bool FBuildPatchFileAttributesImpl::ApplyAttributes(bool bForce)
 {
 	// We need to set attributes for all files in the new build that require it
-	TArray<FString> BuildFiles;
-	NewManifest->GetFileList(BuildFiles);
-	for (const FString& BuildFile : BuildFiles)
-	{
-		// Break if quitting
-		if (FBuildPatchInstallError::HasFatalError())
-		{
-			break;
-		}
-
-		const FFileManifestData* FileManifest = NewManifest->GetFileManifest(BuildFile);
-		if (FileManifest != nullptr)
-		{
-			// Apply
-			const bool bHasAttrib = FileManifest->bIsReadOnly || FileManifest->bIsCompressed || FileManifest->bIsUnixExecutable;
-			if (bHasAttrib || bForce)
-			{
-				SetupFileAttributes(SelectFullFilePath(BuildFile), *FileManifest);
-			}
-		}
-	}
+    {
+        TArray<FString> BuildFiles;
+        NewManifest->GetFileList(BuildFiles);
+        for (const FString& BuildFile : BuildFiles)
+        {
+            // Break if quitting
+            if (FBuildPatchInstallError::HasFatalError())
+            {
+                break;
+            }
+
+            const FFileManifestData* FileManifest = NewManifest->GetFileManifest(BuildFile);
+            if (FileManifest != nullptr)
+            {
+                // Apply
+                const bool bHasAttrib = FileManifest->bIsReadOnly || FileManifest->bIsCompressed || FileManifest->bIsUnixExecutable;
+                if (bHasAttrib || bForce)
+                {
+                    SetupFileAttributes(SelectFullFilePath(BuildFile), *FileManifest);
+                }
+            }
+        }
+    }
 
 	// We also need to check if any attributes have been removed, unless we forced anyway
 	if (OldManifest.IsValid() && !bForce)
diff --git a/Engine/Source/Runtime/Online/BuildPatchServices/Private/BuildPatchGeneration.cpp b/Engine/Source/Runtime/Online/BuildPatchServices/Private/BuildPatchGeneration.cpp
--- a/Engine/Source/Runtime/Online/BuildPatchServices/Private/BuildPatchGeneration.cpp
+++ b/Engine/Source/Runtime/Online/BuildPatchServices/Private/BuildPatchGeneration.cpp
@@ -888,8 +888,8 @@ const uint32 FBuildGenerationChunkCache::FChunkReader::BytesLeft()
 
 /* FBuildGenerationChunkCache implementation
 *****************************************************************************/
-FBuildGenerationChunkCache::FBuildGenerationChunkCache(const FDateTime& DataAgeThreshold)
-	: DataAgeThreshold(DataAgeThreshold)
+FBuildGenerationChunkCache::FBuildGenerationChunkCache(const FDateTime& InDataAgeThreshold)
+	: DataAgeThreshold(InDataAgeThreshold)
 {
 }
 
@@ -904,12 +904,12 @@ TSharedRef< FBuildGenerationChunkCache::FChunkReader > FBuildGenerationChunkCach
 			double OldestAccessTime = FPlatformTime::Seconds();
 			for( auto ChunkCacheIt = ChunkCache.CreateConstIterator(); ChunkCacheIt; ++ChunkCacheIt )
 			{
-				const FString& ChunkFilePath = ChunkCacheIt.Key();
+				const FString& ChunkCacheFilePath = ChunkCacheIt.Key();
 				const FChunkFile& ChunkFile = ChunkCacheIt.Value().Get();
 				if( ChunkFile.GetLastAccessTime() < OldestAccessTime )
 				{
 					OldestAccessTime = ChunkFile.GetLastAccessTime();
-					OldestAccessChunk = ChunkFilePath;
+					OldestAccessChunk = ChunkCacheFilePath;
 				}
 			}
 			ChunkCache.Remove( OldestAccessChunk );
@@ -1474,28 +1474,28 @@ bool FBuildDataGenerator::FindExistingChunkData( const uint64& ChunkHash, const
 					{
 						TArray<FGuid> ChunksReferenced;
 						BuildManifest->GetDataList(ChunksReferenced);
-						for (const auto& ChunkGuid : ChunksReferenced)
+						for (const auto& ReferencedChunkGuid : ChunksReferenced)
 						{
-							uint64 ChunkHash;
-							if (BuildManifest->GetChunkHash(ChunkGuid, ChunkHash))
+							uint64 ReferencedChunkHash;
+							if (BuildManifest->GetChunkHash(ReferencedChunkGuid, ReferencedChunkHash))
 							{
-								if (ChunkHash != 0)
+								if (ReferencedChunkHash != 0)
 								{
-									TArray< FGuid >& HashChunkList = ExistingChunkHashInventory.FindOrAdd(ChunkHash);
-									if (!HashChunkList.Contains(ChunkGuid))
+									TArray< FGuid >& HashChunkList = ExistingChunkHashInventory.FindOrAdd(ReferencedChunkHash);
+									if (!HashChunkList.Contains(ReferencedChunkGuid))
 									{
 										++NumChunksFound;
-										HashChunkList.Add(ChunkGuid);
+										HashChunkList.Add(ReferencedChunkGuid);
 									}
 								}
 								else
 								{
-									GLog->Logf(TEXT("BuildDataGenerator: INFO: Ignored an existing chunk %s with a failed hash value of zero to avoid performance problems while chunking"), *ChunkGuid.ToString());
+									GLog->Logf(TEXT("BuildDataGenerator: INFO: Ignored an existing chunk %s with a failed hash value of zero to avoid performance problems while chunking"), *ReferencedChunkGuid.ToString());
 								}
 							}
 							else
 							{
-								GLog->Logf(TEXT("BuildDataGenerator: WARNING: Missing chunk hash for %s in manifest %s"), *ChunkGuid.ToString(), *ManifestFile);
+								GLog->Logf(TEXT("BuildDataGenerator: WARNING: Missing chunk hash for %s in manifest %s"), *ReferencedChunkGuid.ToString(), *ManifestFile);
 							}
 						}
 					}
@@ -1919,8 +1919,8 @@ void FBuildDataGenerator::StripIgnoredFiles( TArray< FString >& AllFiles, const
 	struct FRemoveMatchingStrings
 	{ 
 		const TSet<FString>& IgnoreList;
-		FRemoveMatchingStrings( const TSet<FString>& IgnoreList )
-			: IgnoreList(IgnoreList) {}
+		FRemoveMatchingStrings( const TSet<FString>& InIgnoreList )
+			: IgnoreList(InIgnoreList) {}
 
 		bool operator()(const FString& RemovalCandidate) const
 		{
diff --git a/Engine/Source/Runtime/Online/BuildPatchServices/Private/BuildPatchManifest.cpp b/Engine/Source/Runtime/Online/BuildPatchServices/Private/BuildPatchManifest.cpp
--- a/Engine/Source/Runtime/Online/BuildPatchServices/Private/BuildPatchManifest.cpp
+++ b/Engine/Source/Runtime/Online/BuildPatchServices/Private/BuildPatchManifest.cpp
@@ -125,9 +125,9 @@ FCustomFieldData::FCustomFieldData()
 {
 }
 
-FCustomFieldData::FCustomFieldData(const FString& Key, const FString& Value)
-	: Key(Key)
-	, Value(Value)
+FCustomFieldData::FCustomFieldData(const FString& InKey, const FString& InValue)
+	: Key(InKey)
+	, Value(InValue)
 {
 }
 
diff --git a/Engine/Source/Runtime/Online/BuildPatchServices/Private/BuildPatchVerification.cpp b/Engine/Source/Runtime/Online/BuildPatchServices/Private/BuildPatchVerification.cpp
--- a/Engine/Source/Runtime/Online/BuildPatchServices/Private/BuildPatchVerification.cpp
+++ b/Engine/Source/Runtime/Online/BuildPatchServices/Private/BuildPatchVerification.cpp
@@ -25,13 +25,13 @@ private:
 	FString SelectFullFilePath(const FString& BuildFile);
 };
 
-FBuildPatchVerificationImpl::FBuildPatchVerificationImpl(const FBuildPatchAppManifestRef& Manifest, const FBuildPatchFloatDelegate& ProgressDelegate, const FBuildPatchBoolRetDelegate& ShouldPauseDelegate, const FString& VerifyDirectory, const FString& StagedFileDirectory, bool bUseStageDirectory)
-	: Manifest(Manifest)
-	, ProgressDelegate(ProgressDelegate)
-	, ShouldPauseDelegate(ShouldPauseDelegate)
-	, VerifyDirectory(VerifyDirectory)
-	, StagedFileDirectory(StagedFileDirectory)
-	, bUseStageDirectory(bUseStageDirectory)
+FBuildPatchVerificationImpl::FBuildPatchVerificationImpl(const FBuildPatchAppManifestRef& InManifest, const FBuildPatchFloatDelegate& InProgressDelegate, const FBuildPatchBoolRetDelegate& InShouldPauseDelegate, const FString& InVerifyDirectory, const FString& InStagedFileDirectory, bool bInUseStageDirectory)
+	: Manifest(InManifest)
+	, ProgressDelegate(InProgressDelegate)
+	, ShouldPauseDelegate(InShouldPauseDelegate)
+	, VerifyDirectory(InVerifyDirectory)
+	, StagedFileDirectory(InStagedFileDirectory)
+	, bUseStageDirectory(bInUseStageDirectory)
 	, CurrentBuildPercentage(0)
 	, CurrentFileWeight(0)
 {}
