Fix ios shadow variables

diff --git a/Engine/Source/Runtime/IOS/MetalRHI/Private/MetalVertexDeclaration.cpp b/Engine/Source/Runtime/IOS/MetalRHI/Private/MetalVertexDeclaration.cpp
--- a/Engine/Source/Runtime/IOS/MetalRHI/Private/MetalVertexDeclaration.cpp
+++ b/Engine/Source/Runtime/IOS/MetalRHI/Private/MetalVertexDeclaration.cpp
@@ -86,15 +86,15 @@ FVertexDeclarationRHIRef FMetalDynamicRHI::RHICreateVertexDeclaration(const FVer
 }
 
 
-void FMetalVertexDeclaration::GenerateLayout(const FVertexDeclarationElementList& Elements)
+void FMetalVertexDeclaration::GenerateLayout(const FVertexDeclarationElementList& InElements)
 {
 	Layout = [[MTLVertexDescriptor alloc] init];
 	TRACK_OBJECT(Layout);
 
 	TMap<uint32, uint32> BufferStrides;
-	for (uint32 ElementIndex = 0; ElementIndex < Elements.Num(); ElementIndex++)
+	for (uint32 ElementIndex = 0; ElementIndex < InElements.Num(); ElementIndex++)
 	{
-		const FVertexElement& Element = Elements[ElementIndex];
+		const FVertexElement& Element = InElements[ElementIndex];
 		
 		checkf(Element.Stride == 0 || Element.Offset + TranslateElementTypeToSize(Element.Type) <= Element.Stride, 
 			TEXT("Stream component is bigger than stride: Offset: %d, Size: %d [Type %d], Stride: %d"), Element.Offset, TranslateElementTypeToSize(Element.Type), (uint32)Element.Type, Element.Stride);
diff --git a/Engine/Source/Runtime/Online/IOS/OnlineSubsystemIOS/Private/OnlineIdentityInterfaceIOS.cpp b/Engine/Source/Runtime/Online/IOS/OnlineSubsystemIOS/Private/OnlineIdentityInterfaceIOS.cpp
--- a/Engine/Source/Runtime/Online/IOS/OnlineSubsystemIOS/Private/OnlineIdentityInterfaceIOS.cpp
+++ b/Engine/Source/Runtime/Online/IOS/OnlineSubsystemIOS/Private/OnlineIdentityInterfaceIOS.cpp
@@ -203,12 +203,12 @@ void FOnlineIdentityIOS::GetUserPrivilege(const FUniqueNetId& UserId, EUserPrivi
 	Delegate.ExecuteIfBound(UserId, Privilege, (uint32)EPrivilegeResults::NoFailures);
 }
 
-FPlatformUserId FOnlineIdentityIOS::GetPlatformUserIdFromUniqueNetId(const FUniqueNetId& UniqueNetId)
+FPlatformUserId FOnlineIdentityIOS::GetPlatformUserIdFromUniqueNetId(const FUniqueNetId& InUniqueNetId)
 {
 	for (int i = 0; i < MAX_LOCAL_PLAYERS; ++i)
 	{
 		auto CurrentUniqueId = GetUniquePlayerId(i);
-		if (CurrentUniqueId.IsValid() && (*CurrentUniqueId == UniqueNetId))
+		if (CurrentUniqueId.IsValid() && (*CurrentUniqueId == InUniqueNetId))
 		{
 			return i;
 		}
diff --git a/Engine/Source/Runtime/Online/IOS/OnlineSubsystemIOS/Private/OnlineLeaderboardsInterfaceIOS.cpp b/Engine/Source/Runtime/Online/IOS/OnlineSubsystemIOS/Private/OnlineLeaderboardsInterfaceIOS.cpp
--- a/Engine/Source/Runtime/Online/IOS/OnlineSubsystemIOS/Private/OnlineLeaderboardsInterfaceIOS.cpp
+++ b/Engine/Source/Runtime/Online/IOS/OnlineSubsystemIOS/Private/OnlineLeaderboardsInterfaceIOS.cpp
@@ -82,7 +82,7 @@ bool FOnlineLeaderboardsIOS::ReadLeaderboardCompletionDelegate(NSArray* players,
                             
                 if (bWasSuccessful)
                 {
-                    bool bWasSuccessful = [scores count] > 0;
+                    bWasSuccessful = [scores count] > 0;
                     UE_LOG(LogOnline, Display, TEXT("FOnlineLeaderboardsIOS::loadScoresWithCompletionHandler() - %s"), (bWasSuccessful ? TEXT("Success!") : TEXT("Failed!, no scores retrieved")));
                     for (GKScore* score in scores)
                     {
@@ -205,7 +205,7 @@ bool FOnlineLeaderboardsIOS::ReadLeaderboards(const TArray< TSharedRef<FUniqueNe
              
                 if (bWasSuccessful)
                 {
-                    bool bWasSuccessful = [players count] > 0;
+                    bWasSuccessful = [players count] > 0;
                     ReadLeaderboardCompletionDelegate(players, InReadObject);
                 }
              }];
diff --git a/Engine/Source/Runtime/Online/IOS/OnlineSubsystemIOS/Private/TurnBasedMatchmakerIOS.cpp b/Engine/Source/Runtime/Online/IOS/OnlineSubsystemIOS/Private/TurnBasedMatchmakerIOS.cpp
--- a/Engine/Source/Runtime/Online/IOS/OnlineSubsystemIOS/Private/TurnBasedMatchmakerIOS.cpp
+++ b/Engine/Source/Runtime/Online/IOS/OnlineSubsystemIOS/Private/TurnBasedMatchmakerIOS.cpp
@@ -94,12 +94,12 @@ void FTurnBasedMatchmakerIOS::DismissMatchmaker()
 
 GKMatchRequest* FTurnBasedMatchmakerIOS::GetGKMatchRequestFromMPMatchRequest(const FTurnBasedMatchRequest& Request) const
 {
-    GKMatchRequest* MatchRequest = [[GKMatchRequest alloc] init];
+    GKMatchRequest* NewMatchRequest = [[GKMatchRequest alloc] init];
 
-    MatchRequest.maxPlayers = Request.GetMaxNumberOfPlayers();
-    MatchRequest.minPlayers = Request.GetMinNumberOfPlayers();
-    MatchRequest.defaultNumberOfPlayers = Request.GetMinNumberOfPlayers();
-    MatchRequest.playerGroup = Request.GetPlayerGroup();
+    NewMatchRequest.maxPlayers = Request.GetMaxNumberOfPlayers();
+    NewMatchRequest.minPlayers = Request.GetMinNumberOfPlayers();
+    NewMatchRequest.defaultNumberOfPlayers = Request.GetMinNumberOfPlayers();
+    NewMatchRequest.playerGroup = Request.GetPlayerGroup();
 
-    return [MatchRequest autorelease];
+    return [NewMatchRequest autorelease];
 }
