[INTEGRATE] merge ue4 <-ue4-launcherDev cl [REDACTED_PHONE])

diff --git a/Engine/Source/Runtime/Online/BuildPatchServices/Private/BuildPatchError.cpp b/Engine/Source/Runtime/Online/BuildPatchServices/Private/BuildPatchError.cpp
--- a/Engine/Source/Runtime/Online/BuildPatchServices/Private/BuildPatchError.cpp
+++ b/Engine/Source/Runtime/Online/BuildPatchServices/Private/BuildPatchError.cpp
@@ -39,7 +39,8 @@ bool FBuildPatchInstallError::IsNoRetryError()
 	FScopeLock ScopeLock( &ThreadLock );
 	return ErrorState == EBuildPatchInstallError::DownloadError
 		|| ErrorState == EBuildPatchInstallError::MoveFileToInstall
-		|| ErrorState == EBuildPatchInstallError::InitializationError;
+		|| ErrorState == EBuildPatchInstallError::InitializationError
+		|| ErrorState == EBuildPatchInstallError::PathLengthExceeded;
 }
 
 FString FBuildPatchInstallError::GetErrorString()
@@ -54,22 +55,42 @@ FString FBuildPatchInstallError::GetErrorString()
 	return ReturnString;
 }
 
-const FText& FBuildPatchInstallError::GetErrorText()
+const FText FBuildPatchInstallError::GetErrorText()
 {
-	static const FText NoError( LOCTEXT( "BuildPatchInstallError_NoError", "The operation was successful." ) );
-	static const FText DownloadError( LOCTEXT( "BuildPatchInstallError_DownloadError", "Could not download patch data. Please try again later." ) );
-	static const FText FileConstructionFail( LOCTEXT( "BuildPatchInstallError_FileConstructionFail", "A file corruption has occurred. Please try again." ) );
-	static const FText MoveFileToInstall( LOCTEXT( "BuildPatchInstallError_MoveFileToInstall", "A file access error has occurred. Please check your running processes." ) );
-	static const FText BuildVerifyFail( LOCTEXT( "BuildPatchInstallError_BuildCorrupt", "The installation is corrupt." ) );
-	static const FText ApplicationClosing( LOCTEXT( "BuildPatchInstallError_ApplicationClosing", "The application is closing." ) );
-	static const FText ApplicationError( LOCTEXT( "BuildPatchInstallError_ApplicationError", "Patching service could not start." ) );
-	static const FText UserCanceled( LOCTEXT( "BuildPatchInstallError_UserCanceled", "User cancelled." ) );
-	static const FText PrerequisiteError( LOCTEXT( "BuildPatchInstallError_PrerequisiteError", "Prerequisites install failed.") );
-	static const FText InitializationError(LOCTEXT("BuildPatchInstallError_InitializationError", "The installer failed to initialize."));
-	static const FText InvalidOrMax( LOCTEXT( "BuildPatchInstallError_InvalidOrMax", "An unknown error ocurred." ) );
+	static const FText DownloadError( LOCTEXT( "BuildPatchInstallError_DownloadError", "Please try again later." ) );
+	static const FText FileConstructionFail( LOCTEXT( "BuildPatchInstallError_FileConstructionFail", "Please try again." ) );
+	static const FText MoveFileToInstall( LOCTEXT( "BuildPatchInstallError_MoveFileToInstall", "Please check your running processes." ) );
+	static const FText PathLengthExceeded(LOCTEXT("BuildPatchInstallError_PathLengthExceeded", "Please specify a shorter install location."));
 
 	FScopeLock ScopeLock( &ThreadLock );
+	FText ErrorText = FText::GetEmpty();
 	switch( ErrorState )
+	{
+		case EBuildPatchInstallError::DownloadError: ErrorText = DownloadError;
+		case EBuildPatchInstallError::FileConstructionFail: ErrorText = FileConstructionFail;
+		case EBuildPatchInstallError::MoveFileToInstall: ErrorText = MoveFileToInstall;
+		case EBuildPatchInstallError::PathLengthExceeded: ErrorText = PathLengthExceeded;
+	}
+	return FText::Format(LOCTEXT("BuildPatchInstallLongError", "{0} {1}"), GetShortErrorText(), ErrorText);
+}
+
+const FText& FBuildPatchInstallError::GetShortErrorText()
+{
+	static const FText NoError(LOCTEXT("BuildPatchInstallShortError_NoError", "The operation was successful."));
+	static const FText DownloadError(LOCTEXT("BuildPatchInstallShortError_DownloadError", "Could not download patch data."));
+	static const FText FileConstructionFail(LOCTEXT("BuildPatchInstallShortError_FileConstructionFail", "A file corruption has occurred."));
+	static const FText MoveFileToInstall(LOCTEXT("BuildPatchInstallShortError_MoveFileToInstall", "A file access error has occurred."));
+	static const FText BuildVerifyFail(LOCTEXT("BuildPatchInstallShortError_BuildCorrupt", "The installation is corrupt."));
+	static const FText ApplicationClosing(LOCTEXT("BuildPatchInstallShortError_ApplicationClosing", "The application is closing."));
+	static const FText ApplicationError(LOCTEXT("BuildPatchInstallShortError_ApplicationError", "Patching service could not start."));
+	static const FText UserCanceled(LOCTEXT("BuildPatchInstallShortError_UserCanceled", "User cancelled."));
+	static const FText PrerequisiteError(LOCTEXT("BuildPatchInstallShortError_PrerequisiteError", "Prerequisites install failed."));
+	static const FText InitializationError(LOCTEXT("BuildPatchInstallShortError_InitializationError", "The installer failed to initialize."));
+	static const FText PathLengthExceeded(LOCTEXT("BuildPatchInstallShortError_PathLengthExceeded", "Maximum path length exceeded."));
+	static const FText InvalidOrMax(LOCTEXT("BuildPatchInstallShortError_InvalidOrMax", "An unknown error ocurred."));
+
+	FScopeLock ScopeLock(&ThreadLock);
+	switch (ErrorState)
 	{
 		case EBuildPatchInstallError::NoError: return NoError;
 		case EBuildPatchInstallError::DownloadError: return DownloadError;
@@ -81,6 +102,7 @@ const FText& FBuildPatchInstallError::GetErrorText()
 		case EBuildPatchInstallError::UserCanceled: return UserCanceled;
 		case EBuildPatchInstallError::PrerequisiteError: return PrerequisiteError;
 		case EBuildPatchInstallError::InitializationError: return InitializationError;
+		case EBuildPatchInstallError::PathLengthExceeded: return PathLengthExceeded;
 		default: return InvalidOrMax;
 	}
 }
@@ -120,6 +142,7 @@ const FString& FBuildPatchInstallError::ToString( const EBuildPatchInstallError:
 	static const FString UserCanceled( "EBuildPatchInstallError::UserCanceled" );
 	static const FString PrerequisiteError( "EBuildPatchInstallError::PrerequisiteError" );
 	static const FString InitializationError("EBuildPatchInstallError::InitializationError");
+	static const FString PathLengthExceeded("EBuildPatchInstallError::PathLengthExceeded");
 	static const FString InvalidOrMax( "EBuildPatchInstallError::InvalidOrMax" );
 
 	FScopeLock ScopeLock( &ThreadLock );
@@ -135,6 +158,7 @@ const FString& FBuildPatchInstallError::ToString( const EBuildPatchInstallError:
 		case EBuildPatchInstallError::UserCanceled: return UserCanceled;
 		case EBuildPatchInstallError::PrerequisiteError: return PrerequisiteError;
 		case EBuildPatchInstallError::InitializationError: return InitializationError;
+		case EBuildPatchInstallError::PathLengthExceeded: return PathLengthExceeded;
 		default: return InvalidOrMax;
 	}
 }
diff --git a/Engine/Source/Runtime/Online/BuildPatchServices/Private/BuildPatchError.h b/Engine/Source/Runtime/Online/BuildPatchServices/Private/BuildPatchError.h
--- a/Engine/Source/Runtime/Online/BuildPatchServices/Private/BuildPatchError.h
+++ b/Engine/Source/Runtime/Online/BuildPatchServices/Private/BuildPatchError.h
@@ -44,7 +44,10 @@ namespace EBuildPatchInstallError
 		PrerequisiteError = 8,
 
 		// An initialization error
-		InitializationError = 9
+		InitializationError = 9,
+
+		// An error occurred creating a file due to excessive path length
+		PathLengthExceeded = 10
 	};
 };
 
@@ -105,7 +108,13 @@ public:
 	 * Static function to get the error text for UI
 	 * @return The error text
 	 */
-	static const FText& GetErrorText();
+	static const FText GetErrorText();
+
+	/**
+	 * Static function to get the truncated error text for UI
+	 * @return The shorter error text
+	 */
+	static const FText& GetShortErrorText();
 
 	/**
 	 * Static function allowing any worker to set fatal error
diff --git a/Engine/Source/Runtime/Online/BuildPatchServices/Private/BuildPatchInstaller.cpp b/Engine/Source/Runtime/Online/BuildPatchServices/Private/BuildPatchInstaller.cpp
--- a/Engine/Source/Runtime/Online/BuildPatchServices/Private/BuildPatchInstaller.cpp
+++ b/Engine/Source/Runtime/Online/BuildPatchServices/Private/BuildPatchInstaller.cpp
@@ -273,6 +273,17 @@ bool FBuildPatchInstaller::RunInstallation(TArray<FString>& CorruptFiles)
 	}
 	GLog->Logf(TEXT("BuildPatchServices: Requiring %d files"), FilesToConstruct.Num());
 
+	// Make sure all the files won't exceed the maximum path length
+	for (const auto& FileToConstruct : FilesToConstruct)
+	{
+		if ((InstallStagingDir / FileToConstruct).Len() >= MAX_PATH)
+		{
+			GWarn->Logf(TEXT("BuildPatchServices: ERROR: Could not create new file due to exceeding maximum path length %s"), *(InstallStagingDir / FileToConstruct));
+			FBuildPatchInstallError::SetFatalError(EBuildPatchInstallError::PathLengthExceeded);
+			return false;
+		}
+	}
+
 	// Create the downloader
 	FBuildPatchDownloader::Create(DataStagingDir, NewBuildManifest, &BuildProgress);
 
@@ -938,6 +949,16 @@ bool FBuildPatchInstaller::IsPaused()
 	return BuildProgress.GetPauseState();
 }
 
+bool FBuildPatchInstaller::IsResumable()
+{
+	FScopeLock Lock( &ThreadLock );
+	if( FBuildPatchInstallError::GetErrorState() == EBuildPatchInstallError::PathLengthExceeded )
+	{
+		return false;
+	}
+	return !BuildStats.ProcessSuccess;
+}
+
 bool FBuildPatchInstaller::HasError()
 {
 	FScopeLock Lock( &ThreadLock );
@@ -968,9 +989,9 @@ FText FBuildPatchInstaller::GetPercentageText()
 	return FText::AsPercent(GetUpdateProgress(), &PercentFormattingOptions);
 }
 
-FText FBuildPatchInstaller::GetStatusText()
+FText FBuildPatchInstaller::GetStatusText( bool ShortError )
 {
-	return BuildProgress.GetStateText();
+	return BuildProgress.GetStateText(ShortError);
 }
 
 float FBuildPatchInstaller::GetUpdateProgress()
diff --git a/Engine/Source/Runtime/Online/BuildPatchServices/Private/BuildPatchInstaller.h b/Engine/Source/Runtime/Online/BuildPatchServices/Private/BuildPatchInstaller.h
--- a/Engine/Source/Runtime/Online/BuildPatchServices/Private/BuildPatchInstaller.h
+++ b/Engine/Source/Runtime/Online/BuildPatchServices/Private/BuildPatchInstaller.h
@@ -126,6 +126,7 @@ public:
 	virtual bool IsComplete() override;
 	virtual bool IsCanceled() override;
 	virtual bool IsPaused() override;
+	virtual bool IsResumable() override;
 	virtual bool HasError() override;
 	//@todo this is deprecated and shouldn't be used anymore [6/4/2014 justin.sargent]
 	virtual FText GetPercentageText() override;
@@ -134,7 +135,7 @@ public:
 	virtual double GetDownloadSpeed() const override;
 	virtual int64 GetInitialDownloadSize() const override;
 	virtual int64 GetTotalDownloaded() const override;
-	virtual FText GetStatusText() override;
+	virtual FText GetStatusText(bool ShortError = false) override;
 	virtual float GetUpdateProgress() override;
 	virtual FBuildInstallStats GetBuildStatistics() override;
 	virtual FText GetErrorText() override;
diff --git a/Engine/Source/Runtime/Online/BuildPatchServices/Private/BuildPatchProgress.cpp b/Engine/Source/Runtime/Online/BuildPatchServices/Private/BuildPatchProgress.cpp
--- a/Engine/Source/Runtime/Online/BuildPatchServices/Private/BuildPatchProgress.cpp
+++ b/Engine/Source/Runtime/Online/BuildPatchServices/Private/BuildPatchProgress.cpp
@@ -114,6 +114,7 @@ void FBuildPatchProgress::Reset()
 	CurrentState = EBuildPatchProgress::Initializing;
 	CurrentProgress = 0.0f;
 	ErrorText = FText::GetEmpty();
+	ShortErrorText = FText::GetEmpty();
 
 	// Initialize array data
 	for( uint32 idx = 0; idx < EBuildPatchProgress::NUM_PROGRESS_STATES; ++idx )
@@ -149,12 +150,12 @@ void FBuildPatchProgress::SetStateWeight( const EBuildPatchProgress::Type& State
 	}
 }
 
-const FText& FBuildPatchProgress::GetStateText()
+const FText& FBuildPatchProgress::GetStateText( bool ShortError )
 {
 	FScopeLock ScopeLock( &ThreadLock );
 	if( CurrentState == EBuildPatchProgress::Error )
 	{
-		return ErrorText;
+		return ShortError ? ShortErrorText : ErrorText;
 	}
 	return EBuildPatchProgress::ToText( CurrentState );
 }
@@ -230,6 +231,7 @@ void FBuildPatchProgress::UpdateProgressInfo()
 	{
 		CurrentState = EBuildPatchProgress::Error;
 		ErrorText = FBuildPatchInstallError::GetErrorText();
+		ShortErrorText = FBuildPatchInstallError::GetShortErrorText();
 		return;
 	}
 
diff --git a/Engine/Source/Runtime/Online/BuildPatchServices/Private/BuildPatchProgress.h b/Engine/Source/Runtime/Online/BuildPatchServices/Private/BuildPatchProgress.h
--- a/Engine/Source/Runtime/Online/BuildPatchServices/Private/BuildPatchProgress.h
+++ b/Engine/Source/Runtime/Online/BuildPatchServices/Private/BuildPatchProgress.h
@@ -101,6 +101,9 @@ private:
 	// The error text that was set when in error state
 	FText ErrorText;
 
+	// The error text that was set when in error state
+	FText ShortErrorText;
+
 	// Critical section to protect variable access
 	FCriticalSection ThreadLock;
 
@@ -132,9 +135,10 @@ public:
 
 	/**
 	 * Gets the text for the current progress state
+	 * @param ShortError		The truncated version of the error
 	 * @return The display text for the current progress state
 	 */
-	const FText& GetStateText();
+	const FText& GetStateText( bool ShortError = false );
 
 	/**
 	 * Gets the current overall progress
diff --git a/Engine/Source/Runtime/Online/BuildPatchServices/Public/Interfaces/IBuildInstaller.h b/Engine/Source/Runtime/Online/BuildPatchServices/Public/Interfaces/IBuildInstaller.h
--- a/Engine/Source/Runtime/Online/BuildPatchServices/Public/Interfaces/IBuildInstaller.h
+++ b/Engine/Source/Runtime/Online/BuildPatchServices/Public/Interfaces/IBuildInstaller.h
@@ -97,6 +97,12 @@ public:
 	 */
 	virtual bool IsPaused() = 0;
 
+	/**
+	 * Get whether the install can be resumed.
+	 * @return	true if installation is resumable
+	 */
+	virtual bool IsResumable() = 0;
+
 	/**
 	 * Get whether the install failed. Only valid if complete.
 	 * @return	true if installation was a failure
@@ -139,9 +145,10 @@ public:
 
 	/**
 	 * Get the text for status of the install process
+	 * @param ShortError		The truncated version of the error
 	 * @return	status of the install process text
 	 */
-	virtual FText GetStatusText() = 0;
+	virtual FText GetStatusText(bool ShortError = false) = 0;
 
 	/**
 	 * Get the update progress
