Merging

//depot/UE4-LauncherDev/Engine/Source/Runtime/Online/...

to //depot/UE4/Engine/Source/Runtime/Online/...

diff --git a/Engine/Source/Runtime/Online/OnlineSubsystem/Public/Interfaces/OnlinePartyInterface.h b/Engine/Source/Runtime/Online/OnlineSubsystem/Public/Interfaces/OnlinePartyInterface.h
--- a/Engine/Source/Runtime/Online/OnlineSubsystem/Public/Interfaces/OnlinePartyInterface.h
+++ b/Engine/Source/Runtime/Online/OnlineSubsystem/Public/Interfaces/OnlinePartyInterface.h
@@ -76,6 +76,38 @@ public:
 	virtual TSharedRef<FUniqueNetId> GetPartyMemberId() const = 0;
 };
 
+/**
+ * Data associated with a party invite
+ */
+
+namespace EInviteState
+{
+	enum Type
+	{
+		Unanswered,
+		Accepted
+	};
+}
+
+class IOnlinePartyInviteData
+{
+public:
+	IOnlinePartyInviteData() : InviteState(EInviteState::Unanswered) {}
+	virtual ~IOnlinePartyInviteData() {}
+
+	/**
+	 * @return id of the user this invite is from
+	 */
+	virtual TSharedRef<FUniqueNetId> GetSenderId() const = 0;
+
+	/**
+	 * @return id of the party associated with this invite
+	 */
+	virtual TSharedRef<FOnlinePartyId> GetPartyId() const = 0;
+
+	EInviteState::Type InviteState;
+};
+
 /**
  * Permissions for joining a party
  */
@@ -83,11 +115,9 @@ namespace EPartyPermissions
 {
 	enum Type
 	{
-		/** any user can discover/join the party */
+		/** any user that discovers the party can join */
 		Public,
-		/** only friends can discover/join the party */
-		FriendsOnly,
-		/** invite required to discover/join the party */
+		/** invite required to join the party */
 		Private
 	};
 }
@@ -98,20 +128,30 @@ namespace EPartyPermissions
 struct FPartyConfiguration
 {
 	FPartyConfiguration()
-		: JoinPerms(EPartyPermissions::FriendsOnly)
+		: Permissions(EPartyPermissions::Public)
 		, MaxMembers(0)
 	{}
 
 	/** Permission for configuring party */
-	EPartyPermissions::Type JoinPerms;
-	/** Maximum active members allowed */
+	EPartyPermissions::Type Permissions;
+	/** Maximum active members allowed. 0 means no maximum. */
 	int32 MaxMembers;
-	/** PartyId to use - largely for debug purposes*/
-	FString PartyId;
-	/** Password to use */
+	/** Human readable nickname */
+	FString Nickname;
+	/** Human readable description */
+	FString Description;
+	/** Password to use. */
 	FString Password;
 };
 
+struct FDebugPartyConfiguration : TSharedFromThis<FDebugPartyConfiguration>
+{
+	/** PartyId to use */
+	FString PartyId;
+	/** Access token to use. */
+	FString AccessToken;
+};
+
 namespace EPartyState
 {
 	enum Type
@@ -129,7 +169,7 @@ namespace EPartyState
 /**
  * Current state associated with a party
  */
-class FOnlinePartyInfo
+class FOnlinePartyInfo : public TSharedFromThis<FOnlinePartyInfo>
 {
 public:
 	FOnlinePartyInfo(const TSharedPtr<FOnlinePartyId>& InPartyId = TSharedPtr<FOnlinePartyId>())
@@ -143,12 +183,16 @@ public:
 	EPartyState::Type State;
 	/** Current state of configuration */
 	FPartyConfiguration Config;
+	/** (optional) Debug configuration */
+	TSharedPtr<FDebugPartyConfiguration> DebugConfig;
 	/** unique id of the party */
 	TSharedPtr<FOnlinePartyId> PartyId;
 	/** leader of the party */
 	TSharedPtr<FUniqueNetId> PartyLeader;
-	/** password needed for private parties */
-	FString Password;
+	/** access token that changes when you switch from public to invite-only */
+	FString AccessToken;
+	/** password needed for to join the room */
+	FString RoomPassword;
 	/** id of chat room associated with the party */
 	FChatRoomId RoomId;
 };
@@ -210,6 +254,24 @@ DECLARE_DELEGATE_FourParams(FOnLeavePartyComplete, const FUniqueNetId& /*LocalUs
  * @param Error - string with error info if any
  */
 DECLARE_DELEGATE_FiveParams(FOnSendPartyInviteComplete, const FUniqueNetId& /*LocalUserId*/, bool /*bWasSuccessful*/, const FOnlinePartyId& /*PartyId*/, const FUniqueNetId& /*RecipientId*/, const FString& /*Error*/);
+/**
+ * Accepting an invite to a user to join party async task completed callback
+ *
+ * @param LocalUserId - id of user that initiated the request
+ * @param bWasSuccessful - true if successfully sent invite
+ * @param PartyId - id associated with the party
+ * @param Error - string with error info if any
+ */
+DECLARE_DELEGATE_FourParams(FOnAcceptPartyInviteComplete, const FUniqueNetId& /*LocalUserId*/, bool /*bWasSuccessful*/, const FOnlinePartyId& /*PartyId*/, const FString& /*Error*/);
+/**
+ * Rejecting an invite to a user to join party async task completed callback
+ *
+ * @param LocalUserId - id of user that initiated the request
+ * @param bWasSuccessful - true if successfully sent invite
+ * @param PartyId - id associated with the party
+ * @param Error - string with error info if any
+ */
+DECLARE_DELEGATE_FourParams(FOnRejectPartyInviteComplete, const FUniqueNetId& /*LocalUserId*/, bool /*bWasSuccessful*/, const FOnlinePartyId& /*PartyId*/, const FString& /*Error*/);
 /**
  * Kicking a member of a party async task completed callback
  *
@@ -297,6 +359,24 @@ typedef FOnPartyDataReceived::FDelegate FOnPartyDataReceivedDelegate;
 DECLARE_MULTICAST_DELEGATE_ThreeParams(FOnPartyMemberDataReceived, const FOnlinePartyId& /*PartyId*/, const FUniqueNetId& /*MemberId*/, const TSharedRef<IOnlinePartyMemberData>& /*PartyMemberData*/);
 typedef FOnPartyMemberDataReceived::FDelegate FOnPartyMemberDataReceivedDelegate;
 
+/**
+ * Notification when an invite is rejected
+ *
+ * @param PartyId - id associated with the party
+ * @param MemberId - id of member that rejected the invite
+ */
+DECLARE_MULTICAST_DELEGATE_TwoParams(FOnPartyInviteAccepted, const FOnlinePartyId& /*PartyId*/, const FUniqueNetId& /*RejectorId*/);
+typedef FOnPartyInviteAccepted::FDelegate FOnPartyInviteAcceptedDelegate;
+
+/**
+ * Notification when an invite is rejected
+ *
+ * @param PartyId - id associated with the party
+ * @param MemberId - id of member that rejected the invite
+ */
+DECLARE_MULTICAST_DELEGATE_TwoParams(FOnPartyInviteRejected, const FOnlinePartyId& /*PartyId*/, const FUniqueNetId& /*RejectorId*/);
+typedef FOnPartyInviteRejected::FDelegate FOnPartyInviteRejectedDelegate;
+
 /**
  * Interface definition for the online party services 
  * Allows for forming a party and communicating with party members
@@ -355,7 +435,7 @@ public:
 	 *
 	 * @return true if task was started
 	 */
-	virtual bool JoinParty(const FUniqueNetId& LocalUserId, const FOnlinePartyId& PartyId, const FOnJoinPartyComplete& Delegate = FOnJoinPartyComplete()) = 0;
+	virtual bool JoinParty(const FUniqueNetId& LocalUserId, const FOnlinePartyId& PartyId, const FString& AccessToken, const FString& Password = TEXT(""), const FOnJoinPartyComplete& Delegate = FOnJoinPartyComplete()) = 0;
 
 	/**
 	 * Leave an existing party
@@ -370,18 +450,44 @@ public:
 	virtual bool LeaveParty(const FUniqueNetId& LocalUserId, const FOnlinePartyId& PartyId, const FOnLeavePartyComplete& Delegate = FOnLeavePartyComplete()) = 0;
 
 	/**
-	 * Invite a user to an existing party
-	 * Recipient will be notified of the new invite (see FOnPartyInviteReceived)
-	 *
-	 * @param LocalUserId - user making the request
-	 * @param PartyId - id of an existing party
-	 * @param RecipientId - id of the user being invited
-	 * @param Delegate - called on completion
-	 *
-	 * @return true if task was started
-	 */
+	* Invite a user to an existing party
+	* Recipient will be notified of the new invite (see FOnPartyInviteReceived)
+	*
+	* @param LocalUserId - user making the request
+	* @param PartyId - id of an existing party
+	* @param RecipientId - id of the user being invited
+	* @param Delegate - called on completion
+	*
+	* @return true if task was started
+	*/
 	virtual bool SendInvite(const FUniqueNetId& LocalUserId, const FOnlinePartyId& PartyId, const FUniqueNetId& RecipientId, const FOnSendPartyInviteComplete& Delegate = FOnSendPartyInviteComplete()) = 0;
 
+	/**
+	* Invite a user to an existing party
+	* Recipient will be notified of the new invite (see FOnPartyInviteReceived)
+	*
+	* @param LocalUserId - user making the request
+	* @param PartyId - id of an existing party
+	* @param RecipientId - id of the user being invited
+	* @param Delegate - called on completion
+	*
+	* @return true if task was started
+	*/
+	virtual bool AcceptInvite(const FUniqueNetId& LocalUserId, const FOnlinePartyId& PartyId, const FOnAcceptPartyInviteComplete& Delegate = FOnAcceptPartyInviteComplete()) = 0;
+
+	/**
+	* Invite a user to an existing party
+	* Recipient will be notified of the new invite (see FOnPartyInviteReceived)
+	*
+	* @param LocalUserId - user making the request
+	* @param PartyId - id of an existing party
+	* @param RecipientId - id of the user being invited
+	* @param Delegate - called on completion
+	*
+	* @return true if task was started
+	*/
+	virtual bool RejectInvite(const FUniqueNetId& LocalUserId, const FOnlinePartyId& PartyId, const FOnRejectPartyInviteComplete& Delegate = FOnRejectPartyInviteComplete()) = 0;
+
 	/**
 	 * Kick a user from an existing party
 	 * Only admin can kick a party member
@@ -454,7 +560,7 @@ public:
 	 *
 	 * @return true if entries found
 	 */
-	virtual bool GetPartyInvites(const FUniqueNetId& LocalUserId, TArray< TSharedRef<FOnlinePartyId> >& OutPartyIds) const = 0;
+	virtual bool GetPartyInvites(const FUniqueNetId& LocalUserId, TArray< TSharedRef<IOnlinePartyInviteData> >& OutPartyInviteData) const = 0;
 
 	/**
 	 * Get info associated with a party
@@ -578,6 +684,22 @@ public:
 	 * @param PartyMemberData - party member data that was updated
 	 */
 	DEFINE_ONLINE_DELEGATE_THREE_PARAM(OnPartyMemberDataReceived, const FOnlinePartyId& /*PartyId*/, const FUniqueNetId& /*MemberId*/, const TSharedRef<IOnlinePartyMemberData>& /*PartyMemberData*/);
+
+	/**
+	 * Notification when an invite is accepted
+	 *
+	 * @param PartyId - id associated with the party
+	 * @param MemberId - id of member that accepted the invite
+	 */
+	DEFINE_ONLINE_DELEGATE_TWO_PARAM(OnPartyInviteAccepted, const FOnlinePartyId& /*PartyId*/, const FUniqueNetId& /*RejectorId*/);
+
+	/**
+	 * Notification when an invite is rejected
+	 *
+	 * @param PartyId - id associated with the party
+	 * @param MemberId - id of member that rejected the invite
+	 */
+	DEFINE_ONLINE_DELEGATE_TWO_PARAM(OnPartyInviteRejected, const FOnlinePartyId& /*PartyId*/, const FUniqueNetId& /*RejectorId*/);
 };
 
 typedef TSharedPtr<IOnlineParty, ESPMode::ThreadSafe> IOnlinePartyPtr;
