Merging using UE4-To-UE4-LauncherDev

Runtime/Online/BuildPatchServices up to CL 2528665

diff --git a/Engine/Source/Runtime/Online/BuildPatchServices/Private/BuildPatchInstaller.cpp b/Engine/Source/Runtime/Online/BuildPatchServices/Private/BuildPatchInstaller.cpp
--- a/Engine/Source/Runtime/Online/BuildPatchServices/Private/BuildPatchInstaller.cpp
+++ b/Engine/Source/Runtime/Online/BuildPatchServices/Private/BuildPatchInstaller.cpp
@@ -136,6 +136,7 @@ uint32 FBuildPatchInstaller::Run()
 		bProcessSuccess = bInstallSuccess && RunVerification(CorruptFiles);
 
 		// Clean staging if INSTALL success
+		BuildProgress.SetStateProgress(EBuildPatchProgress::CleanUp, 0.0f);
 		CleanUpTime = FPlatformTime::Seconds();
 		if (bInstallSuccess)
 		{
@@ -221,6 +222,33 @@ uint32 FBuildPatchInstaller::Run()
 	return bSuccess ? 0 : 1;
 }
 
+bool FBuildPatchInstaller::CheckForExternallyInstalledFiles()
+{
+	// Never return true if we are patching, we always expect existing files in that case
+	if (CurrentBuildManifest.IsValid())
+	{
+		return false;
+	}
+
+	// Check the marker file for a previous installation unfinished
+	if (IPlatformFile::GetPlatformPhysical().FileExists(*PreviousMoveMarker))
+	{
+		return true;
+	}
+
+	// Check if any required file is potentially already in place, by comparing file size as a quick 'same file' check
+	TArray<FString> BuildFiles;
+	NewBuildManifest->GetFileList(BuildFiles);
+	for (const FString& BuildFile : BuildFiles)
+	{
+		if (NewBuildManifest->GetFileSize(BuildFile) == IFileManager::Get().FileSize(*(InstallDirectory / BuildFile)))
+		{
+			return true;
+		}
+	}
+	return false;
+}
+
 bool FBuildPatchInstaller::RunInstallation(TArray<FString>& CorruptFiles)
 {
 	GLog->Logf(TEXT("BuildPatchServices: Starting Installation"));
@@ -242,11 +270,12 @@ bool FBuildPatchInstaller::RunInstallation(TArray<FString>& CorruptFiles)
 	// Remove any inventory
 	FBuildPatchFileConstructor::PurgeFileDataInventory();
 
-	// Check if we should skip out of this process
-	bool bPreviousStagingCompleted = FPaths::FileExists(PreviousMoveMarker);
-	if (bPreviousStagingCompleted)
+	// Check if we should skip out of this process due to existing installation,
+	// that will mean we start with the verification stage
+	bool bFirstTimeRun = CorruptFiles.Num() == 0;
+	if (bFirstTimeRun && CheckForExternallyInstalledFiles())
 	{
-		GLog->Logf(TEXT("BuildPatchServices: Detected previous staging completed"));
+		GLog->Logf(TEXT("BuildPatchServices: Detected previous staging completed, or existing files in target directory"));
 		// Set weights for verify only
 		BuildProgress.SetStateWeight(EBuildPatchProgress::Downloading, 0.0f);
 		BuildProgress.SetStateWeight(EBuildPatchProgress::Installing, 0.0f);
diff --git a/Engine/Source/Runtime/Online/BuildPatchServices/Private/BuildPatchInstaller.h b/Engine/Source/Runtime/Online/BuildPatchServices/Private/BuildPatchInstaller.h
--- a/Engine/Source/Runtime/Online/BuildPatchServices/Private/BuildPatchInstaller.h
+++ b/Engine/Source/Runtime/Online/BuildPatchServices/Private/BuildPatchInstaller.h
@@ -156,6 +156,14 @@ public:
 
 private:
 
+	/**
+	 * Checks the installation directory for any already existing files of the correct size, with may account for manual
+	 * installation. Should be used for new installation detecting existing files.
+	 * NB: Not useful for patches, where we'd expect existing files anyway.
+	 * @return    Returns true if there were potentially already installed files
+	 */
+	bool CheckForExternallyInstalledFiles();
+
 	/**
 	 * Runs the installation process
 	 * @param	CorruptFiles	A list of files that were corrupt, to only install those.
diff --git a/Engine/Source/Runtime/Online/BuildPatchServices/Private/BuildPatchServicesModule.cpp b/Engine/Source/Runtime/Online/BuildPatchServices/Private/BuildPatchServicesModule.cpp
--- a/Engine/Source/Runtime/Online/BuildPatchServices/Private/BuildPatchServicesModule.cpp
+++ b/Engine/Source/Runtime/Online/BuildPatchServices/Private/BuildPatchServicesModule.cpp
@@ -232,6 +232,18 @@ void FBuildPatchServicesModule::SetStagingDirectory( const FString& StagingDir )
 void FBuildPatchServicesModule::SetCloudDirectory( const FString& CloudDir )
 {
 	CloudDirectory = CloudDir;
+
+	// Ensure that we remove any double-slash characters apart from:
+	//   1. A double slash following the URI schema
+	//   2. A double slash at the start of the path, indicating a network share
+	CloudDirectory.ReplaceInline(TEXT("\\"), TEXT("/"));
+	bool bIsNetworkPath = CloudDirectory.StartsWith(TEXT("//"));
+	CloudDirectory.ReplaceInline(TEXT("://"), TEXT(":////"));
+	CloudDirectory.ReplaceInline(TEXT("//"), TEXT("/"));
+	if (bIsNetworkPath)
+	{
+		CloudDirectory.InsertAt(0, TEXT("/"));
+	}
 }
 
 void FBuildPatchServicesModule::SetBackupDirectory( const FString& BackupDir )
diff --git a/Engine/Source/Runtime/Online/BuildPatchServices/Public/Interfaces/IBuildPatchServicesModule.h b/Engine/Source/Runtime/Online/BuildPatchServices/Public/Interfaces/IBuildPatchServicesModule.h
--- a/Engine/Source/Runtime/Online/BuildPatchServices/Public/Interfaces/IBuildPatchServicesModule.h
+++ b/Engine/Source/Runtime/Online/BuildPatchServices/Public/Interfaces/IBuildPatchServicesModule.h
@@ -1,11 +1,9 @@
 // Copyright 1998-2015 Epic Games, Inc. All Rights Reserved.
-
-/*=============================================================================
-	IBuildPatchServicesModule.h: Declares the IBuildPatchServicesModule interface.
-=============================================================================*/
-
 #pragma once
 
+#include "IBuildInstaller.h"
+#include "IBuildManifest.h"
+
 class IAnalyticsProvider;
 class FHttpServiceTracker;
 
