Merging using REVERSE UE4-To-UE4-LauncherDev

CL# 2530011 on 29/04/2015 07:03
------------------------------------------
Merging using REVERSE UE4-LauncherDev-To-UE4-LauncherReleases

CL# 2530001 on 29/04/2015 07:00
------------------------------------------
Fixing bug where verification cannot be cancelled promptly

diff --git a/Engine/Source/Runtime/Online/BuildPatchServices/Private/BuildPatchVerification.cpp b/Engine/Source/Runtime/Online/BuildPatchServices/Private/BuildPatchVerification.cpp
--- a/Engine/Source/Runtime/Online/BuildPatchServices/Private/BuildPatchVerification.cpp
+++ b/Engine/Source/Runtime/Online/BuildPatchServices/Private/BuildPatchVerification.cpp
@@ -52,6 +52,12 @@ bool FBuildPatchVerificationImpl::VerifyAgainstDirectory(TArray<FString>& OutDat
 	Manifest->GetFileList(BuildFiles);
 	for (const FString& BuildFile : BuildFiles)
 	{
+		// Break if quitting
+		if (FBuildPatchInstallError::HasFatalError())
+		{
+			break;
+		}
+
 		// Get file details
 		int64 BuildFileSize = Manifest->GetFileSize(BuildFile);
 		FSHAHashData BuildFileHash;
