Fixed a crash when using UAchievementQueryCallbackProxy::CacheAchievementDescriptions
Thanks to GitHub user [REDACTED_PERSON]!

#github https://github.com/EpicGames/UnrealEngine/pull/1106

diff --git a/Engine/Source/Runtime/Online/OnlineSubsystemUtils/Private/AchievementQueryCallbackProxy.cpp b/Engine/Source/Runtime/Online/OnlineSubsystemUtils/Private/AchievementQueryCallbackProxy.cpp
--- a/Engine/Source/Runtime/Online/OnlineSubsystemUtils/Private/AchievementQueryCallbackProxy.cpp
+++ b/Engine/Source/Runtime/Online/OnlineSubsystemUtils/Private/AchievementQueryCallbackProxy.cpp
@@ -25,6 +25,7 @@ UAchievementQueryCallbackProxy* UAchievementQueryCallbackProxy::CacheAchievement
 	UAchievementQueryCallbackProxy* Proxy = NewObject<UAchievementQueryCallbackProxy>();
 	Proxy->PlayerControllerWeakPtr = PlayerController;
 	Proxy->bFetchDescriptions = true;
+	Proxy->WorldContextObject = WorldContextObject;
 	return Proxy;
 }
 
