Make NULL identity id stable so we can store in replays. PIE still returns random id's.

diff --git a/Engine/Source/Runtime/Online/OnlineSubsystemNull/Private/OnlineIdentityNull.cpp b/Engine/Source/Runtime/Online/OnlineSubsystemNull/Private/OnlineIdentityNull.cpp
--- a/Engine/Source/Runtime/Online/OnlineSubsystemNull/Private/OnlineIdentityNull.cpp
+++ b/Engine/Source/Runtime/Online/OnlineSubsystemNull/Private/OnlineIdentityNull.cpp
@@ -38,7 +38,14 @@ inline FString GenerateRandomUserId(int32 LocalUserNum)
 		HostName = Addr->ToString(false);
 	}
 
-	return FString::Printf(TEXT("%s-%s"), *HostName, *FGuid::NewGuid().ToString());
+	if ( !GIsFirstInstance || GIsEditor )
+	{
+		// If we're not the first instance (or in the editor), return truly random id
+		return FString::Printf(TEXT("%s-%s"), *HostName, *FGuid::NewGuid().ToString());
+	}
+
+	// When possible, return a stable user id
+	return FString::Printf( TEXT( "%s-%s" ), *HostName, *FPlatformMisc::GetMachineId().ToString( EGuidFormats::Digits ) );
 }
 
 bool FOnlineIdentityNull::Login(int32 LocalUserNum, const FOnlineAccountCredentials& AccountCredentials)
