Move OnlineGameSessionClient from LocalPlayer to GameInstance

diff --git a/Engine/Source/Runtime/Engine/Classes/Engine/GameInstance.h b/Engine/Source/Runtime/Engine/Classes/Engine/GameInstance.h
--- a/Engine/Source/Runtime/Engine/Classes/Engine/GameInstance.h
+++ b/Engine/Source/Runtime/Engine/Classes/Engine/GameInstance.h
@@ -49,6 +49,10 @@ protected:
 	// Delegate handle that stores delegate for when an invite is accepted by a user
 	FDelegateHandle OnSessionUserInviteAcceptedDelegateHandle;
 	
+	/** Class to manage online services */
+	UPROPERTY()
+	class UOnlineSession* OnlineSession;
+
 public:
 
 	FString PIEMapName;
@@ -181,6 +185,11 @@ public:
 	/** Start playing back a previously recorded replay. */
 	virtual void PlayReplay(const FString& Name);
 
-private:
 	FTimerManager* TimerManager;
+
+	/** @return online session management object associated with this game instance */
+	class UOnlineSession* GetOnlineSession() const { return OnlineSession; }
+
+	/** @return OnlineSession class to use for this game instance  */
+	virtual TSubclassOf<UOnlineSession> GetOnlineSessionClass();
 };
\ No newline at end of file
diff --git a/Engine/Source/Runtime/Engine/Classes/Engine/LocalPlayer.h b/Engine/Source/Runtime/Engine/Classes/Engine/LocalPlayer.h
--- a/Engine/Source/Runtime/Engine/Classes/Engine/LocalPlayer.h
+++ b/Engine/Source/Runtime/Engine/Classes/Engine/LocalPlayer.h
@@ -177,13 +177,6 @@ private:
 	FSceneViewStateReference ViewState;
 	FSceneViewStateReference StereoViewState;
 
-	/** Class to manage online services */
-	UPROPERTY()
-	class UOnlineSession* OnlineSession;
-
-	/** @return OnlineSession class to use for this player controller  */
-	virtual TSubclassOf<UOnlineSession> GetOnlineSessionClass();
-
 	/** The controller ID which this player accepts input from. */
 	int32 ControllerId;
 
@@ -217,10 +210,6 @@ public:
 	bool HandleToggleStreamingVolumesCommand( const TCHAR* Cmd, FOutputDevice& Ar );
 	bool HandleCancelMatineeCommand( const TCHAR* Cmd, FOutputDevice& Ar );
 	
-	// UPlayer interface
-	virtual void HandleDisconnect(class UWorld *World, class UNetDriver *NetDriver) override;
-	// End of UPlayer interface
-
 protected:
 	/**
 	 * Retrieve the viewpoint of this player.
@@ -292,9 +281,6 @@ public:
 	 */
 	virtual void InitOnlineSession();
 
-	/** @return online session management object associated with this player */
-	UOnlineSession* GetOnlineSession() const { return OnlineSession; }
-
 	/**
 	 * Called when the player is removed from the viewport client
 	 */
diff --git a/Engine/Source/Runtime/Engine/Classes/Engine/Player.h b/Engine/Source/Runtime/Engine/Classes/Engine/Player.h
--- a/Engine/Source/Runtime/Engine/Classes/Engine/Player.h
+++ b/Engine/Source/Runtime/Engine/Classes/Engine/Player.h
@@ -40,9 +40,4 @@ public:
 	 * @param    PC - new player controller to assign to player
 	 */
 	ENGINE_API virtual void SwitchController( class APlayerController* PC );
-
-	/**
-	 * Return the main menu after graceful cleanup
-	 */
-	virtual void HandleDisconnect(class UWorld *World, class UNetDriver *NetDriver) {}
 };
diff --git a/Engine/Source/Runtime/Engine/Classes/GameFramework/OnlineSession.h b/Engine/Source/Runtime/Engine/Classes/GameFramework/OnlineSession.h
--- a/Engine/Source/Runtime/Engine/Classes/GameFramework/OnlineSession.h
+++ b/Engine/Source/Runtime/Engine/Classes/GameFramework/OnlineSession.h
@@ -17,10 +17,10 @@ class ENGINE_API UOnlineSession : public UObject
 public:
 
 	/** Register all delegates needed to manage online sessions.  */
-	virtual void RegisterOnlineDelegates(class UWorld* InWorld) {};
+	virtual void RegisterOnlineDelegates() {};
 
 	/** Tear down all delegates used to manage online sessions.	 */
-	virtual void ClearOnlineDelegates(class UWorld* InWorld) {};
+	virtual void ClearOnlineDelegates() {};
 
 	/** Called to tear down any online sessions and return to main menu	 */
 	virtual void HandleDisconnect(UWorld *World, class UNetDriver *NetDriver);
diff --git a/Engine/Source/Runtime/Engine/Private/GameInstance.cpp b/Engine/Source/Runtime/Engine/Private/GameInstance.cpp
--- a/Engine/Source/Runtime/Engine/Private/GameInstance.cpp
+++ b/Engine/Source/Runtime/Engine/Private/GameInstance.cpp
@@ -59,6 +59,13 @@ void UGameInstance::Init()
 			SessionInt->AddOnSessionUserInviteAcceptedDelegate_Handle(FOnSessionUserInviteAcceptedDelegate::CreateUObject(this, &UGameInstance::HandleSessionUserInviteAccepted));
 		}
 	}
+
+	UClass* SpawnClass = GetOnlineSessionClass();
+	OnlineSession = NewObject<UOnlineSession>(this, SpawnClass);
+	if (OnlineSession)
+	{
+		OnlineSession->RegisterOnlineDelegates();
+	}
 }
 
 void UGameInstance::Shutdown()
@@ -649,23 +656,7 @@ void UGameInstance::OnSessionUserInviteAccepted(const bool bWasSuccess, const in
 	{
 		if (InviteResult.IsValid())
 		{
-			for (ULocalPlayer* LocalPlayer : LocalPlayers)
-			{
-				// Route the call to the actual user that accepted the invite
-				if (LocalPlayer->GetCachedUniqueNetId() == UserId)
-				{
-					LocalPlayer->GetOnlineSession()->OnSessionUserInviteAccepted(bWasSuccess, ControllerId, UserId, InviteResult);
-					return;
-				}
-			}
-
-			// Go ahead and have the active local player handle accepting the invite. A game can detect that the user id is different and handle
-			// it how it needs to.
-			ULocalPlayer* LocalPlayer = GetFirstGamePlayer();
-			if (LocalPlayer)
-			{
-				LocalPlayer->GetOnlineSession()->OnSessionUserInviteAccepted(bWasSuccess, ControllerId, UserId, InviteResult);
-			}
+			GetOnlineSession()->OnSessionUserInviteAccepted(bWasSuccess, ControllerId, UserId, InviteResult);
 		}
 		else
 		{
@@ -784,3 +775,8 @@ void UGameInstance::PlayReplay(const FString& Name)
 		FCoreUObjectDelegates::PostDemoPlay.Broadcast();
 	}
 }
+
+TSubclassOf<UOnlineSession> UGameInstance::GetOnlineSessionClass()
+{
+	return UOnlineSession::StaticClass();
+}
diff --git a/Engine/Source/Runtime/Engine/Private/LocalPlayer.cpp b/Engine/Source/Runtime/Engine/Private/LocalPlayer.cpp
--- a/Engine/Source/Runtime/Engine/Private/LocalPlayer.cpp
+++ b/Engine/Source/Runtime/Engine/Private/LocalPlayer.cpp
@@ -185,47 +185,11 @@ void ULocalPlayer::PlayerAdded(UGameViewportClient* InViewportClient, int32 InCo
 
 void ULocalPlayer::InitOnlineSession()
 {
-	if (OnlineSession == NULL)
-	{
-		UClass* SpawnClass = GetOnlineSessionClass();
-		OnlineSession = NewObject<UOnlineSession>(this, SpawnClass);
-		if (OnlineSession)
-		{
-			UWorld* World = GetWorld();
-			if (World != NULL)
-			{
-				OnlineSession->RegisterOnlineDelegates(World);
-			}
-		}
-	}
+	// FIXME: This may be obsolete, still here to support a few straggler cases that do stuff in child classes
 }
 
 void ULocalPlayer::PlayerRemoved()
 {
-	// Clear all online delegates
-	if (OnlineSession != NULL)
-	{
-		OnlineSession->ClearOnlineDelegates(GetWorld());
-		OnlineSession = NULL;
-	}
-}
-
-TSubclassOf<UOnlineSession> ULocalPlayer::GetOnlineSessionClass()
-{
-	return UOnlineSession::StaticClass();
-}
-
-void ULocalPlayer::HandleDisconnect(UWorld *World, UNetDriver *NetDriver)
-{
-	if (OnlineSession)
-	{
-		OnlineSession->HandleDisconnect(World, NetDriver);
-	}
-	else
-	{
-		// Let the engine cleanup this disconnect
-		GEngine->HandleDisconnect(World, NetDriver);
-	}
 }
 
 bool ULocalPlayer::SpawnPlayActor(const FString& URL,FString& OutError, UWorld* InWorld)
diff --git a/Engine/Source/Runtime/Engine/Private/PlayerController.cpp b/Engine/Source/Runtime/Engine/Private/PlayerController.cpp
--- a/Engine/Source/Runtime/Engine/Private/PlayerController.cpp
+++ b/Engine/Source/Runtime/Engine/Private/PlayerController.cpp
@@ -1611,9 +1611,9 @@ void APlayerController::LocalTravel( const FString& FURL )
 void APlayerController::ClientReturnToMainMenu_Implementation(const FString& ReturnReason)
 {
 	UWorld* World = GetWorld();
-	if (Player)
+	if (GetGameInstance() && GetGameInstance()->GetOnlineSession())
 	{
-		Player->HandleDisconnect(World, World->GetNetDriver());
+		GetGameInstance()->GetOnlineSession()->HandleDisconnect(World, World->GetNetDriver());
 	}
 	else
 	{
@@ -3158,33 +3158,17 @@ void APlayerController::ClientWasKicked_Implementation(const FText& KickReason)
 
 void APlayerController::ClientStartOnlineSession_Implementation()
 {
-	if (IsPrimaryPlayer() && PlayerState)
+	if (IsPrimaryPlayer() && PlayerState && GetGameInstance() && GetGameInstance()->GetOnlineSession())
 	{
-		ULocalPlayer* LP = Cast<ULocalPlayer>(Player);
-		if (LP)
-		{
-			UOnlineSession* Session = LP->GetOnlineSession();
-			if (Session)
-			{
-				Session->StartOnlineSession(PlayerState->SessionName);
-			}
-		}
+		GetGameInstance()->GetOnlineSession()->StartOnlineSession(PlayerState->SessionName);
 	}
 }
 
 void APlayerController::ClientEndOnlineSession_Implementation()
 {
-	if (IsPrimaryPlayer() && PlayerState)
+	if (IsPrimaryPlayer() && PlayerState && GetGameInstance() && GetGameInstance()->GetOnlineSession())
 	{
-		ULocalPlayer* LP = Cast<ULocalPlayer>(Player);
-		if (LP)
-		{
-			UOnlineSession* Session = LP->GetOnlineSession();
-			if (Session)
-			{
-				Session->EndOnlineSession(PlayerState->SessionName);
-			}
-		}
+		GetGameInstance()->GetOnlineSession()->EndOnlineSession(PlayerState->SessionName);
 	}
 }
 
diff --git a/Engine/Source/Runtime/Engine/Private/UnrealEngine.cpp b/Engine/Source/Runtime/Engine/Private/UnrealEngine.cpp
--- a/Engine/Source/Runtime/Engine/Private/UnrealEngine.cpp
+++ b/Engine/Source/Runtime/Engine/Private/UnrealEngine.cpp
@@ -105,6 +105,8 @@
 #include "JsonInternationalizationArchiveSerializer.h"
 #include "JsonInternationalizationManifestSerializer.h"
 
+#include "GameFramework/OnlineSession.h"
+
 DEFINE_LOG_CATEGORY(LogEngine);
 
 IMPLEMENT_MODULE( FEngineModule, Engine );
@@ -7900,13 +7902,12 @@ static inline void CallHandleDisconnectForFailure(UWorld* InWorld, UNetDriver* N
 		}
 	}
 	
-	// A valid world or NetDriver is required to look up a ULocalPlayer.
+	// A valid world or NetDriver is required to look up a GameInstance/ULocalPlayer.
 	if (InWorld)
 	{
-		ULocalPlayer* const LP = GEngine->GetFirstGamePlayer(InWorld);
-		if(ensure(LP))
+		if (InWorld->GetGameInstance() != nullptr && InWorld->GetGameInstance()->GetOnlineSession() != nullptr)
 		{
-			LP->HandleDisconnect(InWorld, NetDriver);
+			InWorld->GetGameInstance()->GetOnlineSession()->HandleDisconnect(InWorld, NetDriver);
 		}
 	}
 	else if(NetDriver && NetDriver->NetDriverName == NAME_PendingNetDriver)
@@ -7915,12 +7916,9 @@ static inline void CallHandleDisconnectForFailure(UWorld* InWorld, UNetDriver* N
 		FWorldContext &Context = GEngine->GetWorldContextFromPendingNetGameNetDriverChecked(NetDriver);
 		check(Context.OwningGameInstance != NULL && Context.OwningGameInstance->GetFirstGamePlayer() != NULL);
 
-		ULocalPlayer* const LP = Context.OwningGameInstance->GetFirstGamePlayer();
-		if(ensure(LP))
+		if ( Context.OwningGameInstance->GetOnlineSession() != nullptr )
 		{
-			// Use the world on the context instead since InWorld is null, it should be valid. This way we can do a travel
-			// in UEngine::HandleDisconnect if it gets called
-			LP->HandleDisconnect(Context.World(), NetDriver);
+			Context.OwningGameInstance->GetOnlineSession()->HandleDisconnect(Context.World(), NetDriver);
 		}
 	}
 	else
diff --git a/Engine/Source/Runtime/Online/OnlineSubsystemUtils/Classes/OnlineSessionClient.h b/Engine/Source/Runtime/Online/OnlineSubsystemUtils/Classes/OnlineSessionClient.h
--- a/Engine/Source/Runtime/Online/OnlineSubsystemUtils/Classes/OnlineSessionClient.h
+++ b/Engine/Source/Runtime/Online/OnlineSubsystemUtils/Classes/OnlineSessionClient.h
@@ -19,10 +19,6 @@ class ONLINESUBSYSTEMUTILS_API UOnlineSessionClient : public UOnlineSession
 	GENERATED_UCLASS_BODY()
 
 protected:
-
-	/** Reference to the online sessions interface */
-	IOnlineSessionPtr SessionInt;
-
 	/** Delegate for destroying a session after previously ending it */
 	FOnEndSessionCompleteDelegate OnEndForJoinSessionCompleteDelegate;
 	/** Delegate for joining a new session after previously destroying it */
@@ -53,23 +49,13 @@ protected:
 	UPROPERTY(Transient)
 	bool bHandlingDisconnect;
 
+	/** @return the current game instance */
+	virtual UGameInstance* GetGameInstance() const;
+
 	/** @return the current game world */
 	virtual UWorld* GetWorld() const override;
 
-	/**
-	 * Get the player controller associated with this session (from the owning ULocalPlayer)
-	 */
-	APlayerController* GetPlayerController();
-
-	/** @return the unique id associated with the player who owns this class */
-	TSharedPtr<FUniqueNetId> GetUniqueId();
-
-	/**
-	 * Helper function to retrieve the controller id of the owning controller
-	 *
-	 * @return controller id of the controller
-	 */
-	int32 GetControllerId();
+	virtual IOnlineSessionPtr GetSessionInt();
 
 	/**
 	 * Chance for the session client to handle the disconnect
@@ -188,8 +174,8 @@ protected:
 public:
 
 	// UOnlineSession interface begin
-	virtual void RegisterOnlineDelegates(class UWorld* InWorld) override;
-	virtual void ClearOnlineDelegates(class UWorld* InWorld) override;
+	virtual void RegisterOnlineDelegates() override;
+	virtual void ClearOnlineDelegates() override;
 	virtual void HandleDisconnect(UWorld *World, UNetDriver *NetDriver) override;
 	virtual void StartOnlineSession(FName SessionName) override;
 	virtual void EndOnlineSession(FName SessionName) override;
diff --git a/Engine/Source/Runtime/Online/OnlineSubsystemUtils/Private/OnlineSessionClient.cpp b/Engine/Source/Runtime/Online/OnlineSubsystemUtils/Private/OnlineSessionClient.cpp
--- a/Engine/Source/Runtime/Online/OnlineSubsystemUtils/Private/OnlineSessionClient.cpp
+++ b/Engine/Source/Runtime/Online/OnlineSubsystemUtils/Private/OnlineSessionClient.cpp
@@ -16,63 +16,37 @@ UOnlineSessionClient::UOnlineSessionClient(const FObjectInitializer& ObjectIniti
 	bIsFromInvite = false;
 }
 
-UWorld* UOnlineSessionClient::GetWorld() const
+UGameInstance* UOnlineSessionClient::GetGameInstance() const
 {
-	ULocalPlayer* LP = Cast<ULocalPlayer>(GetOuter());
-	check(LP);
-	return LP->GetWorld();
+	return CastChecked<UGameInstance>(GetOuter());
 }
 
-APlayerController* UOnlineSessionClient::GetPlayerController()
+UWorld* UOnlineSessionClient::GetWorld() const
 {
-	ULocalPlayer* LP = Cast<ULocalPlayer>(GetOuter());
-	if (LP != nullptr)
-	{
-		return LP->PlayerController;
-	}
-
-	return nullptr;
+	return GetGameInstance()->GetWorld();
 }
 
-int32 UOnlineSessionClient::GetControllerId()
+IOnlineSessionPtr UOnlineSessionClient::GetSessionInt()
 {
-	ULocalPlayer* LP = Cast<ULocalPlayer>(GetOuter());
-	if (LP != nullptr)
+	if (GetWorld() == nullptr)
 	{
-		return LP->GetControllerId();
+		UE_LOG(LogOnline, Warning, TEXT("UOnlineSessionClient::GetSessionInt: Called with NULL world."));
+		return nullptr;
 	}
 
-	return INVALID_CONTROLLERID;
+	return Online::GetSessionInterface(GetWorld());
 }
 
-TSharedPtr<FUniqueNetId> UOnlineSessionClient::GetUniqueId()
+void UOnlineSessionClient::RegisterOnlineDelegates()
 {
-	ULocalPlayer* LP = Cast<ULocalPlayer>(GetOuter());
-	if (LP != nullptr)
-	{
-		return LP->GetPreferredUniqueNetId();
-	}
-
-	return nullptr;
+	OnJoinSessionCompleteDelegate           = FOnJoinSessionCompleteDelegate   ::CreateUObject(this, &UOnlineSessionClient::OnJoinSessionComplete);
+	OnEndForJoinSessionCompleteDelegate     = FOnEndSessionCompleteDelegate    ::CreateUObject(this, &UOnlineSessionClient::OnEndForJoinSessionComplete);
+	OnDestroyForJoinSessionCompleteDelegate = FOnDestroySessionCompleteDelegate::CreateUObject(this, &UOnlineSessionClient::OnDestroyForJoinSessionComplete);
+	OnDestroyForMainMenuCompleteDelegate	= FOnDestroySessionCompleteDelegate::CreateUObject(this, &UOnlineSessionClient::OnDestroyForMainMenuComplete);
 }
 
-void UOnlineSessionClient::RegisterOnlineDelegates(UWorld* InWorld)
+void UOnlineSessionClient::ClearOnlineDelegates()
 {
-	IOnlineSubsystem* OnlineSub = Online::GetSubsystem(InWorld);
-	if (OnlineSub)
-	{
-		OnJoinSessionCompleteDelegate           = FOnJoinSessionCompleteDelegate   ::CreateUObject(this, &UOnlineSessionClient::OnJoinSessionComplete);
-		OnEndForJoinSessionCompleteDelegate     = FOnEndSessionCompleteDelegate    ::CreateUObject(this, &UOnlineSessionClient::OnEndForJoinSessionComplete);
-		OnDestroyForJoinSessionCompleteDelegate = FOnDestroySessionCompleteDelegate::CreateUObject(this, &UOnlineSessionClient::OnDestroyForJoinSessionComplete);
-	}
-
-	// Register disconnect delegate even if we don't have an online system
-	OnDestroyForMainMenuCompleteDelegate = FOnDestroySessionCompleteDelegate::CreateUObject(this, &UOnlineSessionClient::OnDestroyForMainMenuComplete);
-}
-
-void UOnlineSessionClient::ClearOnlineDelegates(UWorld* InWorld)
-{
-	SessionInt = NULL;
 }
 
 void UOnlineSessionClient::OnSessionUserInviteAccepted(bool bWasSuccessful, int32 ControllerId, TSharedPtr<FUniqueNetId> UserId, const FOnlineSessionSearchResult& SearchResult)
@@ -97,7 +71,11 @@ void UOnlineSessionClient::OnSessionUserInviteAccepted(bool bWasSuccessful, int3
 void UOnlineSessionClient::OnEndForJoinSessionComplete(FName SessionName, bool bWasSuccessful)
 {
 	UE_LOG(LogOnline, Verbose, TEXT("OnEndForJoinSessionComplete %s bSuccess: %d"), *SessionName.ToString(), bWasSuccessful);
-	SessionInt->ClearOnEndSessionCompleteDelegate_Handle(OnEndForJoinSessionCompleteDelegateHandle);
+	IOnlineSessionPtr SessionInt = GetSessionInt();
+	if (SessionInt.IsValid())
+	{
+		SessionInt->ClearOnEndSessionCompleteDelegate_Handle(OnEndForJoinSessionCompleteDelegateHandle);
+	}
 	DestroyExistingSession_Impl(OnDestroyForJoinSessionCompleteDelegateHandle, SessionName, OnDestroyForJoinSessionCompleteDelegate);
 }
 
@@ -117,6 +95,8 @@ FDelegateHandle UOnlineSessionClient::EndExistingSession_Impl(FName SessionName,
 {
 	FDelegateHandle Result;
 
+	IOnlineSessionPtr SessionInt = GetSessionInt();
+
 	if (SessionInt.IsValid())
 	{
 		Result = SessionInt->AddOnEndSessionCompleteDelegate_Handle(Delegate);
@@ -133,6 +113,9 @@ FDelegateHandle UOnlineSessionClient::EndExistingSession_Impl(FName SessionName,
 void UOnlineSessionClient::OnDestroyForJoinSessionComplete(FName SessionName, bool bWasSuccessful)
 {
 	UE_LOG(LogOnline, Verbose, TEXT("OnDestroyForJoinSessionComplete %s bSuccess: %d"), *SessionName.ToString(), bWasSuccessful);
+
+	IOnlineSessionPtr SessionInt = GetSessionInt();
+
 	if (SessionInt.IsValid())
 	{
 		SessionInt->ClearOnDestroySessionCompleteDelegate_Handle(OnDestroyForJoinSessionCompleteDelegateHandle);
@@ -149,6 +132,9 @@ void UOnlineSessionClient::OnDestroyForJoinSessionComplete(FName SessionName, bo
 void UOnlineSessionClient::OnDestroyForMainMenuComplete(FName SessionName, bool bWasSuccessful)
 {
 	UE_LOG(LogOnline, Verbose, TEXT("OnDestroyForMainMenuComplete %s bSuccess: %d"), *SessionName.ToString(), bWasSuccessful);
+
+	IOnlineSessionPtr SessionInt = GetSessionInt();
+
 	if (SessionInt.IsValid())
 	{
 		SessionInt->ClearOnDestroySessionCompleteDelegate_Handle(OnDestroyForMainMenuCompleteDelegateHandle);
@@ -175,6 +161,8 @@ void UOnlineSessionClient::DestroyExistingSession(FName SessionName, FOnDestroyS
  */
 void UOnlineSessionClient::DestroyExistingSession_Impl(FDelegateHandle& OutResult, FName SessionName, FOnDestroySessionCompleteDelegate& Delegate)
 {
+	IOnlineSessionPtr SessionInt = GetSessionInt();
+
 	if (SessionInt.IsValid())
 	{
 		OutResult = SessionInt->AddOnDestroySessionCompleteDelegate_Handle(Delegate);
@@ -190,27 +178,33 @@ void UOnlineSessionClient::DestroyExistingSession_Impl(FDelegateHandle& OutResul
 void UOnlineSessionClient::OnJoinSessionComplete(FName SessionName, EOnJoinSessionCompleteResult::Type Result)
 {
 	UE_LOG(LogOnline, Verbose, TEXT("OnJoinSessionComplete %s bSuccess: %d"), *SessionName.ToString(), static_cast<int32>(Result));
-	SessionInt->ClearOnJoinSessionCompleteDelegate_Handle(OnJoinSessionCompleteDelegateHandle);
 
-	if (Result == EOnJoinSessionCompleteResult::Success)
+	IOnlineSessionPtr SessionInt = GetSessionInt();
+
+	if (SessionInt.IsValid())
 	{
-		FString URL;
-		if (SessionInt->GetResolvedConnectString(SessionName, URL))
+		SessionInt->ClearOnJoinSessionCompleteDelegate_Handle(OnJoinSessionCompleteDelegateHandle);
+
+		if (Result == EOnJoinSessionCompleteResult::Success)
 		{
-			APlayerController* PC = GetPlayerController();
-			if (PC)
+			FString URL;
+			if (SessionInt->GetResolvedConnectString(SessionName, URL))
 			{
-				if (bIsFromInvite)
+				APlayerController* PC = GetGameInstance()->GetFirstLocalPlayerController();;
+				if (PC)
 				{
-					URL += TEXT("?bIsFromInvite");
-					bIsFromInvite = false;
+					if (bIsFromInvite)
+					{
+						URL += TEXT("?bIsFromInvite");
+						bIsFromInvite = false;
+					}
+					PC->ClientTravel(URL, TRAVEL_Absolute);
 				}
-				PC->ClientTravel(URL, TRAVEL_Absolute);
 			}
-		}
-		else
-		{
-			UE_LOG(LogOnline, Warning, TEXT("Failed to join session %s"), *SessionName.ToString());
+			else
+			{
+				UE_LOG(LogOnline, Warning, TEXT("Failed to join session %s"), *SessionName.ToString());
+			}
 		}
 	}
 }
@@ -218,20 +212,20 @@ void UOnlineSessionClient::OnJoinSessionComplete(FName SessionName, EOnJoinSessi
 void UOnlineSessionClient::JoinSession(FName SessionName, const FOnlineSessionSearchResult& SearchResult)
 {
 	// Clean up existing sessions if applicable
-	EOnlineSessionState::Type SessionState = SessionInt->GetSessionState(SessionName);
-	if (SessionState != EOnlineSessionState::NoSession)
-	{
-		CachedSessionResult = SearchResult;
-		OnEndForJoinSessionCompleteDelegateHandle = EndExistingSession_Impl(SessionName, OnEndForJoinSessionCompleteDelegate);
-	}
-	else
-	{
-		ULocalPlayer* LocalPlayer = static_cast<ULocalPlayer*>(GetOuter());
-		check(LocalPlayer);
-		UGameInstance * const GameInstance = LocalPlayer->GetGameInstance();
-		check(GameInstance);
+	IOnlineSessionPtr SessionInt = GetSessionInt();
 
-		GameInstance->JoinSession(LocalPlayer, SearchResult);
+	if (SessionInt.IsValid())
+	{
+		EOnlineSessionState::Type SessionState = SessionInt->GetSessionState(SessionName);
+		if (SessionState != EOnlineSessionState::NoSession)
+		{
+			CachedSessionResult = SearchResult;
+			OnEndForJoinSessionCompleteDelegateHandle = EndExistingSession_Impl(SessionName, OnEndForJoinSessionCompleteDelegate);
+		}
+		else
+		{
+			GetGameInstance()->JoinSession(GetGameInstance()->GetFirstGamePlayer(), SearchResult);
+		}
 	}
 }
 
