Fixed format string for beacon opened log message.

diff --git a/Engine/Source/Runtime/Online/OnlineSubsystemUtils/Private/OnlineBeacon.cpp b/Engine/Source/Runtime/Online/OnlineSubsystemUtils/Private/OnlineBeacon.cpp
--- a/Engine/Source/Runtime/Online/OnlineSubsystemUtils/Private/OnlineBeacon.cpp
+++ b/Engine/Source/Runtime/Online/OnlineSubsystemUtils/Private/OnlineBeacon.cpp
@@ -102,7 +102,7 @@ void AOnlineBeacon::NotifyAcceptedConnection(UNetConnection* Connection)
 {
 	check(NetDriver != NULL);
 	check(NetDriver->ServerConnection == NULL);
-	UE_LOG(LogNet, Log, TEXT("Open %s %0.2f %s"), *GetName(), FPlatformTime::StrTimestamp(), *Connection->LowLevelGetRemoteAddress());
+	UE_LOG(LogNet, Log, TEXT("Open %s %s %s"), *GetName(), FPlatformTime::StrTimestamp(), *Connection->LowLevelGetRemoteAddress());
 }
 
 bool AOnlineBeacon::NotifyAcceptingChannel(UChannel* Channel)
