[INTEGRATE] merge ue4-launcherDev->ue4 cl 2545202

diff --git a/Engine/Source/Runtime/Online/BuildPatchServices/Public/Interfaces/IBuildInstaller.h b/Engine/Source/Runtime/Online/BuildPatchServices/Public/Interfaces/IBuildInstaller.h
--- a/Engine/Source/Runtime/Online/BuildPatchServices/Public/Interfaces/IBuildInstaller.h
+++ b/Engine/Source/Runtime/Online/BuildPatchServices/Public/Interfaces/IBuildInstaller.h
@@ -14,6 +14,30 @@ typedef TSharedRef< class IBuildInstaller, ESPMode::ThreadSafe > IBuildInstaller
  */
 struct FBuildInstallStats
 {
+	// Constructor
+	FBuildInstallStats()
+		: NumFilesInBuild(0)
+		, NumFilesOutdated(0)
+		, NumFilesToRemove(0)
+		, NumChunksRequired(0)
+		, ChunksQueuedForDownload(0)
+		, ChunksLocallyAvailable(0)
+		, NumChunksDownloaded(0)
+		, NumChunksRecycled(0)
+		, NumChunksCacheBooted(0)
+		, NumDriveCacheChunkLoads(0)
+		, NumRecycleFailures(0)
+		, NumDriveCacheLoadFailures(0)
+		, TotalDownloadedData(0)
+		, AverageDownloadSpeed(0.0)
+		, TheoreticalDownloadTime(0.0f)
+		, VerifyTime(0.0f)
+		, CleanUpTime(0.0f)
+		, ProcessExecuteTime(0.0f)
+		, ProcessPausedTime(0.0f)
+		, ProcessSuccess(false)
+	{}
+
 	// The name of the app being installed
 	FString AppName;
 	// The version string currently installed, or "NONE"
