Making OnlineAsyncTaskManagerNull work with -nothreading

diff --git a/Engine/Source/Runtime/Online/OnlineSubsystemNull/Private/OnlineAsyncTaskManagerNull.cpp b/Engine/Source/Runtime/Online/OnlineSubsystemNull/Private/OnlineAsyncTaskManagerNull.cpp
--- a/Engine/Source/Runtime/Online/OnlineSubsystemNull/Private/OnlineAsyncTaskManagerNull.cpp
+++ b/Engine/Source/Runtime/Online/OnlineSubsystemNull/Private/OnlineAsyncTaskManagerNull.cpp
@@ -7,6 +7,6 @@
 void FOnlineAsyncTaskManagerNull::OnlineTick()
 {
 	check(NullSubsystem);
-	check(FPlatformTLS::GetCurrentThreadId() == OnlineThreadId);
+	check(FPlatformTLS::GetCurrentThreadId() == OnlineThreadId || !FPlatformProcess::SupportsMultithreading());
 }
 
