BuildPatchServices: Fix for detecting previous move when resuming a patch

diff --git a/Engine/Source/Runtime/Online/BuildPatchServices/Private/BuildPatchInstaller.cpp b/Engine/Source/Runtime/Online/BuildPatchServices/Private/BuildPatchInstaller.cpp
--- a/Engine/Source/Runtime/Online/BuildPatchServices/Private/BuildPatchInstaller.cpp
+++ b/Engine/Source/Runtime/Online/BuildPatchServices/Private/BuildPatchInstaller.cpp
@@ -224,18 +224,18 @@ uint32 FBuildPatchInstaller::Run()
 
 bool FBuildPatchInstaller::CheckForExternallyInstalledFiles()
 {
-	// Never return true if we are patching, we always expect existing files in that case
-	if (CurrentBuildManifest.IsValid())
-	{
-		return false;
-	}
-
 	// Check the marker file for a previous installation unfinished
 	if (IPlatformFile::GetPlatformPhysical().FileExists(*PreviousMoveMarker))
 	{
 		return true;
 	}
 
+	// If we are patching, but without the marker, we should not return true, the existing files will be old installation
+	if (CurrentBuildManifest.IsValid())
+	{
+		return false;
+	}
+
 	// Check if any required file is potentially already in place, by comparing file size as a quick 'same file' check
 	TArray<FString> BuildFiles;
 	NewBuildManifest->GetFileList(BuildFiles);
