Fix typo in filename

diff --git a/Engine/Source/Runtime/Online/Android/OnlineSubsystemGooglePlay/Private/OnlineIdentityInterfaceGooglePlay.cpp b/Engine/Source/Runtime/Online/Android/OnlineSubsystemGooglePlay/Private/OnlineIdentityInterfaceGooglePlay.cpp
new file mode 100644
--- /dev/null
+++ b/Engine/Source/Runtime/Online/Android/OnlineSubsystemGooglePlay/Private/OnlineIdentityInterfaceGooglePlay.cpp
@@ -0,0 +1,196 @@
+// Copyright 1998-2015 Epic Games, Inc. All Rights Reserved.
+
+#include "OnlineSubsystemGooglePlayPrivatePCH.h"
+
+FOnlineIdentityGooglePlay::FPendingConnection FOnlineIdentityGooglePlay::PendingConnectRequest;
+
+FOnlineIdentityGooglePlay::FOnlineIdentityGooglePlay(FOnlineSubsystemGooglePlay* InSubsystem)
+	: bPrevLoggedIn(false)
+	, bLoggedIn(false)
+	, PlayerAlias("NONE")
+	, MainSubsystem(InSubsystem)
+	, bRegisteringUser(false)
+	, bLoggingInUser(false)
+
+{
+	UE_LOG(LogOnline, Display, TEXT("FOnlineIdentityAndroid::FOnlineIdentityAndroid()"));
+	PendingConnectRequest.ConnectionInterface = this;
+}
+
+TSharedPtr<FUserOnlineAccount> FOnlineIdentityGooglePlay::GetUserAccount(const FUniqueNetId& UserId) const
+{
+	return nullptr;
+}
+
+
+TArray<TSharedPtr<FUserOnlineAccount> > FOnlineIdentityGooglePlay::GetAllUserAccounts() const
+{
+	TArray<TSharedPtr<FUserOnlineAccount> > Result;
+
+	return Result;
+}
+
+
+bool FOnlineIdentityGooglePlay::Login(int32 LocalUserNum, const FOnlineAccountCredentials& AccountCredentials) 
+{
+	bool bStartedLogin = false;
+	if (bLoggedIn)
+	{
+		// already logged in so just report all is ok!
+		// Now logged in
+		bStartedLogin = true;
+
+		static const int32 MAX_TEXT_LINE_LEN = 32;
+		TCHAR Line[MAX_TEXT_LINE_LEN + 1] = { 0 };
+		int32 Len = FCString::Snprintf(Line, MAX_TEXT_LINE_LEN, TEXT("%d"), LocalUserNum);
+
+		const FString PlayerId(Line);
+		UniqueNetId = MakeShareable(new FUniqueNetIdString(PlayerId));
+		TriggerOnLoginCompleteDelegates(LocalUserNum, true, *UniqueNetId, TEXT(""));
+	}
+	else if (!PendingConnectRequest.IsConnectionPending)
+	{
+		// Kick the login sequence...
+		bStartedLogin = true;
+		PendingConnectRequest.IsConnectionPending = true;
+	}
+	else
+	{
+		TriggerOnLoginCompleteDelegates(LocalUserNum, false, FUniqueNetIdString(TEXT("")), FString("Already trying to login"));
+	}
+	
+	return bStartedLogin;
+}
+
+
+bool FOnlineIdentityGooglePlay::Logout(int32 LocalUserNum)
+{
+	TriggerOnLogoutCompleteDelegates(LocalUserNum, false);
+	return false;
+}
+
+
+bool FOnlineIdentityGooglePlay::AutoLogin(int32 LocalUserNum)
+{
+	return Login(LocalUserNum, FOnlineAccountCredentials());
+}
+
+
+ELoginStatus::Type FOnlineIdentityGooglePlay::GetLoginStatus(int32 LocalUserNum) const
+{
+	ELoginStatus::Type LoginStatus = ELoginStatus::NotLoggedIn;
+
+	if(LocalUserNum < MAX_LOCAL_PLAYERS && bLoggedIn)
+	{
+		UE_LOG(LogOnline, Display, TEXT("FOnlineIdentityAndroid::GetLoginStatus - Logged in"));
+		LoginStatus = ELoginStatus::LoggedIn;
+	}
+	else
+	{
+		UE_LOG(LogOnline, Display, TEXT("FOnlineIdentityAndroid::GetLoginStatus - Not logged in"));
+	}
+
+	return LoginStatus;
+}
+
+
+ELoginStatus::Type FOnlineIdentityGooglePlay::GetLoginStatus(const FUniqueNetId& UserId) const
+{
+	ELoginStatus::Type LoginStatus = ELoginStatus::NotLoggedIn;
+
+	if(bLoggedIn)
+	{
+		UE_LOG(LogOnline, Display, TEXT("FOnlineIdentityAndroid::GetLoginStatus - Logged in"));
+		LoginStatus = ELoginStatus::LoggedIn;
+	}
+	else
+	{
+		UE_LOG(LogOnline, Display, TEXT("FOnlineIdentityAndroid::GetLoginStatus - Not logged in"));
+	}
+
+	return LoginStatus;
+}
+
+
+TSharedPtr<const FUniqueNetId> FOnlineIdentityGooglePlay::GetUniquePlayerId(int32 LocalUserNum) const
+{
+	TSharedPtr<const FUniqueNetId> NewID = MakeShareable(new FUniqueNetIdString(""));
+	return NewID;
+}
+
+
+TSharedPtr<const FUniqueNetId> FOnlineIdentityGooglePlay::CreateUniquePlayerId(uint8* Bytes, int32 Size)
+{
+	if( Bytes && Size == sizeof(uint64) )
+	{
+		int32 StrLen = FCString::Strlen((TCHAR*)Bytes);
+		if (StrLen > 0)
+		{
+			FString StrId((TCHAR*)Bytes);
+			return MakeShareable(new FUniqueNetIdString(StrId));
+		}
+	}
+	return NULL;
+}
+
+
+TSharedPtr<const FUniqueNetId> FOnlineIdentityGooglePlay::CreateUniquePlayerId(const FString& Str)
+{
+	return MakeShareable(new FUniqueNetIdString(Str));
+}
+
+
+FString FOnlineIdentityGooglePlay::GetPlayerNickname(int32 LocalUserNum) const
+{
+	UE_LOG(LogOnline, Display, TEXT("FOnlineIdentityAndroid::GetPlayerNickname"));
+	return PlayerAlias;
+}
+
+FString FOnlineIdentityGooglePlay::GetPlayerNickname(const FUniqueNetId& UserId) const
+{
+	UE_LOG(LogOnline, Display, TEXT("FOnlineIdentityAndroid::GetPlayerNickname"));
+	return PlayerAlias;
+}
+
+FString FOnlineIdentityGooglePlay::GetAuthToken(int32 LocalUserNum) const
+{
+	UE_LOG(LogOnline, Display, TEXT("FOnlineIdentityAndroid::GetAuthToken not implemented"));
+	FString ResultToken;
+	return ResultToken;
+}
+
+void FOnlineIdentityGooglePlay::Tick(float DeltaTime)
+{
+}
+
+void FOnlineIdentityGooglePlay::OnLoginCompleted(const int playerID, const gpg::AuthStatus errorCode)
+{
+	static const int32 MAX_TEXT_LINE_LEN = 32;
+	TCHAR Line[MAX_TEXT_LINE_LEN + 1] = { 0 };
+	int32 Len = FCString::Snprintf(Line, MAX_TEXT_LINE_LEN, TEXT("%d"), playerID);
+
+	UniqueNetId = MakeShareable(new FUniqueNetIdString(Line));
+	bLoggedIn = errorCode == gpg::AuthStatus::VALID;
+	TriggerOnLoginCompleteDelegates(playerID, bLoggedIn, *UniqueNetId, TEXT(""));
+
+	PendingConnectRequest.IsConnectionPending = false;
+}
+
+void FOnlineIdentityGooglePlay::GetUserPrivilege(const FUniqueNetId& UserId, EUserPrivileges::Type Privilege, const FOnGetUserPrivilegeCompleteDelegate& Delegate)
+{
+	Delegate.ExecuteIfBound(UserId, Privilege, (uint32)EPrivilegeResults::NoFailures);
+}
+
+FPlatformUserId FOnlineIdentityGooglePlay::GetPlatformUserIdFromUniqueNetId(const FUniqueNetId& UniqueNetId)
+{
+	for (int i = 0; i < MAX_LOCAL_PLAYERS; ++i)
+	{
+		auto CurrentUniqueId = GetUniquePlayerId(i);
+		if (CurrentUniqueId.IsValid() && (*CurrentUniqueId == UniqueNetId))
+		{
+			return i;
+		}
+	}
+
+	return PLATFORMUSERID_NONE;
+}
diff --git a/Engine/Source/Runtime/Online/Android/OnlineSubsystemGooglePlay/Private/OnlineIdentityInterfaceGoolgePlay.cpp b/Engine/Source/Runtime/Online/Android/OnlineSubsystemGooglePlay/Private/OnlineIdentityInterfaceGoolgePlay.cpp
deleted file mode 100644
