Exposing build file list on BuildPatchServices manifest

Merging from CL# 2549278

diff --git a/Engine/Source/Runtime/Online/BuildPatchServices/Private/BuildPatchManifest.cpp b/Engine/Source/Runtime/Online/BuildPatchServices/Private/BuildPatchManifest.cpp
--- a/Engine/Source/Runtime/Online/BuildPatchServices/Private/BuildPatchManifest.cpp
+++ b/Engine/Source/Runtime/Online/BuildPatchServices/Private/BuildPatchManifest.cpp
@@ -1153,6 +1153,13 @@ int64 FBuildPatchAppManifest::GetBuildSize() const
 	return TotalBuildSize;
 }
 
+TArray<FString> FBuildPatchAppManifest::GetBuildFileList() const
+{
+	TArray<FString> Filenames;
+	GetFileList(Filenames);
+	return MoveTemp(Filenames);
+}
+
 int64 FBuildPatchAppManifest::GetFileSize(const TArray<FString>& Filenames) const
 {
 	int64 TotalSize = 0;
diff --git a/Engine/Source/Runtime/Online/BuildPatchServices/Private/BuildPatchManifest.h b/Engine/Source/Runtime/Online/BuildPatchServices/Private/BuildPatchManifest.h
--- a/Engine/Source/Runtime/Online/BuildPatchServices/Private/BuildPatchManifest.h
+++ b/Engine/Source/Runtime/Online/BuildPatchServices/Private/BuildPatchManifest.h
@@ -362,6 +362,7 @@ public:
 	virtual const FString& GetPrereqArgs() const override;
 	virtual int64 GetDownloadSize() const override;
 	virtual int64 GetBuildSize() const override;
+	virtual TArray<FString> GetBuildFileList() const override;
 	virtual void GetRemovableFiles(IBuildManifestRef OldManifest, TArray< FString >& RemovableFiles) const override;
 	virtual void GetRemovableFiles(const TCHAR* InstallPath, TArray< FString >& RemovableFiles) const override;
 	virtual bool NeedsResaving() const override;
diff --git a/Engine/Source/Runtime/Online/BuildPatchServices/Public/Interfaces/IBuildManifest.h b/Engine/Source/Runtime/Online/BuildPatchServices/Public/Interfaces/IBuildManifest.h
--- a/Engine/Source/Runtime/Online/BuildPatchServices/Public/Interfaces/IBuildManifest.h
+++ b/Engine/Source/Runtime/Online/BuildPatchServices/Public/Interfaces/IBuildManifest.h
@@ -114,6 +114,12 @@ public:
 	 */
 	virtual int64 GetBuildSize() const = 0;
 
+	/**
+	 * Get the list of files in this build. Filenames are all relative to an install directory
+	 * @return		an array containing build files
+	 */
+	virtual TArray<FString> GetBuildFileList() const = 0;
+
 	/**
 	 * Gets a list of files that were installed with the Old Manifest, but no longer required by this Manifest.
 	 * @param OldManifest	IN		The Build Manifest that is currently installed.
