Changed MTU to 1000 (to be compliant with 3G MTU) from 512 in numerous places in the engine. Created a new file to save net constants (NetConstants.h) where MAX_PACKET_SIZE is defined.

UE-2416

diff --git a/Engine/Source/Runtime/Engine/Classes/Engine/NetConnection.h b/Engine/Source/Runtime/Engine/Classes/Engine/NetConnection.h
--- a/Engine/Source/Runtime/Engine/Classes/Engine/NetConnection.h
+++ b/Engine/Source/Runtime/Engine/Classes/Engine/NetConnection.h
@@ -6,6 +6,7 @@
 
 #pragma once
 
+#include "Net/NetConstants.h"
 #include "Net/DataBunch.h"
 #include "Engine/Channel.h"
 #include "Engine/Player.h"
@@ -20,14 +21,6 @@ class FUniqueNetId;
 	Types.
 -----------------------------------------------------------------------------*/
 
-// Up to this many reliable channel bunches may be buffered.
-enum {RELIABLE_BUFFER         = 256   }; // Power of 2 >= 1.
-enum {MAX_PACKETID            = 16384 }; // Power of 2 >= 1, covering guaranteed loss/misorder time.
-enum {MAX_CHSEQUENCE          = 1024  }; // Power of 2 >RELIABLE_BUFFER, covering loss/misorder time.
-enum {MAX_BUNCH_HEADER_BITS   = 64    };
-enum {MAX_PACKET_HEADER_BITS  = 16    };
-enum {MAX_PACKET_TRAILER_BITS = 1     };
-
 
 class UNetDriver;
 
diff --git a/Engine/Source/Runtime/Engine/Private/NetConnection.cpp b/Engine/Source/Runtime/Engine/Private/NetConnection.cpp
--- a/Engine/Source/Runtime/Engine/Private/NetConnection.cpp
+++ b/Engine/Source/Runtime/Engine/Private/NetConnection.cpp
@@ -137,7 +137,7 @@ void UNetConnection::InitConnection(UNetDriver* InDriver, EConnectionState InSta
 {
 	Driver = InDriver;
 	// We won't be sending any packets, so use a default size
-	MaxPacket = 512;
+	MaxPacket = MAX_PACKET_SIZE;
 	PacketOverhead = 0;
 	State = InState;
 
diff --git a/Engine/Source/Runtime/Engine/Public/Net/NetConstants.h b/Engine/Source/Runtime/Engine/Public/Net/NetConstants.h
new file mode 100644
--- /dev/null
+++ b/Engine/Source/Runtime/Engine/Public/Net/NetConstants.h
@@ -0,0 +1,12 @@
+// Copyright 1998-2015 Epic Games, Inc. All Rights Reserved.
+
+#pragma once
+
+// Up to this many reliable channel bunches may be buffered.
+enum { RELIABLE_BUFFER = 256 }; // Power of 2 >= 1.
+enum { MAX_PACKETID = 16384 }; // Power of 2 >= 1, covering guaranteed loss/misorder time.
+enum { MAX_CHSEQUENCE = 1024 }; // Power of 2 >RELIABLE_BUFFER, covering loss/misorder time.
+enum { MAX_BUNCH_HEADER_BITS = 64 };
+enum { MAX_PACKET_HEADER_BITS = 16 };
+enum { MAX_PACKET_TRAILER_BITS = 1 };
+enum { MAX_PACKET_SIZE = 1000 }; // MTU for the connection
\ No newline at end of file
diff --git a/Engine/Source/Runtime/Online/OnlineSubsystem/OnlineSubsystem.Build.cs b/Engine/Source/Runtime/Online/OnlineSubsystem/OnlineSubsystem.Build.cs
--- a/Engine/Source/Runtime/Online/OnlineSubsystem/OnlineSubsystem.Build.cs
+++ b/Engine/Source/Runtime/Online/OnlineSubsystem/OnlineSubsystem.Build.cs
@@ -27,5 +27,11 @@ public class OnlineSubsystem : ModuleRules
 				"JsonUtilities"
 			}
 		);
+
+        PublicIncludePathModuleNames.Add("Engine"); 
+
+
 	}
+
+  
 }
diff --git a/Engine/Source/Runtime/Online/OnlineSubsystem/Public/LANBeacon.h b/Engine/Source/Runtime/Online/OnlineSubsystem/Public/LANBeacon.h
--- a/Engine/Source/Runtime/Online/OnlineSubsystem/Public/LANBeacon.h
+++ b/Engine/Source/Runtime/Online/OnlineSubsystem/Public/LANBeacon.h
@@ -2,6 +2,8 @@
 
 #pragma once
 
+#include "Net/NetConstants.h"
+
 /**
  * This value indicates which packet version the server is sending. Clients with
  * differing versions will ignore these packets. This prevents crashing when
@@ -16,8 +18,8 @@
 #define LAN_BEACON_PACKET_HEADER_SIZE 16
 
 /** The max size expected in a lan beacon packet */
-#define LAN_BEACON_MAX_PACKET_SIZE 512
-
+#define LAN_BEACON_MAX_PACKET_SIZE MAX_PACKET_SIZE
+	
 // Offsets for various fields
 #define LAN_BEACON_VER_OFFSET 0
 #define LAN_BEACON_PLATFORM_OFFSET 1
diff --git a/Engine/Source/Runtime/Online/OnlineSubsystemUtils/Private/IpConnection.cpp b/Engine/Source/Runtime/Online/OnlineSubsystemUtils/Private/IpConnection.cpp
--- a/Engine/Source/Runtime/Online/OnlineSubsystemUtils/Private/IpConnection.cpp
+++ b/Engine/Source/Runtime/Online/OnlineSubsystemUtils/Private/IpConnection.cpp
@@ -22,7 +22,6 @@ Notes:
 #define IP_HEADER_SIZE     (20)
 #define UDP_HEADER_SIZE    (IP_HEADER_SIZE+8)
 #define SLIP_HEADER_SIZE   (UDP_HEADER_SIZE+4)
-#define WINSOCK_MAX_PACKET (512)
 
 UIpConnection::UIpConnection(const FObjectInitializer& ObjectInitializer) :
 	Super(ObjectInitializer),
@@ -37,7 +36,7 @@ void UIpConnection::InitBase(UNetDriver* InDriver, class FSocket* InSocket, cons
 	// Pass the call up the chain
 	Super::InitBase(InDriver, InSocket, InURL, InState, 
 		// Use the default packet size/overhead unless overridden by a child class
-		InMaxPacket == 0 ? WINSOCK_MAX_PACKET : InMaxPacket,
+		InMaxPacket == 0 ? MAX_PACKET_SIZE : InMaxPacket,
 		InPacketOverhead == 0 ? SLIP_HEADER_SIZE : InPacketOverhead);
 
 	Socket = InSocket;
@@ -48,7 +47,7 @@ void UIpConnection::InitLocalConnection(UNetDriver* InDriver, class FSocket* InS
 {
 	InitBase(InDriver, InSocket, InURL, InState, 
 		// Use the default packet size/overhead unless overridden by a child class
-		InMaxPacket == 0 ? WINSOCK_MAX_PACKET : InMaxPacket,
+		InMaxPacket == 0 ? MAX_PACKET_SIZE : InMaxPacket,
 		InPacketOverhead == 0 ? SLIP_HEADER_SIZE : InPacketOverhead);
 
 	// Figure out IP address from the host URL
@@ -78,7 +77,7 @@ void UIpConnection::InitRemoteConnection(UNetDriver* InDriver, class FSocket* In
 {
 	InitBase(InDriver, InSocket, InURL, InState, 
 		// Use the default packet size/overhead unless overridden by a child class
-		InMaxPacket == 0 ? WINSOCK_MAX_PACKET : InMaxPacket,
+		InMaxPacket == 0 ? MAX_PACKET_SIZE : InMaxPacket,
 		InPacketOverhead == 0 ? SLIP_HEADER_SIZE : InPacketOverhead);
 
 	// Copy the remote IPAddress passed in
diff --git a/Engine/Source/Runtime/Online/OnlineSubsystemUtils/Private/IpNetDriver.cpp b/Engine/Source/Runtime/Online/OnlineSubsystemUtils/Private/IpNetDriver.cpp
--- a/Engine/Source/Runtime/Online/OnlineSubsystemUtils/Private/IpNetDriver.cpp
+++ b/Engine/Source/Runtime/Online/OnlineSubsystemUtils/Private/IpNetDriver.cpp
@@ -17,8 +17,6 @@ Notes:
 	Declarations.
 -----------------------------------------------------------------------------*/
 
-/** Size of the network recv buffer */
-#define NETWORK_MAX_PACKET (576)
 
 UIpNetDriver::UIpNetDriver(const FObjectInitializer& ObjectInitializer)
 	: Super(ObjectInitializer)
@@ -171,7 +169,7 @@ void UIpNetDriver::TickDispatch( float DeltaTime )
 	ISocketSubsystem* SocketSubsystem = GetSocketSubsystem();
 
 	// Process all incoming packets.
-	uint8 Data[NETWORK_MAX_PACKET];
+	uint8 Data[MAX_PACKET_SIZE];
 	TSharedRef<FInternetAddr> FromAddr = SocketSubsystem->CreateInternetAddr();
 	for( ; Socket != NULL; )
 	{
