addition of IAP restoration code
added a new blueprint node which starts the async restore and returns the product ids of the purchases made by the user to be restored
#ios

diff --git a/Engine/Source/Editor/OnlineBlueprintSupport/Classes/K2Node_InAppPurchaseRestore.h b/Engine/Source/Editor/OnlineBlueprintSupport/Classes/K2Node_InAppPurchaseRestore.h
new file mode 100644
--- /dev/null
+++ b/Engine/Source/Editor/OnlineBlueprintSupport/Classes/K2Node_InAppPurchaseRestore.h
@@ -0,0 +1,11 @@
+// Copyright 1998-2015 Epic Games, Inc. All Rights Reserved.
+
+#pragma once
+
+#include "K2Node_InAppPurchaseRestore.generated.h"
+
+UCLASS()
+class ONLINEBLUEPRINTSUPPORT_API UK2Node_InAppPurchaseRestore : public UK2Node_BaseAsyncTask
+{
+	GENERATED_UCLASS_BODY()
+};
diff --git a/Engine/Source/Editor/OnlineBlueprintSupport/Private/K2Node_InAppPurchaseRestore.cpp b/Engine/Source/Editor/OnlineBlueprintSupport/Private/K2Node_InAppPurchaseRestore.cpp
new file mode 100644
--- /dev/null
+++ b/Engine/Source/Editor/OnlineBlueprintSupport/Private/K2Node_InAppPurchaseRestore.cpp
@@ -0,0 +1,18 @@
+// Copyright 1998-2015 Epic Games, Inc. All Rights Reserved.
+
+#include "OnlineBlueprintSupportPrivatePCH.h"
+#include "InAppPurchaseRestoreCallbackProxy.h"
+#include "K2Node_InAppPurchaseRestore.h"
+
+#define LOCTEXT_NAMESPACE "K2Node"
+
+UK2Node_InAppPurchaseRestore::UK2Node_InAppPurchaseRestore(const FObjectInitializer& ObjectInitializer)
+	: Super(ObjectInitializer)
+{
+	ProxyFactoryFunctionName = GET_FUNCTION_NAME_CHECKED(UInAppPurchaseRestoreCallbackProxy, CreateProxyObjectForInAppPurchaseRestore);
+	ProxyFactoryClass = UInAppPurchaseRestoreCallbackProxy::StaticClass();
+
+	ProxyClass = UInAppPurchaseRestoreCallbackProxy::StaticClass();
+}
+
+#undef LOCTEXT_NAMESPACE
diff --git a/Engine/Source/Runtime/Online/IOS/OnlineSubsystemIOS/Private/OnlineStoreInterfaceIOS.cpp b/Engine/Source/Runtime/Online/IOS/OnlineSubsystemIOS/Private/OnlineStoreInterfaceIOS.cpp
--- a/Engine/Source/Runtime/Online/IOS/OnlineSubsystemIOS/Private/OnlineStoreInterfaceIOS.cpp
+++ b/Engine/Source/Runtime/Online/IOS/OnlineSubsystemIOS/Private/OnlineStoreInterfaceIOS.cpp
@@ -45,6 +45,59 @@
     }
 }
 
+-(void)paymentQueueRestoreCompletedTransactionsFinished:(SKPaymentQueue *)queue
+{
+	UE_LOG(LogOnline, Log, TEXT("FStoreKitHelper::paymentQueueRestoreCompletedTransactionsFinished"));
+
+	[FIOSAsyncTask CreateTaskWithBlock : ^ bool(void)
+	{
+		IOnlineSubsystem* OnlineSub = IOnlineSubsystem::Get(FName(TEXT("IOS")));
+		FOnlineStoreInterfaceIOS* StoreInterface = (FOnlineStoreInterfaceIOS*)OnlineSub->GetStoreInterface().Get();
+
+		if (StoreInterface->CachedPurchaseRestoreObject.IsValid())
+		{
+			StoreInterface->CachedPurchaseRestoreObject->ReadState = EOnlineAsyncTaskState::Done;
+		}
+		StoreInterface->TriggerOnInAppPurchaseRestoreCompleteDelegates(EInAppPurchaseState::Restored);
+
+		return true;
+	}];
+}
+
+-(void)paymentQueue: (SKPaymentQueue *)queue restoreCompletedTransactionsFailedWithError : (NSError *)error
+{
+	UE_LOG(LogOnline, Log, TEXT("FStoreKitHelper::failedRestore - %s"), *FString([error localizedDescription]));
+
+	EInAppPurchaseState::Type CompletionState = EInAppPurchaseState::Unknown;
+	switch (error.code)
+	{
+	case SKErrorPaymentCancelled:
+		CompletionState = EInAppPurchaseState::Cancelled;
+		break;
+	case SKErrorClientInvalid:
+	case SKErrorStoreProductNotAvailable:
+	case SKErrorPaymentInvalid:
+		CompletionState = EInAppPurchaseState::Invalid;
+		break;
+	case SKErrorPaymentNotAllowed:
+		CompletionState = EInAppPurchaseState::NotAllowed;
+		break;
+	}
+
+	[FIOSAsyncTask CreateTaskWithBlock : ^ bool(void)
+	{
+		IOnlineSubsystem* OnlineSub = IOnlineSubsystem::Get(FName(TEXT("IOS")));
+		FOnlineStoreInterfaceIOS* StoreInterface = (FOnlineStoreInterfaceIOS*)OnlineSub->GetStoreInterface().Get();
+		if (StoreInterface->CachedPurchaseRestoreObject.IsValid())
+		{
+			StoreInterface->CachedPurchaseRestoreObject->ReadState = EOnlineAsyncTaskState::Done;
+		}
+
+		StoreInterface->TriggerOnInAppPurchaseRestoreCompleteDelegates(CompletionState);
+
+		return true;
+	}];
+}
 
 - (void) completeTransaction: (SKPaymentTransaction *)transaction
 {
@@ -97,9 +150,44 @@
 {
 	UE_LOG(LogOnline, Log, TEXT( "FStoreKitHelper::restoreTransaction" ));
 
-    //[self recordTransaction: transaction];
-    //[self provideContent: transaction.originalTransaction.payment.productIdentifier];
-    [[SKPaymentQueue defaultQueue] finishTransaction: transaction];
+	[FIOSAsyncTask CreateTaskWithBlock : ^ bool(void)
+	{
+		IOnlineSubsystem* OnlineSub = IOnlineSubsystem::Get(FName(TEXT("IOS")));
+		FOnlineStoreInterfaceIOS* StoreInterface = (FOnlineStoreInterfaceIOS*)OnlineSub->GetStoreInterface().Get();
+
+		if (StoreInterface->CachedPurchaseRestoreObject.IsValid())
+		{
+			FString ReceiptData;
+
+#ifdef __IPHONE_7_0
+			if ([IOSAppDelegate GetDelegate].OSVersion >= 7.0)
+			{
+				NSURL* nsReceiptUrl = [[NSBundle mainBundle] appStoreReceiptURL];
+				NSData* nsReceiptData = [NSData dataWithContentsOfURL : nsReceiptUrl];
+				NSString* nsEncodedReceiptData = [nsReceiptData base64EncodedStringWithOptions : NSDataBase64EncodingEndLineWithLineFeed];
+
+				ReceiptData = nsEncodedReceiptData;
+			}
+			else
+#endif
+			{
+#if __IPHONE_OS_VERSION_MIN_REQUIRED < __IPHONE_7_0
+				// If earlier than IOS 7, we will need to use the transactionReceipt
+				NSString* nsEncodedReceiptData = [transaction.originalTransaction.transactionReceipt base64EncodedStringWithOptions : NSDataBase64EncodingEndLineWithLineFeed];
+				ReceiptData = nsEncodedReceiptData;
+#endif
+			}
+
+			FInAppPurchaseRestoreInfo RestoreInfo;
+			RestoreInfo.Identifier = FString(transaction.originalTransaction.payment.productIdentifier);
+			RestoreInfo.ReceiptData = ReceiptData;
+			StoreInterface->CachedPurchaseRestoreObject->ProvidedRestoreInformation.Add(RestoreInfo);
+		}
+
+		return true;
+	}];
+	
+	[[SKPaymentQueue defaultQueue] finishTransaction:transaction];
 }
 
 
@@ -179,6 +267,33 @@
 	[request autorelease];
 }
 
+-(void)requestDidFinish:(SKRequest*)request
+{
+#ifdef __IPHONE_7_0
+	if ([Request isKindOfClass : [SKReceiptRefreshRequest class]])
+	{
+		[[SKPaymentQueue defaultQueue] restoreCompletedTransactions];
+	}
+#endif
+}
+
+-(void)restorePurchases
+{
+#ifdef __IPHONE_7_0
+	if ([IOSAppDelegate GetDelegate].OSVersion >= 7.0)
+	{
+		Request = [[SKReceiptRefreshRequest alloc] init];
+		Request.delegate = self;
+		[Request start];
+	}
+	else
+#endif
+#if __IPHONE_OS_VERSION_MIN_REQUIRED < __IPHONE_7_0
+	{
+		[[SKPaymentQueue defaultQueue] restoreCompletedTransactions];
+	}
+#endif
+}
 @end
 
 
@@ -304,6 +419,25 @@ bool FOnlineStoreInterfaceIOS::BeginPurchase(const FInAppPurchaseProductRequest&
 }
 
 
+bool FOnlineStoreInterfaceIOS::RestorePurchases(FOnlineInAppPurchaseRestoreReadRef& InReadObject)
+{
+	bool bSentAQueryRequest = false;
+	CachedPurchaseRestoreObject = InReadObject;
+
+	if (bIsPurchasing || bIsProductRequestInFlight)
+	{
+		UE_LOG(LogOnline, Verbose, TEXT("FOnlineStoreInterfaceIOS::BeginPurchase - cannot make a purchase whilst one is in transaction."));
+		TriggerOnInAppPurchaseRestoreCompleteDelegates(EInAppPurchaseState::Failed);
+	}
+	else
+	{
+		[StoreHelper restorePurchases];
+		bSentAQueryRequest = true;
+	}
+
+	return bSentAQueryRequest;
+}
+
 void FOnlineStoreInterfaceIOS::ProcessProductsResponse( SKProductsResponse* Response )
 {
 	if( bIsPurchasing )
diff --git a/Engine/Source/Runtime/Online/IOS/OnlineSubsystemIOS/Private/OnlineStoreInterfaceIOS.h b/Engine/Source/Runtime/Online/IOS/OnlineSubsystemIOS/Private/OnlineStoreInterfaceIOS.h
--- a/Engine/Source/Runtime/Online/IOS/OnlineSubsystemIOS/Private/OnlineStoreInterfaceIOS.h
+++ b/Engine/Source/Runtime/Online/IOS/OnlineSubsystemIOS/Private/OnlineStoreInterfaceIOS.h
@@ -9,6 +9,7 @@
 #include <StoreKit/SKError.h>
 #include <StoreKit/SKProduct.h>
 #include <StoreKit/SKProductsRequest.h>
+#include <StoreKit/SKReceiptRefreshRequest.h>
 #include <StoreKit/SKPaymentTransaction.h>
 #include <StoreKit/SKPayment.h>
 #include <StoreKit/SKPaymentQueue.h>
@@ -30,7 +31,7 @@ namespace EInAppPurchaseResult
 
 
 /** Helper class, which allows us to manage IAP product information requests, AND transactions */
-@interface FStoreKitHelper : NSObject<SKProductsRequestDelegate, SKPaymentTransactionObserver>
+@interface FStoreKitHelper : NSObject<SKProductsRequestDelegate, SKPaymentTransactionObserver, SKRequestDelegate>
 {
 };
 /** Store kit request object, holds information about the products we are purchasing, or querying. */
@@ -46,6 +47,8 @@ namespace EInAppPurchaseResult
 /** Helper fn to direct a product request response back to our store interface */
 - (void)productsRequest:(SKProductsRequest *)request didReceiveResponse:(SKProductsResponse *)response;
 
+/** Helper fn to restore previously purchased products */
+-(void)restorePurchases;
 @end
 
 
@@ -64,6 +67,7 @@ public:
 	virtual bool QueryForAvailablePurchases(const TArray<FString>& ProductIDs, FOnlineProductInformationReadRef& InReadObject) override;
 	virtual bool BeginPurchase(const FInAppPurchaseProductRequest& ProductRequest, FOnlineInAppPurchaseTransactionRef& InReadObject) override;
 	virtual bool IsAllowedToMakePurchases() override;
+	virtual bool RestorePurchases(FOnlineInAppPurchaseRestoreReadRef& InReadObject) override;
 	// End IOnlineStore interface
 
 	/**
@@ -77,6 +81,9 @@ public:
 	/** Cached in-app purchase transaction object, used to provide details to the user, of the product that has just been purchased. */
 	FOnlineInAppPurchaseTransactionPtr CachedPurchaseStateObject;
 
+	/** Cached in-app purchase restore transaction object, used to provide details to the developer about what products should be restored */
+	FOnlineInAppPurchaseRestoreReadPtr CachedPurchaseRestoreObject;
+
 
 private:
 	/** Flags which determine the state of our transaction buffer, only one action at a time */
@@ -90,6 +97,9 @@ private:
 
 	/** Delegate fired when a purchase transaction has completed, whether successful or unsuccessful */
 	FOnInAppPurchaseComplete OnPurchaseCompleteDelegate;
+
+	/** Delegate fired when the purchase restoration has completed, whether successful or unsuccessful */
+	FOnInAppPurchaseRestoreComplete OnPurchaseRestoreCompleteDelegate;
 };
 
 typedef TSharedPtr<FOnlineStoreInterfaceIOS, ESPMode::ThreadSafe> FOnlineStoreInterfaceIOSPtr;
diff --git a/Engine/Source/Runtime/Online/OnlineSubsystem/Public/Interfaces/OnlineStoreInterface.h b/Engine/Source/Runtime/Online/OnlineSubsystem/Public/Interfaces/OnlineStoreInterface.h
--- a/Engine/Source/Runtime/Online/OnlineSubsystem/Public/Interfaces/OnlineStoreInterface.h
+++ b/Engine/Source/Runtime/Online/OnlineSubsystem/Public/Interfaces/OnlineStoreInterface.h
@@ -17,6 +17,7 @@ namespace EInAppPurchaseState
 		Cancelled UMETA(DisplayName = "Cancelled"),
 		Invalid UMETA(DisplayName = "Invalid"),
 		NotAllowed UMETA(DisplayName = "NotAllowed"),
+		Restored UMETA(DisplayName = "Restored"),
 		Unknown  UMETA(DisplayName = "Unknown")
 	};
 }
@@ -48,6 +49,15 @@ typedef FOnQueryForAvailablePurchasesComplete::FDelegate FOnQueryForAvailablePur
 DECLARE_MULTICAST_DELEGATE_OneParam(FOnInAppPurchaseComplete, EInAppPurchaseState::Type);
 typedef FOnInAppPurchaseComplete::FDelegate FOnInAppPurchaseCompleteDelegate;
 
+/**
+* Delegate fired when the online session has transitioned to the started state
+*
+* @param SessionName the name of the session the that has transitioned to started
+* @param bWasSuccessful true if the async action completed without error, false if there was an error
+*/
+DECLARE_MULTICAST_DELEGATE_OneParam(FOnInAppPurchaseRestoreComplete, EInAppPurchaseState::Type);
+typedef FOnInAppPurchaseRestoreComplete::FDelegate FOnInAppPurchaseRestoreCompleteDelegate;
+
 
 /**
  * Micro-transaction request information
@@ -97,6 +107,24 @@ struct FInAppPurchaseProductInfo
 };
 
 
+/**
+* Micro-transaction restored purchase information
+*/
+USTRUCT(BlueprintType)
+struct FInAppPurchaseRestoreInfo
+{
+	GENERATED_USTRUCT_BODY()
+
+	// The unique product identifier
+	UPROPERTY(BlueprintReadOnly, Category = ProductInfo)
+	FString Identifier;
+
+	// The localized display price name
+	UPROPERTY(BlueprintReadOnly, Category = ProductInfo)
+		FString ReceiptData;
+};
+
+
 /**
 *	Interface for reading data from an In App Purchase service
 */
@@ -138,6 +166,27 @@ public:
 typedef TSharedRef<FOnlineInAppPurchaseTransaction, ESPMode::ThreadSafe> FOnlineInAppPurchaseTransactionRef;
 typedef TSharedPtr<FOnlineInAppPurchaseTransaction, ESPMode::ThreadSafe> FOnlineInAppPurchaseTransactionPtr;
 
+
+/**
+*	Interface for reading data from an In App Purchase service
+*/
+class ONLINESUBSYSTEM_API FOnlineInAppPurchaseRestoreRead
+{
+public:
+	/** Indicates an error reading data occurred while processing */
+	EOnlineAsyncTaskState::Type ReadState;
+
+	FOnlineInAppPurchaseRestoreRead() :
+		ReadState(EOnlineAsyncTaskState::NotStarted)
+	{
+	}
+
+	TArray<FInAppPurchaseRestoreInfo> ProvidedRestoreInformation;
+};
+
+typedef TSharedRef<FOnlineInAppPurchaseRestoreRead, ESPMode::ThreadSafe> FOnlineInAppPurchaseRestoreReadRef;
+typedef TSharedPtr<FOnlineInAppPurchaseRestoreRead, ESPMode::ThreadSafe> FOnlineInAppPurchaseRestoreReadPtr;
+
 /**
  *	IOnlineStore - Interface class for microtransactions
  */
@@ -172,9 +221,19 @@ public:
 	 * @return - whether a purchase request was sent
 	 */
 	virtual bool BeginPurchase(const FInAppPurchaseProductRequest& ProductRequest, FOnlineInAppPurchaseTransactionRef& InReadObject) = 0;
-	
+
 	/**
 	 * Delegate which is executed when a Purchase completes
 	 */
 	DEFINE_ONLINE_DELEGATE_ONE_PARAM(OnInAppPurchaseComplete, EInAppPurchaseState::Type);
+
+	/**
+	* Restore any purchases previously made
+	*/
+	virtual bool RestorePurchases(FOnlineInAppPurchaseRestoreReadRef& InReadObject) = 0;
+
+	/**
+	* Delegate which is executed when a Purchase completes
+	*/
+	DEFINE_ONLINE_DELEGATE_ONE_PARAM(OnInAppPurchaseRestoreComplete, EInAppPurchaseState::Type);
 };
\ No newline at end of file
diff --git a/Engine/Source/Runtime/Online/OnlineSubsystemUtils/Classes/InAppPurchaseRestoreCallbackProxy.h b/Engine/Source/Runtime/Online/OnlineSubsystemUtils/Classes/InAppPurchaseRestoreCallbackProxy.h
new file mode 100644
--- /dev/null
+++ b/Engine/Source/Runtime/Online/OnlineSubsystemUtils/Classes/InAppPurchaseRestoreCallbackProxy.h
@@ -0,0 +1,66 @@
+// Copyright 1998-2015 Epic Games, Inc. All Rights Reserved.
+#pragma once
+
+#include "OnlineStoreInterface.h"
+#include "InAppPurchaseRestoreCallbackProxy.generated.h"
+
+DECLARE_DYNAMIC_MULTICAST_DELEGATE_TwoParams(FInAppPurchaseRestoreResult, EInAppPurchaseState::Type, CompletionStatus, const TArray<FInAppPurchaseRestoreInfo>&, InAppRestorePurchaseInformation);
+
+UCLASS(MinimalAPI)
+class UInAppPurchaseRestoreCallbackProxy : public UObject
+{
+	GENERATED_UCLASS_BODY()
+
+	// Called when there is a successful In-App Purchase transaction
+	UPROPERTY(BlueprintAssignable)
+	FInAppPurchaseRestoreResult OnSuccess;
+
+	// Called when there is an unsuccessful In-App Purchase transaction
+	UPROPERTY(BlueprintAssignable)
+	FInAppPurchaseRestoreResult OnFailure;
+
+	// Kicks off a transaction for the provided product identifier
+	UFUNCTION(BlueprintCallable, meta = (DisplayName="Make an In-App Purchase"), Category="Online|InAppPurchase")
+	static UInAppPurchaseRestoreCallbackProxy* CreateProxyObjectForInAppPurchaseRestore(class APlayerController* PlayerController);
+
+public:
+
+	// Begin UObject interface
+	virtual void BeginDestroy() override;
+	// End UObject interface
+
+private:
+
+	/** Called by the InAppPurchase system when the transaction has finished */
+	void OnInAppPurchaseRestoreComplete_Delayed();
+	void OnInAppPurchaseRestoreComplete(EInAppPurchaseState::Type CompletionState);
+
+	/** Unregisters our delegate from the In-App Purchase system */
+	void RemoveDelegate();
+
+	/** Triggers the In-App Purchase Restore Transaction for the specifed user */
+	void Trigger(class APlayerController* PlayerController);
+
+private:
+
+	/** Delegate called when a InAppPurchase has been successfully restored */
+	FOnInAppPurchaseRestoreCompleteDelegate InAppPurchaseRestoreCompleteDelegate;
+
+	/** Handle to the registered InAppPurchaseCompleteDelegate */
+	FDelegateHandle InAppPurchaseRestoreCompleteDelegateHandle;
+
+	/** The InAppPurchaseRestore read request */
+	FOnlineInAppPurchaseRestoreReadPtr ReadObject;
+
+	/** Did we fail immediately? */
+	bool bFailedToEvenSubmit;
+
+	/** Pointer to the world, needed to delay the results slightly */
+	TWeakObjectPtr<UWorld> WorldPtr;
+
+	/** Did the purchase succeed? */
+	EInAppPurchaseState::Type SavedPurchaseState;
+	TArray<FInAppPurchaseRestoreInfo> SavedProductInformation;
+
+	FTimerHandle OnInAppPurchaseRestoreComplete_DelayedTimerHandle;
+};
diff --git a/Engine/Source/Runtime/Online/OnlineSubsystemUtils/Private/InAppPurchaseRestoreCallbackProxy.cpp b/Engine/Source/Runtime/Online/OnlineSubsystemUtils/Private/InAppPurchaseRestoreCallbackProxy.cpp
new file mode 100644
--- /dev/null
+++ b/Engine/Source/Runtime/Online/OnlineSubsystemUtils/Private/InAppPurchaseRestoreCallbackProxy.cpp
@@ -0,0 +1,123 @@
+// Copyright 1998-2015 Epic Games, Inc. All Rights Reserved.
+
+#include "OnlineSubsystemUtilsPrivatePCH.h"
+
+//////////////////////////////////////////////////////////////////////////
+// UInAppPurchaseRestoreCallbackProxy
+
+UInAppPurchaseRestoreCallbackProxy::UInAppPurchaseRestoreCallbackProxy(const FObjectInitializer& ObjectInitializer)
+	: Super(ObjectInitializer)
+{
+	WorldPtr = nullptr;
+}
+
+
+void UInAppPurchaseRestoreCallbackProxy::Trigger(APlayerController* PlayerController)
+{
+	bFailedToEvenSubmit = true;
+
+	WorldPtr = (PlayerController != nullptr) ? PlayerController->GetWorld() : nullptr;
+	if (APlayerState* PlayerState = (PlayerController != nullptr) ? PlayerController->PlayerState : nullptr)
+	{
+		if (IOnlineSubsystem* const OnlineSub = IOnlineSubsystem::Get())
+		{
+			IOnlineStorePtr StoreInterface = OnlineSub->GetStoreInterface();
+			if (StoreInterface.IsValid())
+			{
+				bFailedToEvenSubmit = false;
+
+				// Register the completion callback
+				InAppPurchaseRestoreCompleteDelegate = FOnInAppPurchaseCompleteDelegate::CreateUObject(this, &UInAppPurchaseRestoreCallbackProxy::OnInAppPurchaseRestoreComplete);
+				InAppPurchaseRestoreCompleteDelegateHandle = StoreInterface->AddOnInAppPurchaseRestoreCompleteDelegate_Handle(InAppPurchaseRestoreCompleteDelegate);
+
+				// Set-up, and trigger the transaction through the store interface
+				ReadObject = MakeShareable(new FOnlineInAppPurchaseRestoreRead());
+				FOnlineInAppPurchaseRestoreReadRef ReadObjectRef = ReadObject.ToSharedRef();
+				StoreInterface->RestorePurchases(ReadObjectRef);
+			}
+			else
+			{
+				FFrame::KismetExecutionMessage(TEXT("UInAppPurchaseRestoreCallbackProxy::Trigger - In-App Purchases are not supported by Online Subsystem"), ELogVerbosity::Warning);
+			}
+		}
+		else
+		{
+			FFrame::KismetExecutionMessage(TEXT("UInAppPurchaseRestoreCallbackProxy::Trigger - Invalid or uninitialized OnlineSubsystem"), ELogVerbosity::Warning);
+		}
+	}
+	else
+	{
+		FFrame::KismetExecutionMessage(TEXT("UInAppPurchaseRestoreCallbackProxy::Trigger - Invalid player state"), ELogVerbosity::Warning);
+	}
+
+	if (bFailedToEvenSubmit && (PlayerController != NULL))
+	{
+		OnInAppPurchaseRestoreComplete(EInAppPurchaseState::Failed);
+	}
+}
+
+
+void UInAppPurchaseRestoreCallbackProxy::OnInAppPurchaseRestoreComplete(EInAppPurchaseState::Type CompletionState)
+{
+	RemoveDelegate();
+	SavedPurchaseState = CompletionState;
+	if (CompletionState == EInAppPurchaseState::Restored)
+	{
+		SavedProductInformation = ReadObject->ProvidedRestoreInformation;
+	}
+    
+	if (UWorld* World = WorldPtr.Get())
+	{
+		// Use a local timer handle as we don't need to store it for later but we don't need to look for something to clear
+		FTimerHandle TimerHandle;
+		World->GetTimerManager().SetTimer(OnInAppPurchaseRestoreComplete_DelayedTimerHandle, this, &UInAppPurchaseRestoreCallbackProxy::OnInAppPurchaseRestoreComplete_Delayed, 0.001f, false);
+    }
+
+	ReadObject = NULL;
+}
+
+
+void UInAppPurchaseRestoreCallbackProxy::OnInAppPurchaseRestoreComplete_Delayed()
+{    
+	if (SavedPurchaseState == EInAppPurchaseState::Restored)
+	{
+		OnSuccess.Broadcast(SavedPurchaseState, SavedProductInformation);
+	}
+	else
+	{
+		OnFailure.Broadcast(SavedPurchaseState, SavedProductInformation);
+	}
+}
+
+
+void UInAppPurchaseRestoreCallbackProxy::RemoveDelegate()
+{
+	if (!bFailedToEvenSubmit)
+	{
+		if (IOnlineSubsystem* OnlineSub = IOnlineSubsystem::Get())
+		{
+			IOnlineStorePtr InAppPurchases = OnlineSub->GetStoreInterface();
+			if (InAppPurchases.IsValid())
+			{
+				InAppPurchases->ClearOnInAppPurchaseRestoreCompleteDelegate_Handle(InAppPurchaseRestoreCompleteDelegateHandle);
+			}
+		}
+	}
+}
+
+
+void UInAppPurchaseRestoreCallbackProxy::BeginDestroy()
+{
+	ReadObject = NULL;
+	RemoveDelegate();
+
+	Super::BeginDestroy();
+}
+
+
+UInAppPurchaseRestoreCallbackProxy* UInAppPurchaseRestoreCallbackProxy::CreateProxyObjectForInAppPurchaseRestore(class APlayerController* PlayerController)
+{
+	UInAppPurchaseRestoreCallbackProxy* Proxy = NewObject<UInAppPurchaseRestoreCallbackProxy>();
+	Proxy->Trigger(PlayerController);
+	return Proxy;
+}
