initial implementation of CloudKit for cloud database storage.  This does not include key-value storage or iCloud document storage
#ios

diff --git a/Engine/Source/Programs/UnrealBuildTool/IOS/UEBuildIOS.cs b/Engine/Source/Programs/UnrealBuildTool/IOS/UEBuildIOS.cs
--- a/Engine/Source/Programs/UnrealBuildTool/IOS/UEBuildIOS.cs
+++ b/Engine/Source/Programs/UnrealBuildTool/IOS/UEBuildIOS.cs
@@ -297,6 +297,7 @@ namespace UnrealBuildTool
 
             InBuildTarget.GlobalLinkEnvironment.Config.AdditionalFrameworks.Add( new UEBuildFramework( "GameKit" ) );
             InBuildTarget.GlobalLinkEnvironment.Config.AdditionalFrameworks.Add( new UEBuildFramework( "StoreKit" ) );
+			InBuildTarget.GlobalLinkEnvironment.Config.AdditionalFrameworks.Add( new UEBuildFramework( "CloudKit" ) );
 
 			if (RunTimeIOSVersion != "6.1")
 			{
diff --git a/Engine/Source/Programs/UnrealBuildTool/IOS/UEDeployIOS.cs b/Engine/Source/Programs/UnrealBuildTool/IOS/UEDeployIOS.cs
--- a/Engine/Source/Programs/UnrealBuildTool/IOS/UEDeployIOS.cs
+++ b/Engine/Source/Programs/UnrealBuildTool/IOS/UEDeployIOS.cs
@@ -744,11 +744,35 @@ namespace UnrealBuildTool.IOS
 
 		private void WriteEntitlementsFile(string OutputFilename)
 		{
+			// get the settings from the ini file
+			// plist replacements
+			ConfigCacheIni Ini = new ConfigCacheIni(UnrealTargetPlatform.IOS, "Engine", UnrealBuildTool.GetUProjectPath());
+			bool bSupported = false;
+			Ini.GetBool("/Script/IOSRuntimeSettings.IOSRuntimeSettings", "bEnableCloudKitSupport", out bSupported);
+
 			Directory.CreateDirectory(Path.GetDirectoryName(OutputFilename));
 			// we need to have something so Xcode will compile, so we just set the get-task-allow, since we know the value, 
 			// which is based on distribution or not (true means debuggable)
-			File.WriteAllText(OutputFilename, string.Format("<plist><dict><key>get-task-allow</key><{0}/></dict></plist>",
-				/*Config.bForDistribution ? "false" : */"true"));
+			StringBuilder Text = new StringBuilder();
+			Text.AppendLine("<?xml version=\"1.0\" encoding=\"UTF-8\"?>");
+			Text.AppendLine("<!DOCTYPE plist PUBLIC \"-//Apple//DTD PLIST 1.0//EN\" \"http://www.apple.com/DTDs/PropertyList-1.0.dtd\">");
+			Text.AppendLine("<plist version=\"1.0\">");
+			Text.AppendLine("<dict>");
+			Text.AppendLine(string.Format("\t<key>get-task-allow</key><{0}/>",	/*Config.bForDistribution ? "false" : */"true"));
+			if (bSupported)
+			{
+				Text.AppendLine("\t<key>com.apple.developer.icloud-container-identifiers</key>");
+				Text.AppendLine("\t<array>");
+				Text.AppendLine("\t\t<string>iCloud.$(CFBundleIdentifier)</string>");
+				Text.AppendLine("\t</array>");
+				Text.AppendLine("\t<key>com.apple.developer.icloud-services</key>");
+				Text.AppendLine("\t<array>");
+				Text.AppendLine("\t\t<string>CloudKit</string>");
+				Text.AppendLine("\t</array>");
+			}
+			Text.AppendLine("</dict>");
+			Text.AppendLine("</plist>");
+			File.WriteAllText(OutputFilename, Text.ToString());
 		}
 
 		static void SafeFileCopy(FileInfo SourceFile, string DestinationPath, bool bOverwrite)
diff --git a/Engine/Source/Runtime/IOS/IOSRuntimeSettings/Classes/IOSRuntimeSettings.h b/Engine/Source/Runtime/IOS/IOSRuntimeSettings/Classes/IOSRuntimeSettings.h
--- a/Engine/Source/Runtime/IOS/IOSRuntimeSettings/Classes/IOSRuntimeSettings.h
+++ b/Engine/Source/Runtime/IOS/IOSRuntimeSettings/Classes/IOSRuntimeSettings.h
@@ -144,6 +144,10 @@ public:
 	UPROPERTY(GlobalConfig, EditAnywhere, Category = Online, meta = (ConfigHierarchyEditable))
 	uint32 bEnableGameCenterSupport : 1;
 	
+	// Should Cloud Kit support (iOS Online Subsystem) be enabled?
+	UPROPERTY(GlobalConfig, EditAnywhere, Category = Online)
+	uint32 bEnableCloudKitSupport : 1;
+
 	// Whether or not to add support for Metal API (requires IOS8 and A7 processors).
 	UPROPERTY(GlobalConfig, EditAnywhere, Category = Rendering, meta = (DisplayName = "Support Forward Rendering with Metal (A7 and up devices)"))
 	bool bSupportsMetal;
diff --git a/Engine/Source/Runtime/IOS/IOSRuntimeSettings/Private/IOSRuntimeSettings.cpp b/Engine/Source/Runtime/IOS/IOSRuntimeSettings/Private/IOSRuntimeSettings.cpp
--- a/Engine/Source/Runtime/IOS/IOSRuntimeSettings/Private/IOSRuntimeSettings.cpp
+++ b/Engine/Source/Runtime/IOS/IOSRuntimeSettings/Private/IOSRuntimeSettings.cpp
@@ -8,6 +8,7 @@ UIOSRuntimeSettings::UIOSRuntimeSettings(const FObjectInitializer& ObjectInitial
 	: Super(ObjectInitializer)
 {
 	bEnableGameCenterSupport = true;
+	bEnableCloudKitSupport = false;
 	bSupportsPortraitOrientation = true;
 	BundleDisplayName = TEXT("UE4 Game");
 	BundleName = TEXT("MyUE4Game");
diff --git a/Engine/Source/Runtime/Online/IOS/OnlineSubsystemIOS/Private/OnlineSharedCloudInterfaceIOS.cpp b/Engine/Source/Runtime/Online/IOS/OnlineSubsystemIOS/Private/OnlineSharedCloudInterfaceIOS.cpp
new file mode 100644
--- /dev/null
+++ b/Engine/Source/Runtime/Online/IOS/OnlineSubsystemIOS/Private/OnlineSharedCloudInterfaceIOS.cpp
@@ -0,0 +1,190 @@
+// Copyright 1998-2015 Epic Games, Inc. All Rights Reserved.
+
+#include "OnlineSubsystemIOSPrivatePCH.h"
+#include "OnlineSharedCloudInterfaceIOS.h"
+
+//
+// SharedCloud interface implementation
+//
+
+FOnlineSharedCloudInterfaceIOS::~FOnlineSharedCloudInterfaceIOS()
+{
+	ClearSharedFiles();
+}
+
+FCloudFile* FOnlineSharedCloudInterfaceIOS::GetCloudFile(const FString& FileName, bool bCreateIfMissing)
+{
+    FScopeLock ScopeLock(&CloudDataLock);
+	if (FileName.Len() > 0)
+	{
+		for (int32 FileIdx = 0; FileIdx < CloudFileData.Num(); FileIdx++)
+		{
+			FCloudFile* UserFileData = &CloudFileData[FileIdx];
+			if (UserFileData &&
+				UserFileData->FileName == FileName)
+			{
+				return UserFileData;
+			}
+		}
+
+		if (bCreateIfMissing)
+		{
+			return new (CloudFileData)FCloudFile(FileName);
+		}
+	}
+
+	return NULL;
+}
+
+bool FOnlineSharedCloudInterfaceIOS::ClearFiles()
+{
+    FScopeLock ScopeLock(&CloudDataLock);
+	// Delete file contents
+	for (int32 FileIdx = 0; FileIdx < CloudFileData.Num(); FileIdx++)
+	{
+		CloudFileData[FileIdx].Data.Empty();
+	}
+
+	// No async files being handled, so empty them all
+	CloudFileData.Empty();
+	return true;
+}
+
+bool FOnlineSharedCloudInterfaceIOS::ClearCloudFile(const FString& FileName)
+{
+    FScopeLock ScopeLock(&CloudDataLock);
+	int32 FoundIndex = INDEX_NONE;
+	for (int32 FileIdx = 0; FileIdx < CloudFileData.Num(); FileIdx++)
+	{
+		FCloudFile* UserFileData = &CloudFileData[FileIdx];
+		if (UserFileData->FileName == FileName)
+		{
+			// If there is an async task outstanding, fail to empty
+			if (UserFileData->AsyncState == EOnlineAsyncTaskState::InProgress)
+			{
+				return false;
+			}
+
+			UserFileData->Data.Empty();
+			FoundIndex = FileIdx;
+			break;
+		}
+	}
+
+	if (FoundIndex != INDEX_NONE)
+	{
+		CloudFileData.RemoveAtSwap(FoundIndex);
+	}
+	return true;
+}
+
+bool FOnlineSharedCloudInterfaceIOS::GetSharedFileContents(const FSharedContentHandle& SharedHandle, TArray<uint8>& FileContents)
+{
+	FCloudFile* CloudFile = GetCloudFile(SharedHandle.ToString());
+	if (CloudFile && CloudFile->AsyncState == EOnlineAsyncTaskState::Done && CloudFile->Data.Num() > 0)
+	{
+		FileContents = CloudFile->Data;
+		return true;
+	}
+	return false;
+}
+
+bool FOnlineSharedCloudInterfaceIOS::ClearSharedFiles()
+{
+	// NOTE: Return true regardless of if the user was valid or not. We don't care if
+	// there weren't any files in need of clearing, only if there was a failure to clear
+	ClearFiles();
+	return true;
+}
+
+bool FOnlineSharedCloudInterfaceIOS::ClearSharedFile(const FSharedContentHandle& SharedHandle)
+{
+	// NOTE: Return true regardless of if the file exists or not. The only way to return false
+	// would be in we failed to delete, and not deleting a non-existing file is not a failure
+	ClearCloudFile(SharedHandle.ToString());
+	return true;
+}
+
+bool FOnlineSharedCloudInterfaceIOS::ReadSharedFile(const FSharedContentHandle& SharedHandle)
+{
+	FCloudFile* CloudFile = GetCloudFile(SharedHandle.ToString());
+	if (CloudFile)
+	{
+		CloudFile->AsyncState = EOnlineAsyncTaskState::InProgress;
+        __block FSharedContentHandleIOS NewHandle = *((FSharedContentHandleIOS*)&SharedHandle);
+		return [[IOSCloudStorage cloudStorage] readFile:SharedHandle.ToString().GetNSString() sharedDB : true completionHandler : ^ (CKRecord *record, NSError *error)
+		{
+            FScopeLock ScopeLock(&CloudDataLock);
+            FCloudFile* File = GetCloudFile(NewHandle.ToString());
+			if (error)
+			{
+				// TODO: record is potentially not found
+				File->AsyncState = EOnlineAsyncTaskState::Failed;
+				TriggerOnReadSharedFileCompleteDelegates(false, NewHandle);
+				NSLog(@"Error: %@", error);
+			}
+			else
+			{
+				// store the contents in the memory record database
+				NSData* data = (NSData*)record[@"contents"];
+				File->Data.Empty();
+				File->Data.Append((uint8*)data.bytes, data.length);
+				File->AsyncState = EOnlineAsyncTaskState::Done;
+				TriggerOnReadSharedFileCompleteDelegates(true, NewHandle);
+				NSLog(@"Record Read!");
+			}
+		}];
+	}
+    else
+    {
+        TriggerOnReadSharedFileCompleteDelegates(false, SharedHandle);
+    }
+	return false;
+}
+
+bool FOnlineSharedCloudInterfaceIOS::WriteSharedFile(const FUniqueNetId& UserId, const FString& FileName, TArray<uint8>& FileContents)
+{
+	FCloudFile* CloudFile = GetCloudFile(FileName, true);
+	if (CloudFile)
+	{
+		CloudFile->AsyncState = EOnlineAsyncTaskState::InProgress;
+        __block FString NewFile = FileName;
+        __block TArray<uint8> DataContents = FileContents;
+		return [[IOSCloudStorage cloudStorage] writeFile:FileName.GetNSString() contents : [[NSData alloc] initWithBytes:FileContents.GetData() length : FileContents.Num()] sharedDB : true completionHandler : ^ (CKRecord *record, NSError *error)
+		{
+            FScopeLock ScopeLock(&CloudDataLock);
+            FCloudFile* File = GetCloudFile(NewFile);
+			if (error)
+			{
+				// TODO: record is potentially newer on the server
+				File->AsyncState = EOnlineAsyncTaskState::Failed;
+				TSharedRef<FSharedContentHandle> SharedHandle = MakeShareable(new FSharedContentHandleIOS());
+				TriggerOnWriteSharedFileCompleteDelegates(false, UserId, NewFile, SharedHandle);
+				NSLog(@"Error: %@", error);
+			}
+			else
+			{
+				File->Data = DataContents;
+				File->AsyncState = EOnlineAsyncTaskState::Done;
+				TSharedRef<FSharedContentHandle> SharedHandle = MakeShareable(new FSharedContentHandleIOS(NewFile));
+				TriggerOnWriteSharedFileCompleteDelegates(true, UserId, NewFile, SharedHandle);
+				NSLog(@"Record Saved!");
+			}
+		}];
+	}
+    else
+    {
+        TSharedRef<FSharedContentHandle> SharedHandle = MakeShareable(new FSharedContentHandleIOS());
+        TriggerOnWriteSharedFileCompleteDelegates(false, UserId, FileName, SharedHandle);
+    }
+	return false;
+}
+
+void FOnlineSharedCloudInterfaceIOS::GetDummySharedHandlesForTest(TArray< TSharedRef<FSharedContentHandle> >& OutHandles)
+{
+    OutHandles.Add(MakeShareable(new FSharedContentHandleIOS(TEXT("TestData1"))));
+    OutHandles.Add(MakeShareable(new FSharedContentHandleIOS(TEXT("TestData2"))));
+    OutHandles.Add(MakeShareable(new FSharedContentHandleIOS(TEXT("TestData3"))));
+    OutHandles.Add(MakeShareable(new FSharedContentHandleIOS(TEXT("TestData4"))));
+    OutHandles.Add(MakeShareable(new FSharedContentHandleIOS(TEXT("TestData5"))));
+}
\ No newline at end of file
diff --git a/Engine/Source/Runtime/Online/IOS/OnlineSubsystemIOS/Private/OnlineSharedCloudInterfaceIOS.h b/Engine/Source/Runtime/Online/IOS/OnlineSubsystemIOS/Private/OnlineSharedCloudInterfaceIOS.h
new file mode 100644
--- /dev/null
+++ b/Engine/Source/Runtime/Online/IOS/OnlineSubsystemIOS/Private/OnlineSharedCloudInterfaceIOS.h
@@ -0,0 +1,134 @@
+// Copyright 1998-2015 Epic Games, Inc. All Rights Reserved.
+ 
+#pragma once
+
+#include "OnlineSharedCloudInterface.h"
+#include "OnlineUserCloudInterfaceIOS.h"
+
+
+/**
+* iOS specific implementation of a shared file handle
+*/
+class FSharedContentHandleIOS :
+	public FSharedContentHandle
+{
+
+	/** Holds the handle to the shared content */
+	const FString SharedContentHandle;
+
+public:
+
+	FSharedContentHandleIOS()
+		: SharedContentHandle()
+	{
+	}
+
+	/**
+	* Constructs this object with the specified shared content id
+	*
+	* @param InUniqueNetId the id to set ours to
+	*/
+	FSharedContentHandleIOS(const FString& InSharedContentHandle) :
+		SharedContentHandle(InSharedContentHandle)
+	{
+	}
+
+	/**
+	*	Comparison operator
+	*/
+	bool operator==(const FSharedContentHandleIOS& Other) const
+	{
+		return (SharedContentHandle == Other.SharedContentHandle);
+	}
+
+	/**
+	* Get the raw byte representation of this shared content handle
+	* This data is platform dependent and shouldn't be manipulated directly
+	*
+	* @return byte array of size GetSize()
+	*/
+	virtual const uint8* GetBytes() const override
+	{
+		return (uint8*)(&SharedContentHandle);
+	}
+
+	/**
+	* Get the size of this shared content handle
+	*
+	* @return size in bytes of the id representation
+	*/
+	virtual int32 GetSize() const override
+	{
+		return SharedContentHandle.GetAllocatedSize();
+	}
+
+	/**
+	* Check the validity of this shared content handle
+	*
+	* @return true if this is a well formed ID, false otherwise
+	*/
+	virtual bool IsValid() const override
+	{
+		return SharedContentHandle != FString();
+	}
+
+	/**
+	* Platform specific conversion to string representation of data
+	*
+	* @return data in string form
+	*/
+	virtual FString ToString() const override
+	{
+		return SharedContentHandle;
+	}
+
+	/**
+	* Get a human readable representation of this shared content handle
+	* Shouldn't be used for anything other than logging/debugging
+	*
+	* @return handle in string form
+	*/
+	virtual FString ToDebugString() const override
+	{
+		return SharedContentHandle;
+	}
+};
+
+/**
+ * Provides the interface for sharing files already on the cloud with other users
+ */
+class FOnlineSharedCloudInterfaceIOS : public IOnlineSharedCloud
+{
+protected:
+
+	FCloudFile* GetCloudFile(const FString& FileName, bool bCreateIfMissing = false);
+	bool ClearFiles();
+	bool ClearCloudFile(const FString& FileName);
+
+PACKAGE_SCOPE:
+
+	FOnlineSharedCloudInterfaceIOS()
+	{
+	}
+
+public:
+
+	virtual ~FOnlineSharedCloudInterfaceIOS();
+
+	// IOnlineSharedCloud
+	virtual bool GetSharedFileContents(const FSharedContentHandle& SharedHandle, TArray<uint8>& FileContents) override;
+	virtual bool ClearSharedFiles() override;
+	virtual bool ClearSharedFile(const FSharedContentHandle& SharedHandle) override;
+	virtual bool ReadSharedFile(const FSharedContentHandle& SharedHandle) override;
+	virtual bool WriteSharedFile(const FUniqueNetId& UserId, const FString& Filename, TArray<uint8>& FileContents) override;
+	virtual void GetDummySharedHandlesForTest(TArray< TSharedRef<FSharedContentHandle> > & OutHandles) override;
+
+private:
+	/** File cache */
+	TArray<struct FCloudFile> CloudFileData;
+    
+    /** Critical section for thread safe operation on cloud files */
+    FCriticalSection CloudDataLock;
+};
+
+typedef TSharedPtr<FOnlineSharedCloudInterfaceIOS, ESPMode::ThreadSafe> FOnlineSharedCloudIOSPtr;
diff --git a/Engine/Source/Runtime/Online/IOS/OnlineSubsystemIOS/Private/OnlineSubsystemIOS.cpp b/Engine/Source/Runtime/Online/IOS/OnlineSubsystemIOS/Private/OnlineSubsystemIOS.cpp
--- a/Engine/Source/Runtime/Online/IOS/OnlineSubsystemIOS/Private/OnlineSubsystemIOS.cpp
+++ b/Engine/Source/Runtime/Online/IOS/OnlineSubsystemIOS/Private/OnlineSubsystemIOS.cpp
@@ -30,12 +30,12 @@ IOnlineGroupsPtr FOnlineSubsystemIOS::GetGroupsInterface() const
 
 IOnlineSharedCloudPtr FOnlineSubsystemIOS::GetSharedCloudInterface() const
 {
-	return nullptr;
+	return SharedCloudInterface;
 }
 
 IOnlineUserCloudPtr FOnlineSubsystemIOS::GetUserCloudInterface() const
 {
-	return nullptr;
+	return UserCloudInterface;
 }
 
 IOnlineUserCloudPtr FOnlineSubsystemIOS::GetUserCloudInterface(const FString& Key) const
@@ -148,6 +148,8 @@ bool FOnlineSubsystemIOS::Init()
 		AchievementsInterface = MakeShareable(new FOnlineAchievementsIOS(this));
 		ExternalUIInterface = MakeShareable(new FOnlineExternalUIIOS());
         TurnBasedInterface = MakeShareable(new FOnlineTurnBasedIOS());
+        UserCloudInterface = MakeShareable(new FOnlineUserCloudInterfaceIOS());
+        SharedCloudInterface = MakeShareable(new FOnlineSharedCloudInterfaceIOS());
 	}
 
 	if( IsInAppPurchasingEnabled() )
@@ -194,7 +196,9 @@ bool FOnlineSubsystemIOS::IsEnabled()
 {
 	bool bEnableGameCenter = false;
 	GConfig->GetBool(TEXT("/Script/IOSRuntimeSettings.IOSRuntimeSettings"), TEXT("bEnableGameCenterSupport"), bEnableGameCenter, GEngineIni);
-	return bEnableGameCenter;
+    bool bEnableCloudKit = false;
+    GConfig->GetBool(TEXT("/Script/IOSRuntimeSettings.IOSRuntimeSettings"), TEXT("bEnableCloudKitSupport"), bEnableCloudKit, GEngineIni);
+	return bEnableGameCenter || bEnableCloudKit;
 }
 
 bool FOnlineSubsystemIOS::IsInAppPurchasingEnabled()
diff --git a/Engine/Source/Runtime/Online/IOS/OnlineSubsystemIOS/Private/OnlineSubsystemIOSPrivatePCH.h b/Engine/Source/Runtime/Online/IOS/OnlineSubsystemIOS/Private/OnlineSubsystemIOSPrivatePCH.h
--- a/Engine/Source/Runtime/Online/IOS/OnlineSubsystemIOS/Private/OnlineSubsystemIOSPrivatePCH.h
+++ b/Engine/Source/Runtime/Online/IOS/OnlineSubsystemIOS/Private/OnlineSubsystemIOSPrivatePCH.h
@@ -18,6 +18,8 @@
 #include "OnlineAchievementsInterfaceIOS.h"
 #include "OnlineExternalUIInterfaceIOS.h"
 #include "OnlineTurnBasedInterfaceIOS.h"
+#include "OnlineUserCloudInterfaceIOS.h"
+#include "OnlineSharedCloudInterfaceIOS.h"
 
 #include "OnlineSubsystemIOS.h"
 #include "OnlineSubsystemIOSModule.h"
\ No newline at end of file
diff --git a/Engine/Source/Runtime/Online/IOS/OnlineSubsystemIOS/Private/OnlineUserCloudInterfaceIOS.cpp b/Engine/Source/Runtime/Online/IOS/OnlineSubsystemIOS/Private/OnlineUserCloudInterfaceIOS.cpp
new file mode 100644
--- /dev/null
+++ b/Engine/Source/Runtime/Online/IOS/OnlineSubsystemIOS/Private/OnlineUserCloudInterfaceIOS.cpp
@@ -0,0 +1,419 @@
+// Copyright 1998-2015 Epic Games, Inc. All Rights Reserved.
+
+#include "OnlineSubsystemIOSPrivatePCH.h"
+#include "OnlineUserCloudInterfaceIOS.h"
+
+@implementation IOSCloudStorage
+
+@synthesize CloudContainer;
+@synthesize SharedDatabase;
+@synthesize UserDatabase;
+@synthesize iCloudToken;
+
+-(IOSCloudStorage*)init:(bool)registerHandler
+{
+	// get the current iCloud ubiquity token
+	iCloudToken = [NSFileManager defaultManager].ubiquityIdentityToken;
+
+	// register for iCloud change notifications
+	if (registerHandler)
+	{
+		[[NSNotificationCenter defaultCenter] addObserver:self selector : @selector(iCloudAccountAvailabilityChanged:) name: NSUbiquityIdentityDidChangeNotification object : nil];
+	}
+
+	CloudContainer = [CKContainer defaultContainer];
+	SharedDatabase = [CloudContainer publicCloudDatabase];
+//	if (iCloudToken != nil)
+	{
+		UserDatabase = [CloudContainer privateCloudDatabase];
+	}
+//	else
+//	{
+//		UserDatabase = nil;
+//	}
+
+	return self;
+}
+
+-(bool)readFile:(NSString*)fileName sharedDB:(bool)shared completionHandler:(void(^)(CKRecord* record, NSError* error))handler
+{
+	CKDatabase* DB = shared ? SharedDatabase : UserDatabase;
+	if (DB != nil)
+	{
+		CKRecordID* recordId = [[CKRecordID alloc] initWithRecordName:fileName];
+		[DB fetchRecordWithID : recordId completionHandler : handler];
+		return true;
+	}
+	return false;
+}
+
+-(bool)writeFile:(NSString*)fileName contents:(NSData*)fileContents sharedDB:(bool)shared completionHandler:(void(^)(CKRecord* record, NSError* error))handler
+{
+	CKDatabase* DB = shared ? SharedDatabase : UserDatabase;
+	if (DB != nil)
+	{
+		CKRecordID* recordId = [[CKRecordID alloc] initWithRecordName:fileName];
+		CKRecord* record = [[CKRecord alloc] initWithRecordType:@"file" recordID: recordId];
+		record[@"contents"] = fileContents;
+		[DB saveRecord : record completionHandler : handler];
+		return true;
+	}
+	return false;
+}
+
+-(bool)deleteFile:(NSString*)fileName sharedDB:(bool)shared completionHandler:(void(^)(CKRecordID* record, NSError* error))handler
+{
+	CKDatabase* DB = shared ? SharedDatabase : UserDatabase;
+	if (DB != nil)
+	{
+		CKRecordID* recordId = [[CKRecordID alloc] initWithRecordName:fileName];
+		[DB deleteRecordWithID: recordId completionHandler : handler];
+		return true;
+	}
+	return false;
+}
+
+-(bool)query:(bool)shared fetchHandler:(void(^)(CKRecord* record))fetch completionHandler:(void(^)(CKQueryCursor* record, NSError* error))complete
+{
+	CKDatabase* DB = shared ? SharedDatabase : UserDatabase;
+	if (DB != nil)
+	{
+		CKQuery* query = [[CKQuery alloc] initWithRecordType:@"file" predicate:[NSPredicate predicateWithFormat : @"TRUEPREDICATE"]];
+		CKQueryOperation* queryOp = [[CKQueryOperation alloc] initWithQuery:query];
+		queryOp.desiredKeys = @[@"record.recordID.recordName"];
+		queryOp.recordFetchedBlock = fetch;
+		queryOp.queryCompletionBlock = complete;
+		queryOp.resultsLimit = CKQueryOperationMaximumResults;
+		[DB addOperation : queryOp];
+		return true;
+	}
+	return false;
+}
+
+-(void)iCloudAccountAvailabilityChanged:(NSNotification*)notification
+{
+	// access the token
+	id newiCloudToken = [NSFileManager defaultManager].ubiquityIdentityToken;
+
+	// check to see what it has changed to
+	if (newiCloudToken != iCloudToken)
+	{
+		[self init: false];
+	}
+}
+
++(IOSCloudStorage*)cloudStorage
+{
+	static IOSCloudStorage* theStorage = nil;
+
+	if (theStorage == nil)
+	{
+		theStorage = [[IOSCloudStorage alloc] init: true];
+	}
+	return theStorage;
+}
+
+@end
+
+//
+// UserCloud interface implementation
+//
+
+FOnlineUserCloudInterfaceIOS::~FOnlineUserCloudInterfaceIOS()
+{
+	ClearFiles();
+}
+
+FCloudFile* FOnlineUserCloudInterfaceIOS::GetCloudFile(const FString& FileName, bool bCreateIfMissing)
+{
+    FScopeLock ScopeLock(&CloudDataLock);
+	if (FileName.Len() > 0)
+	{
+		for (int32 FileIdx = 0; FileIdx < CloudFileData.Num(); FileIdx++)
+		{
+			FCloudFile* UserFileData = &CloudFileData[FileIdx];
+			if (UserFileData &&
+				UserFileData->FileName == FileName)
+			{
+				return UserFileData;
+			}
+		}
+
+		if (bCreateIfMissing)
+		{
+			return new (CloudFileData)FCloudFile(FileName);
+		}
+	}
+
+	return NULL;
+}
+
+FCloudFileHeader* FOnlineUserCloudInterfaceIOS::GetCloudFileHeader(const FString& FileName, bool bCreateIfMissing)
+{
+    FScopeLock ScopeLock(&CloudDataLock);
+    if (FileName.Len() > 0)
+    {
+        for (int32 FileIdx = 0; FileIdx < CloudMetaData.Num(); FileIdx++)
+        {
+            FCloudFileHeader* UserFileData = &CloudMetaData[FileIdx];
+            if (UserFileData &&
+                UserFileData->FileName == FileName)
+            {
+                return UserFileData;
+            }
+        }
+        
+        if (bCreateIfMissing)
+        {
+            new (CloudMetaData)FCloudFileHeader(FileName, FileName, 0);
+        }
+    }
+    
+    return NULL;
+}
+
+bool FOnlineUserCloudInterfaceIOS::ClearFiles()
+{
+    FScopeLock ScopeLock(&CloudDataLock);
+	// Delete file contents
+	for (int32 FileIdx = 0; FileIdx < CloudFileData.Num(); FileIdx++)
+	{
+		CloudFileData[FileIdx].Data.Empty();
+	}
+
+	// No async files being handled, so empty them all
+	CloudFileData.Empty();
+	return true;
+}
+
+bool FOnlineUserCloudInterfaceIOS::ClearCloudFile(const FString& FileName)
+{
+    FScopeLock ScopeLock(&CloudDataLock);
+	int32 FoundIndex = INDEX_NONE;
+	for (int32 FileIdx = 0; FileIdx < CloudFileData.Num(); FileIdx++)
+	{
+		FCloudFile* UserFileData = &CloudFileData[FileIdx];
+		if (UserFileData->FileName == FileName)
+		{
+			// If there is an async task outstanding, fail to empty
+			if (UserFileData->AsyncState == EOnlineAsyncTaskState::InProgress)
+			{
+				return false;
+			}
+
+			UserFileData->Data.Empty();
+			FoundIndex = FileIdx;
+			break;
+		}
+	}
+
+	if (FoundIndex != INDEX_NONE)
+	{
+		CloudFileData.RemoveAtSwap(FoundIndex);
+	}
+	return true;
+}
+
+bool FOnlineUserCloudInterfaceIOS::GetFileContents(const FUniqueNetId& UserId, const FString& FileName, TArray<uint8>& FileContents)
+{
+	FCloudFile* CloudFile = GetCloudFile(FileName);
+	if (CloudFile && CloudFile->AsyncState == EOnlineAsyncTaskState::Done && CloudFile->Data.Num() > 0)
+	{
+		FileContents = CloudFile->Data;
+		return true;
+	}
+	return false;
+}
+
+bool FOnlineUserCloudInterfaceIOS::ClearFiles(const FUniqueNetId& UserId)
+{
+	// NOTE: Return true regardless of if the user was valid or not. We don't care if
+	// there weren't any files in need of clearing, only if there was a failure to clear
+	ClearFiles();
+	return true;
+}
+
+bool FOnlineUserCloudInterfaceIOS::ClearFile(const FUniqueNetId& UserId, const FString& FileName)
+{
+	// NOTE: Return true regardless of if the file exists or not. The only way to return false
+	// would be in we failed to delete, and not deleting a non-existing file is not a failure
+	ClearCloudFile(FileName);
+	return true;
+}
+
+void FOnlineUserCloudInterfaceIOS::EnumerateUserFiles(const FUniqueNetId& UserId)
+{
+    MetaDataState = EOnlineAsyncTaskState::InProgress;
+	if ([[IOSCloudStorage cloudStorage] query:false fetchHandler:^(CKRecord* record)
+	{
+         FString FileName = record.recordID.recordName;
+         FCloudFileHeader* Header = GetCloudFileHeader(FileName, true);
+	}
+	completionHandler:^(CKQueryCursor* cursor, NSError* error)
+	{
+		if (error)
+		{
+			// TODO: record is potentially not found
+            MetaDataState = EOnlineAsyncTaskState::Failed;
+			TriggerOnEnumerateUserFilesCompleteDelegates(false, UserId);
+			NSLog(@"Error: %@", error);
+		}
+		else
+		{
+			// store the contents in the memory record database
+            MetaDataState = EOnlineAsyncTaskState::Done;
+			TriggerOnEnumerateUserFilesCompleteDelegates(true, UserId);
+			NSLog(@"Enumerated Read!");
+		}
+	}] == false)
+    {
+        TriggerOnEnumerateUserFilesCompleteDelegates(false, UserId);
+    }
+}
+
+void FOnlineUserCloudInterfaceIOS::GetUserFileList(const FUniqueNetId& UserId, TArray<FCloudFileHeader>& UserFiles)
+{
+    while (MetaDataState == EOnlineAsyncTaskState::InProgress)
+    {
+        FPlatformProcess::Sleep(0.01f);
+    }
+    UserFiles = CloudMetaData;
+}
+
+bool FOnlineUserCloudInterfaceIOS::ReadUserFile(const FUniqueNetId& UserId, const FString& FileName)
+{
+	FCloudFile* CloudFile = GetCloudFile(FileName);
+	if (CloudFile)
+	{
+        __block FString NewFile = FileName;
+        CloudFile->AsyncState = EOnlineAsyncTaskState::InProgress;
+		return [[IOSCloudStorage cloudStorage] readFile:FileName.GetNSString() sharedDB : false completionHandler : ^ (CKRecord *record, NSError *error)
+		{
+            FScopeLock ScopeLock(&CloudDataLock);
+            FCloudFile* File = GetCloudFile(NewFile);
+			if (error)
+			{
+				// TODO: record is potentially not found
+				File->AsyncState = EOnlineAsyncTaskState::Failed;
+				TriggerOnReadUserFileCompleteDelegates(false, UserId, NewFile);
+				NSLog(@"Error: %@", error);
+			}
+			else
+			{
+ 				// store the contents in the memory record database
+				NSData* data = (NSData*)record[@"contents"];
+				File->Data.Empty();
+				File->Data.Append((uint8*)data.bytes, data.length);
+				File->AsyncState = EOnlineAsyncTaskState::Done;
+				TriggerOnReadUserFileCompleteDelegates(true, UserId, NewFile);
+				NSLog(@"Record Read!");
+			}
+		}];
+	}
+    else
+    {
+        TriggerOnReadUserFileCompleteDelegates(false, UserId, FileName);
+    }
+	return false;
+}
+
+bool FOnlineUserCloudInterfaceIOS::WriteUserFile(const FUniqueNetId& UserId, const FString& FileName, TArray<uint8>& FileContents)
+{
+	FCloudFile* CloudFile = GetCloudFile(FileName, true);
+	if (CloudFile)
+	{
+        __block TArray<uint8> DataContents = FileContents;
+        __block FString NewFile = FileName;
+		CloudFile->AsyncState = EOnlineAsyncTaskState::InProgress;
+		return [[IOSCloudStorage cloudStorage] writeFile:FileName.GetNSString() contents : [[NSData alloc] initWithBytes:FileContents.GetData() length : FileContents.Num()] sharedDB : false completionHandler : ^ (CKRecord *record, NSError *error)
+		{
+            FScopeLock ScopeLock(&CloudDataLock);
+            FCloudFile* File = GetCloudFile(NewFile);
+			if (error)
+			{
+				// TODO: record is potentially newer on the server
+				File->AsyncState = EOnlineAsyncTaskState::Failed;
+				TriggerOnWriteUserFileCompleteDelegates(false, UserId, NewFile);
+				NSLog(@"Error: %@", error);
+			}
+			else
+			{
+                FCloudFileHeader* Header = GetCloudFileHeader(NewFile, true);
+				File->Data = DataContents;
+				File->AsyncState = EOnlineAsyncTaskState::Done;
+				TriggerOnWriteUserFileCompleteDelegates(true, UserId, NewFile);
+				NSLog(@"Record Saved!");
+			}
+		}];
+	}
+    else
+    {
+        TriggerOnWriteUserFileCompleteDelegates(false, UserId, FileName);
+    }
+	return false;
+}
+
+void FOnlineUserCloudInterfaceIOS::CancelWriteUserFile(const FUniqueNetId& UserId, const FString& FileName)
+{
+	// Not implemented
+}
+
+
+bool FOnlineUserCloudInterfaceIOS::DeleteUserFile(const FUniqueNetId& UserId, const FString& FileName, bool bShouldCloudDelete, bool bShouldLocallyDelete)
+{
+	FCloudFile* CloudFile = GetCloudFile(FileName);
+	if (CloudFile)
+	{
+        if (bShouldCloudDelete)
+        {
+            __block FString NewFile = FileName;
+            __block bool bDeleteLocal = bShouldLocallyDelete;
+            return [[IOSCloudStorage cloudStorage] deleteFile:FileName.GetNSString() sharedDB : false completionHandler : ^ (CKRecordID *record, NSError *error)
+            {
+                if (error)
+                {
+                    // TODO: record is potentially not found
+                    TriggerOnDeleteUserFileCompleteDelegates(false, UserId, NewFile);
+                    NSLog(@"Error: %@", error);
+                }
+                else
+                {
+                    // store the contents in the memory record database
+                    if (bDeleteLocal)
+                    {
+                        ClearCloudFile(NewFile);
+                    }
+                    TriggerOnDeleteUserFileCompleteDelegates(true, UserId, NewFile);
+                    NSLog(@"Record Deleted!");
+                }
+            }];
+        }
+        else if (bShouldLocallyDelete)
+        {
+            ClearCloudFile(FileName);
+            TriggerOnDeleteUserFileCompleteDelegates(true, UserId, FileName);
+            return true;
+        }
+	}
+    else
+    {
+        TriggerOnDeleteUserFileCompleteDelegates(false, UserId, FileName);
+    }
+	return false;
+}
+
+bool FOnlineUserCloudInterfaceIOS::RequestUsageInfo(const FUniqueNetId& UserId)
+{
+	// Not implemented
+	return false;
+}
+
+void FOnlineUserCloudInterfaceIOS::DumpCloudState(const FUniqueNetId& UserId)
+{
+	// Not implemented
+}
+
+void FOnlineUserCloudInterfaceIOS::DumpCloudFileState(const FUniqueNetId& UserId, const FString& FileName)
+{
+	// Not implemented
+}
diff --git a/Engine/Source/Runtime/Online/IOS/OnlineSubsystemIOS/Private/OnlineUserCloudInterfaceIOS.h b/Engine/Source/Runtime/Online/IOS/OnlineSubsystemIOS/Private/OnlineUserCloudInterfaceIOS.h
new file mode 100644
--- /dev/null
+++ b/Engine/Source/Runtime/Online/IOS/OnlineSubsystemIOS/Private/OnlineUserCloudInterfaceIOS.h
@@ -0,0 +1,71 @@
+// Copyright 1998-2015 Epic Games, Inc. All Rights Reserved.
+
+#pragma once
+
+#include "OnlineUserCloudInterface.h"
+#import <CloudKit/CloudKit.h>
+
+@class IOSCloudStorage;
+
+@interface IOSCloudStorage : NSObject
+
+@property(retain) CKContainer* CloudContainer;
+@property(retain) CKDatabase* SharedDatabase;
+@property(retain) CKDatabase* UserDatabase;
+@property(retain) id iCloudToken;
+
+-(IOSCloudStorage*)init:(bool)registerHandler;
+-(bool)readFile:(NSString*)fileName sharedDB:(bool)shared completionHandler:(void(^)(CKRecord* record, NSError* error))handler;
+-(bool)writeFile:(NSString*)fileName contents:(NSData*)fileContents sharedDB:(bool)shared completionHandler:(void(^)(CKRecord* record, NSError* error))handler;
+-(bool)deleteFile:(NSString*)fileName sharedDB:(bool)shared completionHandler:(void(^)(CKRecordID* record, NSError* error))handler;
+-(bool)query:(bool)shared fetchHandler : (void(^)(CKRecord* record))fetch completionHandler : (void(^)(CKQueryCursor* record, NSError* error))complete;
+
+-(void)iCloudAccountAvailabilityChanged:(NSNotification*)notification;
+
++(IOSCloudStorage*)cloudStorage;
+
+@end
+
+/**
+*	FOnlineUserCloudInterfaceIOS - Implementation of user cloud storage for IOS
+*/
+class FOnlineUserCloudInterfaceIOS : public IOnlineUserCloud
+{
+protected:
+
+	FCloudFile* GetCloudFile(const FString& FileName, bool bCreateIfMissing = false);
+    FCloudFileHeader* GetCloudFileHeader(const FString& FileName, bool bCreateIfMissing = false);
+	bool ClearFiles();
+	bool ClearCloudFile(const FString& FileName);
+
+public:
+    FOnlineUserCloudInterfaceIOS(){ MetaDataState = EOnlineAsyncTaskState::Done; }
+	virtual ~FOnlineUserCloudInterfaceIOS();
+
+	// IOnlineUserCloud
+	virtual bool GetFileContents(const FUniqueNetId& UserId, const FString& FileName, TArray<uint8>& FileContents) override;
+	virtual bool ClearFiles(const FUniqueNetId& UserId) override;
+	virtual bool ClearFile(const FUniqueNetId& UserId, const FString& FileName) override;
+	virtual void EnumerateUserFiles(const FUniqueNetId& UserId) override;
+	virtual void GetUserFileList(const FUniqueNetId& UserId, TArray<FCloudFileHeader>& UserFiles) override;
+	virtual bool ReadUserFile(const FUniqueNetId& UserId, const FString& FileName) override;
+	virtual bool WriteUserFile(const FUniqueNetId& UserId, const FString& FileName, TArray<uint8>& FileContents) override;
+	virtual void CancelWriteUserFile(const FUniqueNetId& UserId, const FString& FileName) override;
+	virtual bool DeleteUserFile(const FUniqueNetId& UserId, const FString& FileName, bool bShouldCloudDelete, bool bShouldLocallyDelete) override;
+	virtual bool RequestUsageInfo(const FUniqueNetId& UserId) override;
+	virtual void DumpCloudState(const FUniqueNetId& UserId) override;
+	virtual void DumpCloudFileState(const FUniqueNetId& UserId, const FString& FileName) override;
+
+private:
+	/** File metadata */
+	TArray<struct FCloudFileHeader> CloudMetaData;
+    /** File metadata query state */
+    EOnlineAsyncTaskState::Type MetaDataState;
+	/** File cache */
+	TArray<struct FCloudFile> CloudFileData;
+
+    /** Critical section for thread safe operation on cloud files */
+    FCriticalSection CloudDataLock;
+};
+
+typedef TSharedPtr<FOnlineUserCloudInterfaceIOS, ESPMode::ThreadSafe> FOnlineUserCloudIOSPtr;
diff --git a/Engine/Source/Runtime/Online/IOS/OnlineSubsystemIOS/Public/OnlineSubsystemIOS.h b/Engine/Source/Runtime/Online/IOS/OnlineSubsystemIOS/Public/OnlineSubsystemIOS.h
--- a/Engine/Source/Runtime/Online/IOS/OnlineSubsystemIOS/Public/OnlineSubsystemIOS.h
+++ b/Engine/Source/Runtime/Online/IOS/OnlineSubsystemIOS/Public/OnlineSubsystemIOS.h
@@ -91,6 +91,12 @@ private:
 
     /** Interface to the turnbased multiplayer services */
     FOnlineTurnBasedIOSPtr TurnBasedInterface;
+
+	/** Interface to the user cloud storage */
+	FOnlineUserCloudIOSPtr UserCloudInterface;
+
+	/** Interface to the shared cloud storage */
+	FOnlineSharedCloudIOSPtr SharedCloudInterface;
 };
 
 typedef TSharedPtr<FOnlineSubsystemIOS, ESPMode::ThreadSafe> FOnlineSubsystemIOSPtr;
