CIS fix

diff --git a/Engine/Source/Runtime/Online/Android/OnlineSubsystemGooglePlay/Private/OnlineStoreInterfaceGooglePlay.h b/Engine/Source/Runtime/Online/Android/OnlineSubsystemGooglePlay/Private/OnlineStoreInterfaceGooglePlay.h
--- a/Engine/Source/Runtime/Online/Android/OnlineSubsystemGooglePlay/Private/OnlineStoreInterfaceGooglePlay.h
+++ b/Engine/Source/Runtime/Online/Android/OnlineSubsystemGooglePlay/Private/OnlineStoreInterfaceGooglePlay.h
@@ -62,6 +62,11 @@ public:
 	virtual bool QueryForAvailablePurchases(const TArray<FString>& ProductIDs, FOnlineProductInformationReadRef& InReadObject) override;
 	virtual bool BeginPurchase(const FInAppPurchaseProductRequest& ProductRequest, FOnlineInAppPurchaseTransactionRef& InReadObject) override;
 	virtual bool IsAllowedToMakePurchases() override;
+	virtual bool RestorePurchases(FOnlineInAppPurchaseRestoreReadRef& InReadObject) override
+	{
+		// Not Yet Implemented
+		return false;
+	}
 	// End IOnlineStore 
 
 	void ProcessQueryAvailablePurchasesResults(bool bInSuccessful, const TArray<FInAppPurchaseProductInfo>& AvailablePurchases);
