BuildPatchServices: Fix for missing break statements in refactored switch that causes incorrect error messages for installer failures
Merging using REVERSE UE4-To-UE4-LauncherDev

diff --git a/Engine/Source/Runtime/Online/BuildPatchServices/Private/BuildPatchError.cpp b/Engine/Source/Runtime/Online/BuildPatchServices/Private/BuildPatchError.cpp
--- a/Engine/Source/Runtime/Online/BuildPatchServices/Private/BuildPatchError.cpp
+++ b/Engine/Source/Runtime/Online/BuildPatchServices/Private/BuildPatchError.cpp
@@ -66,10 +66,10 @@ const FText FBuildPatchInstallError::GetErrorText()
 	FText ErrorText = FText::GetEmpty();
 	switch( ErrorState )
 	{
-		case EBuildPatchInstallError::DownloadError: ErrorText = DownloadError;
-		case EBuildPatchInstallError::FileConstructionFail: ErrorText = FileConstructionFail;
-		case EBuildPatchInstallError::MoveFileToInstall: ErrorText = MoveFileToInstall;
-		case EBuildPatchInstallError::PathLengthExceeded: ErrorText = PathLengthExceeded;
+		case EBuildPatchInstallError::DownloadError: ErrorText = DownloadError; break;
+		case EBuildPatchInstallError::FileConstructionFail: ErrorText = FileConstructionFail; break;
+		case EBuildPatchInstallError::MoveFileToInstall: ErrorText = MoveFileToInstall; break;
+		case EBuildPatchInstallError::PathLengthExceeded: ErrorText = PathLengthExceeded; break;
 	}
 	return FText::Format(LOCTEXT("BuildPatchInstallLongError", "{0} {1}"), GetShortErrorText(), ErrorText);
 }
