Support for forcing unique null id's

diff --git a/Engine/Source/Runtime/Online/OnlineSubsystemNull/Private/OnlineIdentityNull.cpp b/Engine/Source/Runtime/Online/OnlineSubsystemNull/Private/OnlineIdentityNull.cpp
--- a/Engine/Source/Runtime/Online/OnlineSubsystemNull/Private/OnlineIdentityNull.cpp
+++ b/Engine/Source/Runtime/Online/OnlineSubsystemNull/Private/OnlineIdentityNull.cpp
@@ -38,14 +38,16 @@ inline FString GenerateRandomUserId(int32 LocalUserNum)
 		HostName = Addr->ToString(false);
 	}
 
-	if ( !GIsFirstInstance || GIsEditor )
+	const bool bForceUniqueId = FParse::Param( FCommandLine::Get(), TEXT( "StableNullID" ) );
+	
+	if ( ( GIsFirstInstance || bForceUniqueId ) && !GIsEditor )
 	{
-		// If we're not the first instance (or in the editor), return truly random id
-		return FString::Printf(TEXT("%s-%s"), *HostName, *FGuid::NewGuid().ToString());
+		// When possible, return a stable user id
+		return FString::Printf( TEXT( "%s-%s" ), *HostName, *FPlatformMisc::GetMachineId().ToString( EGuidFormats::Digits ) );
 	}
 
-	// When possible, return a stable user id
-	return FString::Printf( TEXT( "%s-%s" ), *HostName, *FPlatformMisc::GetMachineId().ToString( EGuidFormats::Digits ) );
+	// If we're not the first instance (or in the editor), return truly random id
+	return FString::Printf( TEXT( "%s-%s" ), *HostName, *FGuid::NewGuid().ToString() );
 }
 
 bool FOnlineIdentityNull::Login(int32 LocalUserNum, const FOnlineAccountCredentials& AccountCredentials)
