Throwing in some quick stats logging for BuildPatchTool, which will be reverted after running some tests.

diff --git a/Engine/Source/Runtime/Online/BuildPatchServices/Private/BuildPatchGeneration.cpp b/Engine/Source/Runtime/Online/BuildPatchServices/Private/BuildPatchGeneration.cpp
--- a/Engine/Source/Runtime/Online/BuildPatchServices/Private/BuildPatchGeneration.cpp
+++ b/Engine/Source/Runtime/Online/BuildPatchServices/Private/BuildPatchGeneration.cpp
@@ -9,6 +9,58 @@
 
 #if WITH_BUILDPATCHGENERATION
 
+namespace FBuildPatchTimers
+{
+	uint64 DataProcessed = 0;
+	volatile int64 AccumOpening = 0;
+	volatile int64 AccumReading = 0;
+	volatile int64 AccumEnqueueing = 0;
+	volatile int64 OtherTime = 0;
+	volatile int64 DataAccessSpeed = 0;
+	volatile int64 DataThroughputSpeed = 0;
+	volatile int64 TotalDataRead = 0;
+	uint32 TotalNumNewChunks = 0;
+	uint32 TotalNumKnownChunks = 0;
+	uint64 TimeStarted = 0;
+
+	uint64 GetCycles()
+	{
+#if PLATFORM_WINDOWS
+		LARGE_INTEGER Cycles;
+		QueryPerformanceCounter(&Cycles);
+		return Cycles.QuadPart;
+#elif PLATFORM_MAC
+		uint64 Cycles = mach_absolute_time();
+		return Cycles;
+#else
+		return FPlatformTime::Cycles();
+#endif
+	}
+
+	static double CyclesToSeconds( const uint64 Cycles )
+	{
+		return FPlatformTime::GetSecondsPerCycle() * Cycles;
+	}
+
+	void LogInfo(bool bForceLog = false)
+	{
+		// Log processed data
+		static uint64 LastLog = FBuildPatchTimers::GetCycles();
+		if (bForceLog || (FBuildPatchTimers::CyclesToSeconds((FBuildPatchTimers::GetCycles() - LastLog)) > 5.0f))
+		{
+			LastLog = FBuildPatchTimers::GetCycles();
+			GLog->Logf(TEXT("FGenerationInfo: Stream: ReadTime:  %s"), *FPlatformTime::PrettyTime(FBuildPatchTimers::CyclesToSeconds(AccumReading)));
+			GLog->Logf(TEXT("FGenerationInfo: Stream: OpenTime:  %s"), *FPlatformTime::PrettyTime(FBuildPatchTimers::CyclesToSeconds(AccumOpening)));
+			GLog->Logf(TEXT("FGenerationInfo: DataAccessSpeed:   %s/s"), *FText::AsMemory(DataAccessSpeed).ToString());
+			GLog->Logf(TEXT("FGenerationInfo: DataThroughput:    %s/s"), *FText::AsMemory(DataThroughputSpeed).ToString());
+			GLog->Logf(TEXT("FGenerationInfo: DataProcessed:     %s"), *FText::AsMemory(DataProcessed).ToString());
+			GLog->Logf(TEXT("FGenerationInfo: TotalTime:         %s"), *FPlatformTime::PrettyTime(FBuildPatchTimers::CyclesToSeconds(FBuildPatchTimers::GetCycles() - TimeStarted)));
+			GLog->Logf(TEXT("FGenerationInfo: NewChunks: %u"), TotalNumNewChunks);
+			GLog->Logf(TEXT("FGenerationInfo: OldChunks: %u\n"), TotalNumKnownChunks);
+		}
+	}
+};
+
 FFileAttributes::FFileAttributes()
 	: bReadOnly(false)
 	, bCompressed(false)
@@ -76,11 +128,18 @@ bool FBuildStream::FBuildStreamReader::Init()
 
 uint32 FBuildStream::FBuildStreamReader::Run()
 {
+	IFileManager* FileManager = &IFileManager::Get();
+
+	uint64 StartTime = FBuildPatchTimers::GetCycles();
 	// Clear the build stream
 	BuildStream->Clear();
 
 	TArray< FString > AllFiles;
-	IFileManager::Get().FindFilesRecursive( AllFiles, *DepotDirectory, TEXT("*.*"), true, false );
+	uint64 FileEnumerationStart = FBuildPatchTimers::GetCycles();
+	FileManager->FindFilesRecursive(AllFiles, *DepotDirectory, TEXT("*.*"), true, false);
+	uint64 FileEnumerationEnd = FBuildPatchTimers::GetCycles();
+	uint64 FileEnumerationTime = FileEnumerationEnd - FileEnumerationStart;
+	GLog->Logf(TEXT("FBuildStreamReader: Enumerated %d files in %s"), AllFiles.Num(), *FPlatformTime::PrettyTime(FBuildPatchTimers::CyclesToSeconds(FileEnumerationTime)));
 	AllFiles.Sort();
 
 	// Remove the files that appear in an ignore list
@@ -89,10 +148,16 @@ uint32 FBuildStream::FBuildStreamReader::Run()
 	// Allocate our file read buffer
 	uint8* FileReadBuffer = new uint8[ FileBufferSize ];
 
+
+	uint64 LastLogged = FBuildPatchTimers::GetCycles();
+	uint64 ProcessBegin = FBuildPatchTimers::GetCycles();
 	for (auto& SourceFile : AllFiles)
 	{
 		// Read the file
-		FArchive* FileReader = IFileManager::Get().CreateFileReader( *SourceFile );
+		uint64 FileOpenStart = FBuildPatchTimers::GetCycles();
+		FArchive* FileReader = FileManager->CreateFileReader(*SourceFile);
+		uint64 FileOpenEnd = FBuildPatchTimers::GetCycles();
+		FPlatformAtomics::InterlockedAdd(&FBuildPatchTimers::AccumOpening, FileOpenEnd - FileOpenStart);
 		const bool bBuildFileOpenSuccess = FileReader != NULL;
 		if( bBuildFileOpenSuccess )
 		{
@@ -106,10 +171,26 @@ uint32 FBuildStream::FBuildStreamReader::Run()
 				while( !FileReader->AtEnd() )
 				{
 					const int64 SizeLeft = FileSize - FileReader->Tell();
-					const uint32 ReadLen = FMath::Min< int64 >( FileBufferSize, SizeLeft );
-					FileReader->Serialize( FileReadBuffer, ReadLen );
+					const uint32 ReadLen = FMath::Min< int64 >(FileBufferSize, SizeLeft);
+					uint64 FileReadStart = FBuildPatchTimers::GetCycles();
+					FileReader->Serialize(FileReadBuffer, ReadLen);
+					uint64 FileReadEnd = FBuildPatchTimers::GetCycles();
+					FPlatformAtomics::InterlockedAdd(&FBuildPatchTimers::AccumReading, FileReadEnd - FileReadStart);
+					FPlatformAtomics::InterlockedAdd(&FBuildPatchTimers::TotalDataRead, ReadLen);
+
 					// Copy into data stream
-					BuildStream->EnqueueData( FileReadBuffer, ReadLen );
+					uint64 EnqueueDataStart = FBuildPatchTimers::GetCycles();
+					BuildStream->EnqueueData(FileReadBuffer, ReadLen);
+					uint64 EnqueueDataEnd = FBuildPatchTimers::GetCycles();
+					FPlatformAtomics::InterlockedAdd(&FBuildPatchTimers::AccumEnqueueing, EnqueueDataEnd - EnqueueDataStart);
+
+					uint64 DataAccess = FBuildPatchTimers::AccumOpening + FBuildPatchTimers::AccumReading;
+					FPlatformAtomics::InterlockedExchange(&FBuildPatchTimers::OtherTime, (FBuildPatchTimers::GetCycles() - StartTime) - DataAccess - FBuildPatchTimers::AccumEnqueueing);
+					double DataAccessTimeFloat = FBuildPatchTimers::CyclesToSeconds(DataAccess);
+					FPlatformAtomics::InterlockedExchange(&FBuildPatchTimers::DataAccessSpeed, FBuildPatchTimers::TotalDataRead / DataAccessTimeFloat);
+					double DataThroughputTimeFloat = FBuildPatchTimers::CyclesToSeconds(FBuildPatchTimers::GetCycles() - ProcessBegin);
+					FPlatformAtomics::InterlockedExchange(&FBuildPatchTimers::DataThroughputSpeed, FBuildPatchTimers::TotalDataRead / DataThroughputTimeFloat);
+
 				}
 			}
 			// Special case zero byte files
@@ -422,11 +503,9 @@ void FBuildDataChunkProcessor::EndNewChunk( const uint64& ChunkHash, const uint8
 
 	}
 
-	if( bLogProgress )
-	{
-		// Output to log for builder info
-		GLog->Logf( TEXT( "%s %s [%d:%d]" ), *BuildManifest->GetAppName(), *BuildManifest->GetVersionString(), NumNewChunks, NumKnownChunks );
-	}
+
+	FBuildPatchTimers::TotalNumNewChunks = NumNewChunks;
+	FBuildPatchTimers::TotalNumKnownChunks = NumKnownChunks;
 }
 
 void FBuildDataChunkProcessor::PushChunk()
@@ -1050,6 +1129,7 @@ bool FBuildDataGenerator::GenerateChunksManifestFromDirectory( const FBuildPatch
 {
 	// Output to log for builder info
 	GLog->Logf(TEXT("Running Chunks Patch Generation for: %u:%s %s"), Settings.AppID, *Settings.AppName, *Settings.BuildVersion);
+	FBuildPatchTimers::TimeStarted = FBuildPatchTimers::GetCycles();
 
 	// Take the build CS
 	FScopeLock SingleConcurrentBuild( &SingleConcurrentBuildCS );
@@ -1125,9 +1205,13 @@ bool FBuildDataGenerator::GenerateChunksManifestFromDirectory( const FBuildPatch
 	double LastProgressLog = FPlatformTime::Seconds();
 	const double TimeGenStarted = LastProgressLog;
 
+	FBuildPatchTimers::LogInfo();
+
 	// Loop through all data
 	while ( !BuildStream->IsEndOfData() )
 	{
+		FBuildPatchTimers::LogInfo();
+
 		// Grab some data from the build stream
 		ReadLen = BuildStream->DequeueData( DataBuffer, DataBufferSize );
 
@@ -1143,6 +1227,8 @@ bool FBuildDataGenerator::GenerateChunksManifestFromDirectory( const FBuildPatch
 		// Process data while we have more
 		while ( ( DataBufferPos < ReadLen ) || ( bNoMoreData && PaddedZeros < RollingHash->GetWindowSize() ) )
 		{
+			FBuildPatchTimers::LogInfo();
+
 			// Prime the rolling hash
 			if( RollingHash->GetNumDataNeeded() > 0 )
 			{
@@ -1207,19 +1293,19 @@ bool FBuildDataGenerator::GenerateChunksManifestFromDirectory( const FBuildPatch
 				}
 			}
 
-			// Log processed data
-			if( ( FPlatformTime::Seconds() - LastProgressLog ) >= 10.0 )
-			{
-				LastProgressLog = FPlatformTime::Seconds();
-				GLog->Logf( TEXT( "Processed %lld bytes." ), ProcessPos );
-			}
+			FBuildPatchTimers::DataProcessed = ProcessPos;
 		}
 	}
+	FBuildPatchTimers::LogInfo();
+
 
 	// The final chunk if any should be finished.
 	// This also triggers the chunk writer thread to exit.
 	DataProcessor.FinalChunk();
 
+	FBuildPatchTimers::LogInfo();
+
+
 	// Handle empty files
 	FSHA1 EmptyHasher;
 	EmptyHasher.Final();
@@ -1242,6 +1328,8 @@ bool FBuildDataGenerator::GenerateChunksManifestFromDirectory( const FBuildPatch
 		}
 	}
 
+	FBuildPatchTimers::LogInfo();
+
 	// Fill out lookups
 	BuildManifest->InitLookups();
 
@@ -1272,6 +1360,8 @@ bool FBuildDataGenerator::GenerateChunksManifestFromDirectory( const FBuildPatch
 	BuildManifest->Data->ManifestFileVersion = EBuildPatchAppManifestVersion::GetLatestJsonVersion();
 	BuildManifest->SaveToFile(JsonFilename, false);
 
+	FBuildPatchTimers::LogInfo(true);
+
 	// Output to log for builder info
 	GLog->Logf(TEXT("Saved manifest to %s"), *JsonFilename);
 
@@ -1562,6 +1652,7 @@ bool FBuildDataGenerator::FindExistingChunkData( const uint64& ChunkHash, const
 
 FString FBuildDataGenerator::DiscoverChunkFilename(const FGuid& ChunkGuid, const uint64& ChunkHash)
 {
+	return FBuildPatchUtils::GetChunkNewFilename(EBuildPatchAppManifestVersion::GetLatestVersion(), FBuildPatchServicesModule::GetCloudDirectory(), ChunkGuid, ChunkHash);
 	static double AccumTime = 0.0;
 	const double StartDiscovery = FPlatformTime::Seconds();
 	const FString CloudDir = FBuildPatchServicesModule::GetCloudDirectory();
