BuildPatchTool will now run a chunked manifest generation as a threaded process for faster build iteration

diff --git a/Engine/Source/Programs/BuildPatchTool/Private/BuildPatchToolMain.cpp b/Engine/Source/Programs/BuildPatchTool/Private/BuildPatchToolMain.cpp
--- a/Engine/Source/Programs/BuildPatchTool/Private/BuildPatchToolMain.cpp
+++ b/Engine/Source/Programs/BuildPatchTool/Private/BuildPatchToolMain.cpp
@@ -83,7 +83,6 @@
 #include "RequiredProgramMainCPPInclude.h"
 
 #include "BuildPatchServices.h"
-
 #include "CoreUObject.h"
 
 // Ensure compiled with patch generation code
@@ -143,6 +142,8 @@ namespace EBuildPatchToolMode
 int32 BuildPatchToolMain( const TCHAR* CommandLine )
 {
 	GEngineLoop.PreInit(CommandLine);
+	GConfig->InitializeConfigSystem();
+	IFileManager::Get().ProcessCommandLineOptions();
 
 	// Add log device
 	if (FParse::Param(FCommandLine::Get(), TEXT("stdout")))
@@ -372,7 +373,6 @@ int32 BuildPatchToolMain( const TCHAR* CommandLine )
 		{
 			// Initialize the configuration system, we can only do this reliably if we have CloudDirectory (i.e. bSuccess is true)
 			IniFile = CloudDirectory / TEXT("BuildPatchTool.ini");
-			GConfig->InitializeConfigSystem();
 
 			if (DataAgeThresholdIdx != INDEX_NONE)
 			{
@@ -406,9 +406,6 @@ int32 BuildPatchToolMain( const TCHAR* CommandLine )
 		}
 	}
 
-	// Initialize the file manager
-	IFileManager::Get().ProcessCommandLineOptions();
-
 	// Check for argument error
 	if( !bSuccess )
 	{
@@ -552,7 +549,7 @@ int32 BuildPatchToolMain( const TCHAR* CommandLine )
 
 INT32_MAIN_INT32_ARGC_TCHAR_ARGV()
 {
-	FString CommandLine;
+	FString CommandLine = TEXT("-usehyperthreading");
 	for( int32 Option = 1; Option < ArgC; Option++ )
 	{
 		CommandLine += TEXT(" ");
diff --git a/Engine/Source/Runtime/Online/BuildPatchServices/BuildPatchServices.Build.cs b/Engine/Source/Runtime/Online/BuildPatchServices/BuildPatchServices.Build.cs
--- a/Engine/Source/Runtime/Online/BuildPatchServices/BuildPatchServices.Build.cs
+++ b/Engine/Source/Runtime/Online/BuildPatchServices/BuildPatchServices.Build.cs
@@ -17,7 +17,7 @@ public class BuildPatchServices : ModuleRules
 		PrivateIncludePaths.AddRange(
 			new string[]
 			{
-				"Developer/NoRedist/BuildPatchServices/Private",
+				"Runtime/Online/BuildPatchServices/Private",
 			}
 		);
 
diff --git a/Engine/Source/Runtime/Online/BuildPatchServices/Private/BuildPatchChunk.cpp b/Engine/Source/Runtime/Online/BuildPatchServices/Private/BuildPatchChunk.cpp
--- a/Engine/Source/Runtime/Online/BuildPatchServices/Private/BuildPatchChunk.cpp
+++ b/Engine/Source/Runtime/Online/BuildPatchServices/Private/BuildPatchChunk.cpp
@@ -361,7 +361,6 @@ void FChunkWriter::FQueuedChunkWriter::SetNoMoreChunks()
 void FChunkWriter::FQueuedChunkWriter::GetChunkFilesizes(TMap<FGuid, int64>& OutChunkFileSizes)
 {
 	FScopeLock ScopeLock(&ChunkFileSizesCS);
-	OutChunkFileSizes.Empty(ChunkFileSizes.Num());
 	OutChunkFileSizes.Append(ChunkFileSizes);
 }
 
diff --git a/Engine/Source/Runtime/Online/BuildPatchServices/Private/BuildPatchGeneration.cpp b/Engine/Source/Runtime/Online/BuildPatchServices/Private/BuildPatchGeneration.cpp
--- a/Engine/Source/Runtime/Online/BuildPatchServices/Private/BuildPatchGeneration.cpp
+++ b/Engine/Source/Runtime/Online/BuildPatchServices/Private/BuildPatchGeneration.cpp
@@ -7,59 +7,14 @@
 
 #include "BuildPatchServicesPrivatePCH.h"
 
-#if WITH_BUILDPATCHGENERATION
+#include "Generation/DataScanner.h"
+#include "Generation/BuildStreamer.h"
+#include "Generation/CloudEnumeration.h"
+#include "Generation/ManifestBuilder.h"
 
-namespace FBuildPatchTimers
-{
-	uint64 DataProcessed = 0;
-	volatile int64 AccumOpening = 0;
-	volatile int64 AccumReading = 0;
-	volatile int64 AccumEnqueueing = 0;
-	volatile int64 OtherTime = 0;
-	volatile int64 DataAccessSpeed = 0;
-	volatile int64 DataThroughputSpeed = 0;
-	volatile int64 TotalDataRead = 0;
-	uint32 TotalNumNewChunks = 0;
-	uint32 TotalNumKnownChunks = 0;
-	uint64 TimeStarted = 0;
-
-	uint64 GetCycles()
-	{
-#if PLATFORM_WINDOWS
-		LARGE_INTEGER Cycles;
-		QueryPerformanceCounter(&Cycles);
-		return Cycles.QuadPart;
-#elif PLATFORM_MAC
-		uint64 Cycles = mach_absolute_time();
-		return Cycles;
-#else
-		return FPlatformTime::Cycles();
-#endif
-	}
-
-	static double CyclesToSeconds( const uint64 Cycles )
-	{
-		return FPlatformTime::GetSecondsPerCycle() * Cycles;
-	}
+using namespace FBuildPatchServices;
 
-	void LogInfo(bool bForceLog = false)
-	{
-		// Log processed data
-		static uint64 LastLog = FBuildPatchTimers::GetCycles();
-		if (bForceLog || (FBuildPatchTimers::CyclesToSeconds((FBuildPatchTimers::GetCycles() - LastLog)) > 5.0f))
-		{
-			LastLog = FBuildPatchTimers::GetCycles();
-			GLog->Logf(TEXT("FGenerationInfo: Stream: ReadTime:  %s"), *FPlatformTime::PrettyTime(FBuildPatchTimers::CyclesToSeconds(AccumReading)));
-			GLog->Logf(TEXT("FGenerationInfo: Stream: OpenTime:  %s"), *FPlatformTime::PrettyTime(FBuildPatchTimers::CyclesToSeconds(AccumOpening)));
-			GLog->Logf(TEXT("FGenerationInfo: DataAccessSpeed:   %s/s"), *FText::AsMemory(DataAccessSpeed).ToString());
-			GLog->Logf(TEXT("FGenerationInfo: DataThroughput:    %s/s"), *FText::AsMemory(DataThroughputSpeed).ToString());
-			GLog->Logf(TEXT("FGenerationInfo: DataProcessed:     %s"), *FText::AsMemory(DataProcessed).ToString());
-			GLog->Logf(TEXT("FGenerationInfo: TotalTime:         %s"), *FPlatformTime::PrettyTime(FBuildPatchTimers::CyclesToSeconds(FBuildPatchTimers::GetCycles() - TimeStarted)));
-			GLog->Logf(TEXT("FGenerationInfo: NewChunks: %u"), TotalNumNewChunks);
-			GLog->Logf(TEXT("FGenerationInfo: OldChunks: %u\n"), TotalNumKnownChunks);
-		}
-	}
-};
+#if WITH_BUILDPATCHGENERATION
 
 FFileAttributes::FFileAttributes()
 	: bReadOnly(false)
@@ -102,599 +57,6 @@ static FString GetSymlinkTarget(const TCHAR* Filename)
 	return TEXT("");
 }
 
-/* FBuildStreamReader implementation
-*****************************************************************************/
-FBuildStream::FBuildStreamReader::FBuildStreamReader()
-	: DepotDirectory(TEXT(""))
-	, BuildStream(NULL)
-	, Thread(NULL)
-{
-}
-
-FBuildStream::FBuildStreamReader::~FBuildStreamReader()
-{
-	if (Thread)
-	{
-		Thread->Kill();
-		delete Thread;
-		Thread = NULL;
-	}
-}
-
-bool FBuildStream::FBuildStreamReader::Init() 
-{
-	return BuildStream != NULL && IFileManager::Get().DirectoryExists( *DepotDirectory );
-}
-
-uint32 FBuildStream::FBuildStreamReader::Run()
-{
-	IFileManager* FileManager = &IFileManager::Get();
-
-	uint64 StartTime = FBuildPatchTimers::GetCycles();
-	// Clear the build stream
-	BuildStream->Clear();
-
-	TArray< FString > AllFiles;
-	uint64 FileEnumerationStart = FBuildPatchTimers::GetCycles();
-	FileManager->FindFilesRecursive(AllFiles, *DepotDirectory, TEXT("*.*"), true, false);
-	uint64 FileEnumerationEnd = FBuildPatchTimers::GetCycles();
-	uint64 FileEnumerationTime = FileEnumerationEnd - FileEnumerationStart;
-	GLog->Logf(TEXT("FBuildStreamReader: Enumerated %d files in %s"), AllFiles.Num(), *FPlatformTime::PrettyTime(FBuildPatchTimers::CyclesToSeconds(FileEnumerationTime)));
-	AllFiles.Sort();
-
-	// Remove the files that appear in an ignore list
-	FBuildDataGenerator::StripIgnoredFiles( AllFiles, DepotDirectory, IgnoreListFile );
-
-	// Allocate our file read buffer
-	uint8* FileReadBuffer = new uint8[ FileBufferSize ];
-
-
-	uint64 LastLogged = FBuildPatchTimers::GetCycles();
-	uint64 ProcessBegin = FBuildPatchTimers::GetCycles();
-	for (auto& SourceFile : AllFiles)
-	{
-		// Read the file
-		uint64 FileOpenStart = FBuildPatchTimers::GetCycles();
-		FArchive* FileReader = FileManager->CreateFileReader(*SourceFile);
-		uint64 FileOpenEnd = FBuildPatchTimers::GetCycles();
-		FPlatformAtomics::InterlockedAdd(&FBuildPatchTimers::AccumOpening, FileOpenEnd - FileOpenStart);
-		const bool bBuildFileOpenSuccess = FileReader != NULL;
-		if( bBuildFileOpenSuccess )
-		{
-			// Make SourceFile the format we want it in and start a new file
-			FPaths::MakePathRelativeTo( SourceFile, *( DepotDirectory + TEXT( "/" ) ) );
-			int64 FileSize = FileReader->TotalSize();
-			// Process files that have bytes
-			if( FileSize > 0 )
-			{
-				BuildStream->BeginNewFile( SourceFile, FileSize );
-				while( !FileReader->AtEnd() )
-				{
-					const int64 SizeLeft = FileSize - FileReader->Tell();
-					const uint32 ReadLen = FMath::Min< int64 >(FileBufferSize, SizeLeft);
-					uint64 FileReadStart = FBuildPatchTimers::GetCycles();
-					FileReader->Serialize(FileReadBuffer, ReadLen);
-					uint64 FileReadEnd = FBuildPatchTimers::GetCycles();
-					FPlatformAtomics::InterlockedAdd(&FBuildPatchTimers::AccumReading, FileReadEnd - FileReadStart);
-					FPlatformAtomics::InterlockedAdd(&FBuildPatchTimers::TotalDataRead, ReadLen);
-
-					// Copy into data stream
-					uint64 EnqueueDataStart = FBuildPatchTimers::GetCycles();
-					BuildStream->EnqueueData(FileReadBuffer, ReadLen);
-					uint64 EnqueueDataEnd = FBuildPatchTimers::GetCycles();
-					FPlatformAtomics::InterlockedAdd(&FBuildPatchTimers::AccumEnqueueing, EnqueueDataEnd - EnqueueDataStart);
-
-					uint64 DataAccess = FBuildPatchTimers::AccumOpening + FBuildPatchTimers::AccumReading;
-					FPlatformAtomics::InterlockedExchange(&FBuildPatchTimers::OtherTime, (FBuildPatchTimers::GetCycles() - StartTime) - DataAccess - FBuildPatchTimers::AccumEnqueueing);
-					double DataAccessTimeFloat = FBuildPatchTimers::CyclesToSeconds(DataAccess);
-					FPlatformAtomics::InterlockedExchange(&FBuildPatchTimers::DataAccessSpeed, FBuildPatchTimers::TotalDataRead / DataAccessTimeFloat);
-					double DataThroughputTimeFloat = FBuildPatchTimers::CyclesToSeconds(FBuildPatchTimers::GetCycles() - ProcessBegin);
-					FPlatformAtomics::InterlockedExchange(&FBuildPatchTimers::DataThroughputSpeed, FBuildPatchTimers::TotalDataRead / DataThroughputTimeFloat);
-
-				}
-			}
-			// Special case zero byte files
-			else if( FileSize == 0 )
-			{
-				BuildStream->AddEmptyFile( SourceFile );
-			}
-			FileReader->Close();
-			delete FileReader;
-		}
-		else
-		{
-			// Not being able to load a required file from the build would be fatal, hard fault.
-			GLog->Logf( TEXT( "FBuildStreamReader: Could not open file from build! %s" ), *SourceFile );
-			GLog->PanicFlushThreadedLogs();
-			// Use bool variable for easier to understand assert message
-			check( bBuildFileOpenSuccess );
-		}
-	}
-
-	// Mark end of build
-	BuildStream->EndOfBuild();
-
-	// Deallocate our file read buffer
-	delete[] FileReadBuffer;
-
-	return 0;
-}
-
-void FBuildStream::FBuildStreamReader::StartThread()
-{
-	Thread = FRunnableThread::Create(this, TEXT("BuildStreamReaderThread"));
-}
-
-/* FBuildStream implementation
-*****************************************************************************/
-FBuildStream::FBuildStream( const FString& RootDirectory, const FString& IgnoreListFile )
-{
-	BuildStreamReader.DepotDirectory = RootDirectory;
-	BuildStreamReader.IgnoreListFile = IgnoreListFile;
-	BuildStreamReader.BuildStream = this;
-	BuildStreamReader.StartThread();
-}
-
-FBuildStream::~FBuildStream()
-{
-}
-
-bool FBuildStream::GetFileSpan( const uint64& StartingIdx, FString& Filename, uint64& FileSize )
-{
-	bool bFound = false;
-	FilesListsCS.Lock();
-	// Find the filename
-	FFileSpan* FileSpan = FilesParsed.Find( StartingIdx );
-	if( FileSpan != NULL )
-	{
-		Filename = FileSpan->Filename;
-		FileSize = FileSpan->Size;
-		bFound = true;
-	}
-	FilesListsCS.Unlock();
-	return bFound;
-}
-
-uint32 FBuildStream::DequeueData( uint8* Buffer, const uint32& ReqSize, const bool WaitForData )
-{
-	// Wait for data
-	if( WaitForData )
-	{
-		while ( DataAvailable() < ReqSize && !IsEndOfBuild() )
-		{
-			FPlatformProcess::Sleep( 0.01f );
-		}
-	}
-
-	BuildDataStreamCS.Lock();
-	uint32 ReadLen = FMath::Min< int64 >( ReqSize, DataAvailable() );
-	ReadLen = BuildDataStream.Dequeue( Buffer, ReadLen );
-	BuildDataStreamCS.Unlock();
-
-	return ReadLen;
-}
-
-const TArray< FString > FBuildStream::GetEmptyFiles()
-{
-	FScopeLock ScopeLock( &FilesListsCS );
-	return EmptyFiles;
-}
-
-bool FBuildStream::IsEndOfBuild()
-{
-	bool rtn;
-	NoMoreDataCS.Lock();
-	rtn = bNoMoreData;
-	NoMoreDataCS.Unlock();
-	return rtn;
-}
-
-bool FBuildStream::IsEndOfData()
-{
-	bool rtn;
-	NoMoreDataCS.Lock();
-	rtn = bNoMoreData;
-	NoMoreDataCS.Unlock();
-	BuildDataStreamCS.Lock();
-	rtn &= BuildDataStream.RingDataUsage() == 0;
-	BuildDataStreamCS.Unlock();
-	return rtn;
-}
-
-void FBuildStream::Clear()
-{
-	EndOfBuild( false );
-
-	BuildDataStreamCS.Lock();
-	BuildDataStream.Empty();
-	BuildDataStreamCS.Unlock();
-
-	FilesListsCS.Lock();
-	FilesParsed.Empty();
-	FilesListsCS.Unlock();
-}
-
-uint32 FBuildStream::SpaceLeft()
-{
-	uint32 rtn;
-	BuildDataStreamCS.Lock();
-	rtn = BuildDataStream.RingDataSize() - BuildDataStream.RingDataUsage();
-	BuildDataStreamCS.Unlock();
-	return rtn;
-}
-
-uint32 FBuildStream::DataAvailable()
-{
-	uint32 rtn;
-	BuildDataStreamCS.Lock();
-	rtn = BuildDataStream.RingDataUsage();
-	BuildDataStreamCS.Unlock();
-	return rtn;
-}
-
-void FBuildStream::BeginNewFile( const FString& Filename, const uint64& FileSize )
-{
-	FFileSpan FileSpan;
-	FileSpan.Filename = Filename;
-	FileSpan.Size = FileSize;
-	BuildDataStreamCS.Lock();
-	FileSpan.StartIdx = BuildDataStream.TotalDataPushed();
-	BuildDataStreamCS.Unlock();
-	FilesListsCS.Lock();
-	FilesParsed.Add( FileSpan.StartIdx, FileSpan );
-	FilesListsCS.Unlock();
-}
-
-void FBuildStream::AddEmptyFile( const FString& Filename )
-{
-	FilesListsCS.Lock();
-	EmptyFiles.Add( Filename );
-	FilesListsCS.Unlock();
-}
-
-void FBuildStream::EnqueueData( const uint8* Buffer, const uint32& Len )
-{
-	// Wait for space
-	while ( SpaceLeft() < Len )
-	{
-		FPlatformProcess::Sleep( 0.01f );
-	}
-	BuildDataStreamCS.Lock();
-	BuildDataStream.Enqueue( Buffer, Len );
-	BuildDataStreamCS.Unlock();
-}
-
-void FBuildStream::EndOfBuild( bool bIsEnd )
-{
-	NoMoreDataCS.Lock();
-	bNoMoreData = bIsEnd;
-	NoMoreDataCS.Unlock();
-}
-
-/* FBuildDataProcessor implementation
-*****************************************************************************/
-FBuildDataChunkProcessor::FBuildDataChunkProcessor( FBuildPatchAppManifestRef InBuildManifest, const FString& InBuildRoot )
-	: NumNewChunks( 0 )
-	, NumKnownChunks( 0 )
-	, BuildRoot( InBuildRoot )
-	, ChunkWriter( FBuildPatchServicesModule::GetCloudDirectory() )
-	, CurrentChunkBufferPos( 0 )
-	, CurrentFile(nullptr)
-	, IsProcessingChunk( false )
-	, IsProcessingChunkPart( false )
-	, IsProcessingFile( false )
-	, ChunkIsPushed( false )
-	, BackupChunkBufferPos( 0 )
-	, BackupProcessingChunk( false )
-	, BackupProcessingChunkPart( false )
-{
-	BuildManifest = InBuildManifest;
-	CurrentChunkGuid.Invalidate();
-	BackupChunkGuid.Invalidate();
-	CurrentChunkBuffer = new uint8[ FBuildPatchData::ChunkDataSize ];
-}
-
-FBuildDataChunkProcessor::~FBuildDataChunkProcessor()
-{
-	delete[] CurrentChunkBuffer;
-}
-
-void FBuildDataChunkProcessor::BeginNewChunk( const bool& bZeroBuffer )
-{
-	check( IsProcessingChunk == false );
-	check( IsProcessingChunkPart == false );
-	IsProcessingChunk = true;
-
-	// NB: If you change this function, make sure you make required changes to 
-	// Backup of chunk if need be at the end of a recognized chunk
-
-	// Erase the chunk buffer data ready for new chunk parts
-	if( bZeroBuffer )
-	{
-		FMemory::Memzero( CurrentChunkBuffer, FBuildPatchData::ChunkDataSize );
-		CurrentChunkBufferPos = 0;
-	}
-
-	// Get a new GUID for this chunk for identification
-	CurrentChunkGuid = FGuid::NewGuid();
-}
-
-void FBuildDataChunkProcessor::BeginNewChunkPart()
-{
-	check( IsProcessingChunkPart == false );
-	check( IsProcessingChunk == true );
-	IsProcessingChunkPart = true;
-
-	// The current file should have a new chunk part setup
-	if (CurrentFile)
-	{
-		CurrentFile->FileChunkParts.Add(FChunkPartData());
-		FChunkPartData& NewPart = CurrentFile->FileChunkParts.Last();
-		NewPart.Guid = CurrentChunkGuid;
-		NewPart.Offset = CurrentChunkBufferPos;
-	}
-}
-
-void FBuildDataChunkProcessor::EndNewChunkPart()
-{
-	check( IsProcessingChunkPart == true );
-	check( IsProcessingChunk == true );
-	IsProcessingChunkPart = false;
-
-	// Current file should have it's last chunk part updated with the size value finalized.
-	if (CurrentFile)
-	{
-		FChunkPartData* ChunkPart = &CurrentFile->FileChunkParts.Last();
-		ChunkPart->Size = CurrentChunkBufferPos - ChunkPart->Offset;
-		check( ChunkPart->Size != 0 );
-	}
-}
-
-void FBuildDataChunkProcessor::EndNewChunk( const uint64& ChunkHash, const uint8* ChunkData, const FGuid& ChunkGuid )
-{
-	check( IsProcessingChunk == true );
-	check( IsProcessingChunkPart == false );
-	IsProcessingChunk = false;
-
-	// A bool that will state whether we should update and log chunking progress
-	bool bLogProgress = false;
-
-	// If the new chunk was recognized, then we will have got a different Guid through, so fix up guid for files using this chunk
-	if( CurrentChunkGuid != ChunkGuid )
-	{
-		bLogProgress = true;
-		++NumKnownChunks;
-		for (auto& FileManifest : BuildManifest->Data->FileManifestList)
-		{
-			for (auto& ChunkPart : FileManifest.FileChunkParts)
-			{
-				if( ChunkPart.Guid == CurrentChunkGuid )
-				{
-					ChunkPart.Guid = ChunkGuid;
-				}
-			}
-		}
-	}
-	// If the chunk is new, count it if we were passed data
-	else if( ChunkData )
-	{
-		bLogProgress = true;
-		++NumNewChunks;
-	}
-
-	// Always queue the chunk if passed data, it will be skipped automatically if existing already, and we need to make
-	// sure the latest version is saved out when recognizing an older version
-	if( ChunkData )
-	{
-		ChunkWriter.QueueChunk( ChunkData, ChunkGuid, ChunkHash );
-
-		// Also add the info to the data
-		if (!BuildManifest->ChunkInfoLookup.Contains(ChunkGuid))
-		{
-			BuildManifest->Data->ChunkList.Add(FChunkInfoData());
-			FChunkInfoData& NewChunk = BuildManifest->Data->ChunkList.Last();
-			NewChunk.Guid = ChunkGuid;
-			NewChunk.Hash = ChunkHash;
-			NewChunk.FileSize = INDEX_NONE;
-			NewChunk.GroupNumber = FCrc::MemCrc32(&ChunkGuid, sizeof(FGuid)) % 100;
-			// Lookup is used just to ensure one copy of each chunk added, we must not use the ptr as array will realloc
-			BuildManifest->ChunkInfoLookup.Add(ChunkGuid, nullptr);
-		}
-
-	}
-
-
-	FBuildPatchTimers::TotalNumNewChunks = NumNewChunks;
-	FBuildPatchTimers::TotalNumKnownChunks = NumKnownChunks;
-}
-
-void FBuildDataChunkProcessor::PushChunk()
-{
-	check( !ChunkIsPushed );
-	ChunkIsPushed = true;
-
-	// Then we can skip over this entire chunk using it for current files in processing
-	BackupProcessingChunk = IsProcessingChunk;
-	BackupProcessingChunkPart = IsProcessingChunkPart;
-	BackupChunkGuid = CurrentChunkGuid;
-	BackupChunkBufferPos = CurrentChunkBufferPos;
-	if( BackupProcessingChunkPart )
-	{
-		EndNewChunkPart();
-	}
-	if( BackupProcessingChunk )
-	{
-		// null ChunkData will stop the incomplete chunk being saved out
-		EndNewChunk( 0, NULL, CurrentChunkGuid );
-	}
-
-	// Start this chunk, must begin from 0
-	CurrentChunkBufferPos = 0;
-	BeginNewChunk( false );
-	BeginNewChunkPart();
-}
-
-void FBuildDataChunkProcessor::PopChunk( const uint64& ChunkHash, const uint8* ChunkData, const FGuid& ChunkGuid )
-{
-	check( ChunkIsPushed );
-	ChunkIsPushed = false;
-
-	EndNewChunkPart();
-	EndNewChunk( ChunkHash, ChunkData, ChunkGuid );
-
-	// Backup data for previous part chunk
-	IsProcessingChunk = BackupProcessingChunk;
-	CurrentChunkGuid = BackupChunkGuid;
-	CurrentChunkBufferPos = BackupChunkBufferPos;
-}
-
-void FBuildDataChunkProcessor::FinalChunk()
-{
-	// If we're processing a chunk part, then we MUST be processing a chunk too
-	check(IsProcessingChunkPart ? IsProcessingChunk : true);
-
-	if( IsProcessingChunk )
-	{
-		// The there is no more data and the last file is processed.
-		// The final chunk should be finished
-		uint64 NewChunkHash = FRollingHash< FBuildPatchData::ChunkDataSize >::GetHashForDataSet( CurrentChunkBuffer );
-		FGuid ChunkGuid = CurrentChunkGuid;
-		FBuildDataGenerator::FindExistingChunkData( NewChunkHash, CurrentChunkBuffer, ChunkGuid );
-		if( IsProcessingChunkPart )
-		{
-			EndNewChunkPart();
-		}
-		EndNewChunk( NewChunkHash, CurrentChunkBuffer, ChunkGuid );
-	}
-
-	// Wait for the chunk writer to finish
-	ChunkWriter.NoMoreChunks();
-	ChunkWriter.WaitForThread();
-}
-
-void FBuildDataChunkProcessor::BeginFile( const FString& FileName )
-{
-	check( IsProcessingFile == false );
-	IsProcessingFile = true;
-
-	// We should start a new file, which begins with the current new chunk and current chunk offset
-	BuildManifest->Data->FileManifestList.Add(FFileManifestData());
-	CurrentFile = &BuildManifest->Data->FileManifestList.Last();
-	CurrentFile->Filename = FileName;
-	CurrentFile->bIsUnixExecutable = IsUnixExecutable(*(BuildRoot / FileName));
-	CurrentFile->SymlinkTarget = GetSymlinkTarget(*(BuildRoot / FileName));
-
-	// Setup for current chunk part
-	if( IsProcessingChunkPart )
-	{
-		CurrentFile->FileChunkParts.Add(FChunkPartData());
-		FChunkPartData& NewPart = CurrentFile->FileChunkParts.Last();
-		NewPart.Guid = CurrentChunkGuid;
-		NewPart.Offset = CurrentChunkBufferPos;
-	}
-}
-
-void FBuildDataChunkProcessor::EndFile()
-{
-	check( IsProcessingFile == true );
-	IsProcessingFile = false;
-
-	if (CurrentFile)
-	{
-		if( IsProcessingChunkPart )
-		{
-			// Current file should have it's last chunk part updated with the size value finalized.
-			FChunkPartData* ChunkPart = &CurrentFile->FileChunkParts.Last();
-			ChunkPart->Size = CurrentChunkBufferPos - ChunkPart->Offset;
-		}
-		// Check file size is correct from chunks
-		int64 FileSystemSize = IFileManager::Get().FileSize(*(BuildRoot / CurrentFile->Filename));
-		CurrentFile->Init();
-		check(CurrentFile->GetFileSize() == FileSystemSize);
-	}
-
-	CurrentFile = nullptr;
-}
-
-void FBuildDataChunkProcessor::SkipKnownByte( const uint8& NextByte, const bool& bStartOfFile, const bool& bEndOfFile, const FString& Filename )
-{
-	// Check for start of new file
-	if( bStartOfFile )
-	{
-		BeginFile( Filename );
-		FileHash.Reset();
-	}
-
-	// Increment position between start and end of file
-	++CurrentChunkBufferPos;
-	FileHash.Update( &NextByte, 1 );
-
-	// Check for end of file
-	if( bEndOfFile )
-	{
-		FileHash.Final();
-		FileHash.GetHash( CurrentFile->FileHash.Hash );
-		EndFile();
-	}
-}
-
-void FBuildDataChunkProcessor::ProcessNewByte( const uint8& NewByte, const bool& bStartOfFile, const bool& bEndOfFile, const FString& Filename )
-{
-	// If we finished a chunk, we begin a new chunk and new chunk part
-	if( !IsProcessingChunk )
-	{
-		BeginNewChunk();
-		BeginNewChunkPart();
-	}
-	// Or if we finished a chunk part (by recognizing a chunk hash), we begin a new one
-	else if( !IsProcessingChunkPart )
-	{
-		BeginNewChunkPart();
-	}
-
-	// Check for start of new file
-	if( bStartOfFile )
-	{
-		BeginFile( Filename );
-		FileHash.Reset();
-	}
-
-	// We add the old byte to our new chunk which will be used as part of the file it belonged to
-	CurrentChunkBuffer[ CurrentChunkBufferPos++ ] = NewByte;
-	FileHash.Update( &NewByte, 1 );
-
-	// Check for end of file
-	if( bEndOfFile )
-	{
-		FileHash.Final();
-		FileHash.GetHash( CurrentFile->FileHash.Hash );
-		EndFile();
-	}
-
-	// Do we have a full new chunk?
-	check( CurrentChunkBufferPos <= FBuildPatchData::ChunkDataSize );
-	if( CurrentChunkBufferPos == FBuildPatchData::ChunkDataSize )
-	{
-		uint64 NewChunkHash = FRollingHash< FBuildPatchData::ChunkDataSize >::GetHashForDataSet( CurrentChunkBuffer );
-		FGuid ChunkGuid = CurrentChunkGuid;
-		FBuildDataGenerator::FindExistingChunkData( NewChunkHash, CurrentChunkBuffer, ChunkGuid );
-		EndNewChunkPart();
-		EndNewChunk( NewChunkHash, CurrentChunkBuffer, ChunkGuid );
-	}
-}
-
-void FBuildDataChunkProcessor::GetChunkStats( uint32& OutNewFiles, uint32& OutKnownFiles )
-{
-	OutNewFiles = NumNewChunks;
-	OutKnownFiles = NumKnownChunks;
-}
-
-const TMap<FGuid, int64>& FBuildDataChunkProcessor::GetChunkFilesizes()
-{
-	ChunkWriter.GetChunkFilesizes(ChunkFileSizes);
-	return ChunkFileSizes;
-}
-
 /* FBuildDataFileProcessor implementation
 *****************************************************************************/
 FBuildDataFileProcessor::FBuildDataFileProcessor(FBuildPatchAppManifestRef InBuildManifest, const FString& InBuildRoot, const FDateTime& InDataThresholdTime)
@@ -1070,8 +432,6 @@ static void AddCustomFieldsToBuildManifest(const TMap<FString, FVariant>& Custom
 
 static void FileAttributesMetaToMap(const FString& AttributesList, TMap<FString, FFileAttributes>& FileAttributesMap)
 {
-	GLog->Logf(TEXT("Parsing file attributes list:-"));
-	checkf(AttributesList.Len() > 0, TEXT("Attributes File List was empty file"));
 	const TCHAR Quote = TEXT('\"');
 	const TCHAR EOFile = TEXT('\0');
 	const TCHAR EOLine = TEXT('\n');
@@ -1098,105 +458,80 @@ static void FileAttributesMetaToMap(const FString& AttributesList, TMap<FString,
 		FString Filename = FString(FilenameEnd - FilenameStart, FilenameStart).Replace(TEXT("\\"), TEXT("/"));
 		FString Keywords(EndOfLine - FilenameEnd, FilenameEnd);
 		FFileAttributes FileAttributes;
-		GLog->Logf(TEXT("    %s"), *Filename);
 		if (Keywords.Contains(TEXT("readonly")))
 		{
 			FileAttributes.bReadOnly = true;
-			GLog->Logf(TEXT("        readonly"), *Filename);
 		}
 		if (Keywords.Contains(TEXT("compressed")))
 		{
 			FileAttributes.bCompressed = true;
-			GLog->Logf(TEXT("        compressed"), *Filename);
 		}
 		if (Keywords.Contains(TEXT("executable")))
 		{
 			FileAttributes.bUnixExecutable = true;
-			GLog->Logf(TEXT("        executable"), *Filename);
-		}
-		if (!(FileAttributes.bReadOnly || FileAttributes.bCompressed || FileAttributes.bUnixExecutable))
-		{
-			GLog->Logf(TEXT("        none"), *Filename);
 		}
 		FileAttributesMap.Add(MoveTemp(Filename), FileAttributes);
 	}
-
 }
 
 /* FBuildDataGenerator implementation
 *****************************************************************************/
 bool FBuildDataGenerator::GenerateChunksManifestFromDirectory( const FBuildPatchSettings& Settings )
 {
-	// Output to log for builder info
-	GLog->Logf(TEXT("Running Chunks Patch Generation for: %u:%s %s"), Settings.AppID, *Settings.AppName, *Settings.BuildVersion);
-	FBuildPatchTimers::TimeStarted = FBuildPatchTimers::GetCycles();
+	uint64 StartTime = FStatsCollector::GetCycles();
 
 	// Take the build CS
 	FScopeLock SingleConcurrentBuild( &SingleConcurrentBuildCS );
 
+	// Create stat collector
+	FStatsCollectorRef StatsCollector = FStatsCollectorFactory::Create();
+
+	// Create a build streamer
+	FBuildStreamerRef BuildStream = FBuildStreamerFactory::Create(Settings.RootDirectory, Settings.IgnoreListFile, StatsCollector);
+
+	// Enumerate Chunks
+	FCloudEnumerationRef CloudEnumeration = FCloudEnumerationFactory::Create(FBuildPatchServicesModule::GetCloudDirectory());
+
+	// Output to log for builder info
+	GLog->Logf(TEXT("Running Chunks Patch Generation for: %u:%s %s"), Settings.AppID, *Settings.AppName, *Settings.BuildVersion);
+
 	// Create our chunk cache
 	const FDateTime Cutoff = Settings.bShouldHonorReuseThreshold ? FDateTime::UtcNow() - FTimespan::FromDays(Settings.DataAgeThreshold) : FDateTime::MinValue();
 	FBuildGenerationChunkCache::Init(Cutoff);
 
-	// Create a manifest
-	FBuildPatchAppManifestRef BuildManifest = MakeShareable( new FBuildPatchAppManifest() );
-
-	// Setup custom fields
-	AddCustomFieldsToBuildManifest(Settings.CustomFields, BuildManifest);
-
+	// Create a manifest builder
+	FManifestDetails ManifestDetails;
+	ManifestDetails.bIsFileData = false;
+	ManifestDetails.AppId = Settings.AppID;
+	ManifestDetails.AppName = Settings.AppName;
+	ManifestDetails.BuildVersion = Settings.BuildVersion;
+	ManifestDetails.LaunchExe = Settings.LaunchExe;
+	ManifestDetails.LaunchCommand = Settings.LaunchCommand;
+	ManifestDetails.PrereqName = Settings.PrereqName;
+	ManifestDetails.PrereqPath = Settings.PrereqPath;
+	ManifestDetails.PrereqArgs = Settings.PrereqArgs;
+	ManifestDetails.CustomFields = Settings.CustomFields;
 	// Get the file attributes
 	FString AttributesList;
-	TMap<FString, FFileAttributes> FileAttributesMap;
 	if (Settings.AttributeListFile.Len() > 0)
 	{
 		FFileHelper::LoadFileToString(AttributesList, *Settings.AttributeListFile);
 		if (!AttributesList.IsEmpty())
 		{
-			FileAttributesMetaToMap(AttributesList, FileAttributesMap);
+			FileAttributesMetaToMap(AttributesList, ManifestDetails.FileAttributesMap);
 		}
 		else
 		{
 			GLog->Logf(TEXT("WARNING: Attributes list file empty"));
 		}
 	}
+	FManifestBuilderRef ManifestBuilder = FManifestBuilderFactory::Create(ManifestDetails, BuildStream);
 
-	// Reset chunk inventory
-	ExistingChunksEnumerated = false;
-	ExistingChunkGuidInventory.Empty();
-	ExistingChunkHashInventory.Empty();
-
-	// Declare a build processor
-	FBuildDataChunkProcessor DataProcessor(BuildManifest, Settings.RootDirectory);
-
-	// Create a build streamer
-	FBuildStream* BuildStream = new FBuildStream(Settings.RootDirectory, Settings.IgnoreListFile);
-
-	// Set the basic details
-	BuildManifest->Data->bIsFileData = false;
-	BuildManifest->Data->AppID = Settings.AppID;
-	BuildManifest->Data->AppName = Settings.AppName;
-	BuildManifest->Data->BuildVersion = Settings.BuildVersion;
-	BuildManifest->Data->LaunchExe = Settings.LaunchExe;
-	BuildManifest->Data->LaunchCommand = Settings.LaunchCommand;
-	BuildManifest->Data->PrereqName = Settings.PrereqName;
-	BuildManifest->Data->PrereqPath = Settings.PrereqPath;
-	BuildManifest->Data->PrereqArgs = Settings.PrereqArgs;
-
-	// Create a data buffer
-	const uint32 DataBufferSize = FBuildPatchData::ChunkDataSize;
-	uint8* DataBuffer = new uint8[ DataBufferSize ];
+	// Chunks matching
+	FDataMatcherRef DataMatcher = FDataMatcherFactory::Create(FBuildPatchServicesModule::GetCloudDirectory());
 
-	// We'll need a rolling hash for chunking
-	FRollingHash< FBuildPatchData::ChunkDataSize >* RollingHash = new FRollingHash< FBuildPatchData::ChunkDataSize >();
-	
-	// Refers to how much data has been processed (into the FBuildDataProcessor)
-	uint64 ProcessPos = 0;
-
-	// Records the current file we are processing
-	FString FileName;
-
-	// And the current file's data left to process
-	uint64 FileDataCount = 0;
+	// Refers to how much data has been dequeued
+	//uint64 ProcessPos = 0;
 
 	// Used to store data read lengths
 	uint32 ReadLen = 0;
@@ -1205,173 +540,64 @@ bool FBuildDataGenerator::GenerateChunksManifestFromDirectory( const FBuildPatch
 	double LastProgressLog = FPlatformTime::Seconds();
 	const double TimeGenStarted = LastProgressLog;
 
-	FBuildPatchTimers::LogInfo();
+	// 50MB Data buffer
+	const int32 DataBufferSize = 1024*1024*50;
+	const int32 OverlapSize = FBuildPatchData::ChunkDataSize - 1;
+	TArray<uint8> DataBuffer;
+	uint64 DataOffset = 0;
 
 	// Loop through all data
-	while ( !BuildStream->IsEndOfData() )
+	while (!BuildStream->IsEndOfData())
 	{
-		FBuildPatchTimers::LogInfo();
-
-		// Grab some data from the build stream
-		ReadLen = BuildStream->DequeueData( DataBuffer, DataBufferSize );
-
-		// A bool says if there's no more data to come from the Build Stream
-		const bool bNoMoreData = BuildStream->IsEndOfData();
-
-		// Refers to how much data from DataBuffer has been passed into the rolling hash
-		uint32 DataBufferPos = 0;
-
-		// Count how many times we pad the rolling hash with zero
-		uint32 PaddedZeros = 0;
-
-		// Process data while we have more
-		while ( ( DataBufferPos < ReadLen ) || ( bNoMoreData && PaddedZeros < RollingHash->GetWindowSize() ) )
+		// Keep the overlap data from previous scanner
+		int32 PreviousSize = DataBuffer.Num();
+		if(PreviousSize > 0)
 		{
-			FBuildPatchTimers::LogInfo();
-
-			// Prime the rolling hash
-			if( RollingHash->GetNumDataNeeded() > 0 )
-			{
-				if( DataBufferPos < ReadLen )
-				{
-					RollingHash->ConsumeByte( DataBuffer[ DataBufferPos++ ] );
-				}
-				else
-				{
-					RollingHash->ConsumeByte( 0 );
-					++PaddedZeros;
-				}
-				// Keep looping until primed
-				continue;
-			}
-
-			// Check if we recognized a chunk
-			FGuid ChunkGuid;
-			const uint64 WindowHash = RollingHash->GetWindowHash();
-			const TRingBuffer< uint8, FBuildPatchData::ChunkDataSize >& WindowData = RollingHash->GetWindowData();
-			bool ChunkRecognised = FindExistingChunkData( WindowHash, WindowData, ChunkGuid );
-			if( ChunkRecognised )
-			{
-				// Process all bytes
-				DataProcessor.PushChunk();
-				const uint32 WindowDataSize = RollingHash->GetWindowSize() - PaddedZeros;
-				for( uint32 i = 0; i < WindowDataSize; ++i )
-				{
-					const bool bStartOfFile = BuildStream->GetFileSpan( ProcessPos, FileName, FileDataCount );
-					const bool bEndOfFile = FileDataCount <= 1;
-					check( FileDataCount > 0 );// If FileDataCount is ever 0, it means this piece of data belongs to no file, so something is wrong
-					DataProcessor.SkipKnownByte( WindowData[i], bStartOfFile, bEndOfFile, FileName );
-					++ProcessPos;
-					--FileDataCount;
-				}
-				uint8* SerialWindowData = new uint8[ FBuildPatchData::ChunkDataSize ];
-				WindowData.Serialize( SerialWindowData );
-				DataProcessor.PopChunk( WindowHash, SerialWindowData, ChunkGuid );
-				delete[] SerialWindowData;
-
-				// Clear
-				RollingHash->Clear();
-			}
-			else
-			{
-				// Process one byte
-				const bool bStartOfFile = BuildStream->GetFileSpan( ProcessPos, FileName, FileDataCount );
-				const bool bEndOfFile = FileDataCount <= 1;
-				DataProcessor.ProcessNewByte( WindowData.Bottom(), bStartOfFile, bEndOfFile, FileName );
-				++ProcessPos;
-				--FileDataCount;
-
-				// Roll
-				if( DataBufferPos < ReadLen )
-				{
-					RollingHash->RollForward( DataBuffer[ DataBufferPos++ ] );
-				}
-				else if( bNoMoreData )
-				{
-					RollingHash->RollForward( 0 );
-					++PaddedZeros;
-				}
-			}
-
-			FBuildPatchTimers::DataProcessed = ProcessPos;
+			check(PreviousSize > OverlapSize);
+			uint8* CopyTo = DataBuffer.GetData();
+			uint8* CopyFrom = CopyTo + (PreviousSize - OverlapSize);
+			FMemory::Memcpy(CopyTo, CopyFrom, OverlapSize);
+			DataBuffer.SetNum(OverlapSize, false);
+			DataOffset += DataBufferSize - OverlapSize;
 		}
-	}
-	FBuildPatchTimers::LogInfo();
-
-
-	// The final chunk if any should be finished.
-	// This also triggers the chunk writer thread to exit.
-	DataProcessor.FinalChunk();
-
-	FBuildPatchTimers::LogInfo();
 
+		// Grab some data from the build stream
+		PreviousSize = DataBuffer.Num();
+		DataBuffer.SetNumUninitialized(DataBufferSize);
+		ReadLen = BuildStream->DequeueData(DataBuffer.GetData() + PreviousSize, DataBufferSize - PreviousSize);
+		DataBuffer.SetNum(PreviousSize + ReadLen, false);
 
-	// Handle empty files
-	FSHA1 EmptyHasher;
-	EmptyHasher.Final();
-	const TArray< FString >& EmptyFileList = BuildStream->GetEmptyFiles();
-	for (const auto& EmptyFile : EmptyFileList)
-	{
-		BuildManifest->Data->FileManifestList.Add(FFileManifestData());
-		FFileManifestData& EmptyFileManifest = BuildManifest->Data->FileManifestList.Last();
-		EmptyFileManifest.Filename = EmptyFile;
-		EmptyHasher.GetHash(EmptyFileManifest.FileHash.Hash);
-	}
-
-	// Add chunk sizes
-	const TMap<FGuid, int64>& ChunkFilesizes = DataProcessor.GetChunkFilesizes();
-	for (FChunkInfoData& ChunkInfo : BuildManifest->Data->ChunkList)
-	{
-		if (ChunkFilesizes.Contains(ChunkInfo.Guid))
+		// Create data processor, waiting for available slot first
+		while(FDataScannerCounter::GetNumIncompleteScanners() > FDataScannerCounter::GetNumRunningScanners())
 		{
-			ChunkInfo.FileSize = ChunkFilesizes[ChunkInfo.Guid];
+			FPlatformProcess::Sleep(0.01f);
+			StatsCollector->LogStats(2.0f);
 		}
-	}
 
-	FBuildPatchTimers::LogInfo();
+		// Pass a data scanner for this piece of data to the manifest builder
+		ManifestBuilder->AddDataScanner(FDataScannerFactory::Create(DataOffset, DataBuffer, CloudEnumeration, DataMatcher, StatsCollector));
 
-	// Fill out lookups
-	BuildManifest->InitLookups();
-
-	// Fill out the file attributes
-	for (const auto& Entry : FileAttributesMap)
-	{
-		const FString& Filename = Entry.Key;
-		const FFileAttributes& Attributes = Entry.Value;
-		if (BuildManifest->FileManifestLookup.Contains(Filename))
-		{
-			FFileManifestData& FileManifest = *BuildManifest->FileManifestLookup[Filename];
-			FileManifest.bIsReadOnly = Attributes.bReadOnly;
-			FileManifest.bIsCompressed = Attributes.bCompressed;
-			// Only overwrite unix exe if true
-			if (Attributes.bUnixExecutable)
-			{
-				FileManifest.bIsUnixExecutable = Attributes.bUnixExecutable;
-			}
-		}
-		else
-		{
-			GLog->Logf(TEXT("File Attributes: File not in build %s"), *Filename);
-		}
+		// Log collected stats
+		StatsCollector->LogStats(2.0f);
 	}
 
 	// Save manifest into the cloud directory
-	FString JsonFilename = FBuildPatchServicesModule::GetCloudDirectory() / FDefaultValueHelper::RemoveWhitespaces(BuildManifest->Data->AppName + BuildManifest->Data->BuildVersion) + TEXT(".manifest");
-	BuildManifest->Data->ManifestFileVersion = EBuildPatchAppManifestVersion::GetLatestJsonVersion();
-	BuildManifest->SaveToFile(JsonFilename, false);
+	FString JsonFilename = FBuildPatchServicesModule::GetCloudDirectory() / FDefaultValueHelper::RemoveWhitespaces(Settings.AppName + Settings.BuildVersion) + TEXT(".manifest");
+	ManifestBuilder->SaveToFile(JsonFilename);
 
-	FBuildPatchTimers::LogInfo(true);
+	// Final value log
+	StatsCollector->LogStats();
 
 	// Output to log for builder info
 	GLog->Logf(TEXT("Saved manifest to %s"), *JsonFilename);
 
+
+	uint64 EndTime = FStatsCollector::GetCycles();
+	GLog->Logf(TEXT("Completed in %s"), *FPlatformTime::PrettyTime(FStatsCollector::CyclesToSeconds(EndTime - StartTime)));
+
 	// Clean up memory
-	delete[] DataBuffer;
-	delete BuildStream;
-	delete RollingHash;
 	FBuildGenerationChunkCache::Shutdown();
 
-	// @TODO LSwift: Detect errors and return false on failure
 	return true;
 }
 
@@ -1444,6 +670,7 @@ bool FBuildDataGenerator::GenerateFilesManifestFromDirectory( const FBuildPatchS
 		const FString& FileName = *FileIt;
 		// Read the file
 		FArchive* FileReader = IFileManager::Get().CreateFileReader( *FileName );
+		checkf(FileReader != nullptr, TEXT("Fatal Error: Could not open build file %s"), *FileName);
 		if( FileReader != NULL )
 		{
 			// Make SourceFile the format we want it in and start a new file
@@ -1469,10 +696,6 @@ bool FBuildDataGenerator::GenerateFilesManifestFromDirectory( const FBuildPatchS
 			delete FileReader;
 			DataProcessor.EndFile();
 		}
-		else
-		{
-			// @TODO LSwift: Handle File error?
-		}
 	}
 
 	// Fill out lookups
@@ -1515,171 +738,6 @@ bool FBuildDataGenerator::GenerateFilesManifestFromDirectory( const FBuildPatchS
 	return true;
 }
 
-bool FBuildDataGenerator::FindExistingChunkData( const uint64& ChunkHash, const uint8* ChunkData, FGuid& ChunkGuid )
-{
-	// Quick code hack
-	TRingBuffer< uint8, FBuildPatchData::ChunkDataSize > ChunkDataRing;
-	const uint32 ChunkDataLen = FBuildPatchData::ChunkDataSize;
-	ChunkDataRing.Enqueue( ChunkData, ChunkDataLen );
-	return FindExistingChunkData( ChunkHash, ChunkDataRing, ChunkGuid );
-}
-
-bool FBuildDataGenerator::FindExistingChunkData( const uint64& ChunkHash, const TRingBuffer< uint8, FBuildPatchData::ChunkDataSize >& ChunkData, FGuid& ChunkGuid )
-{
-	bool bFoundMatchingChunk = false;
-
-	// Perform an inventory on Cloud chunks if not already done
-	if( ExistingChunksEnumerated == false )
-	{
-		IFileManager& FileManager = IFileManager::Get();
-		FString JSONOutput;
-		TSharedRef< TJsonWriter< TCHAR, TPrettyJsonPrintPolicy< TCHAR > > > DebugWriter = TJsonWriterFactory< TCHAR, TPrettyJsonPrintPolicy< TCHAR > >::Create( &JSONOutput );
-		DebugWriter->WriteObjectStart();
-
-		// Find all manifest files
-		const FString CloudDir = FBuildPatchServicesModule::GetCloudDirectory();
-		if (FileManager.DirectoryExists(*CloudDir))
-		{
-			const double StartEnumerate = FPlatformTime::Seconds();
-			TArray<FString> AllManifests;
-			GLog->Logf(TEXT("BuildDataGenerator: Enumerating Manifests from %s"), *CloudDir);
-			FileManager.FindFiles(AllManifests, *(CloudDir / TEXT("*.manifest")), true, false);
-			const double EnumerateTime = FPlatformTime::Seconds() - StartEnumerate;
-			GLog->Logf(TEXT("BuildDataGenerator: Found %d manifests in %.1f seconds"), AllManifests.Num(), EnumerateTime);
-
-			// Load all manifest files
-			uint64 NumChunksFound = 0;
-			const double StartLoadAllManifest = FPlatformTime::Seconds();
-			for (const auto& ManifestFile : AllManifests)
-			{
-				// Determine chunks from manifest file
-				const FString ManifestFilename = CloudDir / ManifestFile;
-				FBuildPatchAppManifestRef BuildManifest = MakeShareable(new FBuildPatchAppManifest());
-				const double StartLoadManifest = FPlatformTime::Seconds();
-				if (BuildManifest->LoadFromFile(ManifestFilename))
-				{
-					const double LoadManifestTime = FPlatformTime::Seconds() - StartLoadManifest;
-					GLog->Logf(TEXT("BuildDataGenerator: Loaded %s in %.1f seconds"), *ManifestFile, LoadManifestTime);
-					if(!BuildManifest->IsFileDataManifest())
-					{
-						TArray<FGuid> ChunksReferenced;
-						BuildManifest->GetDataList(ChunksReferenced);
-						for (const auto& ReferencedChunkGuid : ChunksReferenced)
-						{
-							uint64 ReferencedChunkHash;
-							if (BuildManifest->GetChunkHash(ReferencedChunkGuid, ReferencedChunkHash))
-							{
-								if (ReferencedChunkHash != 0)
-								{
-									TArray< FGuid >& HashChunkList = ExistingChunkHashInventory.FindOrAdd(ReferencedChunkHash);
-									if (!HashChunkList.Contains(ReferencedChunkGuid))
-									{
-										++NumChunksFound;
-										HashChunkList.Add(ReferencedChunkGuid);
-									}
-								}
-								else
-								{
-									GLog->Logf(TEXT("BuildDataGenerator: INFO: Ignored an existing chunk %s with a failed hash value of zero to avoid performance problems while chunking"), *ReferencedChunkGuid.ToString());
-								}
-							}
-							else
-							{
-								GLog->Logf(TEXT("BuildDataGenerator: WARNING: Missing chunk hash for %s in manifest %s"), *ReferencedChunkGuid.ToString(), *ManifestFile);
-							}
-						}
-					}
-					else
-					{
-						GLog->Logf(TEXT("BuildDataGenerator: INFO: Ignoring non-chunked manifest %s"), *ManifestFilename);
-					}
-				}
-				else
-				{
-					GLog->Logf(TEXT("BuildDataGenerator: WARNING: Could not read Manifest file. Data recognition will suffer (%s)"), *ManifestFilename);
-				}
-			}
-			const double LoadAllManifestTime = FPlatformTime::Seconds() - StartLoadAllManifest;
-			GLog->Logf(TEXT("BuildDataGenerator: Used %d manifests to enumerate %llu chunks in %.1f seconds"), AllManifests.Num(), NumChunksFound, LoadAllManifestTime);
-		}
-		else
-		{
-			GLog->Logf(TEXT("BuildDataGenerator: Cloud directory does not exist: %s"), *CloudDir);
-		}
-
-		ExistingChunksEnumerated = true;
-	}
-
-	// Do we have a chunk matching this data?
-	if( ExistingChunkHashInventory.Num() > 0 )
-	{
-		TArray< FGuid >* ChunkList = ExistingChunkHashInventory.Find( ChunkHash );
-		if( ChunkList != NULL )
-		{
-			// We need to load each chunk in this list and compare data
-			for( auto ChunkIt = ChunkList->CreateConstIterator(); ChunkIt && !bFoundMatchingChunk ; ++ChunkIt)
-			{
-				FGuid Guid = *ChunkIt;
-				if (!ExistingChunkGuidInventory.Contains(Guid))
-				{
-					ExistingChunkGuidInventory.Add(Guid, DiscoverChunkFilename(Guid, ChunkHash));
-				}
-				const FString& SourceFile = ExistingChunkGuidInventory[ Guid ];
-				// Read the file
-				uint8* TempChunkData = new uint8[ FBuildPatchData::ChunkDataSize ];
-				ChunkData.Serialize( TempChunkData );
-				bool bChunkIsUsable = true;
-				if( CompareDataToChunk( SourceFile, TempChunkData, Guid, bChunkIsUsable ) )
-				{
-					// We have a chunk match!!
-					bFoundMatchingChunk = true;
-					ChunkGuid = Guid;
-				}
-				// Check if this chunk should be dumped
-				if(!bChunkIsUsable)
-				{
-					GLog->Logf(TEXT("BuildDataGenerator: Chunk %s unusable, removed from inventory."), *ChunkGuid.ToString());
-					ChunkList->Remove(Guid);
-					--ChunkIt;
-				}
-				delete[] TempChunkData;
-			}
-		}
-	}
-
-	return bFoundMatchingChunk;
-}
-
-FString FBuildDataGenerator::DiscoverChunkFilename(const FGuid& ChunkGuid, const uint64& ChunkHash)
-{
-	return FBuildPatchUtils::GetChunkNewFilename(EBuildPatchAppManifestVersion::GetLatestVersion(), FBuildPatchServicesModule::GetCloudDirectory(), ChunkGuid, ChunkHash);
-	static double AccumTime = 0.0;
-	const double StartDiscovery = FPlatformTime::Seconds();
-	const FString CloudDir = FBuildPatchServicesModule::GetCloudDirectory();
-	FString ChunkFilenameChecked;
-	for (EBuildPatchAppManifestVersion::Type VersionCounter = static_cast<EBuildPatchAppManifestVersion::Type>(EBuildPatchAppManifestVersion::LatestPlusOne - 1);
-		VersionCounter >= EBuildPatchAppManifestVersion::Original;
-		VersionCounter = static_cast<EBuildPatchAppManifestVersion::Type>(VersionCounter - 1))
-	{
-		const FString ChunkFilename = VersionCounter < EBuildPatchAppManifestVersion::DataFileRenames ?
-			FBuildPatchUtils::GetChunkOldFilename(CloudDir, ChunkGuid) :
-			FBuildPatchUtils::GetChunkNewFilename(VersionCounter, CloudDir, ChunkGuid, ChunkHash);
-		if (ChunkFilenameChecked != ChunkFilename)
-		{
-			ChunkFilenameChecked = ChunkFilename;
-			if (FPaths::FileExists(ChunkFilename))
-			{
-				const double DiscoveryTime = FPlatformTime::Seconds() - StartDiscovery;
-				AccumTime += DiscoveryTime;
-				GLog->Logf(TEXT("BuildDataGenerator: DiscoverChunkFilename: Chunk %s found in %f secs. Accum time %.2f secs"), *ChunkGuid.ToString(), DiscoveryTime, AccumTime);
-				return ChunkFilename;
-			}
-		}
-	}
-	GLog->Logf(TEXT("BuildDataGenerator: DiscoverChunkFilename: Chunk %s not found"), *ChunkGuid.ToString());
-	return TEXT("");
-}
-
 bool FBuildDataGenerator::FindExistingFileData(const FString& InSourceFile, const FSHAHashData& InFileHash, const FDateTime& DataThresholdTime, FGuid& OutFileGuid)
 {
 	bool bFoundMatchingFile = false;
diff --git a/Engine/Source/Runtime/Online/BuildPatchServices/Private/BuildPatchGeneration.h b/Engine/Source/Runtime/Online/BuildPatchServices/Private/BuildPatchGeneration.h
--- a/Engine/Source/Runtime/Online/BuildPatchServices/Private/BuildPatchGeneration.h
+++ b/Engine/Source/Runtime/Online/BuildPatchServices/Private/BuildPatchGeneration.h
@@ -4,357 +4,10 @@
 	BuildPatchGeneration.h: Declares classes involved with producing build data.
 =============================================================================*/
 
-#ifndef __BuildPatchGeneration_h__
-#define __BuildPatchGeneration_h__
-
 #pragma once
 
 #if WITH_BUILDPATCHGENERATION
 
-/**
- * A threaded class that holds Build stream data as a FIFO data buffer of given size. Will read until full then wait
- * for data to be read out before continuing.
- */
-class FBuildStream
-{
-private:
-	// A private runnable class that reads in data
-	class FBuildStreamReader : public FRunnable
-	{
-	public:
-		// The directory that we want to read from
-		FString DepotDirectory;
-
-		// The path to a text file containing \r\n separated depot relative files to ignore
-		FString IgnoreListFile;
-
-		// Pointer to FBuildStream
-		FBuildStream* BuildStream;
-
-		// Default Constructor
-		FBuildStreamReader();
-
-		// Destructor
-		~FBuildStreamReader();
-
-		// Start FRunnable interface
-		virtual bool Init() override;
-		virtual uint32 Run() override;
-		// End FRunnable interface
-
-		void StartThread();
-
-		FRunnableThread * Thread;
-
-	} BuildStreamReader;
-
-	// A private struct for storing the data span of a file in the build stream
-	struct FFileSpan
-	{
-		// The file's filename
-		FString Filename;
-		// The file data start index
-		uint64 StartIdx;
-		// The size of the file
-		uint64 Size;
-		
-		/**
-		 * Default constructor
-		 */
-		FFileSpan()
-			: StartIdx( 0 )
-			, Size( 0 )
-		{
-		}
-	};
-
-public:
-	/**
-	 * Constructor
-	 */
-	FBuildStream( const FString& RootDirectory, const FString& IgnoreListFile );
-
-	/**
-	 * Default destructor
-	 */
-	~FBuildStream();
-
-	/**
-	 * Retrieves the file details for a specific start index.
-	 * @param	IN	StartingIdx		The data index into the build image.
-	 * @param	OUT Filename		The filename for the file starting at StartingIdx. NB: The variable remains UNCHANGED if no file starts here.
-	 * @param	OUT FileSize		The size of the file starting at StartingIdx. NB: The variable remains UNCHANGED if no file starts here.
-	 * @return	true if the data byte at StartingIdx is the start of a file
-	 */
-	bool GetFileSpan( const uint64& StartingIdx, FString& Filename, uint64& FileSize );
-
-	/**
-	 * Fetches some data from the buffer, also removing it.
-	 * @param	IN	Buffer			Pointer to buffer to receive the data.
-	 * @param	IN	ReqSize			The amount of data to attempt to retrieve.
-	 * @param	IN	WaitForData		Optional: Default true. Whether to wait until there is enough data in the buffer.
-	 * @return	The amount of data retrieved
-	 */
-	uint32 DequeueData( uint8* Buffer, const uint32& ReqSize, const bool WaitForData = true );
-
-	/**
-	 * Gets a list of empty files that the build contains.
-	 * @return	Array of empty files in the build
-	 */
-	const TArray< FString > GetEmptyFiles();
-
-	/**
-	 * Whether the read thread has finished reading the build image.
-	 * @return	true if there is no more data coming into the buffer
-	 */
-	bool IsEndOfBuild();
-
-	/**
-	 * Whether there is any more data available to dequeue from the buffer
-	 * @return	true if there is no more data coming in, and the internal buffer is also empty.
-	 */
-	bool IsEndOfData();
-
-private:
-	/**
-	 * Hide the default constructor so that the params must be given
-	 */
-	FBuildStream(){};
-
-	//////////////////////////////////////////////////////////////////////////
-	// Private functions for manipulating the data store
-private:
-	/**
-	 * Wipes out all data
-	 */
-	void Clear();
-
-	/**
-	 * How much more data the internal buffer can hold before being full
-	 * @return	the number of bytes free in the buffer
-	 */
-	uint32 SpaceLeft();
-
-	/**
-	 * How much data the internal buffer is currently holding
-	 * @return	the number of bytes held in the buffer
-	 */
-	uint32 DataAvailable();
-
-	/**
-	 * Mark the beginning of a new file
-	 * @param	IN	Filename	The local build filename to the file
-	 * @param	IN	FileSize	The size of the file
-	 */
-	void BeginNewFile( const FString& Filename, const uint64& FileSize );
-
-	/**
-	 * We must special case empty files because we don't want them effecting chunk data
-	 * @param	IN	Filename	The local build filename to the file
-	 */
-	void AddEmptyFile( const FString& Filename );
-
-	/**
-	 * Enqueue a buffer of data into the internal buffer.
-	 * @param	IN	Buffer		Pointer to buffer of data.
-	 * @param	IN	Len			The length of the data buffer.
-	 */
-	void EnqueueData( const uint8* Buffer, const uint32& Len );
-
-	/**
-	 * Mark the end of the build image
-	 * @param	IN	bIsEnd		Optional: Default true. Whether to set the build end as true or false
-	 */
-	void EndOfBuild( bool bIsEnd = true );
-
-	//////////////////////////////////////////////////////////////////////////
-	// Private variables
-private:
-	// Holds the filenames referred to index positions in the data stream
-	TMap< uint64, FFileSpan > FilesParsed;
-
-	// Holds a list of empty files
-	TArray< FString > EmptyFiles;
-
-	// A critical section for accessing the FilesParsed and EmptyFiles data.
-	FCriticalSection FilesListsCS;
-
-	// Holds the ring buffer for loaded data code words
-	TRingBuffer< uint8, StreamBufferSize > BuildDataStream;
-
-	// A critical section for accessing the data stream, BuildDataStream
-	FCriticalSection BuildDataStreamCS;
-
-	// Whether the build has been completely read from disk
-	bool bNoMoreData;
-
-	// A critical section for accessing the no more data flag, NoMoreData.
-	FCriticalSection NoMoreDataCS;
-};
-
-/**
- * A class that controls processing of build data to produce the manifest with chunks
- */
-class FBuildDataChunkProcessor
-{
-public:
-
-	/**
-	 * Constructor
-	 * @param	InBuildManifest 	a build manifest that receives the manifest data.
-	 * @param	InBuildRoot			The root directory of the build
-	 */
-	FBuildDataChunkProcessor( FBuildPatchAppManifestRef InBuildManifest, const FString& InBuildRoot );
-	
-	/**
-	 * Default destructor
-	 */
-	~FBuildDataChunkProcessor();
-
-	/**
-	 * Start a new chunk
-	 * @param	bZeroBuffer 	Optional: Default true. Whether to clear the chunk buffer with zero
-	 */
-	void BeginNewChunk( const bool& bZeroBuffer = true );
-
-	/**
-	 * Start a new chunk part
-	 */
-	void BeginNewChunkPart();
-	
-	/**
-	 * End the current chunk part
-	 */
-	void EndNewChunkPart();
-	
-	/**
-	 * End the current chunk
-	 * @param	ChunkHash		The chunk's hash value
-	 * @param	ChunkData		Pointer to the chunk data, must be of FBuildPatchData::ChunkDataSize length
-	 * @param	ChunkGuid		The chunk GUID of this chunk. If it matches CurrentChunkGuid then this is a new chunk, otherwise it is the guid of the recognized chunk
-	 */
-	void EndNewChunk( const uint64& ChunkHash, const uint8* ChunkData, const FGuid& ChunkGuid );
-	
-	/**
-	 * Begins insertion of recognized chunk
-	 */
-	void PushChunk();
-
-	/**
-	 * Ends the insertion of recognized chunk
-	 * @param	ChunkHash		The hash value for the recognised chunk
-	 * @param	ChunkData		The data for the recognised chunk
-	 * @param	ChunkGuid		The GUID for the recognised chunk
-	 */
-	void PopChunk( const uint64& ChunkHash, const uint8* ChunkData, const FGuid& ChunkGuid );
-
-	/**
-	 * When the build image has been fully covered, call to end any current chunk we have (which will be zero padded)
-	 */
-	void FinalChunk();
-
-	/**
-	 * Mark the beginning of a new file
-	 * @param	FileName	The file's local filename
-	 */
-	void BeginFile( const FString& FileName );
-
-	/**
-	 * Mark the end of the current file
-	 */
-	void EndFile();
-
-	/**
-	 * Skips over the next byte from a recognized chunk
-	 * @param	NextByte		The next byte
-	 * @param	bStartOfFile	Whether this byte is the beginning of a file
-	 * @param	bEndOfFile		Whether this byte is the end of a file
-	 * @param	Filename		The filename of the file, only required if start of new file
-	 */
-	void SkipKnownByte( const uint8& NextByte, const bool& bStartOfFile, const bool& bEndOfFile, const FString& Filename );
-
-	/**
-	 * Process a new byte into the current new chunk
-	 * @param	NewByte			The new byte
-	 * @param	bStartOfFile	Whether this byte is the beginning of a file
-	 * @param	bEndOfFile		Whether this byte is the end of a file
-	 * @param	Filename		The filename of the file, only required if start of new file
-	 */
-	void ProcessNewByte( const uint8& NewByte, const bool& bStartOfFile, const bool& bEndOfFile, const FString& Filename );
-
-	/**
-	 * Get the new/know chunks stats
-	 * @param	OutNewChunks	Receives the number of new chunks found
-	 * @param	OutKnownChunks	Receives the number of recognised chunks found
-	 */
-	void GetChunkStats( uint32& OutNewChunks, uint32& OutKnownChunks );
-
-	/**
-	 * Get the chunk filesizes that have been saved out
-	 * @return ref to the Map of GUID to filesize
-	 */
-	const TMap<FGuid, int64>& GetChunkFilesizes();
-
-private:
-
-	// Keep track of the number of new chunks generated
-	uint32 NumNewChunks;
-
-	// Keep track of the number of recognised chunks
-	uint32 NumKnownChunks;
-
-	// Holds the build manifest that is receiving the information
-	FBuildPatchAppManifestPtr BuildManifest;
-
-	// Holds the build manifest that is receiving the information
-	FString BuildRoot;
-
-	// The chunk writer system
-	FChunkWriter ChunkWriter;
-
-	// A GUID for the current chunk
-	FGuid CurrentChunkGuid;
-
-	// The current chunk data
-	uint8* CurrentChunkBuffer;
-
-	// The current index into the chunk data
-	uint32 CurrentChunkBufferPos;
-
-	// Holds the current file being processed. Receives the chunk parts that we pass over.
-	FFileManifestData* CurrentFile;
-
-	// Holds an SHA hash calculator
-	FSHA1 FileHash;
-
-	// Are we currently processing a chunk
-	bool IsProcessingChunk;
-
-	// Are we currently processing a chunk part
-	bool IsProcessingChunkPart;
-
-	// Are we currently processing a file
-	bool IsProcessingFile;
-
-	// Have we pushed to a recognized chunk
-	bool ChunkIsPushed;
-
-	// Backup of CurrentChunkGuid when pushing a chunk
-	FGuid BackupChunkGuid;
-
-	// Backup of CurrentChunkBufferPos when pushing a chunk
-	uint32 BackupChunkBufferPos;
-
-	// Backup of IsProcessingChunk when pushing a chunk
-	bool BackupProcessingChunk;
-
-	// Backup of IsProcessingChunkPart when pushing a chunk
-	bool BackupProcessingChunkPart;
-
-	// Cached map of chunk file sizes
-	TMap<FGuid, int64> ChunkFileSizes;
-};
-
 /**
  * A class that controls processing of build data to produce the manifest with files
  */
@@ -693,5 +346,3 @@ struct FFileAttributes
 };
 
 #endif //WITH_BUILDPATCHGENERATION
-
-#endif // __BuildPatchGeneration_h__
diff --git a/Engine/Source/Runtime/Online/BuildPatchServices/Private/BuildPatchHash.h b/Engine/Source/Runtime/Online/BuildPatchServices/Private/BuildPatchHash.h
--- a/Engine/Source/Runtime/Online/BuildPatchServices/Private/BuildPatchHash.h
+++ b/Engine/Source/Runtime/Online/BuildPatchServices/Private/BuildPatchHash.h
@@ -1,15 +1,13 @@
 // Copyright 1998-2015 Epic Games, Inc. All Rights Reserved.
+#pragma once
+
+#include "Core/RingBuffer.h"
 
 /*=============================================================================
 	BuildPatchHash.h: Declares and implements template classes for use with the 
 	Build and Patching hash functionality.
 =============================================================================*/
 
-#ifndef __BuildPatchHash_h__
-#define __BuildPatchHash_h__
-
-#pragma once
-
 /**
  * A macro for barrel rolling a 64 bit value n times to the left.
  * @param Value		The value to shift
@@ -17,158 +15,6 @@
  */
 #define ROTLEFT_64B( Value, Shifts ) Value = ( ( ( Value ) << ( ( Shifts ) % 64 ) ) | ( ( Value ) >> ( ( 64 - ( ( Shifts ) % 64 ) ) % 64 ) ) )
 
-/**
- * A ring buffer template class.
- */
-template< typename DataType, uint32 BufferDataSize >
-class TRingBuffer
-{
-public:
-	/**
-	 * Default Constructor
-	 */
-	TRingBuffer();
-
-	/**
-	 * Default Destructor
-	 */
-	~TRingBuffer();
-
-	/**
-	 * Clears memory and indexes
-	 */
-	void Empty();
-
-	/**
-	 * Pushes a data word to the end of the FIFO. WILL OVERWRITE OLDEST if full.
-	 * @param Val	The data word to push
-	 */
-	void Enqueue( const DataType& Val );
-
-	/**
-	 * Pushes a buffer of data words to the end of the FIFO. WILL OVERWRITE OLDEST if full.
-	 * @param ValBuf		The buffer pointer
-	 * @param BufLen		The length to copy
-	 */
-	void Enqueue( const DataType* ValBuf, const uint32& BufLen );
-
-	/**
-	 * Take the next data word from the FIFO buffer.
-	 * @param OUT	ValOut	The variable to receive the data word
-	 * @return	true if the buffer was not empty, false otherwise.
-	 */
-	bool Dequeue( DataType& ValOut );
-
-	/**
-	 * Take the next set of data words from the FIFO buffer.
-	 * @param ValBuf		The buffer to receive the data
-	 * @param BufLen		The number of words requested
-	 * @return	the number of words actually copied
-	 */
-	uint32 Dequeue( DataType* ValBuf, const uint32& BufLen );
-
-	/**
-	 * Get the next data word from the FIFO buffer without removing it.
-	 * @param OUT	ValOut	The variable to receive the data word
-	 * @return	true if the buffer was not empty, false otherwise.
-	 */
-	bool Peek( DataType& ValOut ) const;
-
-	/**
-	 * Get the next set of data words from the FIFO buffer without removing them.
-	 * @param ValBuf		The buffer to receive the data
-	 * @param BufLen		The number of words requested
-	 * @return	the number of words actually copied
-	 */
-	uint32 Peek( DataType* ValBuf, const uint32& BufLen ) const;
-
-	/**
-	 * Compare the memory in the FIFO to the memory in the given buffer
-	 * @param SerialBuffer	The buffer containing data to compare
-	 * @param CompareLen		The number of words to compare
-	 * @return	< 0 if data < SerialBuffer.. 0 if data == SerialBuffer... > 0 if data > SerialBuffer
-	 */
-	int32 SerialCompare( const DataType* SerialBuffer, uint32 CompareLen ) const;
-
-	/**
-	 * Serializes the internal buffer into the given buffer
-	 * @param SerialBuffer	A preallocated buffer to copy data into
-	 */
-	void Serialize( DataType* SerialBuffer ) const;
-
-	/**
-	 * Square bracket operators for accessing data in the buffer by FIFO index. [0] returns the next entry in the FIFO that would get provided by Dequeue or Peek.
-	 * @param Index	The index into the FIFO buffer
-	 * @return	The data word
-	 */
-	FORCEINLINE const DataType& operator[]( const int32& Index ) const;
-	FORCEINLINE DataType& operator[]( const int32& Index );
-
-	/**
-	 * Gets the last data word in the FIFO (i.e. most recently pushed).
-	 * @return	The last data word
-	 */
-	FORCEINLINE const DataType& Top() const;
-	FORCEINLINE DataType& Top();
-
-	/**
-	 * Gets the first data word in the FIFO (i.e. oldest).
-	 * @return	The first data word
-	 */
-	FORCEINLINE const DataType& Bottom() const;
-	FORCEINLINE DataType& Bottom();
-
-	/**
-	 * Gets the buffer index that the last data word is stored in.
-	 * @return	The index of Top()
-	 */
-	FORCEINLINE const uint32 TopIndex() const;
-
-	/**
-	 * Gets the buffer index that the first data word is stored in.
-	 * @return	The index of Bottom()
-	 */
-	FORCEINLINE const uint32 BottomIndex() const;
-
-	/**
-	 * Gets the buffer index that the next enqueued word will get stored in.
-	 * @return	The index that the next enqueued word will have
-	 */
-	FORCEINLINE const uint32 NextIndex() const;
-
-	/**
-	 * Gets the size of the data buffer.
-	 * @return	Template arg BufferDataSize
-	 */
-	FORCEINLINE const uint32 RingDataSize() const;
-
-	/**
-	 * Gets the number of words currently in the FIFO.
-	 * @return	data size
-	 */
-	FORCEINLINE const uint32 RingDataUsage() const;
-
-	/**
-	 * Gets the total number of words that have been pushed through this buffer since clearing
-	 * @return	total length of data passed through
-	 */
-	FORCEINLINE const uint64 TotalDataPushed() const;
-
-
-private:
-	// The data memory
-	DataType* Data;
-
-	// Value keeping track of free space
-	uint32 NumDataAvailable;
-
-	// Value keeping track of the next data index
-	uint32 DataIndex;
-
-	// Value to keep track of total amount of data Enqueued
-	uint64 TotalNumDataPushed;
-};
-
 /**
  * A static struct containing 64bit polynomial constant and hash table lookup for use with FRollingHash.
  */
@@ -274,206 +120,6 @@ namespace FCycPoly64Hash
 	uint64 GetHashForDataSet(const uint8* DataSet, const uint32 DataSize, const uint64 State = 0);
 }
 
-/* TRingBuffer implementation
-*****************************************************************************/
-template< typename DataType, uint32 BufferDataSize >
-TRingBuffer< DataType, BufferDataSize >::TRingBuffer()
-{
-	Data = new DataType[ BufferDataSize ];
-	Empty();
-}
-
-template< typename DataType, uint32 BufferDataSize >
-TRingBuffer< DataType, BufferDataSize >::~TRingBuffer()
-{
-	delete[] Data;
-}
-
-template< typename DataType, uint32 BufferDataSize >
-void TRingBuffer< DataType, BufferDataSize >::Empty()
-{
-	FMemory::Memzero( Data, sizeof( DataType ) * BufferDataSize );
-	NumDataAvailable = 0;
-	DataIndex = 0;
-	TotalNumDataPushed = 0;
-}
-
-template< typename DataType, uint32 BufferDataSize >
-void TRingBuffer< DataType, BufferDataSize >::Enqueue( const DataType& Val )
-{
-	Data[ DataIndex++ ] = Val;
-	DataIndex %= BufferDataSize;
-	++TotalNumDataPushed;
-	NumDataAvailable = FMath::Clamp< uint32 >( NumDataAvailable + 1, 0, BufferDataSize );
-}
-
-template< typename DataType, uint32 BufferDataSize >
-void TRingBuffer< DataType, BufferDataSize >::Enqueue( const DataType* ValBuf, const uint32& BufLen )
-{
-	for( uint32 i = 0; i < BufLen; ++i )
-	{
-		Enqueue( ValBuf[i] );
-	}
-}
-
-template< typename DataType, uint32 BufferDataSize >
-bool TRingBuffer< DataType, BufferDataSize >::Dequeue( DataType& ValOut )
-{
-	if( NumDataAvailable > 0 )
-	{
-		ValOut = Bottom();
-		NumDataAvailable = FMath::Clamp< uint32 >( NumDataAvailable - 1, 0, BufferDataSize );
-		return true;
-	}
-	else
-	{
-		return false;
-	}
-}
-
-template< typename DataType, uint32 BufferDataSize >
-uint32 TRingBuffer< DataType, BufferDataSize >::Dequeue( DataType* ValBuf, const uint32& BufLen )
-{
-	uint32 Count = 0;
-	for( uint32 i = 0; i < BufLen; ++i )
-	{
-		if( Dequeue( ValBuf[i] ) )
-		{
-			++Count;
-		}
-		else
-		{
-			break;
-		}
-	}
-	return Count;
-}
-
-template< typename DataType, uint32 BufferDataSize >
-bool TRingBuffer< DataType, BufferDataSize >::Peek( DataType& ValOut ) const
-{
-	if( NumDataAvailable > 0 )
-	{
-		ValOut = Bottom();
-		return true;
-	}
-	else
-	{
-		return false;
-	}
-}
-
-template< typename DataType, uint32 BufferDataSize >
-uint32 TRingBuffer< DataType, BufferDataSize >::Peek( DataType* ValBuf, const uint32& BufLen ) const
-{
-	uint32 Count = 0;
-	for( uint32 i = 0; i < BufLen; ++i )
-	{
-		if( Peek( ValBuf[i] ) )
-		{
-			++Count;
-		}
-		else
-		{
-			break;
-		}
-	}
-	return Count;
-}
-
-template< typename DataType, uint32 BufferDataSize >
-int32 TRingBuffer< DataType, BufferDataSize >::SerialCompare( const DataType* SerialBuffer, uint32 CompareLen ) const
-{
-	check( CompareLen <= BufferDataSize);
-	const uint32 FirstPartLen = BufferDataSize - DataIndex;
-	int32 FirstCmp = FMemory::Memcmp( Data + BottomIndex(), SerialBuffer, sizeof( DataType ) * FMath::Min( FirstPartLen, CompareLen ) );
-	if( FirstCmp != 0 || FirstPartLen >= CompareLen)
-	{
-		return FirstCmp;
-	}
-	return FMemory::Memcmp( Data, &SerialBuffer[ FirstPartLen ], sizeof( DataType ) * ( CompareLen - FirstPartLen ) );
-}
-
-template< typename DataType, uint32 BufferDataSize >
-void TRingBuffer< DataType, BufferDataSize >::Serialize( DataType* SerialBuffer ) const
-{
-	const uint32 FirstPartLen = BufferDataSize - DataIndex;
- 	FMemory::Memcpy( SerialBuffer, Data + BottomIndex(), sizeof( DataType ) * FirstPartLen );
- 	FMemory::Memcpy( &SerialBuffer[ FirstPartLen ], Data, sizeof( DataType ) * ( BufferDataSize - FirstPartLen ) );
-}
-
-template< typename DataType, uint32 BufferDataSize >
-FORCEINLINE const DataType& TRingBuffer< DataType, BufferDataSize >::operator[]( const int32& Index ) const
-{
-	return Data[ ( BottomIndex() + Index ) % BufferDataSize ];
-}
-
-template< typename DataType, uint32 BufferDataSize >
-FORCEINLINE DataType& TRingBuffer< DataType, BufferDataSize >::operator[]( const int32& Index )
-{
-	return Data[ ( BottomIndex() + Index ) % BufferDataSize ];
-}
-
-template< typename DataType, uint32 BufferDataSize >
-FORCEINLINE const DataType& TRingBuffer< DataType, BufferDataSize >::Top() const
-{
-	return Data[ TopIndex() ];
-}
-
-template< typename DataType, uint32 BufferDataSize >
-FORCEINLINE DataType& TRingBuffer< DataType, BufferDataSize >::Top()
-{
-	return Data[ TopIndex() ];
-}
-
-template< typename DataType, uint32 BufferDataSize >
-FORCEINLINE const DataType& TRingBuffer< DataType, BufferDataSize >::Bottom() const
-{
-	return Data[ BottomIndex() ];
-}
-
-template< typename DataType, uint32 BufferDataSize >
-FORCEINLINE DataType& TRingBuffer< DataType, BufferDataSize >::Bottom()
-{
-	return Data[ BottomIndex() ];
-}
-
-template< typename DataType, uint32 BufferDataSize >
-FORCEINLINE const uint32 TRingBuffer< DataType, BufferDataSize >::TopIndex() const
-{
-	return ( DataIndex + BufferDataSize - 1 ) % BufferDataSize;
-}
-
-template< typename DataType, uint32 BufferDataSize >
-FORCEINLINE const uint32 TRingBuffer< DataType, BufferDataSize >::BottomIndex() const
-{
-	return ( DataIndex + BufferDataSize - NumDataAvailable ) % BufferDataSize;
-}
-
-template< typename DataType, uint32 BufferDataSize >
-FORCEINLINE const uint32 TRingBuffer< DataType, BufferDataSize >::NextIndex() const
-{
-	return DataIndex;
-}
-
-template< typename DataType, uint32 BufferDataSize >
-FORCEINLINE const uint32 TRingBuffer< DataType, BufferDataSize >::RingDataSize() const
-{
-	return BufferDataSize;
-}
-
-template< typename DataType, uint32 BufferDataSize >
-FORCEINLINE const uint32 TRingBuffer< DataType, BufferDataSize >::RingDataUsage() const
-{
-	return NumDataAvailable;
-}
-
-template< typename DataType, uint32 BufferDataSize >
-FORCEINLINE const uint64 TRingBuffer< DataType, BufferDataSize >::TotalDataPushed() const
-{
-	return TotalNumDataPushed;
-}
-
 /* FRollingHash implementation
 *****************************************************************************/
 template< uint32 WindowSize >
@@ -569,8 +215,8 @@ uint64 FRollingHash< WindowSize >::GetHashForDataSet( const uint8* DataSet )
 }
 
 /**
-* A static function to perform sanity checks on the rolling hash class.
-*/
+ * A static function to perform sanity checks on the rolling hash class.
+ */
 static bool CheckRollingHashAlgorithm()
 {
 	bool bCheckOk = true;
@@ -610,5 +256,3 @@ static bool CheckRollingHashAlgorithm()
 
 	return bCheckOk;
 }
-
-#endif // __BuildPatchHash_h__
diff --git a/Engine/Source/Runtime/Online/BuildPatchServices/Private/BuildPatchManifest.h b/Engine/Source/Runtime/Online/BuildPatchServices/Private/BuildPatchManifest.h
--- a/Engine/Source/Runtime/Online/BuildPatchServices/Private/BuildPatchManifest.h
+++ b/Engine/Source/Runtime/Online/BuildPatchServices/Private/BuildPatchManifest.h
@@ -318,6 +318,12 @@ struct FFileChunkPart
 	{}
 };
 
+// Required to allow private access to manifest builder for now..
+namespace FBuildPatchServices
+{
+	class FManifestBuilderImpl;
+}
+
 /**
  * Declare the FBuildPatchAppManifest object class. This holds the UObject data, and the implemented build manifest functionality
  */
@@ -326,9 +332,9 @@ class FBuildPatchAppManifest
 {
 	// Allow access to build processor classes
 	friend class FBuildDataGenerator;
-	friend class FBuildDataChunkProcessor;
 	friend class FBuildDataFileProcessor;
 	friend class FBuildPatchInstaller;
+	friend class FBuildPatchServices::FManifestBuilderImpl;
 public:
 
 	/**
diff --git a/Engine/Source/Runtime/Online/BuildPatchServices/Private/Core/RingBuffer.h b/Engine/Source/Runtime/Online/BuildPatchServices/Private/Core/RingBuffer.h
new file mode 100644
--- /dev/null
+++ b/Engine/Source/Runtime/Online/BuildPatchServices/Private/Core/RingBuffer.h
@@ -0,0 +1,344 @@
+// Copyright 1998-2015 Epic Games, Inc. All Rights Reserved.
+#pragma once
+
+template< typename DataType, uint32 BufferDataSize >
+class TRingBuffer
+{
+public:
+	/**
+	 * Default Constructor
+	 */
+	TRingBuffer();
+
+	/**
+	 * Default Destructor
+	 */
+	~TRingBuffer();
+
+	/**
+	 * Clears memory and indexes
+	 */
+	void Empty();
+
+	/**
+	 * Pushes a data word to the end of the FIFO. WILL OVERWRITE OLDEST if full.
+	 * @param Val	The data word to push
+	 */
+	void Enqueue(const DataType& Val);
+
+	/**
+	 * Pushes a buffer of data words to the end of the FIFO. WILL OVERWRITE OLDEST if full.
+	 * @param ValBuf		The buffer pointer
+	 * @param BufLen		The length to copy
+	 */
+	void Enqueue(const DataType* ValBuf, const uint32& BufLen);
+
+	/**
+	 * Take the next data word from the FIFO buffer.
+	 * @param OUT	ValOut	The variable to receive the data word
+	 * @return	true if the buffer was not empty, false otherwise.
+	 */
+	bool Dequeue(DataType& ValOut);
+
+	/**
+	 * Take the next set of data words from the FIFO buffer.
+	 * @param ValBuf		The buffer to receive the data
+	 * @param BufLen		The number of words requested
+	 * @return	the number of words actually copied
+	 */
+	uint32 Dequeue(DataType* ValBuf, const uint32& BufLen);
+
+	/**
+	 * Get the next data word from the FIFO buffer without removing it.
+	 * @param OUT	ValOut	The variable to receive the data word
+	 * @return	true if the buffer was not empty, false otherwise.
+	 */
+	bool Peek(DataType& ValOut) const;
+
+	/**
+	 * Get the next set of data words from the FIFO buffer without removing them.
+	 * @param ValBuf		The buffer to receive the data
+	 * @param BufLen		The number of words requested
+	 * @return	the number of words actually copied
+	 */
+	uint32 Peek(DataType* ValBuf, const uint32& BufLen) const;
+
+	/**
+	 * Compare the memory in the FIFO to the memory in the given buffer
+	 * @param SerialBuffer	The buffer containing data to compare
+	 * @param CompareLen		The number of words to compare
+	 * @return	< 0 if data < SerialBuffer.. 0 if data == SerialBuffer... > 0 if data > SerialBuffer
+	 */
+	int32 SerialCompare(const DataType* SerialBuffer, uint32 CompareLen) const;
+
+	/**
+	 * Serializes the internal buffer into the given buffer
+	 * @param SerialBuffer	A preallocated buffer to copy data into
+	 */
+	void Serialize(DataType* SerialBuffer) const;
+
+	/**
+	 * Square bracket operators for accessing data in the buffer by FIFO index. [0] returns the next entry in the FIFO that would get provided by Dequeue or Peek.
+	 * @param Index	The index into the FIFO buffer
+	 * @return	The data word
+	 */
+	FORCEINLINE const DataType& operator[](const int32& Index) const;
+	FORCEINLINE DataType& operator[](const int32& Index);
+
+	/**
+	 * Gets the last data word in the FIFO (i.e. most recently pushed).
+	 * @return	The last data word
+	 */
+	FORCEINLINE const DataType& Top() const;
+	FORCEINLINE DataType& Top();
+
+	/**
+	 * Gets the first data word in the FIFO (i.e. oldest).
+	 * @return	The first data word
+	 */
+	FORCEINLINE const DataType& Bottom() const;
+	FORCEINLINE DataType& Bottom();
+
+	/**
+	 * Gets the buffer index that the last data word is stored in.
+	 * @return	The index of Top()
+	 */
+	FORCEINLINE uint32 TopIndex() const;
+
+	/**
+	 * Gets the buffer index that the first data word is stored in.
+	 * @return	The index of Bottom()
+	 */
+	FORCEINLINE uint32 BottomIndex() const;
+
+	/**
+	 * Gets the buffer index that the next enqueued word will get stored in.
+	 * @return	The index that the next enqueued word will have
+	 */
+	FORCEINLINE uint32 NextIndex() const;
+
+	/**
+	 * Gets the size of the data buffer.
+	 * @return	Template arg BufferDataSize
+	 */
+	FORCEINLINE uint32 RingDataSize() const;
+
+	/**
+	 * Gets the number of words currently in the FIFO.
+	 * @return	data size
+	 */
+	FORCEINLINE uint32 RingDataUsage() const;
+
+	/**
+	 * Gets the total number of words that have been pushed through this buffer since clearing
+	 * @return	total length of data passed through
+	 */
+	FORCEINLINE uint64 TotalDataPushed() const;
+
+
+private:
+	// The data memory
+	DataType* Data;
+
+	// Value keeping track of free space
+	uint32 NumDataAvailable;
+
+	// Value keeping track of the next data index
+	uint32 DataIndex;
+
+	// Value to keep track of total amount of data Enqueued
+	uint64 TotalNumDataPushed;
+};
+
+/* TRingBuffer implementation
+*****************************************************************************/
+template< typename DataType, uint32 BufferDataSize >
+TRingBuffer< DataType, BufferDataSize >::TRingBuffer()
+{
+	Data = new DataType[BufferDataSize];
+	Empty();
+}
+
+template< typename DataType, uint32 BufferDataSize >
+TRingBuffer< DataType, BufferDataSize >::~TRingBuffer()
+{
+	delete[] Data;
+}
+
+template< typename DataType, uint32 BufferDataSize >
+void TRingBuffer< DataType, BufferDataSize >::Empty()
+{
+	FMemory::Memzero(Data, sizeof(DataType)* BufferDataSize);
+	NumDataAvailable = 0;
+	DataIndex = 0;
+	TotalNumDataPushed = 0;
+}
+
+template< typename DataType, uint32 BufferDataSize >
+void TRingBuffer< DataType, BufferDataSize >::Enqueue(const DataType& Val)
+{
+	Data[DataIndex++] = Val;
+	DataIndex %= BufferDataSize;
+	++TotalNumDataPushed;
+	NumDataAvailable = FMath::Clamp< uint32 >(NumDataAvailable + 1, 0, BufferDataSize);
+}
+
+template< typename DataType, uint32 BufferDataSize >
+void TRingBuffer< DataType, BufferDataSize >::Enqueue(const DataType* ValBuf, const uint32& BufLen)
+{
+	check(BufLen <= BufferDataSize);
+	const uint32 FirstPartLen = BufferDataSize - DataIndex;
+	FMemory::Memcpy(Data + DataIndex, ValBuf, sizeof(DataType)* FMath::Min(FirstPartLen, BufLen));
+	if (FirstPartLen < BufLen)
+	{
+		FMemory::Memcpy(Data, ValBuf + FirstPartLen, sizeof(DataType)* (BufLen - FirstPartLen));
+	}
+	DataIndex += BufLen;
+	DataIndex %= BufferDataSize;
+	TotalNumDataPushed += BufLen;
+	NumDataAvailable = FMath::Clamp<uint32>(NumDataAvailable + BufLen, 0, BufferDataSize);
+}
+
+template< typename DataType, uint32 BufferDataSize >
+bool TRingBuffer< DataType, BufferDataSize >::Dequeue(DataType& ValOut)
+{
+	if (Peek(ValOut))
+	{
+		NumDataAvailable = FMath::Clamp< uint32 >(NumDataAvailable - 1, 0, BufferDataSize);
+		return true;
+	}
+	else
+	{
+		return false;
+	}
+}
+
+template< typename DataType, uint32 BufferDataSize >
+uint32 TRingBuffer< DataType, BufferDataSize >::Dequeue(DataType* ValBuf, const uint32& BufLen)
+{
+	const uint32 DataProvided = Peek(ValBuf, BufLen);
+	NumDataAvailable -= DataProvided;
+	return DataProvided;
+}
+
+template< typename DataType, uint32 BufferDataSize >
+bool TRingBuffer< DataType, BufferDataSize >::Peek(DataType& ValOut) const
+{
+	if (NumDataAvailable > 0)
+	{
+		ValOut = Bottom();
+		return true;
+	}
+	else
+	{
+		return false;
+	}
+}
+
+template< typename DataType, uint32 BufferDataSize >
+uint32 TRingBuffer< DataType, BufferDataSize >::Peek(DataType* ValBuf, const uint32& BufLen) const
+{
+	const uint32 DataProvided = FMath::Min(BufLen, NumDataAvailable);
+	const uint32 BottomIdx = BottomIndex();
+	const uint32 BottomPartLen = BufferDataSize - BottomIdx;
+	FMemory::Memcpy(ValBuf, Data + BottomIdx, sizeof(DataType)* FMath::Min(BottomPartLen, DataProvided));
+	if (BottomPartLen < DataProvided)
+	{
+		FMemory::Memcpy(ValBuf + BottomPartLen, Data, sizeof(DataType)* (DataProvided - BottomPartLen));
+	}
+	return DataProvided;
+}
+
+template< typename DataType, uint32 BufferDataSize >
+int32 TRingBuffer< DataType, BufferDataSize >::SerialCompare(const DataType* SerialBuffer, uint32 CompareLen) const
+{
+	check(CompareLen <= BufferDataSize);
+	const uint32 FirstPartLen = BufferDataSize - DataIndex;
+	int32 FirstCmp = FMemory::Memcmp(Data + BottomIndex(), SerialBuffer, sizeof(DataType)* FMath::Min(FirstPartLen, CompareLen));
+	if (FirstCmp != 0 || FirstPartLen >= CompareLen)
+	{
+		return FirstCmp;
+	}
+	return FMemory::Memcmp(Data, &SerialBuffer[FirstPartLen], sizeof(DataType)* (CompareLen - FirstPartLen));
+}
+
+template< typename DataType, uint32 BufferDataSize >
+void TRingBuffer< DataType, BufferDataSize >::Serialize(DataType* SerialBuffer) const
+{
+	const uint32 BottomIdx = BottomIndex();
+	const uint32 BottomPartLen = BufferDataSize - BottomIdx;
+	FMemory::Memcpy(SerialBuffer, Data + BottomIdx, sizeof(DataType)* BottomPartLen);
+	FMemory::Memcpy(SerialBuffer + BottomPartLen, Data, sizeof(DataType)* (BufferDataSize - BottomPartLen));
+}
+
+template< typename DataType, uint32 BufferDataSize >
+FORCEINLINE const DataType& TRingBuffer< DataType, BufferDataSize >::operator[](const int32& Index) const
+{
+	return Data[(BottomIndex() + Index) % BufferDataSize];
+}
+
+template< typename DataType, uint32 BufferDataSize >
+FORCEINLINE DataType& TRingBuffer< DataType, BufferDataSize >::operator[](const int32& Index)
+{
+	return Data[(BottomIndex() + Index) % BufferDataSize];
+}
+
+template< typename DataType, uint32 BufferDataSize >
+FORCEINLINE const DataType& TRingBuffer< DataType, BufferDataSize >::Top() const
+{
+	return Data[TopIndex()];
+}
+
+template< typename DataType, uint32 BufferDataSize >
+FORCEINLINE DataType& TRingBuffer< DataType, BufferDataSize >::Top()
+{
+	return Data[TopIndex()];
+}
+
+template< typename DataType, uint32 BufferDataSize >
+FORCEINLINE const DataType& TRingBuffer< DataType, BufferDataSize >::Bottom() const
+{
+	return Data[BottomIndex()];
+}
+
+template< typename DataType, uint32 BufferDataSize >
+FORCEINLINE DataType& TRingBuffer< DataType, BufferDataSize >::Bottom()
+{
+	return Data[BottomIndex()];
+}
+
+template< typename DataType, uint32 BufferDataSize >
+FORCEINLINE uint32 TRingBuffer< DataType, BufferDataSize >::TopIndex() const
+{
+	return (DataIndex + BufferDataSize - 1) % BufferDataSize;
+}
+
+template< typename DataType, uint32 BufferDataSize >
+FORCEINLINE uint32 TRingBuffer< DataType, BufferDataSize >::BottomIndex() const
+{
+	return (DataIndex + BufferDataSize - NumDataAvailable) % BufferDataSize;
+}
+
+template< typename DataType, uint32 BufferDataSize >
+FORCEINLINE uint32 TRingBuffer< DataType, BufferDataSize >::NextIndex() const
+{
+	return DataIndex;
+}
+
+template< typename DataType, uint32 BufferDataSize >
+FORCEINLINE uint32 TRingBuffer< DataType, BufferDataSize >::RingDataSize() const
+{
+	return BufferDataSize;
+}
+
+template< typename DataType, uint32 BufferDataSize >
+FORCEINLINE uint32 TRingBuffer< DataType, BufferDataSize >::RingDataUsage() const
+{
+	return NumDataAvailable;
+}
+
+template< typename DataType, uint32 BufferDataSize >
+FORCEINLINE uint64 TRingBuffer< DataType, BufferDataSize >::TotalDataPushed() const
+{
+	return TotalNumDataPushed;
+}
diff --git a/Engine/Source/Runtime/Online/BuildPatchServices/Private/Generation/BuildStreamer.cpp b/Engine/Source/Runtime/Online/BuildPatchServices/Private/Generation/BuildStreamer.cpp
new file mode 100644
--- /dev/null
+++ b/Engine/Source/Runtime/Online/BuildPatchServices/Private/Generation/BuildStreamer.cpp
@@ -0,0 +1,423 @@
+// Copyright 1998-2015 Epic Games, Inc. All Rights Reserved.
+
+#if WITH_BUILDPATCHGENERATION
+
+#include "BuildPatchServicesPrivatePCH.h"
+
+#include "BuildStreamer.h"
+#include "Core/RingBuffer.h"
+#include "StatsCollector.h"
+
+namespace FBuildPatchServices
+{
+	// Buffer sizes
+	enum
+	{
+		FileBufferSize = 1024 * 1024 * 10,      // 10 MB
+		StreamBufferSize = 1024 * 1024 * 100    // 100 MB
+	};
+
+	static bool IsUnixExecutable(const TCHAR* Filename)
+	{
+#if PLATFORM_MAC
+		struct stat FileInfo;
+		if (stat(TCHAR_TO_UTF8(Filename), &FileInfo) == 0)
+		{
+			return (FileInfo.st_mode & S_IXUSR) != 0;
+		}
+#endif
+		return false;
+	}
+
+	static FString GetSymlinkTarget(const TCHAR* Filename)
+	{
+#if PLATFORM_MAC
+		ANSICHAR SymlinkTarget[MAX_PATH] = { 0 };
+		if (readlink(TCHAR_TO_UTF8(Filename), SymlinkTarget, MAX_PATH) != -1)
+		{
+			return UTF8_TO_TCHAR(SymlinkTarget);
+		}
+#endif
+		return TEXT("");
+	}
+
+	class FDataStream
+	{
+	public:
+		FDataStream();
+		~FDataStream();
+		void Clear();
+		uint32 FreeSpace() const;
+		uint32 UsedSpace() const;
+		uint64 TotalDataPushed() const;
+		void EnqueueData(const uint8* Buffer, const uint32& Len);
+		uint32 DequeueData(uint8* Buffer, const uint32& ReqSize, bool WaitForData);
+		bool IsEndOfStream() const;
+		void SetEndOfStream();
+
+	private:
+		mutable FCriticalSection BuildDataStreamCS;
+		TRingBuffer< uint8, FBuildPatchServices::StreamBufferSize > BuildDataStream;
+		FThreadSafeBool EndOfStream;
+	};
+
+	class FBuildStreamerImpl
+		: public FBuildStreamer
+	{
+
+	public:
+		FBuildStreamerImpl(const FString& BuildRoot, const FString& IgnoreListFile, const FStatsCollectorRef& StatsCollector, IFileManager* FileManager);
+		virtual ~FBuildStreamerImpl();
+
+		virtual uint32 DequeueData(uint8* Buffer, uint32 ReqSize, bool WaitForData = true) override;
+		virtual bool GetFileSpan(uint64 StartingIdx, FFileSpan& FileSpan) const override;
+		virtual TArray<FString> GetEmptyFiles() const override;
+		virtual bool IsEndOfBuild() const override;
+		virtual bool IsEndOfData() const override;
+
+	private:
+		void ReadData();
+		void AddFile(FFileSpan FileSpan);
+		void AddEmptyFile(FString Filename);
+		void SetFileHash(uint64 StartIdx, FSHA1& FileHash);
+		void StripIgnoredFiles(TArray<FString>& AllFiles);
+
+	private:
+		const FString BuildRoot;
+		const FString IgnoreListFile;
+		FStatsCollectorRef StatsCollector;
+		IFileManager* FileManager;
+		FDataStream DataStream;
+		mutable FCriticalSection FilesCS;
+		TMap<uint64, FFileSpan> Files;
+		TSet<FString> EmptyFiles;
+		TFuture<void> Future;
+		FThreadSafeBool bShouldAbort;
+	};
+
+	FDataStream::FDataStream()
+		: EndOfStream(false)
+	{
+	}
+
+	FDataStream::~FDataStream()
+	{
+	}
+
+	void FDataStream::Clear()
+	{
+		BuildDataStreamCS.Lock();
+		BuildDataStream.Empty();
+		BuildDataStreamCS.Unlock();
+	}
+
+	uint32 FDataStream::FreeSpace() const
+	{
+		uint32 rtn;
+		BuildDataStreamCS.Lock();
+		rtn = BuildDataStream.RingDataSize() - BuildDataStream.RingDataUsage();
+		BuildDataStreamCS.Unlock();
+		return rtn;
+	}
+
+	uint32 FDataStream::UsedSpace() const
+	{
+		uint32 rtn;
+		BuildDataStreamCS.Lock();
+		rtn = BuildDataStream.RingDataUsage();
+		BuildDataStreamCS.Unlock();
+		return rtn;
+	}
+
+	uint64 FDataStream::TotalDataPushed() const
+	{
+		uint64 rtn;
+		BuildDataStreamCS.Lock();
+		rtn = BuildDataStream.TotalDataPushed();
+		BuildDataStreamCS.Unlock();
+		return rtn;
+	}
+
+	void FDataStream::EnqueueData(const uint8* Buffer, const uint32& Len)
+	{
+		checkf(!EndOfStream, TEXT("More data was added after specifying the end of stream"));
+		BuildDataStreamCS.Lock();
+		while (FreeSpace() < Len)
+		{
+			BuildDataStreamCS.Unlock();
+			FPlatformProcess::Sleep(0.01f);
+			BuildDataStreamCS.Lock();
+		}
+		BuildDataStream.Enqueue(Buffer, Len);
+		BuildDataStreamCS.Unlock();
+	}
+
+	uint32 FDataStream::DequeueData(uint8* Buffer, const uint32& ReqSize, bool WaitForData)
+	{
+		BuildDataStreamCS.Lock();
+		uint32 ReadLen = BuildDataStream.Dequeue(Buffer, ReqSize);
+		if (WaitForData && ReadLen < ReqSize)
+		{
+			while (ReadLen < ReqSize && !EndOfStream)
+			{
+				BuildDataStreamCS.Unlock();
+				FPlatformProcess::Sleep(0.01f);
+				BuildDataStreamCS.Lock();
+				ReadLen += BuildDataStream.Dequeue(&Buffer[ReadLen], ReqSize - ReadLen);
+			}
+		}
+		BuildDataStreamCS.Unlock();
+		return ReadLen;
+	}
+
+	bool FDataStream::IsEndOfStream() const
+	{
+		return EndOfStream;
+	}
+
+	void FDataStream::SetEndOfStream()
+	{
+		EndOfStream = true;
+	}
+
+	FBuildStreamerImpl::FBuildStreamerImpl(const FString& InBuildRoot, const FString& InIgnoreListFile, const FStatsCollectorRef& InStatsCollector, IFileManager* InFileManager)
+		: BuildRoot(InBuildRoot)
+		, IgnoreListFile(InIgnoreListFile)
+		, StatsCollector(InStatsCollector)
+		, FileManager(InFileManager)
+		, bShouldAbort(false)
+	{
+		TFunction<void()> Task = [this]() { ReadData(); };
+		Future = Async(EAsyncExecution::Thread, Task);
+	}
+
+	FBuildStreamerImpl::~FBuildStreamerImpl()
+	{
+		bShouldAbort = true;
+		Future.Wait();
+	}
+
+	uint32 FBuildStreamerImpl::DequeueData(uint8* Buffer, uint32 ReqSize, bool WaitForData /*= true*/)
+	{
+		return DataStream.DequeueData(Buffer, ReqSize, WaitForData);
+	}
+
+	bool FBuildStreamerImpl::GetFileSpan(uint64 StartingIdx, FFileSpan& FileSpan) const
+	{
+		FScopeLock ScopeLock(&FilesCS);
+		if (Files.Contains(StartingIdx))
+		{
+			FileSpan = Files[StartingIdx];
+			return true;
+		}
+		return false;
+	}
+
+	TArray<FString> FBuildStreamerImpl::GetEmptyFiles() const
+	{
+		FScopeLock ScopeLock(&FilesCS);
+		return EmptyFiles.Array();
+	}
+
+	bool FBuildStreamerImpl::IsEndOfBuild() const
+	{
+		return DataStream.IsEndOfStream();
+	}
+
+	bool FBuildStreamerImpl::IsEndOfData() const
+	{
+		return DataStream.IsEndOfStream() && DataStream.UsedSpace() == 0;
+	}
+
+	void FBuildStreamerImpl::ReadData()
+	{
+		// Stats
+		auto* StatFileOpenTime = StatsCollector->CreateStat(TEXT("Stream: Open Time"), EStatFormat::Timer);
+		auto* StatFileReadTime = StatsCollector->CreateStat(TEXT("Stream: Read Time"), EStatFormat::Timer);
+		auto* StatFileHashTime = StatsCollector->CreateStat(TEXT("Stream: Hash Time"), EStatFormat::Timer);
+		auto* StatDataEnqueueTime = StatsCollector->CreateStat(TEXT("Stream: Enqueue Time"), EStatFormat::Timer);
+		auto* StatDataAccessSpeed = StatsCollector->CreateStat(TEXT("Stream: Data Access Speed"), EStatFormat::DataSpeed);
+		auto* StatPotentialThroughput = StatsCollector->CreateStat(TEXT("Stream: Potential Throughtput"), EStatFormat::DataSpeed);
+		auto* StatThroughput = StatsCollector->CreateStat(TEXT("Stream: Throughput"), EStatFormat::DataSpeed);
+		auto* StatTotalDataRead = StatsCollector->CreateStat(TEXT("Stream: Total Data Read"), EStatFormat::DataSize);
+		uint64 TempValue;
+
+		// Clear the build stream
+		DataStream.Clear();
+
+		// Enumerate build files
+		TArray< FString > AllFiles;
+		uint32 FileEnumerationStart = FStatsCollector::GetCycles();
+		FileManager->FindFilesRecursive(AllFiles, *BuildRoot, TEXT("*.*"), true, false);
+		uint32 FileEnumerationEnd = FStatsCollector::GetCycles();
+		uint32 FileEnumerationTime = FileEnumerationEnd - FileEnumerationStart;
+		GLog->Logf(TEXT("FBuildStreamReader: Enumerated %d files in %s"), AllFiles.Num(), *FPlatformTime::PrettyTime(FStatsCollector::CyclesToSeconds(FileEnumerationTime)));
+
+		// Remove the files that appear in the ignore list
+		AllFiles.Sort();
+		StripIgnoredFiles(AllFiles);
+
+		// Track file hashes
+		FSHA1 FileHash;
+
+		// Allocate our file read buffer
+		uint8* FileReadBuffer = new uint8[FBuildPatchServices::FileBufferSize];
+
+		uint64 StartTime = FStatsCollector::GetCycles();
+		for (auto& SourceFile : AllFiles)
+		{
+			if (bShouldAbort)
+			{
+				break;
+			}
+
+			// Read the file
+			FStatsCollector::AccumulateTimeBegin(TempValue);
+			FArchive* FileReader = FileManager->CreateFileReader(*SourceFile);
+			bool bIsUnixExe = IsUnixExecutable(*SourceFile);
+			FString SymlinkTarget = GetSymlinkTarget(*SourceFile);
+			FStatsCollector::AccumulateTimeEnd(StatFileOpenTime, TempValue);
+			const bool bBuildFileOpenSuccess = FileReader != nullptr;
+			if (bBuildFileOpenSuccess)
+			{
+				// Make SourceFile the format we want it in and start a new file
+				FPaths::MakePathRelativeTo(SourceFile, *(BuildRoot + TEXT("/")));
+				int64 FileSize = FileReader->TotalSize();
+				// Process files that have bytes
+				if (FileSize > 0)
+				{
+					FileHash.Reset();
+					uint64 FileStartIdx = DataStream.TotalDataPushed();
+					AddFile(FFileSpan(SourceFile, FileSize, FileStartIdx, bIsUnixExe, SymlinkTarget));
+					while (!FileReader->AtEnd() && !bShouldAbort)
+					{
+						// Read data file
+						const int64 SizeLeft = FileSize - FileReader->Tell();
+						const uint32 ReadLen = FMath::Min< int64 >(FBuildPatchServices::FileBufferSize, SizeLeft);
+						FStatsCollector::AccumulateTimeBegin(TempValue);
+						FileReader->Serialize(FileReadBuffer, ReadLen);
+						FStatsCollector::AccumulateTimeEnd(StatFileReadTime, TempValue);
+						FStatsCollector::Accumulate(StatTotalDataRead, ReadLen);
+
+						// Hash data file
+						FStatsCollector::AccumulateTimeBegin(TempValue);
+						FileHash.Update(FileReadBuffer, ReadLen);
+						FStatsCollector::AccumulateTimeEnd(StatFileHashTime, TempValue);
+
+						// Copy into data stream
+						FStatsCollector::AccumulateTimeBegin(TempValue);
+						DataStream.EnqueueData(FileReadBuffer, ReadLen);
+						FStatsCollector::AccumulateTimeEnd(StatDataEnqueueTime, TempValue);
+
+						// Calculate other stats
+						FStatsCollector::Set(StatDataAccessSpeed, *StatTotalDataRead / FStatsCollector::CyclesToSeconds(*StatFileOpenTime + *StatFileReadTime));
+						FStatsCollector::Set(StatPotentialThroughput, *StatTotalDataRead / FStatsCollector::CyclesToSeconds(*StatFileOpenTime + *StatFileReadTime + *StatFileHashTime));
+						FStatsCollector::Set(StatThroughput, *StatTotalDataRead / FStatsCollector::CyclesToSeconds(FStatsCollector::GetCycles() - StartTime));
+					}
+					FileHash.Final();
+					SetFileHash(FileStartIdx, FileHash);
+				}
+				// Special case zero byte files
+				else if (FileSize == 0)
+				{
+					AddEmptyFile(SourceFile);
+				}
+				FileReader->Close();
+				delete FileReader;
+			}
+			else
+			{
+				// Not being able to load a required file from the build would be fatal, hard fault.
+				GLog->Logf(TEXT("FBuildStreamReader: Could not open file from build! %s"), *SourceFile);
+				GLog->PanicFlushThreadedLogs();
+				// Use bool variable for easier to understand assert message
+				check(bBuildFileOpenSuccess);
+			}
+		}
+
+		// Mark end of build
+		DataStream.SetEndOfStream();
+
+		// Deallocate our file read buffer
+		delete[] FileReadBuffer;
+	}
+
+	void FBuildStreamerImpl::AddFile(FFileSpan FileSpan)
+	{
+		FScopeLock ScopeLock(&FilesCS);
+		Files.Add(FileSpan.StartIdx, MoveTemp(FileSpan));
+	}
+
+	void FBuildStreamerImpl::AddEmptyFile(FString Filename)
+	{
+		FScopeLock ScopeLock(&FilesCS);
+		EmptyFiles.Add(MoveTemp(Filename));
+	}
+
+	void FBuildStreamerImpl::SetFileHash(uint64 StartIdx, FSHA1& FileHash)
+	{
+		FScopeLock ScopeLock(&FilesCS);
+		FileHash.GetHash(Files[StartIdx].SHAHash.Hash);
+	}
+
+	void FBuildStreamerImpl::StripIgnoredFiles(TArray<FString>& AllFiles)
+	{
+		struct FRemoveMatchingStrings
+		{
+			const TSet<FString>& IgnoreList;
+			FRemoveMatchingStrings(const TSet<FString>& InIgnoreList)
+				: IgnoreList(InIgnoreList) {}
+
+			bool operator()(const FString& RemovalCandidate) const
+			{
+				const bool bRemove = IgnoreList.Contains(RemovalCandidate);
+				if (bRemove)
+				{
+					GLog->Logf(TEXT("    - %s"), *RemovalCandidate);
+				}
+				return bRemove;
+			}
+		};
+
+		GLog->Logf(TEXT("Stripping ignorable files"));
+		const int32 OriginalNumFiles = AllFiles.Num();
+		FString IgnoreFileList = TEXT("");
+		FFileHelper::LoadFileToString(IgnoreFileList, *IgnoreListFile);
+		TArray< FString > IgnoreFiles;
+		IgnoreFileList.ParseIntoArray(IgnoreFiles, TEXT("\r\n"), true);
+
+		// Normalize all paths first
+		for (FString& Filename : AllFiles)
+		{
+			FPaths::NormalizeFilename(Filename);
+		}
+		for (FString& Filename : IgnoreFiles)
+		{
+			int32 TabLocation = Filename.Find(TEXT("\t"));
+			if (TabLocation != INDEX_NONE)
+			{
+				// Strip tab separated timestamp if it exists
+				Filename = Filename.Left(TabLocation);
+			}
+			Filename = BuildRoot / Filename;
+			FPaths::NormalizeFilename(Filename);
+		}
+
+		// Convert ignore list to set
+		TSet<FString> IgnoreSet(MoveTemp(IgnoreFiles));
+
+		// Filter file list
+		FRemoveMatchingStrings FileFilter(IgnoreSet);
+		AllFiles.RemoveAll(FileFilter);
+
+		const int32 NewNumFiles = AllFiles.Num();
+		GLog->Logf(TEXT("Stripped %d ignorable file(s)"), (OriginalNumFiles - NewNumFiles));
+	}
+
+	FBuildStreamerRef FBuildStreamerFactory::Create(const FString& BuildRoot, const FString& IgnoreListFile, const FStatsCollectorRef& StatsCollector, IFileManager* FileManager)
+	{
+		return MakeShareable(new FBuildStreamerImpl(BuildRoot, IgnoreListFile, StatsCollector, FileManager));
+	}
+}
+
+#endif
diff --git a/Engine/Source/Runtime/Online/BuildPatchServices/Private/Generation/BuildStreamer.h b/Engine/Source/Runtime/Online/BuildPatchServices/Private/Generation/BuildStreamer.h
new file mode 100644
--- /dev/null
+++ b/Engine/Source/Runtime/Online/BuildPatchServices/Private/Generation/BuildStreamer.h
@@ -0,0 +1,128 @@
+// Copyright 1998-2015 Epic Games, Inc. All Rights Reserved.
+#pragma once
+
+#if WITH_BUILDPATCHGENERATION
+
+#include "StatsCollector.h"
+
+namespace FBuildPatchServices
+{
+	struct FFileSpan
+	{
+		FFileSpan(const FString& InFilename, uint64 InSize, uint64 InStartIdx, bool InIsUnixExecutable, const FString& InSymlinkTarget)
+			: Filename(InFilename)
+			, Size(InSize)
+			, StartIdx(InStartIdx)
+			, IsUnixExecutable(InIsUnixExecutable)
+			, SymlinkTarget(InSymlinkTarget)
+		{
+		}
+
+		FFileSpan(const FFileSpan& CopyFrom)
+			: Filename(CopyFrom.Filename)
+			, Size(CopyFrom.Size)
+			, StartIdx(CopyFrom.StartIdx)
+			, SHAHash(CopyFrom.SHAHash)
+			, IsUnixExecutable(CopyFrom.IsUnixExecutable)
+			, SymlinkTarget(CopyFrom.SymlinkTarget)
+		{
+		}
+
+		FFileSpan(FFileSpan&& MoveFrom)
+			: Filename(MoveTemp(MoveFrom.Filename))
+			, Size(MoveTemp(MoveFrom.Size))
+			, StartIdx(MoveTemp(MoveFrom.StartIdx))
+			, SHAHash(MoveTemp(MoveFrom.SHAHash))
+			, IsUnixExecutable(MoveTemp(MoveFrom.IsUnixExecutable))
+			, SymlinkTarget(MoveTemp(MoveFrom.SymlinkTarget))
+		{
+		}
+
+		FFileSpan()
+			: Size(0)
+			, StartIdx(0)
+			, IsUnixExecutable(false)
+		{
+		}
+
+		FORCEINLINE FFileSpan& operator=(const FFileSpan& CopyFrom)
+		{
+			Filename = CopyFrom.Filename;
+			Size = CopyFrom.Size;
+			StartIdx = CopyFrom.StartIdx;
+			SHAHash = CopyFrom.SHAHash;
+			IsUnixExecutable = CopyFrom.IsUnixExecutable;
+			SymlinkTarget = CopyFrom.SymlinkTarget;
+			return *this;
+		}
+
+		FORCEINLINE FFileSpan& operator=(FFileSpan&& MoveFrom)
+		{
+			Filename = MoveTemp(MoveFrom.Filename);
+			Size = MoveTemp(MoveFrom.Size);
+			StartIdx = MoveTemp(MoveFrom.StartIdx);
+			SHAHash = MoveTemp(MoveFrom.SHAHash);
+			IsUnixExecutable = MoveTemp(MoveFrom.IsUnixExecutable);
+			SymlinkTarget = MoveTemp(MoveFrom.SymlinkTarget);
+			return *this;
+		}
+
+		FString Filename;
+		uint64 Size;
+		uint64 StartIdx;
+		FSHAHash SHAHash;
+		bool IsUnixExecutable;
+		FString SymlinkTarget;
+	};
+
+	class FBuildStreamer
+	{
+	public:
+
+		/**
+		 * Fetches some data from the buffer, also removing it.
+		 * @param	IN	Buffer			Pointer to buffer to receive the data.
+		 * @param	IN	ReqSize			The amount of data to attempt to retrieve.
+		 * @param	IN	WaitForData		Optional: Default true. Whether to wait until there is enough data in the buffer.
+		 * @return	The amount of data retrieved
+		 */
+		virtual uint32 DequeueData(uint8* Buffer, uint32 ReqSize, bool WaitForData = true) = 0;
+
+		/**
+		 * Retrieves the file details for a specific start index.
+		 * @param	IN	StartingIdx		The data index into the build image.
+		 * @param	OUT FileSpan		Receives a copy of the file span data.
+		 * @return	true if the data byte at StartingIdx is the start of a file, false indicates that FileSpan was not set.
+		 */
+		virtual bool GetFileSpan(uint64 StartingIdx, FFileSpan& FileSpan) const = 0;
+
+		/**
+		 * Gets a list of empty files that the build contains.
+		 * @return	Array of empty files in the build
+		 */
+		virtual TArray<FString> GetEmptyFiles() const = 0;
+
+		/**
+		 * Whether the read thread has finished reading the build image.
+		 * @return	true if there is no more data coming into the buffer
+		 */
+		virtual bool IsEndOfBuild() const = 0;
+
+		/**
+		 * Whether there is any more data available to dequeue from the buffer
+		 * @return	true if there is no more data coming in, and the internal buffer is also empty.
+		 */
+		virtual bool IsEndOfData() const = 0;
+	};
+
+	typedef TSharedRef<FBuildStreamer, ESPMode::ThreadSafe> FBuildStreamerRef;
+	typedef TSharedPtr<FBuildStreamer, ESPMode::ThreadSafe> FBuildStreamerPtr;
+
+	class FBuildStreamerFactory
+	{
+	public:
+		static FBuildStreamerRef Create(const FString& BuildRoot, const FString& IgnoreListFile, const FStatsCollectorRef& StatsCollector, IFileManager* FileManager = &IFileManager::Get());
+	};
+}
+
+#endif
diff --git a/Engine/Source/Runtime/Online/BuildPatchServices/Private/Generation/CloudEnumeration.cpp b/Engine/Source/Runtime/Online/BuildPatchServices/Private/Generation/CloudEnumeration.cpp
new file mode 100644
--- /dev/null
+++ b/Engine/Source/Runtime/Online/BuildPatchServices/Private/Generation/CloudEnumeration.cpp
@@ -0,0 +1,161 @@
+// Copyright 1998-2015 Epic Games, Inc. All Rights Reserved.
+
+#if WITH_BUILDPATCHGENERATION
+
+#include "BuildPatchServicesPrivatePCH.h"
+
+#include "CloudEnumeration.h"
+
+namespace FBuildPatchServices
+{
+	class FCloudEnumerationImpl
+		: public FCloudEnumeration
+	{
+	public:
+		FCloudEnumerationImpl(const FString& CloudDirectory);
+		virtual ~FCloudEnumerationImpl();
+
+		virtual TSet<FGuid> GetChunkSet(uint64 ChunkHash) const override;
+		virtual TMap<uint64, TSet<FGuid>> GetChunkInventory() const override;
+		virtual TMap<FGuid, int64> GetChunkFileSizes() const override;
+	private:
+		void EnumerateCloud();
+		void EnumerateManifestData(const FBuildPatchAppManifestRef& Manifest);
+
+	private:
+		const FString CloudDirectory;
+		mutable FCriticalSection InventoryCS;
+		TMap<uint64, TSet<FGuid>> ChunkInventory;
+		TMap<FGuid, int64> ChunkFileSizes;
+		TMap<FSHAHash, TSet<FGuid>> FileInventory;
+		uint64 NumChunksFound;
+		uint64 NumFilesFound;
+		TFuture<void> Future;
+	};
+
+	FCloudEnumerationImpl::FCloudEnumerationImpl(const FString& InCloudDirectory)
+		: CloudDirectory(InCloudDirectory)
+		, NumChunksFound(0)
+		, NumFilesFound(0)
+	{
+		TFunction<void()> Task = [this]() { EnumerateCloud(); };
+		Future = Async(EAsyncExecution::Thread, Task);
+	}
+
+	FCloudEnumerationImpl::~FCloudEnumerationImpl()
+	{
+	}
+
+	TSet<FGuid> FCloudEnumerationImpl::GetChunkSet(uint64 ChunkHash) const
+	{
+		FScopeLock ScopeLock(&InventoryCS);
+		Future.Wait();
+		return ChunkInventory.FindRef(ChunkHash);
+	}
+
+	TMap<uint64, TSet<FGuid>> FCloudEnumerationImpl::GetChunkInventory() const
+	{
+		FScopeLock ScopeLock(&InventoryCS);
+		Future.Wait();
+		return ChunkInventory;
+	}
+
+	TMap<FGuid, int64> FCloudEnumerationImpl::GetChunkFileSizes() const
+	{
+		FScopeLock ScopeLock(&InventoryCS);
+		Future.Wait();
+		return ChunkFileSizes;
+	}
+
+	void FCloudEnumerationImpl::EnumerateCloud()
+	{
+		FScopeLock ScopeLock(&InventoryCS);
+
+		IFileManager& FileManager = IFileManager::Get();
+		FString JSONOutput;
+		TSharedRef< TJsonWriter< TCHAR, TPrettyJsonPrintPolicy< TCHAR > > > DebugWriter = TJsonWriterFactory< TCHAR, TPrettyJsonPrintPolicy< TCHAR > >::Create(&JSONOutput);
+		DebugWriter->WriteObjectStart();
+
+		// Find all manifest files
+		if (FileManager.DirectoryExists(*CloudDirectory))
+		{
+			const double StartEnumerate = FPlatformTime::Seconds();
+			TArray<FString> AllManifests;
+			GLog->Logf(TEXT("BuildDataGenerator: Enumerating Manifests from %s"), *CloudDirectory);
+			FileManager.FindFiles(AllManifests, *(CloudDirectory / TEXT("*.manifest")), true, false);
+			const double EnumerateTime = FPlatformTime::Seconds() - StartEnumerate;
+			GLog->Logf(TEXT("BuildDataGenerator: Found %d manifests in %.1f seconds"), AllManifests.Num(), EnumerateTime);
+
+			// Load all manifest files
+			const double StartLoadAllManifest = FPlatformTime::Seconds();
+			for (const auto& ManifestFile : AllManifests)
+			{
+				// Determine chunks from manifest file
+				const FString ManifestFilename = CloudDirectory / ManifestFile;
+				FBuildPatchAppManifestRef BuildManifest = MakeShareable(new FBuildPatchAppManifest());
+				const double StartLoadManifest = FPlatformTime::Seconds();
+				if (BuildManifest->LoadFromFile(ManifestFilename))
+				{
+					const double LoadManifestTime = FPlatformTime::Seconds() - StartLoadManifest;
+					GLog->Logf(TEXT("BuildDataGenerator: Loaded %s in %.1f seconds"), *ManifestFile, LoadManifestTime);
+					EnumerateManifestData(BuildManifest);
+				}
+				else
+				{
+					GLog->Logf(TEXT("BuildDataGenerator: WARNING: Could not read Manifest file. Data recognition will suffer (%s)"), *ManifestFilename);
+				}
+			}
+			const double LoadAllManifestTime = FPlatformTime::Seconds() - StartLoadAllManifest;
+			GLog->Logf(TEXT("BuildDataGenerator: Used %d manifests to enumerate %llu chunks in %.1f seconds"), AllManifests.Num(), NumChunksFound, LoadAllManifestTime);
+		}
+		else
+		{
+			GLog->Logf(TEXT("BuildDataGenerator: Cloud directory does not exist: %s"), *CloudDirectory);
+		}
+	}
+
+	void FCloudEnumerationImpl::EnumerateManifestData(const FBuildPatchAppManifestRef& Manifest)
+	{
+		TArray<FGuid> DataList;
+		Manifest->GetDataList(DataList);
+		if (!Manifest->IsFileDataManifest())
+		{
+			for (const auto& DataGuid : DataList)
+			{
+				uint64 DataChunkHash;
+				if (Manifest->GetChunkHash(DataGuid, DataChunkHash))
+				{
+					if (DataChunkHash != 0)
+					{
+						TSet<FGuid>& HashChunkList = ChunkInventory.FindOrAdd(DataChunkHash);
+						if (!HashChunkList.Contains(DataGuid))
+						{
+							++NumChunksFound;
+							HashChunkList.Add(DataGuid);
+							ChunkFileSizes.Add(DataGuid, Manifest->GetDataSize(DataGuid));
+						}
+					}
+					else
+					{
+						GLog->Logf(TEXT("BuildDataGenerator: INFO: Ignored an existing chunk %s with a failed hash value of zero to avoid performance problems while chunking"), *DataGuid.ToString());
+					}
+				}
+				else
+				{
+					GLog->Logf(TEXT("BuildDataGenerator: WARNING: Missing chunk hash for %s in manifest %s %s"), *DataGuid.ToString(), *Manifest->GetAppName(), *Manifest->GetVersionString());
+				}
+			}
+		}
+		else
+		{
+			GLog->Logf(TEXT("BuildDataGenerator: INFO: Ignoring non-chunked manifest %s %s"), *Manifest->GetAppName(), *Manifest->GetVersionString());
+		}
+	}
+
+	FCloudEnumerationRef FCloudEnumerationFactory::Create(const FString& CloudDirectory)
+	{
+		return MakeShareable(new FCloudEnumerationImpl(CloudDirectory));
+	}
+}
+
+#endif
diff --git a/Engine/Source/Runtime/Online/BuildPatchServices/Private/Generation/CloudEnumeration.h b/Engine/Source/Runtime/Online/BuildPatchServices/Private/Generation/CloudEnumeration.h
new file mode 100644
--- /dev/null
+++ b/Engine/Source/Runtime/Online/BuildPatchServices/Private/Generation/CloudEnumeration.h
@@ -0,0 +1,26 @@
+// Copyright 1998-2015 Epic Games, Inc. All Rights Reserved.
+#pragma once
+
+#if WITH_BUILDPATCHGENERATION
+
+namespace FBuildPatchServices
+{
+	class FCloudEnumeration
+	{
+	public:
+		virtual TSet<FGuid> GetChunkSet(uint64 ChunkHash) const = 0;
+		virtual TMap<uint64, TSet<FGuid>> GetChunkInventory() const = 0;
+		virtual TMap<FGuid, int64> GetChunkFileSizes() const = 0;
+	};
+
+	typedef TSharedRef<FCloudEnumeration, ESPMode::ThreadSafe> FCloudEnumerationRef;
+	typedef TSharedPtr<FCloudEnumeration, ESPMode::ThreadSafe> FCloudEnumerationPtr;
+
+	class FCloudEnumerationFactory
+	{
+	public:
+		static FCloudEnumerationRef Create(const FString& CloudDirectory);
+	};
+}
+
+#endif
diff --git a/Engine/Source/Runtime/Online/BuildPatchServices/Private/Generation/DataMatcher.cpp b/Engine/Source/Runtime/Online/BuildPatchServices/Private/Generation/DataMatcher.cpp
new file mode 100644
--- /dev/null
+++ b/Engine/Source/Runtime/Online/BuildPatchServices/Private/Generation/DataMatcher.cpp
@@ -0,0 +1,72 @@
+// Copyright 1998-2015 Epic Games, Inc. All Rights Reserved.
+
+#if WITH_BUILDPATCHGENERATION
+
+#include "BuildPatchServicesPrivatePCH.h"
+
+#include "DataMatcher.h"
+#include "../BuildPatchGeneration.h"
+
+namespace FBuildPatchServices
+{
+	class FDataMatcherImpl
+		: public FDataMatcher
+	{
+	public:
+		FDataMatcherImpl(const FString& CloudDirectory);
+		virtual ~FDataMatcherImpl();
+
+		virtual bool CompareData(const FGuid& DataGuid, uint64 Hash, const TArray<uint8>& NewData) const override;
+
+	private:
+		const FString CloudDirectory;
+	};
+
+	FDataMatcherImpl::FDataMatcherImpl(const FString& InCloudDirectory)
+		: CloudDirectory(InCloudDirectory)
+	{
+
+	}
+
+	FDataMatcherImpl::~FDataMatcherImpl()
+	{
+
+	}
+
+	bool FDataMatcherImpl::CompareData(const FGuid& DataGuid, uint64 Hash, const TArray<uint8>& NewData) const
+	{
+		static FCriticalSection LegacyNonThreadSafeInterface;
+		FScopeLock ScopeLock(&LegacyNonThreadSafeInterface);
+
+		bool bMatching = false;
+
+		// Read the file
+		FString ChunkFilePath = FBuildPatchUtils::GetChunkNewFilename(EBuildPatchAppManifestVersion::GetLatestVersion(), CloudDirectory, DataGuid, Hash);
+		TSharedRef< FBuildGenerationChunkCache::FChunkReader > ChunkReader = FBuildGenerationChunkCache::Get().GetChunkReader(ChunkFilePath);
+		if (ChunkReader->IsValidChunk() && DataGuid == ChunkReader->GetChunkGuid())
+		{
+			// Default true
+			bMatching = true;
+			// Compare per small block (for early outing!)
+			const uint32 CompareSize = 64;
+			uint8* ReadBuffer;
+			uint32 NumCompared = 0;
+			while (bMatching && ChunkReader->BytesLeft() > 0 && NumCompared < FBuildPatchData::ChunkDataSize)
+			{
+				const uint32 ReadLen = FMath::Min< uint32 >(CompareSize, ChunkReader->BytesLeft());
+				ChunkReader->ReadNextBytes(&ReadBuffer, ReadLen);
+				bMatching = FMemory::Memcmp(&NewData[NumCompared], ReadBuffer, ReadLen) == 0;
+				NumCompared += ReadLen;
+			}
+		}
+
+		return bMatching;
+	}
+
+	FDataMatcherRef FDataMatcherFactory::Create(const FString& CloudDirectory)
+	{
+		return MakeShareable(new FDataMatcherImpl(CloudDirectory));
+	}
+}
+
+#endif
diff --git a/Engine/Source/Runtime/Online/BuildPatchServices/Private/Generation/DataMatcher.h b/Engine/Source/Runtime/Online/BuildPatchServices/Private/Generation/DataMatcher.h
new file mode 100644
--- /dev/null
+++ b/Engine/Source/Runtime/Online/BuildPatchServices/Private/Generation/DataMatcher.h
@@ -0,0 +1,24 @@
+// Copyright 1998-2015 Epic Games, Inc. All Rights Reserved.
+#pragma once
+
+#if WITH_BUILDPATCHGENERATION
+
+namespace FBuildPatchServices
+{
+	class FDataMatcher
+	{
+	public:
+		virtual bool CompareData(const FGuid& DataGuid, uint64 Hash, const TArray<uint8>& NewData) const = 0;
+	};
+
+	typedef TSharedRef<FDataMatcher, ESPMode::ThreadSafe> FDataMatcherRef;
+	typedef TSharedPtr<FDataMatcher, ESPMode::ThreadSafe> FDataMatcherPtr;
+
+	class FDataMatcherFactory
+	{
+	public:
+		static FDataMatcherRef Create(const FString& CloudDirectory);
+	};
+}
+
+#endif
diff --git a/Engine/Source/Runtime/Online/BuildPatchServices/Private/Generation/DataScanner.cpp b/Engine/Source/Runtime/Online/BuildPatchServices/Private/Generation/DataScanner.cpp
new file mode 100644
--- /dev/null
+++ b/Engine/Source/Runtime/Online/BuildPatchServices/Private/Generation/DataScanner.cpp
@@ -0,0 +1,422 @@
+// Copyright 1998-2015 Epic Games, Inc. All Rights Reserved.
+
+#if WITH_BUILDPATCHGENERATION
+
+#include "BuildPatchServicesPrivatePCH.h"
+
+#include "DataScanner.h"
+
+#include "ThreadingBase.h"
+#include "Async.h"
+#include "../BuildPatchHash.h"
+#include "../BuildPatchChunk.h"
+#include "../BuildPatchManifest.h"
+#include "StatsCollector.h"
+
+namespace FBuildPatchServices
+{
+	const uint32 WindowSize = FBuildPatchData::ChunkDataSize;
+
+	struct FScopeCounter
+	{
+	public:
+		FScopeCounter(FThreadSafeCounter* Counter);
+		~FScopeCounter();
+
+	private:
+		FThreadSafeCounter* Counter;
+	};
+
+	class FDataStructure
+	{
+	public:
+		FDataStructure(const uint64 DataOffset);
+		~FDataStructure();
+
+		FORCEINLINE const FGuid& GetCurrentChunkId() const;
+
+		FORCEINLINE void PushKnownChunk(const FGuid& MatchId, uint32 DataSize);
+		FORCEINLINE void PushUnknownByte();
+
+		FORCEINLINE void RemapCurrentChunk(const FGuid& NewId);
+		FORCEINLINE void CompleteCurrentChunk();
+
+		FORCEINLINE TArray<FChunkPart> GetFinalDataStructure();
+
+	private:
+		FGuid NewChunkGuid;
+		TArray<FChunkPart> DataStructure;
+	};
+
+	class FDataScannerImpl
+		: public FDataScanner
+	{
+	public:
+		FDataScannerImpl(const uint64 DataOffset, const TArray<uint8>& Data, const FCloudEnumerationRef& CloudEnumeration, const FDataMatcherRef& DataMatcher, const FStatsCollectorRef& StatsCollector);
+		virtual ~FDataScannerImpl();
+
+		virtual bool IsComplete() override;
+		virtual FDataScanResult GetResultWhenComplete() override;
+
+	private:
+		FDataScanResult ScanData();
+		bool ProcessCurrentWindow();
+
+	private:
+		const uint64 DataStartOffset;
+		TArray<uint8> Data;
+		FCloudEnumerationRef CloudEnumeration;
+		FDataMatcherRef DataMatcher;
+		FStatsCollectorRef StatsCollector;
+		FThreadSafeBool bIsComplete;
+		FThreadSafeBool bShouldAbort;
+		TFuture<FDataScanResult> FutureResult;
+
+	public:
+		static FThreadSafeCounter NumIncompleteScanners;
+		static FThreadSafeCounter NumRunningScanners;
+	};
+
+	FScopeCounter::FScopeCounter(FThreadSafeCounter* InCounter)
+		: Counter(InCounter)
+	{
+		Counter->Increment();
+	}
+
+	FScopeCounter::~FScopeCounter()
+	{
+		Counter->Decrement();
+	}
+
+	FDataStructure::FDataStructure(const uint64 DataOffset)
+	{
+		DataStructure.AddZeroed();
+		DataStructure.Top().DataOffset = DataOffset;
+		DataStructure.Top().ChunkGuid = NewChunkGuid = FGuid::NewGuid();
+	}
+
+	FDataStructure::~FDataStructure()
+	{
+	}
+
+	const FGuid& FDataStructure::GetCurrentChunkId() const
+	{
+		return NewChunkGuid;
+	}
+
+	void FDataStructure::PushKnownChunk(const FGuid& PotentialMatch, uint32 NumDataInWindow)
+	{
+		if (DataStructure.Top().PartSize > 0)
+		{
+			// Add for matched
+			DataStructure.AddUninitialized();
+			// Add for next
+			DataStructure.AddUninitialized();
+
+			// Fill out info
+			FChunkPart& PreviousChunkPart = DataStructure[DataStructure.Num() - 3];
+			FChunkPart& MatchedChunkPart = DataStructure[DataStructure.Num() - 2];
+			FChunkPart& NextChunkPart = DataStructure[DataStructure.Num() - 1];
+
+			MatchedChunkPart.DataOffset = PreviousChunkPart.DataOffset + PreviousChunkPart.PartSize;
+			MatchedChunkPart.PartSize = NumDataInWindow;
+			MatchedChunkPart.ChunkOffset = 0;
+			MatchedChunkPart.ChunkGuid = PotentialMatch;
+
+			NextChunkPart.DataOffset = MatchedChunkPart.DataOffset + MatchedChunkPart.PartSize;
+			NextChunkPart.ChunkGuid = PreviousChunkPart.ChunkGuid;
+			NextChunkPart.ChunkOffset = PreviousChunkPart.ChunkOffset + PreviousChunkPart.PartSize;
+			NextChunkPart.PartSize = 0;
+		}
+		else
+		{
+			// Add for next
+			DataStructure.AddZeroed();
+
+			// Fill out info
+			FChunkPart& MatchedChunkPart = DataStructure[DataStructure.Num() - 2];
+			FChunkPart& NextChunkPart = DataStructure[DataStructure.Num() - 1];
+
+			MatchedChunkPart.PartSize = NumDataInWindow;
+			MatchedChunkPart.ChunkOffset = 0;
+			MatchedChunkPart.ChunkGuid = PotentialMatch;
+
+			NextChunkPart.DataOffset = MatchedChunkPart.DataOffset + MatchedChunkPart.PartSize;
+			NextChunkPart.ChunkGuid = NewChunkGuid;
+		}
+	}
+
+	void FDataStructure::PushUnknownByte()
+	{
+		DataStructure.Top().PartSize++;
+	}
+
+	void FDataStructure::RemapCurrentChunk(const FGuid& NewId)
+	{
+		for (auto& DataPiece : DataStructure)
+		{
+			if (DataPiece.ChunkGuid == NewChunkGuid)
+			{
+				DataPiece.ChunkGuid = NewId;
+			}
+		}
+		NewChunkGuid = NewId;
+	}
+
+	void FDataStructure::CompleteCurrentChunk()
+	{
+		DataStructure.AddZeroed();
+		FChunkPart& PreviousPart = DataStructure[DataStructure.Num() - 2];
+
+		// Create next chunk
+		NewChunkGuid = FGuid::NewGuid();
+		DataStructure.Top().DataOffset = PreviousPart.DataOffset + PreviousPart.PartSize;
+		DataStructure.Top().ChunkGuid = NewChunkGuid;
+	}
+
+	TArray<FChunkPart> FDataStructure::GetFinalDataStructure()
+	{
+		if (DataStructure.Top().PartSize == 0)
+		{
+			DataStructure.Pop(false);
+		}
+		return MoveTemp(DataStructure);
+	}
+
+	FDataScannerImpl::FDataScannerImpl(const uint64 InDataOffset, const TArray<uint8>& InData, const FCloudEnumerationRef& InCloudEnumeration, const FDataMatcherRef& InDataMatcher, const FStatsCollectorRef& InStatsCollector)
+		: DataStartOffset(InDataOffset)
+		, Data(InData)
+		, CloudEnumeration(InCloudEnumeration)
+		, DataMatcher(InDataMatcher)
+		, StatsCollector(InStatsCollector)
+		, bIsComplete(false)
+		, bShouldAbort(false)
+	{
+		NumIncompleteScanners.Increment();
+		TFunction<FDataScanResult()> Task = [this]()
+		{
+			FDataScanResult Result = ScanData();
+			FDataScannerImpl::NumIncompleteScanners.Decrement();
+			return MoveTemp(Result);
+		};
+		FutureResult = Async(EAsyncExecution::ThreadPool, Task);
+	}
+
+	FDataScannerImpl::~FDataScannerImpl()
+	{
+		// Make sure the task is complete
+		bShouldAbort = true;
+		FutureResult.Wait();
+	}
+
+	bool FDataScannerImpl::IsComplete()
+	{
+		return bIsComplete;
+	}
+
+	FDataScanResult FDataScannerImpl::GetResultWhenComplete()
+	{
+		return MoveTemp(FutureResult.Get());
+	}
+
+	FDataScanResult FDataScannerImpl::ScanData()
+	{
+		// Count running scanners
+		FScopeCounter ScopeCounter(&NumRunningScanners);
+
+		// Create statistics
+		auto* StatRunningScanners = StatsCollector->CreateStat(TEXT("Scanner: Running Scanners"), EStatFormat::Value);
+		auto* StatMatchedData = StatsCollector->CreateStat(TEXT("Scanner: Matched Data"), EStatFormat::DataSize);
+		FStatsCollector::Accumulate(StatRunningScanners, 1);
+
+		// Init data
+		FRollingHash<WindowSize> RollingHash;
+		FChunkWriter ChunkWriter(FBuildPatchServicesModule::GetCloudDirectory());
+		FDataStructure DataStructure(DataStartOffset);
+		TMap<FGuid, FChunkInfo> ChunkInfoLookup;
+		TArray<uint8> ChunkBuffer;
+		TArray<uint8> NewChunkBuffer;
+		uint32 PaddedZeros = 0;
+		ChunkInfoLookup.Reserve(Data.Num() / WindowSize);
+		ChunkBuffer.SetNumUninitialized(WindowSize);
+		NewChunkBuffer.Reserve(WindowSize);
+
+		// Get a copy of the chunk inventory
+		TMap<uint64, TSet<FGuid>> ChunkInventory = CloudEnumeration->GetChunkInventory();
+		TMap<FGuid, int64> ChunkFileSizes = CloudEnumeration->GetChunkFileSizes();
+
+		// Loop over and process all data
+		for (int32 idx = 0; (idx < Data.Num() || PaddedZeros < WindowSize) && !bShouldAbort; ++idx)
+		{
+			// Consume data
+			const uint32 NumDataNeeded = RollingHash.GetNumDataNeeded();
+			if (NumDataNeeded > 0)
+			{
+				uint32 NumConsumedBytes = 0;
+				if (idx < Data.Num())
+				{
+					NumConsumedBytes = FMath::Min<uint32>(NumDataNeeded, Data.Num() - idx);
+					RollingHash.ConsumeBytes(&Data[idx], NumConsumedBytes);
+					idx += NumConsumedBytes - 1;
+				}
+				// Zero Pad?
+				if (NumConsumedBytes < NumDataNeeded)
+				{
+					TArray<uint8> Zeros;
+					Zeros.AddZeroed(NumDataNeeded - NumConsumedBytes);
+					RollingHash.ConsumeBytes(Zeros.GetData(), Zeros.Num());
+					PaddedZeros = Zeros.Num();
+				}
+				check(RollingHash.GetNumDataNeeded() == 0);
+				continue;
+			}
+
+			const uint64 NumDataInWindow = WindowSize - PaddedZeros;
+			const uint64 WindowHash = RollingHash.GetWindowHash();
+			// Try find match
+			bool bFoundChunkMatch = false;
+			if (ChunkInventory.Contains(WindowHash))
+			{
+				RollingHash.GetWindowData().Serialize(ChunkBuffer.GetData());
+				for (FGuid& PotentialMatch : ChunkInventory.FindRef(WindowHash))
+				{
+					if (DataMatcher->CompareData(PotentialMatch, WindowHash, ChunkBuffer))
+					{
+						bFoundChunkMatch = true;
+						DataStructure.PushKnownChunk(PotentialMatch, NumDataInWindow);
+						FChunkInfo& ChunkInfo = ChunkInfoLookup.FindOrAdd(PotentialMatch);
+						ChunkInfo.Hash = WindowHash;
+						ChunkInfo.IsNew = false;
+						break;
+					}
+				}
+			}
+
+			if (bFoundChunkMatch)
+			{
+				FStatsCollector::Accumulate(StatMatchedData, NumDataInWindow);
+				// Clear matched window
+				RollingHash.Clear();
+				// Decrement idx to include current byte in next window
+				--idx;
+			}
+			else
+			{
+				// Collect unrecognized bytes
+				NewChunkBuffer.Add(RollingHash.GetWindowData().Bottom());
+				DataStructure.PushUnknownByte();
+				if (NumDataInWindow == 1)
+				{
+					NewChunkBuffer.AddZeroed(WindowSize - NewChunkBuffer.Num());
+				}
+				if (NewChunkBuffer.Num() == WindowSize)
+				{
+					const uint64 NewChunkHash = FRollingHash<WindowSize>::GetHashForDataSet(NewChunkBuffer.GetData());
+					for (FGuid& PotentialMatch : ChunkInventory.FindRef(NewChunkHash))
+					{
+						if (DataMatcher->CompareData(PotentialMatch, NewChunkHash, NewChunkBuffer))
+						{
+							bFoundChunkMatch = true;
+							DataStructure.RemapCurrentChunk(PotentialMatch);
+							FChunkInfo& ChunkInfo = ChunkInfoLookup.FindOrAdd(PotentialMatch);
+							ChunkInfo.Hash = NewChunkHash;
+							ChunkInfo.IsNew = false;
+							break;
+						}
+					}
+					if (!bFoundChunkMatch)
+					{
+						const FGuid& NewChunkGuid = DataStructure.GetCurrentChunkId();
+						ChunkWriter.QueueChunk(NewChunkBuffer.GetData(), NewChunkGuid, NewChunkHash);
+						FChunkInfo& ChunkInfo = ChunkInfoLookup.FindOrAdd(NewChunkGuid);
+						ChunkInfo.Hash = NewChunkHash;
+						ChunkInfo.IsNew = true;
+					}
+					else
+					{
+						FStatsCollector::Accumulate(StatMatchedData, WindowSize);
+					}
+					DataStructure.CompleteCurrentChunk();
+					NewChunkBuffer.Empty(WindowSize);
+				}
+
+				// Roll byte into window
+				if (idx < Data.Num())
+				{
+					RollingHash.RollForward(Data[idx]);
+				}
+				else
+				{
+					RollingHash.RollForward(0);
+					++PaddedZeros;
+				}
+			}
+		}
+
+		// Collect left-overs
+		if (NewChunkBuffer.Num() > 0)
+		{
+			NewChunkBuffer.AddZeroed(WindowSize - NewChunkBuffer.Num());
+			bool bFoundChunkMatch = false;
+			const uint64 NewChunkHash = FRollingHash<WindowSize>::GetHashForDataSet(NewChunkBuffer.GetData());
+			for (FGuid& PotentialMatch : ChunkInventory.FindRef(NewChunkHash))
+			{
+				if (DataMatcher->CompareData(PotentialMatch, NewChunkHash, NewChunkBuffer))
+				{
+					bFoundChunkMatch = true;
+					DataStructure.RemapCurrentChunk(PotentialMatch);
+					FChunkInfo& ChunkInfo = ChunkInfoLookup.FindOrAdd(PotentialMatch);
+					ChunkInfo.Hash = NewChunkHash;
+					ChunkInfo.IsNew = false;
+					break;
+				}
+			}
+			if (!bFoundChunkMatch)
+			{
+				const FGuid& NewChunkGuid = DataStructure.GetCurrentChunkId();
+				ChunkWriter.QueueChunk(NewChunkBuffer.GetData(), NewChunkGuid, NewChunkHash);
+				FChunkInfo& ChunkInfo = ChunkInfoLookup.FindOrAdd(NewChunkGuid);
+				ChunkInfo.Hash = NewChunkHash;
+				ChunkInfo.IsNew = true;
+			}
+		}
+
+		// Wait for the chunk writer to finish, and fill out chunk file sizes
+		ChunkWriter.NoMoreChunks();
+		ChunkWriter.WaitForThread();
+		ChunkWriter.GetChunkFilesizes(ChunkFileSizes);
+		for (auto& ChunkInfo : ChunkInfoLookup)
+		{
+			ChunkInfo.Value.ChunkFileSize = ChunkFileSizes[ChunkInfo.Key];
+		}
+
+		// Empty data to save RAM
+		Data.Empty();
+
+		FStatsCollector::Accumulate(StatRunningScanners, -1);
+		bIsComplete = true;
+		return FDataScanResult(
+			MoveTemp(DataStructure.GetFinalDataStructure()),
+			MoveTemp(ChunkInfoLookup));
+	}
+
+	FThreadSafeCounter FDataScannerImpl::NumIncompleteScanners;
+	FThreadSafeCounter FDataScannerImpl::NumRunningScanners;
+
+	int32 FDataScannerCounter::GetNumIncompleteScanners()
+	{
+		return FDataScannerImpl::NumIncompleteScanners.GetValue();
+	}
+
+	int32 FDataScannerCounter::GetNumRunningScanners()
+	{
+		return FDataScannerImpl::NumRunningScanners.GetValue();
+	}
+
+	FDataScannerRef FDataScannerFactory::Create(const uint64 DataOffset, const TArray<uint8>& Data, const FCloudEnumerationRef& CloudEnumeration, const FDataMatcherRef& DataMatcher, const FStatsCollectorRef& StatsCollector)
+	{
+		return MakeShareable(new FDataScannerImpl(DataOffset, Data, CloudEnumeration, DataMatcher, StatsCollector));
+	}
+}
+
+#endif
diff --git a/Engine/Source/Runtime/Online/BuildPatchServices/Private/Generation/DataScanner.h b/Engine/Source/Runtime/Online/BuildPatchServices/Private/Generation/DataScanner.h
new file mode 100644
--- /dev/null
+++ b/Engine/Source/Runtime/Online/BuildPatchServices/Private/Generation/DataScanner.h
@@ -0,0 +1,100 @@
+// Copyright 1998-2015 Epic Games, Inc. All Rights Reserved.
+#pragma once
+
+#if WITH_BUILDPATCHGENERATION
+
+#include "CloudEnumeration.h"
+#include "DataMatcher.h"
+#include "StatsCollector.h"
+
+namespace FBuildPatchServices
+{
+	struct FChunkPart
+	{
+		// Offset into provided data
+		uint64 DataOffset;
+		// Size of part
+		uint64 PartSize;
+		// Offset into chunk
+		uint64 ChunkOffset;
+		// Chunk identity
+		FGuid ChunkGuid;
+	};
+
+	struct FChunkInfo
+	{
+		// The hash for this chunk
+		uint64 Hash;
+		// The size of this chunk file
+		int64 ChunkFileSize;
+		// The group number this chunk is stored in
+		uint8 GroupNumber;
+		// Whether this chunk is new
+		bool IsNew;
+	};
+
+	struct FDataScanResult
+	{
+		// The structure of the data
+		TArray<FChunkPart> DataStructure;
+		// Information about chunks in the data
+		TMap<FGuid, FChunkInfo> ChunkInfo;
+
+		FDataScanResult() {}
+		~FDataScanResult() {}
+
+		FDataScanResult(TArray<FChunkPart> InDataStructure, TMap<FGuid, FChunkInfo> InChunkInfo)
+			: DataStructure(MoveTemp(InDataStructure))
+			, ChunkInfo(MoveTemp(InChunkInfo))
+		{}
+
+		FDataScanResult(const FDataScanResult& CopyFrom)
+			: DataStructure(CopyFrom.DataStructure)
+			, ChunkInfo(CopyFrom.ChunkInfo)
+		{}
+
+		FDataScanResult(FDataScanResult&& MoveFrom)
+			: DataStructure(MoveTemp(MoveFrom.DataStructure))
+			, ChunkInfo(MoveTemp(MoveFrom.ChunkInfo))
+		{}
+
+		FORCEINLINE FDataScanResult& operator=(const FDataScanResult& CopyFrom)
+		{
+			DataStructure = CopyFrom.DataStructure;
+			ChunkInfo = CopyFrom.ChunkInfo;
+			return *this;
+		}
+
+		FORCEINLINE FDataScanResult& operator=(FDataScanResult&& MoveFrom)
+		{
+			DataStructure = MoveTemp(MoveFrom.DataStructure);
+			ChunkInfo = MoveTemp(MoveFrom.ChunkInfo);
+			return *this;
+		}
+	};
+
+	class FDataScanner
+	{
+	public:
+		virtual bool IsComplete() = 0;
+		virtual FDataScanResult GetResultWhenComplete() = 0;
+	};
+
+	typedef TSharedRef<FDataScanner, ESPMode::ThreadSafe> FDataScannerRef;
+	typedef TSharedPtr<FDataScanner, ESPMode::ThreadSafe> FDataScannerPtr;
+
+	class FDataScannerCounter
+	{
+	public:
+		static int32 GetNumIncompleteScanners();
+		static int32 GetNumRunningScanners();
+	};
+
+	class FDataScannerFactory
+	{
+	public:
+		static FDataScannerRef Create(const uint64 DataOffset, const TArray<uint8>& Data, const FCloudEnumerationRef& CloudEnumeration, const FDataMatcherRef& DataMatcher, const FStatsCollectorRef& StatsCollector);
+	};
+}
+
+#endif
diff --git a/Engine/Source/Runtime/Online/BuildPatchServices/Private/Generation/ManifestBuilder.cpp b/Engine/Source/Runtime/Online/BuildPatchServices/Private/Generation/ManifestBuilder.cpp
new file mode 100644
--- /dev/null
+++ b/Engine/Source/Runtime/Online/BuildPatchServices/Private/Generation/ManifestBuilder.cpp
@@ -0,0 +1,348 @@
+// Copyright 1998-2015 Epic Games, Inc. All Rights Reserved.
+
+#if WITH_BUILDPATCHGENERATION
+
+#include "BuildPatchServicesPrivatePCH.h"
+
+#include "ManifestBuilder.h"
+#include "BuildStreamer.h"
+
+namespace FBuildPatchServices
+{
+	class FFileManifestBuilder
+	{
+	public:
+		FFileManifestBuilder();
+		~FFileManifestBuilder();
+
+		uint64 CurrentDataPos;
+
+		FFileSpan FileSpan;
+
+		FFileManifestData* FileManifest;
+	};
+
+	FFileManifestBuilder::FFileManifestBuilder()
+		: CurrentDataPos(0)
+		, FileManifest(nullptr)
+	{
+
+	}
+
+	FFileManifestBuilder::~FFileManifestBuilder()
+	{
+
+	}
+
+	class FManifestBuilderImpl
+		: public FManifestBuilder
+	{
+	public:
+		FManifestBuilderImpl(const FManifestDetails& Details, const FBuildStreamerRef& BuildStreamer);
+		virtual ~FManifestBuilderImpl();
+
+		virtual void AddDataScanner(FDataScannerRef Scanner) override;
+		virtual void SaveToFile(const FString& Filename) override;
+	private:
+		void BuildManifest();
+		FDataScannerPtr GetNextScanner();
+
+	private:
+		FBuildStreamerRef BuildStreamer;
+		FBuildPatchAppManifestRef Manifest;
+		TMap<FString, FFileAttributes> FileAttributesMap;
+		FCriticalSection DataScannersCS;
+		TArray<FDataScannerRef> DataScanners;
+		FThreadSafeBool EndOfData;
+		TFuture<void> Future;
+		FEvent* CheckForWork;
+
+		FFileManifestBuilder FileBuilder;
+	};
+
+	FManifestBuilderImpl::FManifestBuilderImpl(const FManifestDetails& InDetails, const FBuildStreamerRef& InBuildStreamer)
+		: BuildStreamer(InBuildStreamer)
+		, Manifest(MakeShareable(new FBuildPatchAppManifest()))
+		, FileAttributesMap(InDetails.FileAttributesMap)
+		, EndOfData(false)
+		, CheckForWork(FPlatformProcess::GetSynchEventFromPool(true))
+	{
+		Manifest->Data->bIsFileData = InDetails.bIsFileData;
+		Manifest->Data->AppID = InDetails.AppId;
+		Manifest->Data->AppName = InDetails.AppName;
+		Manifest->Data->BuildVersion = InDetails.BuildVersion;
+		Manifest->Data->LaunchExe = InDetails.LaunchExe;
+		Manifest->Data->LaunchCommand = InDetails.LaunchCommand;
+		Manifest->Data->PrereqName = InDetails.PrereqName;
+		Manifest->Data->PrereqPath = InDetails.PrereqPath;
+		Manifest->Data->PrereqArgs = InDetails.PrereqArgs;
+		for (const auto& CustomField : InDetails.CustomFields)
+		{
+			int32 VarType = CustomField.Value.GetType();
+			if (VarType == EVariantTypes::Float || VarType == EVariantTypes::Double)
+			{
+				Manifest->SetCustomField(CustomField.Key, (double)CustomField.Value);
+			}
+			else if (VarType == EVariantTypes::Int8 || VarType == EVariantTypes::Int16 || VarType == EVariantTypes::Int32 || VarType == EVariantTypes::Int64 ||
+				VarType == EVariantTypes::UInt8 || VarType == EVariantTypes::UInt16 || VarType == EVariantTypes::UInt32 || VarType == EVariantTypes::UInt64)
+			{
+				Manifest->SetCustomField(CustomField.Key, (int64)CustomField.Value);
+			}
+			else if (VarType == EVariantTypes::String)
+			{
+				Manifest->SetCustomField(CustomField.Key, CustomField.Value.GetValue<FString>());
+			}
+		}
+		TFunction<void()> Task = [this]() { BuildManifest(); };
+		Future = Async(EAsyncExecution::Thread, Task);
+	}
+
+	FManifestBuilderImpl::~FManifestBuilderImpl()
+	{
+		EndOfData = true;
+		CheckForWork->Trigger();
+		Future.Wait();
+		FPlatformProcess::ReturnSynchEventToPool(CheckForWork);
+		CheckForWork = nullptr;
+	}
+
+	void FManifestBuilderImpl::AddDataScanner(FDataScannerRef Scanner)
+	{
+		DataScannersCS.Lock();
+		DataScanners.Add(MoveTemp(Scanner));
+		DataScannersCS.Unlock();
+		CheckForWork->Trigger();
+	}
+
+	void FManifestBuilderImpl::SaveToFile(const FString& Filename)
+	{
+		// Wait for builder to complete
+		EndOfData = true;
+		CheckForWork->Trigger();
+		Future.Wait();
+
+		Manifest->Data->ManifestFileVersion = EBuildPatchAppManifestVersion::GetLatestJsonVersion();
+		Manifest->SaveToFile(Filename, false);
+	}
+
+	void FManifestBuilderImpl::BuildManifest()
+	{
+		TMap<FGuid, FChunkInfo> ChunkInfoLookup;
+		bool Running = true;
+		while (Running)
+		{
+			FDataScannerPtr NextScanner = GetNextScanner();
+			if (NextScanner.IsValid())
+			{
+				FDataScanResult ScanResult = NextScanner->GetResultWhenComplete();
+				ChunkInfoLookup.Append(ScanResult.ChunkInfo);
+
+				// Always reverse for now
+				if (ScanResult.DataStructure.Num() > 0)
+				{
+					FChunkPart& ChunkPart = ScanResult.DataStructure[0];
+					if (ChunkPart.DataOffset != FileBuilder.CurrentDataPos)
+					{
+						check(ChunkPart.DataOffset < FileBuilder.CurrentDataPos); // Missing data!
+
+						bool FoundPosition = false;
+						uint64 DataCount = 0;
+						for (int32 FileIdx = 0; FileIdx < Manifest->Data->FileManifestList.Num() && !FoundPosition; ++FileIdx)
+						{
+							FFileManifestData& FileManifest = Manifest->Data->FileManifestList[FileIdx];
+							FileManifest.Init();
+							uint64 FileStartIdx = DataCount;
+							uint64 FileEndIdx = FileStartIdx + FileManifest.GetFileSize();
+							if (FileEndIdx > ChunkPart.DataOffset)
+							{
+								for (int32 ChunkIdx = 0; ChunkIdx < FileManifest.FileChunkParts.Num() && !FoundPosition; ++ChunkIdx)
+								{
+									FChunkPartData& ChunkPartData = FileManifest.FileChunkParts[ChunkIdx];
+									uint64 ChunkPartEndIdx = DataCount + ChunkPartData.Size;
+									if (ChunkPartEndIdx < ChunkPart.DataOffset)
+									{
+										DataCount += ChunkPartData.Size;
+									}
+									else if (ChunkPartEndIdx > ChunkPart.DataOffset)
+									{
+										ChunkPartData.Size = ChunkPart.DataOffset - DataCount;
+										FileBuilder.CurrentDataPos = DataCount + ChunkPartData.Size;
+										FileManifest.FileChunkParts.SetNum(ChunkIdx + 1, false);
+										FileManifest.FileChunkParts.Emplace();
+										Manifest->Data->FileManifestList.SetNum(FileIdx + 1, false);
+										FileBuilder.FileManifest = &Manifest->Data->FileManifestList.Last();
+										bool FoundFile = BuildStreamer->GetFileSpan(FileStartIdx, FileBuilder.FileSpan);
+										check(FoundFile); // Incorrect positional tracking
+										FoundPosition = true;
+									}
+									else
+									{
+										FileBuilder.CurrentDataPos = DataCount + ChunkPartData.Size;
+										FileManifest.FileChunkParts.SetNum(ChunkIdx + 1, false);
+										FileManifest.FileChunkParts.Emplace();
+										Manifest->Data->FileManifestList.SetNum(FileIdx + 1, false);
+										FileBuilder.FileManifest = &Manifest->Data->FileManifestList.Last();
+										bool FoundFile = BuildStreamer->GetFileSpan(FileStartIdx, FileBuilder.FileSpan);
+										check(FoundFile); // Incorrect positional tracking
+										FoundPosition = true;
+									}
+								}
+							}
+							else if (FileEndIdx < ChunkPart.DataOffset)
+							{
+								DataCount += FileManifest.GetFileSize();
+							}
+							else
+							{
+								FileBuilder.FileManifest = nullptr;
+								FileBuilder.CurrentDataPos = DataCount + FileManifest.GetFileSize();
+								Manifest->Data->FileManifestList.SetNum(FileIdx + 1, false);
+								FoundPosition = true;
+							}
+						}
+
+						check(ChunkPart.DataOffset == FileBuilder.CurrentDataPos);
+						check(FileBuilder.FileManifest == nullptr || FileBuilder.FileSpan.Filename == Manifest->Data->FileManifestList.Last().Filename);
+					}
+				}
+
+				for (int32 idx = 0; idx < ScanResult.DataStructure.Num(); ++idx)
+				{
+					FChunkPart& ChunkPart = ScanResult.DataStructure[idx];
+					// Starting new file?
+					if (FileBuilder.FileManifest == nullptr)
+					{
+						Manifest->Data->FileManifestList.Emplace();
+						FileBuilder.FileManifest = &Manifest->Data->FileManifestList.Last();
+
+						bool FoundFile = BuildStreamer->GetFileSpan(FileBuilder.CurrentDataPos, FileBuilder.FileSpan);
+						check(FoundFile); // Incorrect positional tracking
+
+						FileBuilder.FileManifest->Filename = FileBuilder.FileSpan.Filename;
+						FileBuilder.FileManifest->FileChunkParts.Emplace();
+					}
+
+					FChunkPartData& FileChunkPartData = FileBuilder.FileManifest->FileChunkParts.Last();
+					FileChunkPartData.Guid = ChunkPart.ChunkGuid;
+					FileChunkPartData.Offset = (FileBuilder.CurrentDataPos - ChunkPart.DataOffset) + ChunkPart.ChunkOffset;
+
+					// Process data into file manifests
+					int64 FileDataLeft = (FileBuilder.FileSpan.StartIdx + FileBuilder.FileSpan.Size) - FileBuilder.CurrentDataPos;
+					int64 ChunkDataLeft = (ChunkPart.DataOffset + ChunkPart.PartSize) - FileBuilder.CurrentDataPos;
+					check(FileDataLeft > 0);
+					check(ChunkDataLeft > 0);
+
+					if (ChunkDataLeft >= FileDataLeft)
+					{
+						FileBuilder.CurrentDataPos += FileDataLeft;
+						FileChunkPartData.Size = FileDataLeft;
+					}
+					else
+					{
+						FileBuilder.CurrentDataPos += ChunkDataLeft;
+						FileChunkPartData.Size = ChunkDataLeft;
+					}
+
+					FileDataLeft = (FileBuilder.FileSpan.StartIdx + FileBuilder.FileSpan.Size) - FileBuilder.CurrentDataPos;
+					ChunkDataLeft = (ChunkPart.DataOffset + ChunkPart.PartSize) - FileBuilder.CurrentDataPos;
+					check(FileDataLeft == 0 || ChunkDataLeft == 0);
+					// End of file?
+					if (FileDataLeft == 0)
+					{
+						// Fill out rest of data??
+						FFileSpan FileSpan;
+						bool FoundFile = BuildStreamer->GetFileSpan(FileBuilder.FileSpan.StartIdx, FileSpan);
+						check(FoundFile); // Incorrect positional tracking
+						check(FileSpan.Filename == FileBuilder.FileManifest->Filename);
+						FMemory::Memcpy(FileBuilder.FileManifest->FileHash.Hash, FileSpan.SHAHash.Hash, FSHA1::DigestSize);
+						FFileAttributes Attributes = FileAttributesMap.FindRef(FileSpan.Filename);
+						FileBuilder.FileManifest->bIsUnixExecutable = Attributes.bUnixExecutable || FileSpan.IsUnixExecutable;
+						FileBuilder.FileManifest->SymlinkTarget = FileSpan.SymlinkTarget;
+						FileBuilder.FileManifest->bIsReadOnly = Attributes.bReadOnly;
+						FileBuilder.FileManifest->bIsCompressed = Attributes.bCompressed;
+						FileBuilder.FileManifest->Init();
+						check(FileBuilder.FileManifest->GetFileSize() == FileBuilder.FileSpan.Size);
+						FileBuilder.FileManifest = nullptr;
+					}
+					else if (ChunkDataLeft == 0)
+					{
+						FileBuilder.FileManifest->FileChunkParts.Emplace();
+					}
+
+					// Continue with this chunk?
+					if (ChunkDataLeft > 0)
+					{
+						--idx;
+					}
+				}
+			}
+			else
+			{
+				if (EndOfData)
+				{
+					Running = false;
+				}
+				else
+				{
+					CheckForWork->Wait();
+					CheckForWork->Reset();
+				}
+			}
+		}
+
+		// Fill out chunk list from only chunks that remain referenced
+		TSet<FGuid> ReferencedChunks;
+		for (const auto& FileManifest : Manifest->Data->FileManifestList)
+		{
+			for (const auto& ChunkPart : FileManifest.FileChunkParts)
+			{
+				if (ReferencedChunks.Contains(ChunkPart.Guid) == false)
+				{
+					auto& ChunkInfo = ChunkInfoLookup[ChunkPart.Guid];
+					ReferencedChunks.Add(ChunkPart.Guid);
+					Manifest->Data->ChunkList.Emplace();
+					auto& ChunkInfoData = Manifest->Data->ChunkList.Last();
+					ChunkInfoData.Guid = ChunkPart.Guid;
+					ChunkInfoData.Hash = ChunkInfo.Hash;
+					ChunkInfoData.FileSize = ChunkInfo.ChunkFileSize;
+					ChunkInfoData.GroupNumber = FCrc::MemCrc32(&ChunkPart.Guid, sizeof(FGuid)) % 100;
+				}
+			}
+		}
+
+		// Get empty files
+		FSHA1 EmptyHasher;
+		EmptyHasher.Final();
+		const TArray< FString >& EmptyFileList = BuildStreamer->GetEmptyFiles();
+		for (const auto& EmptyFile : EmptyFileList)
+		{
+			Manifest->Data->FileManifestList.Emplace();
+			FFileManifestData& EmptyFileManifest = Manifest->Data->FileManifestList.Last();
+			EmptyFileManifest.Filename = EmptyFile;
+			EmptyHasher.GetHash(EmptyFileManifest.FileHash.Hash);
+		}
+
+		// Fill out lookups
+		Manifest->InitLookups();
+	}
+
+	FDataScannerPtr FManifestBuilderImpl::GetNextScanner()
+	{
+		FDataScannerPtr Result;
+		DataScannersCS.Lock();
+		if (DataScanners.Num() > 0)
+		{
+			Result = DataScanners[0];
+			DataScanners.RemoveAt(0);
+		}
+		DataScannersCS.Unlock();
+		return Result;
+	}
+
+	FManifestBuilderRef FManifestBuilderFactory::Create(const FManifestDetails& Details, const FBuildStreamerRef& BuildStreamer)
+	{
+		return MakeShareable(new FManifestBuilderImpl(Details, BuildStreamer));
+	}
+}
+
+#endif
diff --git a/Engine/Source/Runtime/Online/BuildPatchServices/Private/Generation/ManifestBuilder.h b/Engine/Source/Runtime/Online/BuildPatchServices/Private/Generation/ManifestBuilder.h
new file mode 100644
--- /dev/null
+++ b/Engine/Source/Runtime/Online/BuildPatchServices/Private/Generation/ManifestBuilder.h
@@ -0,0 +1,54 @@
+// Copyright 1998-2015 Epic Games, Inc. All Rights Reserved.
+#pragma once
+
+#if WITH_BUILDPATCHGENERATION
+
+#include "DataScanner.h"
+#include "BuildStreamer.h"
+
+namespace FBuildPatchServices
+{
+	struct FManifestDetails
+	{
+		// Whether this manifest is describing nochunks build
+		bool bIsFileData;
+		// The ID of the app of this build
+		uint32 AppId;
+		// The name of the app of this build
+		FString AppName;
+		// The version string for this build
+		FString BuildVersion;
+		// The local exe path that would launch this build
+		FString LaunchExe;
+		// The command line that would launch this build
+		FString LaunchCommand;
+		// The display name of the prerequisites installer
+		FString PrereqName;
+		// The path to the prerequisites installer
+		FString PrereqPath;
+		// The command line arguments for the prerequisites installer
+		FString PrereqArgs;
+		// Map of custom fields to add to the manifest
+		TMap<FString, FVariant> CustomFields;
+		// Map of file attributes
+		TMap<FString, FFileAttributes> FileAttributesMap;
+	};
+
+	class FManifestBuilder
+	{
+	public:
+		virtual void AddDataScanner(FDataScannerRef Scanner) = 0;
+		virtual void SaveToFile(const FString& Filename) = 0;
+	};
+
+	typedef TSharedRef<FManifestBuilder, ESPMode::ThreadSafe> FManifestBuilderRef;
+	typedef TSharedPtr<FManifestBuilder, ESPMode::ThreadSafe> FManifestBuilderPtr;
+
+	class FManifestBuilderFactory
+	{
+	public:
+		static FManifestBuilderRef Create(const FManifestDetails& Details, const FBuildStreamerRef& BuildStreamer);
+	};
+}
+
+#endif
diff --git a/Engine/Source/Runtime/Online/BuildPatchServices/Private/Generation/StatsCollector.cpp b/Engine/Source/Runtime/Online/BuildPatchServices/Private/Generation/StatsCollector.cpp
new file mode 100644
--- /dev/null
+++ b/Engine/Source/Runtime/Online/BuildPatchServices/Private/Generation/StatsCollector.cpp
@@ -0,0 +1,147 @@
+// Copyright 1998-2015 Epic Games, Inc. All Rights Reserved.
+#pragma once
+
+#include "BuildPatchServicesPrivatePCH.h"
+
+#include "StatsCollector.h"
+
+inline uint32 GetTypeHash(volatile const int64* A)
+{
+	return PointerHash((const void*)A);
+}
+
+namespace FBuildPatchServices
+{
+
+	uint64 FStatsCollector::GetCycles()
+	{
+#if PLATFORM_WINDOWS
+		LARGE_INTEGER Cycles;
+		QueryPerformanceCounter(&Cycles);
+		return Cycles.QuadPart;
+#elif PLATFORM_MAC
+		uint64 Cycles = mach_absolute_time();
+		return Cycles;
+#else
+		return FPlatformTime::Cycles();
+#endif
+	}
+
+	double FStatsCollector::CyclesToSeconds(const uint64 Cycles)
+	{
+		return FPlatformTime::GetSecondsPerCycle() * Cycles;
+	}
+
+	void FStatsCollector::AccumulateTimeBegin(uint64& TempValue)
+	{
+		TempValue = FStatsCollector::GetCycles();
+	}
+
+	void FStatsCollector::AccumulateTimeEnd(volatile int64* Stat, uint64& TempValue)
+	{
+		FPlatformAtomics::InterlockedAdd(Stat, FStatsCollector::GetCycles() - TempValue);
+	}
+
+	void FStatsCollector::Accumulate(volatile int64* Stat, int64 Amount)
+	{
+		FPlatformAtomics::InterlockedAdd(Stat, Amount);
+	}
+
+	void FStatsCollector::Set(volatile int64* Stat, int64 Value)
+	{
+		FPlatformAtomics::InterlockedExchange(Stat, Value);
+	}
+
+	class FStatsCollectorImpl
+		: public FStatsCollector
+	{
+	public:
+		FStatsCollectorImpl();
+		~FStatsCollectorImpl();
+
+		virtual volatile int64* CreateStat(const FString& Name, EStatFormat Type, int64 InitialValue = 0) override;
+		virtual void LogStats(float TimeBetweenLogs = 0.0f) override;
+
+	private:
+		FCriticalSection DataCS;
+		TMap<FString, volatile int64*> AddedStats;
+		TMap<volatile int64*, FString> AtomicNameMap;
+		TMap<volatile int64*, EStatFormat> AtomicFormatMap;
+		uint64 LastLogged;
+		int32 LongestName;
+	};
+
+	FStatsCollectorImpl::FStatsCollectorImpl()
+		: LastLogged(FStatsCollector::GetCycles())
+		, LongestName(0)
+	{
+	}
+
+	FStatsCollectorImpl::~FStatsCollectorImpl()
+	{
+		FScopeLock ScopeLock(&DataCS);
+		for (auto& Stat : AddedStats)
+		{
+			delete Stat.Value;
+		}
+		AddedStats.Empty();
+		AtomicNameMap.Empty();
+		AtomicFormatMap.Empty();
+	}
+
+	volatile int64* FStatsCollectorImpl::CreateStat(const FString& Name, EStatFormat Type, int64 InitialValue)
+	{
+		FScopeLock ScopeLock(&DataCS);
+		if (AddedStats.Contains(Name) == false)
+		{
+			volatile int64* NewInteger = new volatile int64(InitialValue);
+			AddedStats.Add(Name, NewInteger);
+			AtomicNameMap.Add(NewInteger, Name + TEXT(": "));
+			AtomicFormatMap.Add(NewInteger, Type);
+			LongestName = FMath::Max<uint32>(LongestName, Name.Len() + 2);
+			for (auto& Stat : AtomicNameMap)
+			{
+				while (Stat.Value.Len() < LongestName)
+				{
+					Stat.Value += TEXT(" ");
+				}
+			}
+		}
+		return AddedStats[Name];
+	}
+
+	void FStatsCollectorImpl::LogStats(float TimeBetweenLogs)
+	{
+		const uint64 Cycles = FStatsCollector::GetCycles();
+		if (FStatsCollector::CyclesToSeconds(Cycles - LastLogged) >= TimeBetweenLogs)
+		{
+			LastLogged = Cycles;
+			FScopeLock ScopeLock(&DataCS);
+			GLog->Log(TEXT("/-------- FStatsCollector Log ---------------------"));
+			for (auto& Stat : AtomicNameMap)
+			{
+				switch (AtomicFormatMap[Stat.Key])
+				{
+				case EStatFormat::Timer:
+					GLog->Logf(TEXT("| %s%s"), *Stat.Value, *FPlatformTime::PrettyTime(FStatsCollector::CyclesToSeconds(*Stat.Key)));
+					break;
+				case EStatFormat::DataSize:
+					GLog->Logf(TEXT("| %s%s"), *Stat.Value, *FText::AsMemory(*Stat.Key).ToString());
+					break;
+				case EStatFormat::DataSpeed:
+					GLog->Logf(TEXT("| %s%s/s"), *Stat.Value, *FText::AsMemory(*Stat.Key).ToString());
+					break;
+				default:
+					GLog->Logf(TEXT("| %s%lld"), *Stat.Value, *Stat.Key);
+					break;
+				}
+			}
+			GLog->Log(TEXT("\\--------------------------------------------------"));
+		}
+	}
+
+	FStatsCollectorRef FStatsCollectorFactory::Create()
+	{
+		return MakeShareable(new FStatsCollectorImpl());
+	}
+}
diff --git a/Engine/Source/Runtime/Online/BuildPatchServices/Private/Generation/StatsCollector.h b/Engine/Source/Runtime/Online/BuildPatchServices/Private/Generation/StatsCollector.h
new file mode 100644
--- /dev/null
+++ b/Engine/Source/Runtime/Online/BuildPatchServices/Private/Generation/StatsCollector.h
@@ -0,0 +1,41 @@
+// Copyright 1998-2015 Epic Games, Inc. All Rights Reserved.
+#pragma once
+
+namespace FBuildPatchServices
+{
+	enum EStatFormat
+	{
+		// Using the AccumulateTimeBegin and End functions, measured in cycles
+		Timer,
+		// Value measured in Bytes
+		DataSize,
+		// Value measured in Bytes per second
+		DataSpeed,
+		// Generic int64 output
+		Value
+	};
+
+	class FStatsCollector
+	{
+	public:
+		virtual volatile int64* CreateStat(const FString& Name, EStatFormat Type, int64 InitialValue = 0) = 0;
+		virtual void LogStats(float TimeBetweenLogs = 0.0f) = 0;
+
+	public:
+		static uint64 GetCycles();
+		static double CyclesToSeconds(const uint64 Cycles);
+		static void AccumulateTimeBegin(uint64& TempValue);
+		static void AccumulateTimeEnd(volatile int64* Stat, uint64& TempValue);
+		static void Accumulate(volatile int64* Stat, int64 Amount);
+		static void Set(volatile int64* Stat, int64 Value);
+	};
+
+	typedef TSharedRef<FStatsCollector, ESPMode::ThreadSafe> FStatsCollectorRef;
+	typedef TSharedPtr<FStatsCollector, ESPMode::ThreadSafe> FStatsCollectorPtr;
+
+	class FStatsCollectorFactory
+	{
+	public:
+		static FStatsCollectorRef Create();
+	};
+}
