Removing hyperthreading commandline force, and Mac Build fixes for TMap usage.

diff --git a/Engine/Source/Programs/BuildPatchTool/Private/BuildPatchToolMain.cpp b/Engine/Source/Programs/BuildPatchTool/Private/BuildPatchToolMain.cpp
--- a/Engine/Source/Programs/BuildPatchTool/Private/BuildPatchToolMain.cpp
+++ b/Engine/Source/Programs/BuildPatchTool/Private/BuildPatchToolMain.cpp
@@ -549,7 +549,7 @@ int32 BuildPatchToolMain( const TCHAR* CommandLine )
 
 INT32_MAIN_INT32_ARGC_TCHAR_ARGV()
 {
-	FString CommandLine = TEXT("-usehyperthreading");
+	FString CommandLine;
 	for( int32 Option = 1; Option < ArgC; Option++ )
 	{
 		CommandLine += TEXT(" ");
diff --git a/Engine/Source/Runtime/Online/BuildPatchServices/Private/Generation/StatsCollector.cpp b/Engine/Source/Runtime/Online/BuildPatchServices/Private/Generation/StatsCollector.cpp
--- a/Engine/Source/Runtime/Online/BuildPatchServices/Private/Generation/StatsCollector.cpp
+++ b/Engine/Source/Runtime/Online/BuildPatchServices/Private/Generation/StatsCollector.cpp
@@ -1,14 +1,9 @@
 // Copyright 1998-2015 Epic Games, Inc. All Rights Reserved.
-#pragma once
 
 #include "BuildPatchServicesPrivatePCH.h"
 
 #include "StatsCollector.h"
 
-inline uint32 GetTypeHash(volatile const int64* A)
-{
-	return PointerHash((const void*)A);
-}
 
 namespace FBuildPatchServices
 {
@@ -57,7 +52,7 @@ namespace FBuildPatchServices
 	{
 	public:
 		FStatsCollectorImpl();
-		~FStatsCollectorImpl();
+		virtual ~FStatsCollectorImpl();
 
 		virtual volatile int64* CreateStat(const FString& Name, EStatFormat Type, int64 InitialValue = 0) override;
 		virtual void LogStats(float TimeBetweenLogs = 0.0f) override;
@@ -65,8 +60,8 @@ namespace FBuildPatchServices
 	private:
 		FCriticalSection DataCS;
 		TMap<FString, volatile int64*> AddedStats;
-		TMap<volatile int64*, FString> AtomicNameMap;
-		TMap<volatile int64*, EStatFormat> AtomicFormatMap;
+		TMap<int64*, FString> AtomicNameMap;
+		TMap<int64*, EStatFormat> AtomicFormatMap;
 		uint64 LastLogged;
 		int32 LongestName;
 	};
@@ -96,8 +91,8 @@ namespace FBuildPatchServices
 		{
 			volatile int64* NewInteger = new volatile int64(InitialValue);
 			AddedStats.Add(Name, NewInteger);
-			AtomicNameMap.Add(NewInteger, Name + TEXT(": "));
-			AtomicFormatMap.Add(NewInteger, Type);
+			AtomicNameMap.Add((int64*)NewInteger, Name + TEXT(": "));
+			AtomicFormatMap.Add((int64*)NewInteger, Type);
 			LongestName = FMath::Max<uint32>(LongestName, Name.Len() + 2);
 			for (auto& Stat : AtomicNameMap)
 			{
@@ -118,21 +113,22 @@ namespace FBuildPatchServices
 			LastLogged = Cycles;
 			FScopeLock ScopeLock(&DataCS);
 			GLog->Log(TEXT("/-------- FStatsCollector Log ---------------------"));
-			for (auto& Stat : AtomicNameMap)
+			for (auto& Stat : AddedStats)
 			{
-				switch (AtomicFormatMap[Stat.Key])
+				int64* NameLookup = (int64*)Stat.Value;
+				switch (AtomicFormatMap[NameLookup])
 				{
 				case EStatFormat::Timer:
-					GLog->Logf(TEXT("| %s%s"), *Stat.Value, *FPlatformTime::PrettyTime(FStatsCollector::CyclesToSeconds(*Stat.Key)));
+					GLog->Logf(TEXT("| %s%s"), *AtomicNameMap[NameLookup], *FPlatformTime::PrettyTime(FStatsCollector::CyclesToSeconds(*Stat.Value)));
 					break;
 				case EStatFormat::DataSize:
-					GLog->Logf(TEXT("| %s%s"), *Stat.Value, *FText::AsMemory(*Stat.Key).ToString());
+					GLog->Logf(TEXT("| %s%s"), *AtomicNameMap[NameLookup], *FText::AsMemory(*Stat.Value).ToString());
 					break;
 				case EStatFormat::DataSpeed:
-					GLog->Logf(TEXT("| %s%s/s"), *Stat.Value, *FText::AsMemory(*Stat.Key).ToString());
+					GLog->Logf(TEXT("| %s%s/s"), *AtomicNameMap[NameLookup], *FText::AsMemory(*Stat.Value).ToString());
 					break;
 				default:
-					GLog->Logf(TEXT("| %s%lld"), *Stat.Value, *Stat.Key);
+					GLog->Logf(TEXT("| %s%lld"), *AtomicNameMap[NameLookup], *Stat.Value);
 					break;
 				}
 			}
