Back out changelist 2579828

diff --git a/Engine/Source/Programs/BuildPatchTool/Private/BuildPatchToolMain.cpp b/Engine/Source/Programs/BuildPatchTool/Private/BuildPatchToolMain.cpp
--- a/Engine/Source/Programs/BuildPatchTool/Private/BuildPatchToolMain.cpp
+++ b/Engine/Source/Programs/BuildPatchTool/Private/BuildPatchToolMain.cpp
@@ -549,7 +549,7 @@ int32 BuildPatchToolMain( const TCHAR* CommandLine )
 
 INT32_MAIN_INT32_ARGC_TCHAR_ARGV()
 {
-	FString CommandLine;
+	FString CommandLine = TEXT("-usehyperthreading");
 	for( int32 Option = 1; Option < ArgC; Option++ )
 	{
 		CommandLine += TEXT(" ");
diff --git a/Engine/Source/Runtime/Online/BuildPatchServices/Private/Generation/StatsCollector.cpp b/Engine/Source/Runtime/Online/BuildPatchServices/Private/Generation/StatsCollector.cpp
--- a/Engine/Source/Runtime/Online/BuildPatchServices/Private/Generation/StatsCollector.cpp
+++ b/Engine/Source/Runtime/Online/BuildPatchServices/Private/Generation/StatsCollector.cpp
@@ -1,9 +1,14 @@
 // Copyright 1998-2015 Epic Games, Inc. All Rights Reserved.
+#pragma once
 
 #include "BuildPatchServicesPrivatePCH.h"
 
 #include "StatsCollector.h"
 
+inline uint32 GetTypeHash(volatile const int64* A)
+{
+	return PointerHash((const void*)A);
+}
 
 namespace FBuildPatchServices
 {
@@ -52,7 +57,7 @@ namespace FBuildPatchServices
 	{
 	public:
 		FStatsCollectorImpl();
-		virtual ~FStatsCollectorImpl();
+		~FStatsCollectorImpl();
 
 		virtual volatile int64* CreateStat(const FString& Name, EStatFormat Type, int64 InitialValue = 0) override;
 		virtual void LogStats(float TimeBetweenLogs = 0.0f) override;
@@ -60,8 +65,8 @@ namespace FBuildPatchServices
 	private:
 		FCriticalSection DataCS;
 		TMap<FString, volatile int64*> AddedStats;
-		TMap<int64*, FString> AtomicNameMap;
-		TMap<int64*, EStatFormat> AtomicFormatMap;
+		TMap<volatile int64*, FString> AtomicNameMap;
+		TMap<volatile int64*, EStatFormat> AtomicFormatMap;
 		uint64 LastLogged;
 		int32 LongestName;
 	};
@@ -91,8 +96,8 @@ namespace FBuildPatchServices
 		{
 			volatile int64* NewInteger = new volatile int64(InitialValue);
 			AddedStats.Add(Name, NewInteger);
-			AtomicNameMap.Add((int64*)NewInteger, Name + TEXT(": "));
-			AtomicFormatMap.Add((int64*)NewInteger, Type);
+			AtomicNameMap.Add(NewInteger, Name + TEXT(": "));
+			AtomicFormatMap.Add(NewInteger, Type);
 			LongestName = FMath::Max<uint32>(LongestName, Name.Len() + 2);
 			for (auto& Stat : AtomicNameMap)
 			{
@@ -113,22 +118,21 @@ namespace FBuildPatchServices
 			LastLogged = Cycles;
 			FScopeLock ScopeLock(&DataCS);
 			GLog->Log(TEXT("/-------- FStatsCollector Log ---------------------"));
-			for (auto& Stat : AddedStats)
+			for (auto& Stat : AtomicNameMap)
 			{
-				int64* NameLookup = (int64*)Stat.Value;
-				switch (AtomicFormatMap[NameLookup])
+				switch (AtomicFormatMap[Stat.Key])
 				{
 				case EStatFormat::Timer:
-					GLog->Logf(TEXT("| %s%s"), *AtomicNameMap[NameLookup], *FPlatformTime::PrettyTime(FStatsCollector::CyclesToSeconds(*Stat.Value)));
+					GLog->Logf(TEXT("| %s%s"), *Stat.Value, *FPlatformTime::PrettyTime(FStatsCollector::CyclesToSeconds(*Stat.Key)));
 					break;
 				case EStatFormat::DataSize:
-					GLog->Logf(TEXT("| %s%s"), *AtomicNameMap[NameLookup], *FText::AsMemory(*Stat.Value).ToString());
+					GLog->Logf(TEXT("| %s%s"), *Stat.Value, *FText::AsMemory(*Stat.Key).ToString());
 					break;
 				case EStatFormat::DataSpeed:
-					GLog->Logf(TEXT("| %s%s/s"), *AtomicNameMap[NameLookup], *FText::AsMemory(*Stat.Value).ToString());
+					GLog->Logf(TEXT("| %s%s/s"), *Stat.Value, *FText::AsMemory(*Stat.Key).ToString());
 					break;
 				default:
-					GLog->Logf(TEXT("| %s%lld"), *AtomicNameMap[NameLookup], *Stat.Value);
+					GLog->Logf(TEXT("| %s%lld"), *Stat.Value, *Stat.Key);
 					break;
 				}
 			}
