Back out changelist 2579805

diff --git a/Engine/Source/Programs/BuildPatchTool/Private/BuildPatchToolMain.cpp b/Engine/Source/Programs/BuildPatchTool/Private/BuildPatchToolMain.cpp
--- a/Engine/Source/Programs/BuildPatchTool/Private/BuildPatchToolMain.cpp
+++ b/Engine/Source/Programs/BuildPatchTool/Private/BuildPatchToolMain.cpp
@@ -83,6 +83,7 @@
 #include "RequiredProgramMainCPPInclude.h"
 
 #include "BuildPatchServices.h"
+
 #include "CoreUObject.h"
 
 // Ensure compiled with patch generation code
@@ -142,8 +143,6 @@ namespace EBuildPatchToolMode
 int32 BuildPatchToolMain( const TCHAR* CommandLine )
 {
 	GEngineLoop.PreInit(CommandLine);
-	GConfig->InitializeConfigSystem();
-	IFileManager::Get().ProcessCommandLineOptions();
 
 	// Add log device
 	if (FParse::Param(FCommandLine::Get(), TEXT("stdout")))
@@ -373,6 +372,7 @@ int32 BuildPatchToolMain( const TCHAR* CommandLine )
 		{
 			// Initialize the configuration system, we can only do this reliably if we have CloudDirectory (i.e. bSuccess is true)
 			IniFile = CloudDirectory / TEXT("BuildPatchTool.ini");
+			GConfig->InitializeConfigSystem();
 
 			if (DataAgeThresholdIdx != INDEX_NONE)
 			{
@@ -406,6 +406,9 @@ int32 BuildPatchToolMain( const TCHAR* CommandLine )
 		}
 	}
 
+	// Initialize the file manager
+	IFileManager::Get().ProcessCommandLineOptions();
+
 	// Check for argument error
 	if( !bSuccess )
 	{
@@ -549,7 +552,7 @@ int32 BuildPatchToolMain( const TCHAR* CommandLine )
 
 INT32_MAIN_INT32_ARGC_TCHAR_ARGV()
 {
-	FString CommandLine = TEXT("-usehyperthreading");
+	FString CommandLine;
 	for( int32 Option = 1; Option < ArgC; Option++ )
 	{
 		CommandLine += TEXT(" ");
diff --git a/Engine/Source/Runtime/Online/BuildPatchServices/BuildPatchServices.Build.cs b/Engine/Source/Runtime/Online/BuildPatchServices/BuildPatchServices.Build.cs
--- a/Engine/Source/Runtime/Online/BuildPatchServices/BuildPatchServices.Build.cs
+++ b/Engine/Source/Runtime/Online/BuildPatchServices/BuildPatchServices.Build.cs
@@ -17,7 +17,7 @@ public class BuildPatchServices : ModuleRules
 		PrivateIncludePaths.AddRange(
 			new string[]
 			{
-				"Runtime/Online/BuildPatchServices/Private",
+				"Developer/NoRedist/BuildPatchServices/Private",
 			}
 		);
 
diff --git a/Engine/Source/Runtime/Online/BuildPatchServices/Private/BuildPatchChunk.cpp b/Engine/Source/Runtime/Online/BuildPatchServices/Private/BuildPatchChunk.cpp
--- a/Engine/Source/Runtime/Online/BuildPatchServices/Private/BuildPatchChunk.cpp
+++ b/Engine/Source/Runtime/Online/BuildPatchServices/Private/BuildPatchChunk.cpp
@@ -361,6 +361,7 @@ void FChunkWriter::FQueuedChunkWriter::SetNoMoreChunks()
 void FChunkWriter::FQueuedChunkWriter::GetChunkFilesizes(TMap<FGuid, int64>& OutChunkFileSizes)
 {
 	FScopeLock ScopeLock(&ChunkFileSizesCS);
+	OutChunkFileSizes.Empty(ChunkFileSizes.Num());
 	OutChunkFileSizes.Append(ChunkFileSizes);
 }
 
diff --git a/Engine/Source/Runtime/Online/BuildPatchServices/Private/BuildPatchGeneration.cpp b/Engine/Source/Runtime/Online/BuildPatchServices/Private/BuildPatchGeneration.cpp
--- a/Engine/Source/Runtime/Online/BuildPatchServices/Private/BuildPatchGeneration.cpp
+++ b/Engine/Source/Runtime/Online/BuildPatchServices/Private/BuildPatchGeneration.cpp
@@ -7,14 +7,59 @@
 
 #include "BuildPatchServicesPrivatePCH.h"
 
-#include "Generation/DataScanner.h"
-#include "Generation/BuildStreamer.h"
-#include "Generation/CloudEnumeration.h"
-#include "Generation/ManifestBuilder.h"
+#if WITH_BUILDPATCHGENERATION
 
-using namespace FBuildPatchServices;
+namespace FBuildPatchTimers
+{
+	uint64 DataProcessed = 0;
+	volatile int64 AccumOpening = 0;
+	volatile int64 AccumReading = 0;
+	volatile int64 AccumEnqueueing = 0;
+	volatile int64 OtherTime = 0;
+	volatile int64 DataAccessSpeed = 0;
+	volatile int64 DataThroughputSpeed = 0;
+	volatile int64 TotalDataRead = 0;
+	uint32 TotalNumNewChunks = 0;
+	uint32 TotalNumKnownChunks = 0;
+	uint64 TimeStarted = 0;
+
+	uint64 GetCycles()
+	{
+#if PLATFORM_WINDOWS
+		LARGE_INTEGER Cycles;
+		QueryPerformanceCounter(&Cycles);
+		return Cycles.QuadPart;
+#elif PLATFORM_MAC
+		uint64 Cycles = mach_absolute_time();
+		return Cycles;
+#else
+		return FPlatformTime::Cycles();
+#endif
+	}
 
-#if WITH_BUILDPATCHGENERATION
+	static double CyclesToSeconds( const uint64 Cycles )
+	{
+		return FPlatformTime::GetSecondsPerCycle() * Cycles;
+	}
+
+	void LogInfo(bool bForceLog = false)
+	{
+		// Log processed data
+		static uint64 LastLog = FBuildPatchTimers::GetCycles();
+		if (bForceLog || (FBuildPatchTimers::CyclesToSeconds((FBuildPatchTimers::GetCycles() - LastLog)) > 5.0f))
+		{
+			LastLog = FBuildPatchTimers::GetCycles();
+			GLog->Logf(TEXT("FGenerationInfo: Stream: ReadTime:  %s"), *FPlatformTime::PrettyTime(FBuildPatchTimers::CyclesToSeconds(AccumReading)));
+			GLog->Logf(TEXT("FGenerationInfo: Stream: OpenTime:  %s"), *FPlatformTime::PrettyTime(FBuildPatchTimers::CyclesToSeconds(AccumOpening)));
+			GLog->Logf(TEXT("FGenerationInfo: DataAccessSpeed:   %s/s"), *FText::AsMemory(DataAccessSpeed).ToString());
+			GLog->Logf(TEXT("FGenerationInfo: DataThroughput:    %s/s"), *FText::AsMemory(DataThroughputSpeed).ToString());
+			GLog->Logf(TEXT("FGenerationInfo: DataProcessed:     %s"), *FText::AsMemory(DataProcessed).ToString());
+			GLog->Logf(TEXT("FGenerationInfo: TotalTime:         %s"), *FPlatformTime::PrettyTime(FBuildPatchTimers::CyclesToSeconds(FBuildPatchTimers::GetCycles() - TimeStarted)));
+			GLog->Logf(TEXT("FGenerationInfo: NewChunks: %u"), TotalNumNewChunks);
+			GLog->Logf(TEXT("FGenerationInfo: OldChunks: %u\n"), TotalNumKnownChunks);
+		}
+	}
+};
 
 FFileAttributes::FFileAttributes()
 	: bReadOnly(false)
@@ -57,6 +102,599 @@ static FString GetSymlinkTarget(const TCHAR* Filename)
 	return TEXT("");
 }
 
+/* FBuildStreamReader implementation
+*****************************************************************************/
+FBuildStream::FBuildStreamReader::FBuildStreamReader()
+	: DepotDirectory(TEXT(""))
+	, BuildStream(NULL)
+	, Thread(NULL)
+{
+}
+
+FBuildStream::FBuildStreamReader::~FBuildStreamReader()
+{
+	if (Thread)
+	{
+		Thread->Kill();
+		delete Thread;
+		Thread = NULL;
+	}
+}
+
+bool FBuildStream::FBuildStreamReader::Init() 
+{
+	return BuildStream != NULL && IFileManager::Get().DirectoryExists( *DepotDirectory );
+}
+
+uint32 FBuildStream::FBuildStreamReader::Run()
+{
+	IFileManager* FileManager = &IFileManager::Get();
+
+	uint64 StartTime = FBuildPatchTimers::GetCycles();
+	// Clear the build stream
+	BuildStream->Clear();
+
+	TArray< FString > AllFiles;
+	uint64 FileEnumerationStart = FBuildPatchTimers::GetCycles();
+	FileManager->FindFilesRecursive(AllFiles, *DepotDirectory, TEXT("*.*"), true, false);
+	uint64 FileEnumerationEnd = FBuildPatchTimers::GetCycles();
+	uint64 FileEnumerationTime = FileEnumerationEnd - FileEnumerationStart;
+	GLog->Logf(TEXT("FBuildStreamReader: Enumerated %d files in %s"), AllFiles.Num(), *FPlatformTime::PrettyTime(FBuildPatchTimers::CyclesToSeconds(FileEnumerationTime)));
+	AllFiles.Sort();
+
+	// Remove the files that appear in an ignore list
+	FBuildDataGenerator::StripIgnoredFiles( AllFiles, DepotDirectory, IgnoreListFile );
+
+	// Allocate our file read buffer
+	uint8* FileReadBuffer = new uint8[ FileBufferSize ];
+
+
+	uint64 LastLogged = FBuildPatchTimers::GetCycles();
+	uint64 ProcessBegin = FBuildPatchTimers::GetCycles();
+	for (auto& SourceFile : AllFiles)
+	{
+		// Read the file
+		uint64 FileOpenStart = FBuildPatchTimers::GetCycles();
+		FArchive* FileReader = FileManager->CreateFileReader(*SourceFile);
+		uint64 FileOpenEnd = FBuildPatchTimers::GetCycles();
+		FPlatformAtomics::InterlockedAdd(&FBuildPatchTimers::AccumOpening, FileOpenEnd - FileOpenStart);
+		const bool bBuildFileOpenSuccess = FileReader != NULL;
+		if( bBuildFileOpenSuccess )
+		{
+			// Make SourceFile the format we want it in and start a new file
+			FPaths::MakePathRelativeTo( SourceFile, *( DepotDirectory + TEXT( "/" ) ) );
+			int64 FileSize = FileReader->TotalSize();
+			// Process files that have bytes
+			if( FileSize > 0 )
+			{
+				BuildStream->BeginNewFile( SourceFile, FileSize );
+				while( !FileReader->AtEnd() )
+				{
+					const int64 SizeLeft = FileSize - FileReader->Tell();
+					const uint32 ReadLen = FMath::Min< int64 >(FileBufferSize, SizeLeft);
+					uint64 FileReadStart = FBuildPatchTimers::GetCycles();
+					FileReader->Serialize(FileReadBuffer, ReadLen);
+					uint64 FileReadEnd = FBuildPatchTimers::GetCycles();
+					FPlatformAtomics::InterlockedAdd(&FBuildPatchTimers::AccumReading, FileReadEnd - FileReadStart);
+					FPlatformAtomics::InterlockedAdd(&FBuildPatchTimers::TotalDataRead, ReadLen);
+
+					// Copy into data stream
+					uint64 EnqueueDataStart = FBuildPatchTimers::GetCycles();
+					BuildStream->EnqueueData(FileReadBuffer, ReadLen);
+					uint64 EnqueueDataEnd = FBuildPatchTimers::GetCycles();
+					FPlatformAtomics::InterlockedAdd(&FBuildPatchTimers::AccumEnqueueing, EnqueueDataEnd - EnqueueDataStart);
+
+					uint64 DataAccess = FBuildPatchTimers::AccumOpening + FBuildPatchTimers::AccumReading;
+					FPlatformAtomics::InterlockedExchange(&FBuildPatchTimers::OtherTime, (FBuildPatchTimers::GetCycles() - StartTime) - DataAccess - FBuildPatchTimers::AccumEnqueueing);
+					double DataAccessTimeFloat = FBuildPatchTimers::CyclesToSeconds(DataAccess);
+					FPlatformAtomics::InterlockedExchange(&FBuildPatchTimers::DataAccessSpeed, FBuildPatchTimers::TotalDataRead / DataAccessTimeFloat);
+					double DataThroughputTimeFloat = FBuildPatchTimers::CyclesToSeconds(FBuildPatchTimers::GetCycles() - ProcessBegin);
+					FPlatformAtomics::InterlockedExchange(&FBuildPatchTimers::DataThroughputSpeed, FBuildPatchTimers::TotalDataRead / DataThroughputTimeFloat);
+
+				}
+			}
+			// Special case zero byte files
+			else if( FileSize == 0 )
+			{
+				BuildStream->AddEmptyFile( SourceFile );
+			}
+			FileReader->Close();
+			delete FileReader;
+		}
+		else
+		{
+			// Not being able to load a required file from the build would be fatal, hard fault.
+			GLog->Logf( TEXT( "FBuildStreamReader: Could not open file from build! %s" ), *SourceFile );
+			GLog->PanicFlushThreadedLogs();
+			// Use bool variable for easier to understand assert message
+			check( bBuildFileOpenSuccess );
+		}
+	}
+
+	// Mark end of build
+	BuildStream->EndOfBuild();
+
+	// Deallocate our file read buffer
+	delete[] FileReadBuffer;
+
+	return 0;
+}
+
+void FBuildStream::FBuildStreamReader::StartThread()
+{
+	Thread = FRunnableThread::Create(this, TEXT("BuildStreamReaderThread"));
+}
+
+/* FBuildStream implementation
+*****************************************************************************/
+FBuildStream::FBuildStream( const FString& RootDirectory, const FString& IgnoreListFile )
+{
+	BuildStreamReader.DepotDirectory = RootDirectory;
+	BuildStreamReader.IgnoreListFile = IgnoreListFile;
+	BuildStreamReader.BuildStream = this;
+	BuildStreamReader.StartThread();
+}
+
+FBuildStream::~FBuildStream()
+{
+}
+
+bool FBuildStream::GetFileSpan( const uint64& StartingIdx, FString& Filename, uint64& FileSize )
+{
+	bool bFound = false;
+	FilesListsCS.Lock();
+	// Find the filename
+	FFileSpan* FileSpan = FilesParsed.Find( StartingIdx );
+	if( FileSpan != NULL )
+	{
+		Filename = FileSpan->Filename;
+		FileSize = FileSpan->Size;
+		bFound = true;
+	}
+	FilesListsCS.Unlock();
+	return bFound;
+}
+
+uint32 FBuildStream::DequeueData( uint8* Buffer, const uint32& ReqSize, const bool WaitForData )
+{
+	// Wait for data
+	if( WaitForData )
+	{
+		while ( DataAvailable() < ReqSize && !IsEndOfBuild() )
+		{
+			FPlatformProcess::Sleep( 0.01f );
+		}
+	}
+
+	BuildDataStreamCS.Lock();
+	uint32 ReadLen = FMath::Min< int64 >( ReqSize, DataAvailable() );
+	ReadLen = BuildDataStream.Dequeue( Buffer, ReadLen );
+	BuildDataStreamCS.Unlock();
+
+	return ReadLen;
+}
+
+const TArray< FString > FBuildStream::GetEmptyFiles()
+{
+	FScopeLock ScopeLock( &FilesListsCS );
+	return EmptyFiles;
+}
+
+bool FBuildStream::IsEndOfBuild()
+{
+	bool rtn;
+	NoMoreDataCS.Lock();
+	rtn = bNoMoreData;
+	NoMoreDataCS.Unlock();
+	return rtn;
+}
+
+bool FBuildStream::IsEndOfData()
+{
+	bool rtn;
+	NoMoreDataCS.Lock();
+	rtn = bNoMoreData;
+	NoMoreDataCS.Unlock();
+	BuildDataStreamCS.Lock();
+	rtn &= BuildDataStream.RingDataUsage() == 0;
+	BuildDataStreamCS.Unlock();
+	return rtn;
+}
+
+void FBuildStream::Clear()
+{
+	EndOfBuild( false );
+
+	BuildDataStreamCS.Lock();
+	BuildDataStream.Empty();
+	BuildDataStreamCS.Unlock();
+
+	FilesListsCS.Lock();
+	FilesParsed.Empty();
+	FilesListsCS.Unlock();
+}
+
+uint32 FBuildStream::SpaceLeft()
+{
+	uint32 rtn;
+	BuildDataStreamCS.Lock();
+	rtn = BuildDataStream.RingDataSize() - BuildDataStream.RingDataUsage();
+	BuildDataStreamCS.Unlock();
+	return rtn;
+}
+
+uint32 FBuildStream::DataAvailable()
+{
+	uint32 rtn;
+	BuildDataStreamCS.Lock();
+	rtn = BuildDataStream.RingDataUsage();
+	BuildDataStreamCS.Unlock();
+	return rtn;
+}
+
+void FBuildStream::BeginNewFile( const FString& Filename, const uint64& FileSize )
+{
+	FFileSpan FileSpan;
+	FileSpan.Filename = Filename;
+	FileSpan.Size = FileSize;
+	BuildDataStreamCS.Lock();
+	FileSpan.StartIdx = BuildDataStream.TotalDataPushed();
+	BuildDataStreamCS.Unlock();
+	FilesListsCS.Lock();
+	FilesParsed.Add( FileSpan.StartIdx, FileSpan );
+	FilesListsCS.Unlock();
+}
+
+void FBuildStream::AddEmptyFile( const FString& Filename )
+{
+	FilesListsCS.Lock();
+	EmptyFiles.Add( Filename );
+	FilesListsCS.Unlock();
+}
+
+void FBuildStream::EnqueueData( const uint8* Buffer, const uint32& Len )
+{
+	// Wait for space
+	while ( SpaceLeft() < Len )
+	{
+		FPlatformProcess::Sleep( 0.01f );
+	}
+	BuildDataStreamCS.Lock();
+	BuildDataStream.Enqueue( Buffer, Len );
+	BuildDataStreamCS.Unlock();
+}
+
+void FBuildStream::EndOfBuild( bool bIsEnd )
+{
+	NoMoreDataCS.Lock();
+	bNoMoreData = bIsEnd;
+	NoMoreDataCS.Unlock();
+}
+
+/* FBuildDataProcessor implementation
+*****************************************************************************/
+FBuildDataChunkProcessor::FBuildDataChunkProcessor( FBuildPatchAppManifestRef InBuildManifest, const FString& InBuildRoot )
+	: NumNewChunks( 0 )
+	, NumKnownChunks( 0 )
+	, BuildRoot( InBuildRoot )
+	, ChunkWriter( FBuildPatchServicesModule::GetCloudDirectory() )
+	, CurrentChunkBufferPos( 0 )
+	, CurrentFile(nullptr)
+	, IsProcessingChunk( false )
+	, IsProcessingChunkPart( false )
+	, IsProcessingFile( false )
+	, ChunkIsPushed( false )
+	, BackupChunkBufferPos( 0 )
+	, BackupProcessingChunk( false )
+	, BackupProcessingChunkPart( false )
+{
+	BuildManifest = InBuildManifest;
+	CurrentChunkGuid.Invalidate();
+	BackupChunkGuid.Invalidate();
+	CurrentChunkBuffer = new uint8[ FBuildPatchData::ChunkDataSize ];
+}
+
+FBuildDataChunkProcessor::~FBuildDataChunkProcessor()
+{
+	delete[] CurrentChunkBuffer;
+}
+
+void FBuildDataChunkProcessor::BeginNewChunk( const bool& bZeroBuffer )
+{
+	check( IsProcessingChunk == false );
+	check( IsProcessingChunkPart == false );
+	IsProcessingChunk = true;
+
+	// NB: If you change this function, make sure you make required changes to 
+	// Backup of chunk if need be at the end of a recognized chunk
+
+	// Erase the chunk buffer data ready for new chunk parts
+	if( bZeroBuffer )
+	{
+		FMemory::Memzero( CurrentChunkBuffer, FBuildPatchData::ChunkDataSize );
+		CurrentChunkBufferPos = 0;
+	}
+
+	// Get a new GUID for this chunk for identification
+	CurrentChunkGuid = FGuid::NewGuid();
+}
+
+void FBuildDataChunkProcessor::BeginNewChunkPart()
+{
+	check( IsProcessingChunkPart == false );
+	check( IsProcessingChunk == true );
+	IsProcessingChunkPart = true;
+
+	// The current file should have a new chunk part setup
+	if (CurrentFile)
+	{
+		CurrentFile->FileChunkParts.Add(FChunkPartData());
+		FChunkPartData& NewPart = CurrentFile->FileChunkParts.Last();
+		NewPart.Guid = CurrentChunkGuid;
+		NewPart.Offset = CurrentChunkBufferPos;
+	}
+}
+
+void FBuildDataChunkProcessor::EndNewChunkPart()
+{
+	check( IsProcessingChunkPart == true );
+	check( IsProcessingChunk == true );
+	IsProcessingChunkPart = false;
+
+	// Current file should have it's last chunk part updated with the size value finalized.
+	if (CurrentFile)
+	{
+		FChunkPartData* ChunkPart = &CurrentFile->FileChunkParts.Last();
+		ChunkPart->Size = CurrentChunkBufferPos - ChunkPart->Offset;
+		check( ChunkPart->Size != 0 );
+	}
+}
+
+void FBuildDataChunkProcessor::EndNewChunk( const uint64& ChunkHash, const uint8* ChunkData, const FGuid& ChunkGuid )
+{
+	check( IsProcessingChunk == true );
+	check( IsProcessingChunkPart == false );
+	IsProcessingChunk = false;
+
+	// A bool that will state whether we should update and log chunking progress
+	bool bLogProgress = false;
+
+	// If the new chunk was recognized, then we will have got a different Guid through, so fix up guid for files using this chunk
+	if( CurrentChunkGuid != ChunkGuid )
+	{
+		bLogProgress = true;
+		++NumKnownChunks;
+		for (auto& FileManifest : BuildManifest->Data->FileManifestList)
+		{
+			for (auto& ChunkPart : FileManifest.FileChunkParts)
+			{
+				if( ChunkPart.Guid == CurrentChunkGuid )
+				{
+					ChunkPart.Guid = ChunkGuid;
+				}
+			}
+		}
+	}
+	// If the chunk is new, count it if we were passed data
+	else if( ChunkData )
+	{
+		bLogProgress = true;
+		++NumNewChunks;
+	}
+
+	// Always queue the chunk if passed data, it will be skipped automatically if existing already, and we need to make
+	// sure the latest version is saved out when recognizing an older version
+	if( ChunkData )
+	{
+		ChunkWriter.QueueChunk( ChunkData, ChunkGuid, ChunkHash );
+
+		// Also add the info to the data
+		if (!BuildManifest->ChunkInfoLookup.Contains(ChunkGuid))
+		{
+			BuildManifest->Data->ChunkList.Add(FChunkInfoData());
+			FChunkInfoData& NewChunk = BuildManifest->Data->ChunkList.Last();
+			NewChunk.Guid = ChunkGuid;
+			NewChunk.Hash = ChunkHash;
+			NewChunk.FileSize = INDEX_NONE;
+			NewChunk.GroupNumber = FCrc::MemCrc32(&ChunkGuid, sizeof(FGuid)) % 100;
+			// Lookup is used just to ensure one copy of each chunk added, we must not use the ptr as array will realloc
+			BuildManifest->ChunkInfoLookup.Add(ChunkGuid, nullptr);
+		}
+
+	}
+
+
+	FBuildPatchTimers::TotalNumNewChunks = NumNewChunks;
+	FBuildPatchTimers::TotalNumKnownChunks = NumKnownChunks;
+}
+
+void FBuildDataChunkProcessor::PushChunk()
+{
+	check( !ChunkIsPushed );
+	ChunkIsPushed = true;
+
+	// Then we can skip over this entire chunk using it for current files in processing
+	BackupProcessingChunk = IsProcessingChunk;
+	BackupProcessingChunkPart = IsProcessingChunkPart;
+	BackupChunkGuid = CurrentChunkGuid;
+	BackupChunkBufferPos = CurrentChunkBufferPos;
+	if( BackupProcessingChunkPart )
+	{
+		EndNewChunkPart();
+	}
+	if( BackupProcessingChunk )
+	{
+		// null ChunkData will stop the incomplete chunk being saved out
+		EndNewChunk( 0, NULL, CurrentChunkGuid );
+	}
+
+	// Start this chunk, must begin from 0
+	CurrentChunkBufferPos = 0;
+	BeginNewChunk( false );
+	BeginNewChunkPart();
+}
+
+void FBuildDataChunkProcessor::PopChunk( const uint64& ChunkHash, const uint8* ChunkData, const FGuid& ChunkGuid )
+{
+	check( ChunkIsPushed );
+	ChunkIsPushed = false;
+
+	EndNewChunkPart();
+	EndNewChunk( ChunkHash, ChunkData, ChunkGuid );
+
+	// Backup data for previous part chunk
+	IsProcessingChunk = BackupProcessingChunk;
+	CurrentChunkGuid = BackupChunkGuid;
+	CurrentChunkBufferPos = BackupChunkBufferPos;
+}
+
+void FBuildDataChunkProcessor::FinalChunk()
+{
+	// If we're processing a chunk part, then we MUST be processing a chunk too
+	check(IsProcessingChunkPart ? IsProcessingChunk : true);
+
+	if( IsProcessingChunk )
+	{
+		// The there is no more data and the last file is processed.
+		// The final chunk should be finished
+		uint64 NewChunkHash = FRollingHash< FBuildPatchData::ChunkDataSize >::GetHashForDataSet( CurrentChunkBuffer );
+		FGuid ChunkGuid = CurrentChunkGuid;
+		FBuildDataGenerator::FindExistingChunkData( NewChunkHash, CurrentChunkBuffer, ChunkGuid );
+		if( IsProcessingChunkPart )
+		{
+			EndNewChunkPart();
+		}
+		EndNewChunk( NewChunkHash, CurrentChunkBuffer, ChunkGuid );
+	}
+
+	// Wait for the chunk writer to finish
+	ChunkWriter.NoMoreChunks();
+	ChunkWriter.WaitForThread();
+}
+
+void FBuildDataChunkProcessor::BeginFile( const FString& FileName )
+{
+	check( IsProcessingFile == false );
+	IsProcessingFile = true;
+
+	// We should start a new file, which begins with the current new chunk and current chunk offset
+	BuildManifest->Data->FileManifestList.Add(FFileManifestData());
+	CurrentFile = &BuildManifest->Data->FileManifestList.Last();
+	CurrentFile->Filename = FileName;
+	CurrentFile->bIsUnixExecutable = IsUnixExecutable(*(BuildRoot / FileName));
+	CurrentFile->SymlinkTarget = GetSymlinkTarget(*(BuildRoot / FileName));
+
+	// Setup for current chunk part
+	if( IsProcessingChunkPart )
+	{
+		CurrentFile->FileChunkParts.Add(FChunkPartData());
+		FChunkPartData& NewPart = CurrentFile->FileChunkParts.Last();
+		NewPart.Guid = CurrentChunkGuid;
+		NewPart.Offset = CurrentChunkBufferPos;
+	}
+}
+
+void FBuildDataChunkProcessor::EndFile()
+{
+	check( IsProcessingFile == true );
+	IsProcessingFile = false;
+
+	if (CurrentFile)
+	{
+		if( IsProcessingChunkPart )
+		{
+			// Current file should have it's last chunk part updated with the size value finalized.
+			FChunkPartData* ChunkPart = &CurrentFile->FileChunkParts.Last();
+			ChunkPart->Size = CurrentChunkBufferPos - ChunkPart->Offset;
+		}
+		// Check file size is correct from chunks
+		int64 FileSystemSize = IFileManager::Get().FileSize(*(BuildRoot / CurrentFile->Filename));
+		CurrentFile->Init();
+		check(CurrentFile->GetFileSize() == FileSystemSize);
+	}
+
+	CurrentFile = nullptr;
+}
+
+void FBuildDataChunkProcessor::SkipKnownByte( const uint8& NextByte, const bool& bStartOfFile, const bool& bEndOfFile, const FString& Filename )
+{
+	// Check for start of new file
+	if( bStartOfFile )
+	{
+		BeginFile( Filename );
+		FileHash.Reset();
+	}
+
+	// Increment position between start and end of file
+	++CurrentChunkBufferPos;
+	FileHash.Update( &NextByte, 1 );
+
+	// Check for end of file
+	if( bEndOfFile )
+	{
+		FileHash.Final();
+		FileHash.GetHash( CurrentFile->FileHash.Hash );
+		EndFile();
+	}
+}
+
+void FBuildDataChunkProcessor::ProcessNewByte( const uint8& NewByte, const bool& bStartOfFile, const bool& bEndOfFile, const FString& Filename )
+{
+	// If we finished a chunk, we begin a new chunk and new chunk part
+	if( !IsProcessingChunk )
+	{
+		BeginNewChunk();
+		BeginNewChunkPart();
+	}
+	// Or if we finished a chunk part (by recognizing a chunk hash), we begin a new one
+	else if( !IsProcessingChunkPart )
+	{
+		BeginNewChunkPart();
+	}
+
+	// Check for start of new file
+	if( bStartOfFile )
+	{
+		BeginFile( Filename );
+		FileHash.Reset();
+	}
+
+	// We add the old byte to our new chunk which will be used as part of the file it belonged to
+	CurrentChunkBuffer[ CurrentChunkBufferPos++ ] = NewByte;
+	FileHash.Update( &NewByte, 1 );
+
+	// Check for end of file
+	if( bEndOfFile )
+	{
+		FileHash.Final();
+		FileHash.GetHash( CurrentFile->FileHash.Hash );
+		EndFile();
+	}
+
+	// Do we have a full new chunk?
+	check( CurrentChunkBufferPos <= FBuildPatchData::ChunkDataSize );
+	if( CurrentChunkBufferPos == FBuildPatchData::ChunkDataSize )
+	{
+		uint64 NewChunkHash = FRollingHash< FBuildPatchData::ChunkDataSize >::GetHashForDataSet( CurrentChunkBuffer );
+		FGuid ChunkGuid = CurrentChunkGuid;
+		FBuildDataGenerator::FindExistingChunkData( NewChunkHash, CurrentChunkBuffer, ChunkGuid );
+		EndNewChunkPart();
+		EndNewChunk( NewChunkHash, CurrentChunkBuffer, ChunkGuid );
+	}
+}
+
+void FBuildDataChunkProcessor::GetChunkStats( uint32& OutNewFiles, uint32& OutKnownFiles )
+{
+	OutNewFiles = NumNewChunks;
+	OutKnownFiles = NumKnownChunks;
+}
+
+const TMap<FGuid, int64>& FBuildDataChunkProcessor::GetChunkFilesizes()
+{
+	ChunkWriter.GetChunkFilesizes(ChunkFileSizes);
+	return ChunkFileSizes;
+}
+
 /* FBuildDataFileProcessor implementation
 *****************************************************************************/
 FBuildDataFileProcessor::FBuildDataFileProcessor(FBuildPatchAppManifestRef InBuildManifest, const FString& InBuildRoot, const FDateTime& InDataThresholdTime)
@@ -432,6 +1070,8 @@ static void AddCustomFieldsToBuildManifest(const TMap<FString, FVariant>& Custom
 
 static void FileAttributesMetaToMap(const FString& AttributesList, TMap<FString, FFileAttributes>& FileAttributesMap)
 {
+	GLog->Logf(TEXT("Parsing file attributes list:-"));
+	checkf(AttributesList.Len() > 0, TEXT("Attributes File List was empty file"));
 	const TCHAR Quote = TEXT('\"');
 	const TCHAR EOFile = TEXT('\0');
 	const TCHAR EOLine = TEXT('\n');
@@ -458,80 +1098,105 @@ static void FileAttributesMetaToMap(const FString& AttributesList, TMap<FString,
 		FString Filename = FString(FilenameEnd - FilenameStart, FilenameStart).Replace(TEXT("\\"), TEXT("/"));
 		FString Keywords(EndOfLine - FilenameEnd, FilenameEnd);
 		FFileAttributes FileAttributes;
+		GLog->Logf(TEXT("    %s"), *Filename);
 		if (Keywords.Contains(TEXT("readonly")))
 		{
 			FileAttributes.bReadOnly = true;
+			GLog->Logf(TEXT("        readonly"), *Filename);
 		}
 		if (Keywords.Contains(TEXT("compressed")))
 		{
 			FileAttributes.bCompressed = true;
+			GLog->Logf(TEXT("        compressed"), *Filename);
 		}
 		if (Keywords.Contains(TEXT("executable")))
 		{
 			FileAttributes.bUnixExecutable = true;
+			GLog->Logf(TEXT("        executable"), *Filename);
+		}
+		if (!(FileAttributes.bReadOnly || FileAttributes.bCompressed || FileAttributes.bUnixExecutable))
+		{
+			GLog->Logf(TEXT("        none"), *Filename);
 		}
 		FileAttributesMap.Add(MoveTemp(Filename), FileAttributes);
 	}
+
 }
 
 /* FBuildDataGenerator implementation
 *****************************************************************************/
 bool FBuildDataGenerator::GenerateChunksManifestFromDirectory( const FBuildPatchSettings& Settings )
 {
-	uint64 StartTime = FStatsCollector::GetCycles();
+	// Output to log for builder info
+	GLog->Logf(TEXT("Running Chunks Patch Generation for: %u:%s %s"), Settings.AppID, *Settings.AppName, *Settings.BuildVersion);
+	FBuildPatchTimers::TimeStarted = FBuildPatchTimers::GetCycles();
 
 	// Take the build CS
 	FScopeLock SingleConcurrentBuild( &SingleConcurrentBuildCS );
 
-	// Create stat collector
-	FStatsCollectorRef StatsCollector = FStatsCollectorFactory::Create();
-
-	// Create a build streamer
-	FBuildStreamerRef BuildStream = FBuildStreamerFactory::Create(Settings.RootDirectory, Settings.IgnoreListFile, StatsCollector);
-
-	// Enumerate Chunks
-	FCloudEnumerationRef CloudEnumeration = FCloudEnumerationFactory::Create(FBuildPatchServicesModule::GetCloudDirectory());
-
-	// Output to log for builder info
-	GLog->Logf(TEXT("Running Chunks Patch Generation for: %u:%s %s"), Settings.AppID, *Settings.AppName, *Settings.BuildVersion);
-
 	// Create our chunk cache
 	const FDateTime Cutoff = Settings.bShouldHonorReuseThreshold ? FDateTime::UtcNow() - FTimespan::FromDays(Settings.DataAgeThreshold) : FDateTime::MinValue();
 	FBuildGenerationChunkCache::Init(Cutoff);
 
-	// Create a manifest builder
-	FManifestDetails ManifestDetails;
-	ManifestDetails.bIsFileData = false;
-	ManifestDetails.AppId = Settings.AppID;
-	ManifestDetails.AppName = Settings.AppName;
-	ManifestDetails.BuildVersion = Settings.BuildVersion;
-	ManifestDetails.LaunchExe = Settings.LaunchExe;
-	ManifestDetails.LaunchCommand = Settings.LaunchCommand;
-	ManifestDetails.PrereqName = Settings.PrereqName;
-	ManifestDetails.PrereqPath = Settings.PrereqPath;
-	ManifestDetails.PrereqArgs = Settings.PrereqArgs;
-	ManifestDetails.CustomFields = Settings.CustomFields;
+	// Create a manifest
+	FBuildPatchAppManifestRef BuildManifest = MakeShareable( new FBuildPatchAppManifest() );
+
+	// Setup custom fields
+	AddCustomFieldsToBuildManifest(Settings.CustomFields, BuildManifest);
+
 	// Get the file attributes
 	FString AttributesList;
+	TMap<FString, FFileAttributes> FileAttributesMap;
 	if (Settings.AttributeListFile.Len() > 0)
 	{
 		FFileHelper::LoadFileToString(AttributesList, *Settings.AttributeListFile);
 		if (!AttributesList.IsEmpty())
 		{
-			FileAttributesMetaToMap(AttributesList, ManifestDetails.FileAttributesMap);
+			FileAttributesMetaToMap(AttributesList, FileAttributesMap);
 		}
 		else
 		{
 			GLog->Logf(TEXT("WARNING: Attributes list file empty"));
 		}
 	}
-	FManifestBuilderRef ManifestBuilder = FManifestBuilderFactory::Create(ManifestDetails, BuildStream);
 
-	// Chunks matching
-	FDataMatcherRef DataMatcher = FDataMatcherFactory::Create(FBuildPatchServicesModule::GetCloudDirectory());
+	// Reset chunk inventory
+	ExistingChunksEnumerated = false;
+	ExistingChunkGuidInventory.Empty();
+	ExistingChunkHashInventory.Empty();
+
+	// Declare a build processor
+	FBuildDataChunkProcessor DataProcessor(BuildManifest, Settings.RootDirectory);
+
+	// Create a build streamer
+	FBuildStream* BuildStream = new FBuildStream(Settings.RootDirectory, Settings.IgnoreListFile);
+
+	// Set the basic details
+	BuildManifest->Data->bIsFileData = false;
+	BuildManifest->Data->AppID = Settings.AppID;
+	BuildManifest->Data->AppName = Settings.AppName;
+	BuildManifest->Data->BuildVersion = Settings.BuildVersion;
+	BuildManifest->Data->LaunchExe = Settings.LaunchExe;
+	BuildManifest->Data->LaunchCommand = Settings.LaunchCommand;
+	BuildManifest->Data->PrereqName = Settings.PrereqName;
+	BuildManifest->Data->PrereqPath = Settings.PrereqPath;
+	BuildManifest->Data->PrereqArgs = Settings.PrereqArgs;
+
+	// Create a data buffer
+	const uint32 DataBufferSize = FBuildPatchData::ChunkDataSize;
+	uint8* DataBuffer = new uint8[ DataBufferSize ];
 
-	// Refers to how much data has been dequeued
-	//uint64 ProcessPos = 0;
+	// We'll need a rolling hash for chunking
+	FRollingHash< FBuildPatchData::ChunkDataSize >* RollingHash = new FRollingHash< FBuildPatchData::ChunkDataSize >();
+	
+	// Refers to how much data has been processed (into the FBuildDataProcessor)
+	uint64 ProcessPos = 0;
+
+	// Records the current file we are processing
+	FString FileName;
+
+	// And the current file's data left to process
+	uint64 FileDataCount = 0;
 
 	// Used to store data read lengths
 	uint32 ReadLen = 0;
@@ -540,64 +1205,173 @@ bool FBuildDataGenerator::GenerateChunksManifestFromDirectory( const FBuildPatch
 	double LastProgressLog = FPlatformTime::Seconds();
 	const double TimeGenStarted = LastProgressLog;
 
-	// 50MB Data buffer
-	const int32 DataBufferSize = 1024*1024*50;
-	const int32 OverlapSize = FBuildPatchData::ChunkDataSize - 1;
-	TArray<uint8> DataBuffer;
-	uint64 DataOffset = 0;
+	FBuildPatchTimers::LogInfo();
 
 	// Loop through all data
-	while (!BuildStream->IsEndOfData())
+	while ( !BuildStream->IsEndOfData() )
 	{
-		// Keep the overlap data from previous scanner
-		int32 PreviousSize = DataBuffer.Num();
-		if(PreviousSize > 0)
+		FBuildPatchTimers::LogInfo();
+
+		// Grab some data from the build stream
+		ReadLen = BuildStream->DequeueData( DataBuffer, DataBufferSize );
+
+		// A bool says if there's no more data to come from the Build Stream
+		const bool bNoMoreData = BuildStream->IsEndOfData();
+
+		// Refers to how much data from DataBuffer has been passed into the rolling hash
+		uint32 DataBufferPos = 0;
+
+		// Count how many times we pad the rolling hash with zero
+		uint32 PaddedZeros = 0;
+
+		// Process data while we have more
+		while ( ( DataBufferPos < ReadLen ) || ( bNoMoreData && PaddedZeros < RollingHash->GetWindowSize() ) )
 		{
-			check(PreviousSize > OverlapSize);
-			uint8* CopyTo = DataBuffer.GetData();
-			uint8* CopyFrom = CopyTo + (PreviousSize - OverlapSize);
-			FMemory::Memcpy(CopyTo, CopyFrom, OverlapSize);
-			DataBuffer.SetNum(OverlapSize, false);
-			DataOffset += DataBufferSize - OverlapSize;
+			FBuildPatchTimers::LogInfo();
+
+			// Prime the rolling hash
+			if( RollingHash->GetNumDataNeeded() > 0 )
+			{
+				if( DataBufferPos < ReadLen )
+				{
+					RollingHash->ConsumeByte( DataBuffer[ DataBufferPos++ ] );
+				}
+				else
+				{
+					RollingHash->ConsumeByte( 0 );
+					++PaddedZeros;
+				}
+				// Keep looping until primed
+				continue;
+			}
+
+			// Check if we recognized a chunk
+			FGuid ChunkGuid;
+			const uint64 WindowHash = RollingHash->GetWindowHash();
+			const TRingBuffer< uint8, FBuildPatchData::ChunkDataSize >& WindowData = RollingHash->GetWindowData();
+			bool ChunkRecognised = FindExistingChunkData( WindowHash, WindowData, ChunkGuid );
+			if( ChunkRecognised )
+			{
+				// Process all bytes
+				DataProcessor.PushChunk();
+				const uint32 WindowDataSize = RollingHash->GetWindowSize() - PaddedZeros;
+				for( uint32 i = 0; i < WindowDataSize; ++i )
+				{
+					const bool bStartOfFile = BuildStream->GetFileSpan( ProcessPos, FileName, FileDataCount );
+					const bool bEndOfFile = FileDataCount <= 1;
+					check( FileDataCount > 0 );// If FileDataCount is ever 0, it means this piece of data belongs to no file, so something is wrong
+					DataProcessor.SkipKnownByte( WindowData[i], bStartOfFile, bEndOfFile, FileName );
+					++ProcessPos;
+					--FileDataCount;
+				}
+				uint8* SerialWindowData = new uint8[ FBuildPatchData::ChunkDataSize ];
+				WindowData.Serialize( SerialWindowData );
+				DataProcessor.PopChunk( WindowHash, SerialWindowData, ChunkGuid );
+				delete[] SerialWindowData;
+
+				// Clear
+				RollingHash->Clear();
+			}
+			else
+			{
+				// Process one byte
+				const bool bStartOfFile = BuildStream->GetFileSpan( ProcessPos, FileName, FileDataCount );
+				const bool bEndOfFile = FileDataCount <= 1;
+				DataProcessor.ProcessNewByte( WindowData.Bottom(), bStartOfFile, bEndOfFile, FileName );
+				++ProcessPos;
+				--FileDataCount;
+
+				// Roll
+				if( DataBufferPos < ReadLen )
+				{
+					RollingHash->RollForward( DataBuffer[ DataBufferPos++ ] );
+				}
+				else if( bNoMoreData )
+				{
+					RollingHash->RollForward( 0 );
+					++PaddedZeros;
+				}
+			}
+
+			FBuildPatchTimers::DataProcessed = ProcessPos;
 		}
+	}
+	FBuildPatchTimers::LogInfo();
 
-		// Grab some data from the build stream
-		PreviousSize = DataBuffer.Num();
-		DataBuffer.SetNumUninitialized(DataBufferSize);
-		ReadLen = BuildStream->DequeueData(DataBuffer.GetData() + PreviousSize, DataBufferSize - PreviousSize);
-		DataBuffer.SetNum(PreviousSize + ReadLen, false);
 
-		// Create data processor, waiting for available slot first
-		while(FDataScannerCounter::GetNumIncompleteScanners() > FDataScannerCounter::GetNumRunningScanners())
+	// The final chunk if any should be finished.
+	// This also triggers the chunk writer thread to exit.
+	DataProcessor.FinalChunk();
+
+	FBuildPatchTimers::LogInfo();
+
+
+	// Handle empty files
+	FSHA1 EmptyHasher;
+	EmptyHasher.Final();
+	const TArray< FString >& EmptyFileList = BuildStream->GetEmptyFiles();
+	for (const auto& EmptyFile : EmptyFileList)
+	{
+		BuildManifest->Data->FileManifestList.Add(FFileManifestData());
+		FFileManifestData& EmptyFileManifest = BuildManifest->Data->FileManifestList.Last();
+		EmptyFileManifest.Filename = EmptyFile;
+		EmptyHasher.GetHash(EmptyFileManifest.FileHash.Hash);
+	}
+
+	// Add chunk sizes
+	const TMap<FGuid, int64>& ChunkFilesizes = DataProcessor.GetChunkFilesizes();
+	for (FChunkInfoData& ChunkInfo : BuildManifest->Data->ChunkList)
+	{
+		if (ChunkFilesizes.Contains(ChunkInfo.Guid))
 		{
-			FPlatformProcess::Sleep(0.01f);
-			StatsCollector->LogStats(2.0f);
+			ChunkInfo.FileSize = ChunkFilesizes[ChunkInfo.Guid];
 		}
+	}
 
-		// Pass a data scanner for this piece of data to the manifest builder
-		ManifestBuilder->AddDataScanner(FDataScannerFactory::Create(DataOffset, DataBuffer, CloudEnumeration, DataMatcher, StatsCollector));
+	FBuildPatchTimers::LogInfo();
 
-		// Log collected stats
-		StatsCollector->LogStats(2.0f);
+	// Fill out lookups
+	BuildManifest->InitLookups();
+
+	// Fill out the file attributes
+	for (const auto& Entry : FileAttributesMap)
+	{
+		const FString& Filename = Entry.Key;
+		const FFileAttributes& Attributes = Entry.Value;
+		if (BuildManifest->FileManifestLookup.Contains(Filename))
+		{
+			FFileManifestData& FileManifest = *BuildManifest->FileManifestLookup[Filename];
+			FileManifest.bIsReadOnly = Attributes.bReadOnly;
+			FileManifest.bIsCompressed = Attributes.bCompressed;
+			// Only overwrite unix exe if true
+			if (Attributes.bUnixExecutable)
+			{
+				FileManifest.bIsUnixExecutable = Attributes.bUnixExecutable;
+			}
+		}
+		else
+		{
+			GLog->Logf(TEXT("File Attributes: File not in build %s"), *Filename);
+		}
 	}
 
 	// Save manifest into the cloud directory
-	FString JsonFilename = FBuildPatchServicesModule::GetCloudDirectory() / FDefaultValueHelper::RemoveWhitespaces(Settings.AppName + Settings.BuildVersion) + TEXT(".manifest");
-	ManifestBuilder->SaveToFile(JsonFilename);
+	FString JsonFilename = FBuildPatchServicesModule::GetCloudDirectory() / FDefaultValueHelper::RemoveWhitespaces(BuildManifest->Data->AppName + BuildManifest->Data->BuildVersion) + TEXT(".manifest");
+	BuildManifest->Data->ManifestFileVersion = EBuildPatchAppManifestVersion::GetLatestJsonVersion();
+	BuildManifest->SaveToFile(JsonFilename, false);
 
-	// Final value log
-	StatsCollector->LogStats();
+	FBuildPatchTimers::LogInfo(true);
 
 	// Output to log for builder info
 	GLog->Logf(TEXT("Saved manifest to %s"), *JsonFilename);
 
-
-	uint64 EndTime = FStatsCollector::GetCycles();
-	GLog->Logf(TEXT("Completed in %s"), *FPlatformTime::PrettyTime(FStatsCollector::CyclesToSeconds(EndTime - StartTime)));
-
 	// Clean up memory
+	delete[] DataBuffer;
+	delete BuildStream;
+	delete RollingHash;
 	FBuildGenerationChunkCache::Shutdown();
 
+	// @TODO LSwift: Detect errors and return false on failure
 	return true;
 }
 
@@ -670,7 +1444,6 @@ bool FBuildDataGenerator::GenerateFilesManifestFromDirectory( const FBuildPatchS
 		const FString& FileName = *FileIt;
 		// Read the file
 		FArchive* FileReader = IFileManager::Get().CreateFileReader( *FileName );
-		checkf(FileReader != nullptr, TEXT("Fatal Error: Could not open build file %s"), *FileName);
 		if( FileReader != NULL )
 		{
 			// Make SourceFile the format we want it in and start a new file
@@ -696,6 +1469,10 @@ bool FBuildDataGenerator::GenerateFilesManifestFromDirectory( const FBuildPatchS
 			delete FileReader;
 			DataProcessor.EndFile();
 		}
+		else
+		{
+			// @TODO LSwift: Handle File error?
+		}
 	}
 
 	// Fill out lookups
@@ -738,6 +1515,171 @@ bool FBuildDataGenerator::GenerateFilesManifestFromDirectory( const FBuildPatchS
 	return true;
 }
 
+bool FBuildDataGenerator::FindExistingChunkData( const uint64& ChunkHash, const uint8* ChunkData, FGuid& ChunkGuid )
+{
+	// Quick code hack
+	TRingBuffer< uint8, FBuildPatchData::ChunkDataSize > ChunkDataRing;
+	const uint32 ChunkDataLen = FBuildPatchData::ChunkDataSize;
+	ChunkDataRing.Enqueue( ChunkData, ChunkDataLen );
+	return FindExistingChunkData( ChunkHash, ChunkDataRing, ChunkGuid );
+}
+
+bool FBuildDataGenerator::FindExistingChunkData( const uint64& ChunkHash, const TRingBuffer< uint8, FBuildPatchData::ChunkDataSize >& ChunkData, FGuid& ChunkGuid )
+{
+	bool bFoundMatchingChunk = false;
+
+	// Perform an inventory on Cloud chunks if not already done
+	if( ExistingChunksEnumerated == false )
+	{
+		IFileManager& FileManager = IFileManager::Get();
+		FString JSONOutput;
+		TSharedRef< TJsonWriter< TCHAR, TPrettyJsonPrintPolicy< TCHAR > > > DebugWriter = TJsonWriterFactory< TCHAR, TPrettyJsonPrintPolicy< TCHAR > >::Create( &JSONOutput );
+		DebugWriter->WriteObjectStart();
+
+		// Find all manifest files
+		const FString CloudDir = FBuildPatchServicesModule::GetCloudDirectory();
+		if (FileManager.DirectoryExists(*CloudDir))
+		{
+			const double StartEnumerate = FPlatformTime::Seconds();
+			TArray<FString> AllManifests;
+			GLog->Logf(TEXT("BuildDataGenerator: Enumerating Manifests from %s"), *CloudDir);
+			FileManager.FindFiles(AllManifests, *(CloudDir / TEXT("*.manifest")), true, false);
+			const double EnumerateTime = FPlatformTime::Seconds() - StartEnumerate;
+			GLog->Logf(TEXT("BuildDataGenerator: Found %d manifests in %.1f seconds"), AllManifests.Num(), EnumerateTime);
+
+			// Load all manifest files
+			uint64 NumChunksFound = 0;
+			const double StartLoadAllManifest = FPlatformTime::Seconds();
+			for (const auto& ManifestFile : AllManifests)
+			{
+				// Determine chunks from manifest file
+				const FString ManifestFilename = CloudDir / ManifestFile;
+				FBuildPatchAppManifestRef BuildManifest = MakeShareable(new FBuildPatchAppManifest());
+				const double StartLoadManifest = FPlatformTime::Seconds();
+				if (BuildManifest->LoadFromFile(ManifestFilename))
+				{
+					const double LoadManifestTime = FPlatformTime::Seconds() - StartLoadManifest;
+					GLog->Logf(TEXT("BuildDataGenerator: Loaded %s in %.1f seconds"), *ManifestFile, LoadManifestTime);
+					if(!BuildManifest->IsFileDataManifest())
+					{
+						TArray<FGuid> ChunksReferenced;
+						BuildManifest->GetDataList(ChunksReferenced);
+						for (const auto& ReferencedChunkGuid : ChunksReferenced)
+						{
+							uint64 ReferencedChunkHash;
+							if (BuildManifest->GetChunkHash(ReferencedChunkGuid, ReferencedChunkHash))
+							{
+								if (ReferencedChunkHash != 0)
+								{
+									TArray< FGuid >& HashChunkList = ExistingChunkHashInventory.FindOrAdd(ReferencedChunkHash);
+									if (!HashChunkList.Contains(ReferencedChunkGuid))
+									{
+										++NumChunksFound;
+										HashChunkList.Add(ReferencedChunkGuid);
+									}
+								}
+								else
+								{
+									GLog->Logf(TEXT("BuildDataGenerator: INFO: Ignored an existing chunk %s with a failed hash value of zero to avoid performance problems while chunking"), *ReferencedChunkGuid.ToString());
+								}
+							}
+							else
+							{
+								GLog->Logf(TEXT("BuildDataGenerator: WARNING: Missing chunk hash for %s in manifest %s"), *ReferencedChunkGuid.ToString(), *ManifestFile);
+							}
+						}
+					}
+					else
+					{
+						GLog->Logf(TEXT("BuildDataGenerator: INFO: Ignoring non-chunked manifest %s"), *ManifestFilename);
+					}
+				}
+				else
+				{
+					GLog->Logf(TEXT("BuildDataGenerator: WARNING: Could not read Manifest file. Data recognition will suffer (%s)"), *ManifestFilename);
+				}
+			}
+			const double LoadAllManifestTime = FPlatformTime::Seconds() - StartLoadAllManifest;
+			GLog->Logf(TEXT("BuildDataGenerator: Used %d manifests to enumerate %llu chunks in %.1f seconds"), AllManifests.Num(), NumChunksFound, LoadAllManifestTime);
+		}
+		else
+		{
+			GLog->Logf(TEXT("BuildDataGenerator: Cloud directory does not exist: %s"), *CloudDir);
+		}
+
+		ExistingChunksEnumerated = true;
+	}
+
+	// Do we have a chunk matching this data?
+	if( ExistingChunkHashInventory.Num() > 0 )
+	{
+		TArray< FGuid >* ChunkList = ExistingChunkHashInventory.Find( ChunkHash );
+		if( ChunkList != NULL )
+		{
+			// We need to load each chunk in this list and compare data
+			for( auto ChunkIt = ChunkList->CreateConstIterator(); ChunkIt && !bFoundMatchingChunk ; ++ChunkIt)
+			{
+				FGuid Guid = *ChunkIt;
+				if (!ExistingChunkGuidInventory.Contains(Guid))
+				{
+					ExistingChunkGuidInventory.Add(Guid, DiscoverChunkFilename(Guid, ChunkHash));
+				}
+				const FString& SourceFile = ExistingChunkGuidInventory[ Guid ];
+				// Read the file
+				uint8* TempChunkData = new uint8[ FBuildPatchData::ChunkDataSize ];
+				ChunkData.Serialize( TempChunkData );
+				bool bChunkIsUsable = true;
+				if( CompareDataToChunk( SourceFile, TempChunkData, Guid, bChunkIsUsable ) )
+				{
+					// We have a chunk match!!
+					bFoundMatchingChunk = true;
+					ChunkGuid = Guid;
+				}
+				// Check if this chunk should be dumped
+				if(!bChunkIsUsable)
+				{
+					GLog->Logf(TEXT("BuildDataGenerator: Chunk %s unusable, removed from inventory."), *ChunkGuid.ToString());
+					ChunkList->Remove(Guid);
+					--ChunkIt;
+				}
+				delete[] TempChunkData;
+			}
+		}
+	}
+
+	return bFoundMatchingChunk;
+}
+
+FString FBuildDataGenerator::DiscoverChunkFilename(const FGuid& ChunkGuid, const uint64& ChunkHash)
+{
+	return FBuildPatchUtils::GetChunkNewFilename(EBuildPatchAppManifestVersion::GetLatestVersion(), FBuildPatchServicesModule::GetCloudDirectory(), ChunkGuid, ChunkHash);
+	static double AccumTime = 0.0;
+	const double StartDiscovery = FPlatformTime::Seconds();
+	const FString CloudDir = FBuildPatchServicesModule::GetCloudDirectory();
+	FString ChunkFilenameChecked;
+	for (EBuildPatchAppManifestVersion::Type VersionCounter = static_cast<EBuildPatchAppManifestVersion::Type>(EBuildPatchAppManifestVersion::LatestPlusOne - 1);
+		VersionCounter >= EBuildPatchAppManifestVersion::Original;
+		VersionCounter = static_cast<EBuildPatchAppManifestVersion::Type>(VersionCounter - 1))
+	{
+		const FString ChunkFilename = VersionCounter < EBuildPatchAppManifestVersion::DataFileRenames ?
+			FBuildPatchUtils::GetChunkOldFilename(CloudDir, ChunkGuid) :
+			FBuildPatchUtils::GetChunkNewFilename(VersionCounter, CloudDir, ChunkGuid, ChunkHash);
+		if (ChunkFilenameChecked != ChunkFilename)
+		{
+			ChunkFilenameChecked = ChunkFilename;
+			if (FPaths::FileExists(ChunkFilename))
+			{
+				const double DiscoveryTime = FPlatformTime::Seconds() - StartDiscovery;
+				AccumTime += DiscoveryTime;
+				GLog->Logf(TEXT("BuildDataGenerator: DiscoverChunkFilename: Chunk %s found in %f secs. Accum time %.2f secs"), *ChunkGuid.ToString(), DiscoveryTime, AccumTime);
+				return ChunkFilename;
+			}
+		}
+	}
+	GLog->Logf(TEXT("BuildDataGenerator: DiscoverChunkFilename: Chunk %s not found"), *ChunkGuid.ToString());
+	return TEXT("");
+}
+
 bool FBuildDataGenerator::FindExistingFileData(const FString& InSourceFile, const FSHAHashData& InFileHash, const FDateTime& DataThresholdTime, FGuid& OutFileGuid)
 {
 	bool bFoundMatchingFile = false;
diff --git a/Engine/Source/Runtime/Online/BuildPatchServices/Private/BuildPatchGeneration.h b/Engine/Source/Runtime/Online/BuildPatchServices/Private/BuildPatchGeneration.h
--- a/Engine/Source/Runtime/Online/BuildPatchServices/Private/BuildPatchGeneration.h
+++ b/Engine/Source/Runtime/Online/BuildPatchServices/Private/BuildPatchGeneration.h
@@ -4,10 +4,357 @@
 	BuildPatchGeneration.h: Declares classes involved with producing build data.
 =============================================================================*/
 
+#ifndef __BuildPatchGeneration_h__
+#define __BuildPatchGeneration_h__
+
 #pragma once
 
 #if WITH_BUILDPATCHGENERATION
 
+/**
+ * A threaded class that holds Build stream data as a FIFO data buffer of given size. Will read until full then wait
+ * for data to be read out before continuing.
+ */
+class FBuildStream
+{
+private:
+	// A private runnable class that reads in data
+	class FBuildStreamReader : public FRunnable
+	{
+	public:
+		// The directory that we want to read from
+		FString DepotDirectory;
+
+		// The path to a text file containing \r\n separated depot relative files to ignore
+		FString IgnoreListFile;
+
+		// Pointer to FBuildStream
+		FBuildStream* BuildStream;
+
+		// Default Constructor
+		FBuildStreamReader();
+
+		// Destructor
+		~FBuildStreamReader();
+
+		// Start FRunnable interface
+		virtual bool Init() override;
+		virtual uint32 Run() override;
+		// End FRunnable interface
+
+		void StartThread();
+
+		FRunnableThread * Thread;
+
+	} BuildStreamReader;
+
+	// A private struct for storing the data span of a file in the build stream
+	struct FFileSpan
+	{
+		// The file's filename
+		FString Filename;
+		// The file data start index
+		uint64 StartIdx;
+		// The size of the file
+		uint64 Size;
+		
+		/**
+		 * Default constructor
+		 */
+		FFileSpan()
+			: StartIdx( 0 )
+			, Size( 0 )
+		{
+		}
+	};
+
+public:
+	/**
+	 * Constructor
+	 */
+	FBuildStream( const FString& RootDirectory, const FString& IgnoreListFile );
+
+	/**
+	 * Default destructor
+	 */
+	~FBuildStream();
+
+	/**
+	 * Retrieves the file details for a specific start index.
+	 * @param	IN	StartingIdx		The data index into the build image.
+	 * @param	OUT Filename		The filename for the file starting at StartingIdx. NB: The variable remains UNCHANGED if no file starts here.
+	 * @param	OUT FileSize		The size of the file starting at StartingIdx. NB: The variable remains UNCHANGED if no file starts here.
+	 * @return	true if the data byte at StartingIdx is the start of a file
+	 */
+	bool GetFileSpan( const uint64& StartingIdx, FString& Filename, uint64& FileSize );
+
+	/**
+	 * Fetches some data from the buffer, also removing it.
+	 * @param	IN	Buffer			Pointer to buffer to receive the data.
+	 * @param	IN	ReqSize			The amount of data to attempt to retrieve.
+	 * @param	IN	WaitForData		Optional: Default true. Whether to wait until there is enough data in the buffer.
+	 * @return	The amount of data retrieved
+	 */
+	uint32 DequeueData( uint8* Buffer, const uint32& ReqSize, const bool WaitForData = true );
+
+	/**
+	 * Gets a list of empty files that the build contains.
+	 * @return	Array of empty files in the build
+	 */
+	const TArray< FString > GetEmptyFiles();
+
+	/**
+	 * Whether the read thread has finished reading the build image.
+	 * @return	true if there is no more data coming into the buffer
+	 */
+	bool IsEndOfBuild();
+
+	/**
+	 * Whether there is any more data available to dequeue from the buffer
+	 * @return	true if there is no more data coming in, and the internal buffer is also empty.
+	 */
+	bool IsEndOfData();
+
+private:
+	/**
+	 * Hide the default constructor so that the params must be given
+	 */
+	FBuildStream(){};
+
+	//////////////////////////////////////////////////////////////////////////
+	// Private functions for manipulating the data store
+private:
+	/**
+	 * Wipes out all data
+	 */
+	void Clear();
+
+	/**
+	 * How much more data the internal buffer can hold before being full
+	 * @return	the number of bytes free in the buffer
+	 */
+	uint32 SpaceLeft();
+
+	/**
+	 * How much data the internal buffer is currently holding
+	 * @return	the number of bytes held in the buffer
+	 */
+	uint32 DataAvailable();
+
+	/**
+	 * Mark the beginning of a new file
+	 * @param	IN	Filename	The local build filename to the file
+	 * @param	IN	FileSize	The size of the file
+	 */
+	void BeginNewFile( const FString& Filename, const uint64& FileSize );
+
+	/**
+	 * We must special case empty files because we don't want them effecting chunk data
+	 * @param	IN	Filename	The local build filename to the file
+	 */
+	void AddEmptyFile( const FString& Filename );
+
+	/**
+	 * Enqueue a buffer of data into the internal buffer.
+	 * @param	IN	Buffer		Pointer to buffer of data.
+	 * @param	IN	Len			The length of the data buffer.
+	 */
+	void EnqueueData( const uint8* Buffer, const uint32& Len );
+
+	/**
+	 * Mark the end of the build image
+	 * @param	IN	bIsEnd		Optional: Default true. Whether to set the build end as true or false
+	 */
+	void EndOfBuild( bool bIsEnd = true );
+
+	//////////////////////////////////////////////////////////////////////////
+	// Private variables
+private:
+	// Holds the filenames referred to index positions in the data stream
+	TMap< uint64, FFileSpan > FilesParsed;
+
+	// Holds a list of empty files
+	TArray< FString > EmptyFiles;
+
+	// A critical section for accessing the FilesParsed and EmptyFiles data.
+	FCriticalSection FilesListsCS;
+
+	// Holds the ring buffer for loaded data code words
+	TRingBuffer< uint8, StreamBufferSize > BuildDataStream;
+
+	// A critical section for accessing the data stream, BuildDataStream
+	FCriticalSection BuildDataStreamCS;
+
+	// Whether the build has been completely read from disk
+	bool bNoMoreData;
+
+	// A critical section for accessing the no more data flag, NoMoreData.
+	FCriticalSection NoMoreDataCS;
+};
+
+/**
+ * A class that controls processing of build data to produce the manifest with chunks
+ */
+class FBuildDataChunkProcessor
+{
+public:
+
+	/**
+	 * Constructor
+	 * @param	InBuildManifest 	a build manifest that receives the manifest data.
+	 * @param	InBuildRoot			The root directory of the build
+	 */
+	FBuildDataChunkProcessor( FBuildPatchAppManifestRef InBuildManifest, const FString& InBuildRoot );
+	
+	/**
+	 * Default destructor
+	 */
+	~FBuildDataChunkProcessor();
+
+	/**
+	 * Start a new chunk
+	 * @param	bZeroBuffer 	Optional: Default true. Whether to clear the chunk buffer with zero
+	 */
+	void BeginNewChunk( const bool& bZeroBuffer = true );
+
+	/**
+	 * Start a new chunk part
+	 */
+	void BeginNewChunkPart();
+	
+	/**
+	 * End the current chunk part
+	 */
+	void EndNewChunkPart();
+	
+	/**
+	 * End the current chunk
+	 * @param	ChunkHash		The chunk's hash value
+	 * @param	ChunkData		Pointer to the chunk data, must be of FBuildPatchData::ChunkDataSize length
+	 * @param	ChunkGuid		The chunk GUID of this chunk. If it matches CurrentChunkGuid then this is a new chunk, otherwise it is the guid of the recognized chunk
+	 */
+	void EndNewChunk( const uint64& ChunkHash, const uint8* ChunkData, const FGuid& ChunkGuid );
+	
+	/**
+	 * Begins insertion of recognized chunk
+	 */
+	void PushChunk();
+
+	/**
+	 * Ends the insertion of recognized chunk
+	 * @param	ChunkHash		The hash value for the recognised chunk
+	 * @param	ChunkData		The data for the recognised chunk
+	 * @param	ChunkGuid		The GUID for the recognised chunk
+	 */
+	void PopChunk( const uint64& ChunkHash, const uint8* ChunkData, const FGuid& ChunkGuid );
+
+	/**
+	 * When the build image has been fully covered, call to end any current chunk we have (which will be zero padded)
+	 */
+	void FinalChunk();
+
+	/**
+	 * Mark the beginning of a new file
+	 * @param	FileName	The file's local filename
+	 */
+	void BeginFile( const FString& FileName );
+
+	/**
+	 * Mark the end of the current file
+	 */
+	void EndFile();
+
+	/**
+	 * Skips over the next byte from a recognized chunk
+	 * @param	NextByte		The next byte
+	 * @param	bStartOfFile	Whether this byte is the beginning of a file
+	 * @param	bEndOfFile		Whether this byte is the end of a file
+	 * @param	Filename		The filename of the file, only required if start of new file
+	 */
+	void SkipKnownByte( const uint8& NextByte, const bool& bStartOfFile, const bool& bEndOfFile, const FString& Filename );
+
+	/**
+	 * Process a new byte into the current new chunk
+	 * @param	NewByte			The new byte
+	 * @param	bStartOfFile	Whether this byte is the beginning of a file
+	 * @param	bEndOfFile		Whether this byte is the end of a file
+	 * @param	Filename		The filename of the file, only required if start of new file
+	 */
+	void ProcessNewByte( const uint8& NewByte, const bool& bStartOfFile, const bool& bEndOfFile, const FString& Filename );
+
+	/**
+	 * Get the new/know chunks stats
+	 * @param	OutNewChunks	Receives the number of new chunks found
+	 * @param	OutKnownChunks	Receives the number of recognised chunks found
+	 */
+	void GetChunkStats( uint32& OutNewChunks, uint32& OutKnownChunks );
+
+	/**
+	 * Get the chunk filesizes that have been saved out
+	 * @return ref to the Map of GUID to filesize
+	 */
+	const TMap<FGuid, int64>& GetChunkFilesizes();
+
+private:
+
+	// Keep track of the number of new chunks generated
+	uint32 NumNewChunks;
+
+	// Keep track of the number of recognised chunks
+	uint32 NumKnownChunks;
+
+	// Holds the build manifest that is receiving the information
+	FBuildPatchAppManifestPtr BuildManifest;
+
+	// Holds the build manifest that is receiving the information
+	FString BuildRoot;
+
+	// The chunk writer system
+	FChunkWriter ChunkWriter;
+
+	// A GUID for the current chunk
+	FGuid CurrentChunkGuid;
+
+	// The current chunk data
+	uint8* CurrentChunkBuffer;
+
+	// The current index into the chunk data
+	uint32 CurrentChunkBufferPos;
+
+	// Holds the current file being processed. Receives the chunk parts that we pass over.
+	FFileManifestData* CurrentFile;
+
+	// Holds an SHA hash calculator
+	FSHA1 FileHash;
+
+	// Are we currently processing a chunk
+	bool IsProcessingChunk;
+
+	// Are we currently processing a chunk part
+	bool IsProcessingChunkPart;
+
+	// Are we currently processing a file
+	bool IsProcessingFile;
+
+	// Have we pushed to a recognized chunk
+	bool ChunkIsPushed;
+
+	// Backup of CurrentChunkGuid when pushing a chunk
+	FGuid BackupChunkGuid;
+
+	// Backup of CurrentChunkBufferPos when pushing a chunk
+	uint32 BackupChunkBufferPos;
+
+	// Backup of IsProcessingChunk when pushing a chunk
+	bool BackupProcessingChunk;
+
+	// Backup of IsProcessingChunkPart when pushing a chunk
+	bool BackupProcessingChunkPart;
+
+	// Cached map of chunk file sizes
+	TMap<FGuid, int64> ChunkFileSizes;
+};
+
 /**
  * A class that controls processing of build data to produce the manifest with files
  */
@@ -346,3 +693,5 @@ struct FFileAttributes
 };
 
 #endif //WITH_BUILDPATCHGENERATION
+
+#endif // __BuildPatchGeneration_h__
diff --git a/Engine/Source/Runtime/Online/BuildPatchServices/Private/BuildPatchHash.h b/Engine/Source/Runtime/Online/BuildPatchServices/Private/BuildPatchHash.h
--- a/Engine/Source/Runtime/Online/BuildPatchServices/Private/BuildPatchHash.h
+++ b/Engine/Source/Runtime/Online/BuildPatchServices/Private/BuildPatchHash.h
@@ -1,13 +1,15 @@
 // Copyright 1998-2015 Epic Games, Inc. All Rights Reserved.
-#pragma once
-
-#include "Core/RingBuffer.h"
 
 /*=============================================================================
 	BuildPatchHash.h: Declares and implements template classes for use with the 
 	Build and Patching hash functionality.
 =============================================================================*/
 
+#ifndef __BuildPatchHash_h__
+#define __BuildPatchHash_h__
+
+#pragma once
+
 /**
  * A macro for barrel rolling a 64 bit value n times to the left.
  * @param Value		The value to shift
@@ -15,6 +17,158 @@
  */
 #define ROTLEFT_64B( Value, Shifts ) Value = ( ( ( Value ) << ( ( Shifts ) % 64 ) ) | ( ( Value ) >> ( ( 64 - ( ( Shifts ) % 64 ) ) % 64 ) ) )
 
+/**
+ * A ring buffer template class.
+ */
+template< typename DataType, uint32 BufferDataSize >
+class TRingBuffer
+{
+public:
+	/**
+	 * Default Constructor
+	 */
+	TRingBuffer();
+
+	/**
+	 * Default Destructor
+	 */
+	~TRingBuffer();
+
+	/**
+	 * Clears memory and indexes
+	 */
+	void Empty();
+
+	/**
+	 * Pushes a data word to the end of the FIFO. WILL OVERWRITE OLDEST if full.
+	 * @param Val	The data word to push
+	 */
+	void Enqueue( const DataType& Val );
+
+	/**
+	 * Pushes a buffer of data words to the end of the FIFO. WILL OVERWRITE OLDEST if full.
+	 * @param ValBuf		The buffer pointer
+	 * @param BufLen		The length to copy
+	 */
+	void Enqueue( const DataType* ValBuf, const uint32& BufLen );
+
+	/**
+	 * Take the next data word from the FIFO buffer.
+	 * @param OUT	ValOut	The variable to receive the data word
+	 * @return	true if the buffer was not empty, false otherwise.
+	 */
+	bool Dequeue( DataType& ValOut );
+
+	/**
+	 * Take the next set of data words from the FIFO buffer.
+	 * @param ValBuf		The buffer to receive the data
+	 * @param BufLen		The number of words requested
+	 * @return	the number of words actually copied
+	 */
+	uint32 Dequeue( DataType* ValBuf, const uint32& BufLen );
+
+	/**
+	 * Get the next data word from the FIFO buffer without removing it.
+	 * @param OUT	ValOut	The variable to receive the data word
+	 * @return	true if the buffer was not empty, false otherwise.
+	 */
+	bool Peek( DataType& ValOut ) const;
+
+	/**
+	 * Get the next set of data words from the FIFO buffer without removing them.
+	 * @param ValBuf		The buffer to receive the data
+	 * @param BufLen		The number of words requested
+	 * @return	the number of words actually copied
+	 */
+	uint32 Peek( DataType* ValBuf, const uint32& BufLen ) const;
+
+	/**
+	 * Compare the memory in the FIFO to the memory in the given buffer
+	 * @param SerialBuffer	The buffer containing data to compare
+	 * @param CompareLen		The number of words to compare
+	 * @return	< 0 if data < SerialBuffer.. 0 if data == SerialBuffer... > 0 if data > SerialBuffer
+	 */
+	int32 SerialCompare( const DataType* SerialBuffer, uint32 CompareLen ) const;
+
+	/**
+	 * Serializes the internal buffer into the given buffer
+	 * @param SerialBuffer	A preallocated buffer to copy data into
+	 */
+	void Serialize( DataType* SerialBuffer ) const;
+
+	/**
+	 * Square bracket operators for accessing data in the buffer by FIFO index. [0] returns the next entry in the FIFO that would get provided by Dequeue or Peek.
+	 * @param Index	The index into the FIFO buffer
+	 * @return	The data word
+	 */
+	FORCEINLINE const DataType& operator[]( const int32& Index ) const;
+	FORCEINLINE DataType& operator[]( const int32& Index );
+
+	/**
+	 * Gets the last data word in the FIFO (i.e. most recently pushed).
+	 * @return	The last data word
+	 */
+	FORCEINLINE const DataType& Top() const;
+	FORCEINLINE DataType& Top();
+
+	/**
+	 * Gets the first data word in the FIFO (i.e. oldest).
+	 * @return	The first data word
+	 */
+	FORCEINLINE const DataType& Bottom() const;
+	FORCEINLINE DataType& Bottom();
+
+	/**
+	 * Gets the buffer index that the last data word is stored in.
+	 * @return	The index of Top()
+	 */
+	FORCEINLINE const uint32 TopIndex() const;
+
+	/**
+	 * Gets the buffer index that the first data word is stored in.
+	 * @return	The index of Bottom()
+	 */
+	FORCEINLINE const uint32 BottomIndex() const;
+
+	/**
+	 * Gets the buffer index that the next enqueued word will get stored in.
+	 * @return	The index that the next enqueued word will have
+	 */
+	FORCEINLINE const uint32 NextIndex() const;
+
+	/**
+	 * Gets the size of the data buffer.
+	 * @return	Template arg BufferDataSize
+	 */
+	FORCEINLINE const uint32 RingDataSize() const;
+
+	/**
+	 * Gets the number of words currently in the FIFO.
+	 * @return	data size
+	 */
+	FORCEINLINE const uint32 RingDataUsage() const;
+
+	/**
+	 * Gets the total number of words that have been pushed through this buffer since clearing
+	 * @return	total length of data passed through
+	 */
+	FORCEINLINE const uint64 TotalDataPushed() const;
+
+
+private:
+	// The data memory
+	DataType* Data;
+
+	// Value keeping track of free space
+	uint32 NumDataAvailable;
+
+	// Value keeping track of the next data index
+	uint32 DataIndex;
+
+	// Value to keep track of total amount of data Enqueued
+	uint64 TotalNumDataPushed;
+};
+
 /**
  * A static struct containing 64bit polynomial constant and hash table lookup for use with FRollingHash.
  */
@@ -120,6 +274,206 @@ namespace FCycPoly64Hash
 	uint64 GetHashForDataSet(const uint8* DataSet, const uint32 DataSize, const uint64 State = 0);
 }
 
+/* TRingBuffer implementation
+*****************************************************************************/
+template< typename DataType, uint32 BufferDataSize >
+TRingBuffer< DataType, BufferDataSize >::TRingBuffer()
+{
+	Data = new DataType[ BufferDataSize ];
+	Empty();
+}
+
+template< typename DataType, uint32 BufferDataSize >
+TRingBuffer< DataType, BufferDataSize >::~TRingBuffer()
+{
+	delete[] Data;
+}
+
+template< typename DataType, uint32 BufferDataSize >
+void TRingBuffer< DataType, BufferDataSize >::Empty()
+{
+	FMemory::Memzero( Data, sizeof( DataType ) * BufferDataSize );
+	NumDataAvailable = 0;
+	DataIndex = 0;
+	TotalNumDataPushed = 0;
+}
+
+template< typename DataType, uint32 BufferDataSize >
+void TRingBuffer< DataType, BufferDataSize >::Enqueue( const DataType& Val )
+{
+	Data[ DataIndex++ ] = Val;
+	DataIndex %= BufferDataSize;
+	++TotalNumDataPushed;
+	NumDataAvailable = FMath::Clamp< uint32 >( NumDataAvailable + 1, 0, BufferDataSize );
+}
+
+template< typename DataType, uint32 BufferDataSize >
+void TRingBuffer< DataType, BufferDataSize >::Enqueue( const DataType* ValBuf, const uint32& BufLen )
+{
+	for( uint32 i = 0; i < BufLen; ++i )
+	{
+		Enqueue( ValBuf[i] );
+	}
+}
+
+template< typename DataType, uint32 BufferDataSize >
+bool TRingBuffer< DataType, BufferDataSize >::Dequeue( DataType& ValOut )
+{
+	if( NumDataAvailable > 0 )
+	{
+		ValOut = Bottom();
+		NumDataAvailable = FMath::Clamp< uint32 >( NumDataAvailable - 1, 0, BufferDataSize );
+		return true;
+	}
+	else
+	{
+		return false;
+	}
+}
+
+template< typename DataType, uint32 BufferDataSize >
+uint32 TRingBuffer< DataType, BufferDataSize >::Dequeue( DataType* ValBuf, const uint32& BufLen )
+{
+	uint32 Count = 0;
+	for( uint32 i = 0; i < BufLen; ++i )
+	{
+		if( Dequeue( ValBuf[i] ) )
+		{
+			++Count;
+		}
+		else
+		{
+			break;
+		}
+	}
+	return Count;
+}
+
+template< typename DataType, uint32 BufferDataSize >
+bool TRingBuffer< DataType, BufferDataSize >::Peek( DataType& ValOut ) const
+{
+	if( NumDataAvailable > 0 )
+	{
+		ValOut = Bottom();
+		return true;
+	}
+	else
+	{
+		return false;
+	}
+}
+
+template< typename DataType, uint32 BufferDataSize >
+uint32 TRingBuffer< DataType, BufferDataSize >::Peek( DataType* ValBuf, const uint32& BufLen ) const
+{
+	uint32 Count = 0;
+	for( uint32 i = 0; i < BufLen; ++i )
+	{
+		if( Peek( ValBuf[i] ) )
+		{
+			++Count;
+		}
+		else
+		{
+			break;
+		}
+	}
+	return Count;
+}
+
+template< typename DataType, uint32 BufferDataSize >
+int32 TRingBuffer< DataType, BufferDataSize >::SerialCompare( const DataType* SerialBuffer, uint32 CompareLen ) const
+{
+	check( CompareLen <= BufferDataSize);
+	const uint32 FirstPartLen = BufferDataSize - DataIndex;
+	int32 FirstCmp = FMemory::Memcmp( Data + BottomIndex(), SerialBuffer, sizeof( DataType ) * FMath::Min( FirstPartLen, CompareLen ) );
+	if( FirstCmp != 0 || FirstPartLen >= CompareLen)
+	{
+		return FirstCmp;
+	}
+	return FMemory::Memcmp( Data, &SerialBuffer[ FirstPartLen ], sizeof( DataType ) * ( CompareLen - FirstPartLen ) );
+}
+
+template< typename DataType, uint32 BufferDataSize >
+void TRingBuffer< DataType, BufferDataSize >::Serialize( DataType* SerialBuffer ) const
+{
+	const uint32 FirstPartLen = BufferDataSize - DataIndex;
+ 	FMemory::Memcpy( SerialBuffer, Data + BottomIndex(), sizeof( DataType ) * FirstPartLen );
+ 	FMemory::Memcpy( &SerialBuffer[ FirstPartLen ], Data, sizeof( DataType ) * ( BufferDataSize - FirstPartLen ) );
+}
+
+template< typename DataType, uint32 BufferDataSize >
+FORCEINLINE const DataType& TRingBuffer< DataType, BufferDataSize >::operator[]( const int32& Index ) const
+{
+	return Data[ ( BottomIndex() + Index ) % BufferDataSize ];
+}
+
+template< typename DataType, uint32 BufferDataSize >
+FORCEINLINE DataType& TRingBuffer< DataType, BufferDataSize >::operator[]( const int32& Index )
+{
+	return Data[ ( BottomIndex() + Index ) % BufferDataSize ];
+}
+
+template< typename DataType, uint32 BufferDataSize >
+FORCEINLINE const DataType& TRingBuffer< DataType, BufferDataSize >::Top() const
+{
+	return Data[ TopIndex() ];
+}
+
+template< typename DataType, uint32 BufferDataSize >
+FORCEINLINE DataType& TRingBuffer< DataType, BufferDataSize >::Top()
+{
+	return Data[ TopIndex() ];
+}
+
+template< typename DataType, uint32 BufferDataSize >
+FORCEINLINE const DataType& TRingBuffer< DataType, BufferDataSize >::Bottom() const
+{
+	return Data[ BottomIndex() ];
+}
+
+template< typename DataType, uint32 BufferDataSize >
+FORCEINLINE DataType& TRingBuffer< DataType, BufferDataSize >::Bottom()
+{
+	return Data[ BottomIndex() ];
+}
+
+template< typename DataType, uint32 BufferDataSize >
+FORCEINLINE const uint32 TRingBuffer< DataType, BufferDataSize >::TopIndex() const
+{
+	return ( DataIndex + BufferDataSize - 1 ) % BufferDataSize;
+}
+
+template< typename DataType, uint32 BufferDataSize >
+FORCEINLINE const uint32 TRingBuffer< DataType, BufferDataSize >::BottomIndex() const
+{
+	return ( DataIndex + BufferDataSize - NumDataAvailable ) % BufferDataSize;
+}
+
+template< typename DataType, uint32 BufferDataSize >
+FORCEINLINE const uint32 TRingBuffer< DataType, BufferDataSize >::NextIndex() const
+{
+	return DataIndex;
+}
+
+template< typename DataType, uint32 BufferDataSize >
+FORCEINLINE const uint32 TRingBuffer< DataType, BufferDataSize >::RingDataSize() const
+{
+	return BufferDataSize;
+}
+
+template< typename DataType, uint32 BufferDataSize >
+FORCEINLINE const uint32 TRingBuffer< DataType, BufferDataSize >::RingDataUsage() const
+{
+	return NumDataAvailable;
+}
+
+template< typename DataType, uint32 BufferDataSize >
+FORCEINLINE const uint64 TRingBuffer< DataType, BufferDataSize >::TotalDataPushed() const
+{
+	return TotalNumDataPushed;
+}
+
 /* FRollingHash implementation
 *****************************************************************************/
 template< uint32 WindowSize >
@@ -215,8 +569,8 @@ uint64 FRollingHash< WindowSize >::GetHashForDataSet( const uint8* DataSet )
 }
 
 /**
- * A static function to perform sanity checks on the rolling hash class.
- */
+* A static function to perform sanity checks on the rolling hash class.
+*/
 static bool CheckRollingHashAlgorithm()
 {
 	bool bCheckOk = true;
@@ -256,3 +610,5 @@ static bool CheckRollingHashAlgorithm()
 
 	return bCheckOk;
 }
+
+#endif // __BuildPatchHash_h__
diff --git a/Engine/Source/Runtime/Online/BuildPatchServices/Private/BuildPatchManifest.h b/Engine/Source/Runtime/Online/BuildPatchServices/Private/BuildPatchManifest.h
--- a/Engine/Source/Runtime/Online/BuildPatchServices/Private/BuildPatchManifest.h
+++ b/Engine/Source/Runtime/Online/BuildPatchServices/Private/BuildPatchManifest.h
@@ -318,12 +318,6 @@ struct FFileChunkPart
 	{}
 };
 
-// Required to allow private access to manifest builder for now..
-namespace FBuildPatchServices
-{
-	class FManifestBuilderImpl;
-}
-
 /**
  * Declare the FBuildPatchAppManifest object class. This holds the UObject data, and the implemented build manifest functionality
  */
@@ -332,9 +326,9 @@ class FBuildPatchAppManifest
 {
 	// Allow access to build processor classes
 	friend class FBuildDataGenerator;
+	friend class FBuildDataChunkProcessor;
 	friend class FBuildDataFileProcessor;
 	friend class FBuildPatchInstaller;
-	friend class FBuildPatchServices::FManifestBuilderImpl;
 public:
 
 	/**
diff --git a/Engine/Source/Runtime/Online/BuildPatchServices/Private/Core/RingBuffer.h b/Engine/Source/Runtime/Online/BuildPatchServices/Private/Core/RingBuffer.h
deleted file mode 100644
diff --git a/Engine/Source/Runtime/Online/BuildPatchServices/Private/Generation/BuildStreamer.cpp b/Engine/Source/Runtime/Online/BuildPatchServices/Private/Generation/BuildStreamer.cpp
deleted file mode 100644
diff --git a/Engine/Source/Runtime/Online/BuildPatchServices/Private/Generation/BuildStreamer.h b/Engine/Source/Runtime/Online/BuildPatchServices/Private/Generation/BuildStreamer.h
deleted file mode 100644
diff --git a/Engine/Source/Runtime/Online/BuildPatchServices/Private/Generation/CloudEnumeration.cpp b/Engine/Source/Runtime/Online/BuildPatchServices/Private/Generation/CloudEnumeration.cpp
deleted file mode 100644
diff --git a/Engine/Source/Runtime/Online/BuildPatchServices/Private/Generation/CloudEnumeration.h b/Engine/Source/Runtime/Online/BuildPatchServices/Private/Generation/CloudEnumeration.h
deleted file mode 100644
diff --git a/Engine/Source/Runtime/Online/BuildPatchServices/Private/Generation/DataMatcher.cpp b/Engine/Source/Runtime/Online/BuildPatchServices/Private/Generation/DataMatcher.cpp
deleted file mode 100644
diff --git a/Engine/Source/Runtime/Online/BuildPatchServices/Private/Generation/DataMatcher.h b/Engine/Source/Runtime/Online/BuildPatchServices/Private/Generation/DataMatcher.h
deleted file mode 100644
diff --git a/Engine/Source/Runtime/Online/BuildPatchServices/Private/Generation/DataScanner.cpp b/Engine/Source/Runtime/Online/BuildPatchServices/Private/Generation/DataScanner.cpp
deleted file mode 100644
diff --git a/Engine/Source/Runtime/Online/BuildPatchServices/Private/Generation/DataScanner.h b/Engine/Source/Runtime/Online/BuildPatchServices/Private/Generation/DataScanner.h
deleted file mode 100644
diff --git a/Engine/Source/Runtime/Online/BuildPatchServices/Private/Generation/ManifestBuilder.cpp b/Engine/Source/Runtime/Online/BuildPatchServices/Private/Generation/ManifestBuilder.cpp
deleted file mode 100644
diff --git a/Engine/Source/Runtime/Online/BuildPatchServices/Private/Generation/ManifestBuilder.h b/Engine/Source/Runtime/Online/BuildPatchServices/Private/Generation/ManifestBuilder.h
deleted file mode 100644
diff --git a/Engine/Source/Runtime/Online/BuildPatchServices/Private/Generation/StatsCollector.cpp b/Engine/Source/Runtime/Online/BuildPatchServices/Private/Generation/StatsCollector.cpp
deleted file mode 100644
diff --git a/Engine/Source/Runtime/Online/BuildPatchServices/Private/Generation/StatsCollector.h b/Engine/Source/Runtime/Online/BuildPatchServices/Private/Generation/StatsCollector.h
deleted file mode 100644
