Ability to save json in condensed format

diff --git a/Engine/Source/Runtime/Online/OnlineSubsystem/Public/OnlineJsonSerializer.h b/Engine/Source/Runtime/Online/OnlineSubsystem/Public/OnlineJsonSerializer.h
--- a/Engine/Source/Runtime/Online/OnlineSubsystem/Public/OnlineJsonSerializer.h
+++ b/Engine/Source/Runtime/Online/OnlineSubsystem/Public/OnlineJsonSerializer.h
@@ -144,11 +144,12 @@ struct FOnlineJsonSerializerBase
 /**
  * Implements the abstract serializer interface hiding the underlying writer object
  */
+template <class CharType = TCHAR, class PrintPolicy = TPrettyJsonPrintPolicy<CharType> >
 class FOnlineJsonSerializerWriter :
 	public FOnlineJsonSerializerBase
 {
 	/** The object to write the JSON output to */
-	TSharedRef<TJsonWriter<> > JsonWriter;
+	TSharedRef<TJsonWriter<CharType, PrintPolicy> > JsonWriter;
 
 public:
 
@@ -157,7 +158,7 @@ public:
 	 *
 	 * @param InJsonWriter the object to write the JSON output to
 	 */
-	FOnlineJsonSerializerWriter(TSharedRef<TJsonWriter<> > InJsonWriter) :
+	FOnlineJsonSerializerWriter(TSharedRef<TJsonWriter<CharType, PrintPolicy> > InJsonWriter) :
 		JsonWriter(InJsonWriter)
 	{
 	}
@@ -605,16 +606,28 @@ struct FOnlineJsonSerializable
 	/**
 	 * Serializes this object to its JSON string form
 	 *
+	 * @param bPrettyPrint - If true, will use the pretty json formatter
 	 * @return the corresponding json string
 	 */
-	virtual const FString ToJson()
+	virtual const FString ToJson(bool bPrettyPrint=true)
 	{
 		FString JsonStr;
-		TSharedRef<TJsonWriter<> > JsonWriter = TJsonWriterFactory<>::Create(&JsonStr);
-		FOnlineJsonSerializerWriter Serializer(JsonWriter);
-		Serialize(Serializer);
-		JsonWriter->Close();
+		if (bPrettyPrint)
+		{
+			TSharedRef<TJsonWriter<> > JsonWriter = TJsonWriterFactory<>::Create(&JsonStr);
+			FOnlineJsonSerializerWriter<> Serializer(JsonWriter);
+			Serialize(Serializer);
+			JsonWriter->Close();
+		}
+		else
+		{
+			TSharedRef< TJsonWriter< TCHAR, TCondensedJsonPrintPolicy< TCHAR > > > JsonWriter = TJsonWriterFactory< TCHAR, TCondensedJsonPrintPolicy< TCHAR > >::Create( &JsonStr );
+			FOnlineJsonSerializerWriter<TCHAR, TCondensedJsonPrintPolicy< TCHAR >> Serializer(JsonWriter);
+			Serialize(Serializer);
+			JsonWriter->Close();
+		}
 		return JsonStr;
+
 	}
 	/**
 	 * Serializes the contents of a JSON string into this object
