Merging using UE4-To-UE4-LauncherDev

#online,mcp
- fix for invalid parsing of response payload for non-errors

diff --git a/Engine/Source/Runtime/Online/HTTP/Public/Interfaces/IHttpResponse.h b/Engine/Source/Runtime/Online/HTTP/Public/Interfaces/IHttpResponse.h
--- a/Engine/Source/Runtime/Online/HTTP/Public/Interfaces/IHttpResponse.h
+++ b/Engine/Source/Runtime/Online/HTTP/Public/Interfaces/IHttpResponse.h
@@ -99,7 +99,7 @@ namespace EHttpResponseCodes
 	 */
 	inline bool IsOk(int32 StatusCode)
 	{
-		return StatusCode >= Ok && StatusCode < BadRequest;
+		return StatusCode >= Ok && StatusCode <= PartialContent;
 	}
 }
 
