Merging using UE4-To-UE4-LauncherDev

Replacing FPaths::FileExists with platform function

diff --git a/Engine/Source/Runtime/Online/BuildPatchServices/Private/BuildPatchCompactifier.cpp b/Engine/Source/Runtime/Online/BuildPatchServices/Private/BuildPatchCompactifier.cpp
--- a/Engine/Source/Runtime/Online/BuildPatchServices/Private/BuildPatchCompactifier.cpp
+++ b/Engine/Source/Runtime/Online/BuildPatchServices/Private/BuildPatchCompactifier.cpp
@@ -280,7 +280,7 @@ bool FBuildDataCompactifier::DeleteNonReferencedManifests(TArray<FString>& AllMa
 			if (!bPreview && !bNoPatchDelete)
 			{
 				IFileManager::Get().Delete(*ManifestPath);
-				if (FPaths::FileExists(ManifestPath))
+				if (FPlatformFileManager::Get().GetPlatformFile().FileExists(*ManifestPath))
 				{
 					// Something went wrong ... the file still exists!
 					GLog->Logf(ELogVerbosity::Error, TEXT("Compactify could not delete manifest file %s"), *Manifest);
diff --git a/Engine/Source/Runtime/Online/BuildPatchServices/Private/BuildPatchFileConstructor.cpp b/Engine/Source/Runtime/Online/BuildPatchServices/Private/BuildPatchFileConstructor.cpp
--- a/Engine/Source/Runtime/Online/BuildPatchServices/Private/BuildPatchFileConstructor.cpp
+++ b/Engine/Source/Runtime/Online/BuildPatchServices/Private/BuildPatchFileConstructor.cpp
@@ -52,7 +52,7 @@ public:
 		, BuildManifest( InBuildManifest )
 	{
 		// Load data from previous resume file
-		bHasResumeData = FPaths::FileExists( ResumeDataFile );
+		bHasResumeData = FPlatformFileManager::Get().GetPlatformFile().FileExists(*ResumeDataFile);
 		GLog->Logf(TEXT("BuildPatchResumeData file found %d"), bHasResumeData);
 		if( bHasResumeData )
 		{
diff --git a/Engine/Source/Runtime/Online/BuildPatchServices/Private/BuildPatchGeneration.cpp b/Engine/Source/Runtime/Online/BuildPatchServices/Private/BuildPatchGeneration.cpp
--- a/Engine/Source/Runtime/Online/BuildPatchServices/Private/BuildPatchGeneration.cpp
+++ b/Engine/Source/Runtime/Online/BuildPatchServices/Private/BuildPatchGeneration.cpp
@@ -955,11 +955,11 @@ bool FBuildDataGenerator::SaveOutFileData(const FString& SourceFile, const FSHAH
 	IFileManager& FileManager = IFileManager::Get();
 
 	const FString NewFilename = FBuildPatchUtils::GetFileNewFilename( EBuildPatchAppManifestVersion::GetLatestVersion(), FBuildPatchServicesModule::GetCloudDirectory(), FileGuid, FileHash );
-	bAlreadySaved = FPaths::FileExists( NewFilename );
+	bAlreadySaved = FPlatformFileManager::Get().GetPlatformFile().FileExists(*NewFilename);
 
 #if SAVE_OLD_FILEDATA_FILENAMES
 	const FString OldFilename = FBuildPatchUtils::GetFileOldFilename( FBuildPatchServicesModule::GetCloudDirectory(), FileGuid );
-	bAlreadySaved = bAlreadySaved && FPaths::FileExists( OldFilename );
+	bAlreadySaved = bAlreadySaved && FPlatformFileManager::Get().GetPlatformFile().FileExists(*OldFilename);
 #endif
 
 	if( !bAlreadySaved )
diff --git a/Engine/Source/Runtime/Online/BuildPatchServices/Private/BuildPatchInstaller.cpp b/Engine/Source/Runtime/Online/BuildPatchServices/Private/BuildPatchInstaller.cpp
--- a/Engine/Source/Runtime/Online/BuildPatchServices/Private/BuildPatchInstaller.cpp
+++ b/Engine/Source/Runtime/Online/BuildPatchServices/Private/BuildPatchInstaller.cpp
@@ -552,7 +552,7 @@ bool FBuildPatchInstaller::RunBackupAndMove()
 				const FString DestFilename = InstallDirectory / ConstructionFile;
 				const float FileIndexFloat = ConstructionFilesIt.GetIndex();
 				// Skip files not constructed
-				if (!FPaths::FileExists(SrcFilename))
+				if (!FPlatformFileManager::Get().GetPlatformFile().FileExists(*SrcFilename))
 				{
 					BuildProgress.SetStateProgress(EBuildPatchProgress::MovingToInstall, FileIndexFloat / NumConstructionFilesFloat);
 					continue;
@@ -715,13 +715,13 @@ bool FBuildPatchInstaller::BackupFileIfNecessary(const FString& Filename, bool b
 		return true;
 	}
 	// Skip if no file to backup
-	const bool bInstalledFileExists = FPaths::FileExists(InstalledFilename);
+	const bool bInstalledFileExists = FPlatformFileManager::Get().GetPlatformFile().FileExists(*InstalledFilename);
 	if (!bInstalledFileExists)
 	{
 		return true;
 	}
 	// Skip if already backed up
-	const bool bAlreadyBackedUp = FPaths::FileExists(BackupFilename);
+	const bool bAlreadyBackedUp = FPlatformFileManager::Get().GetPlatformFile().FileExists(*BackupFilename);
 	if (bAlreadyBackedUp)
 	{
 		return true;
