Allow HTTP redirects to be followed in the CurlHTTPRequest implementation

Fixes https://jira.ol.epicgames.net/browse/UE-16988

diff --git a/Engine/Source/Runtime/Online/HTTP/Private/Curl/CurlHttp.cpp b/Engine/Source/Runtime/Online/HTTP/Private/Curl/CurlHttp.cpp
--- a/Engine/Source/Runtime/Online/HTTP/Private/Curl/CurlHttp.cpp
+++ b/Engine/Source/Runtime/Online/HTTP/Private/Curl/CurlHttp.cpp
@@ -46,6 +46,9 @@ FCurlHttpRequest::FCurlHttpRequest(CURLM * InMultiHandle)
 			curl_easy_setopt(EasyHandle, CURLOPT_SSL_VERIFYPEER, 0L);
 		}
 
+		// allow http redirects to be followed
+		curl_easy_setopt(EasyHandle, CURLOPT_FOLLOWLOCATION, 1L);
+
 		// required for all multi-threaded handles
 		curl_easy_setopt(EasyHandle, CURLOPT_NOSIGNAL, 1L);
 
