Changed the order of AudioUnit and AudioToolbox linking in two more places to solve issues with Mac 10.11 SDK

diff --git a/Engine/Source/Runtime/Mac/UnrealAudioCoreAudio/UnrealAudioCoreAudio.Build.cs b/Engine/Source/Runtime/Mac/UnrealAudioCoreAudio/UnrealAudioCoreAudio.Build.cs
--- a/Engine/Source/Runtime/Mac/UnrealAudioCoreAudio/UnrealAudioCoreAudio.Build.cs
+++ b/Engine/Source/Runtime/Mac/UnrealAudioCoreAudio/UnrealAudioCoreAudio.Build.cs
@@ -25,8 +25,8 @@ public class UnrealAudioCoreAudio : ModuleRules
 		PublicFrameworks.AddRange(
 			new string[] { 
 				"CoreAudio", 
-				"AudioToolbox", 
-				"AudioUnit" 
+				"AudioUnit",
+				"AudioToolbox"
 			}
 		);
 	}
diff --git a/Engine/Source/Runtime/Online/Voice/Voice.Build.cs b/Engine/Source/Runtime/Online/Voice/Voice.Build.cs
--- a/Engine/Source/Runtime/Online/Voice/Voice.Build.cs
+++ b/Engine/Source/Runtime/Online/Voice/Voice.Build.cs
@@ -27,7 +27,7 @@ public class Voice : ModuleRules
 		}
 		else if(Target.Platform == UnrealTargetPlatform.Mac)
 		{
-			PublicFrameworks.AddRange(new string[] { "CoreAudio", "AudioToolbox", "AudioUnit" });
+			PublicFrameworks.AddRange(new string[] { "CoreAudio", "AudioUnit", "AudioToolbox" });
 		}
 
 		AddThirdPartyPrivateStaticDependencies(Target, "libOpus");
