Merging using UE4-LauncherDev-To-UE4-LauncherReleases

Adding the setting of file attributes to contribute to an Installing phase of a installer. This fixes an issue where Initialising or Verifying appears to take a long time.
Also optimising the file attributes setting code by reading file attributes before attemping to change them, this has a large effect on speed as most files will be setup already, especially for a repair operation.

diff --git a/Engine/Source/Runtime/Online/BuildPatchServices/Private/BuildPatchFileAttributes.cpp b/Engine/Source/Runtime/Online/BuildPatchServices/Private/BuildPatchFileAttributes.cpp
--- a/Engine/Source/Runtime/Online/BuildPatchServices/Private/BuildPatchFileAttributes.cpp
+++ b/Engine/Source/Runtime/Online/BuildPatchServices/Private/BuildPatchFileAttributes.cpp
@@ -5,85 +5,56 @@
 class FBuildPatchFileAttributesImpl : public FBuildPatchFileAttributes
 {
 public:
-	FBuildPatchFileAttributesImpl(const FBuildPatchAppManifestRef& NewManifest, const FBuildPatchAppManifestPtr& OldManifest, const FString& VerifyDirectory, const FString& StagedFileDirectory, bool bUseStageDirectory);
+	FBuildPatchFileAttributesImpl(const FBuildPatchAppManifestRef& NewManifest, const FBuildPatchAppManifestPtr& OldManifest, const FString& VerifyDirectory, const FString& StagedFileDirectory, bool bUseStageDirectory, FBuildPatchProgress* BuildProgress);
 
 	virtual ~FBuildPatchFileAttributesImpl(){}
 
 	virtual bool ApplyAttributes(bool bForce) override;
 
+private:
+	FString SelectFullFilePath(const FString& BuildFile);
+	void SetupFileAttributes(const FString& FilePath, const FFileManifestData& FileManifest);
+	bool GetCurrentFileAttributes(const FString& FilePath, bool& OutFileExists, bool& OutIsReadonly, bool& OutIsCompressed, bool& OutIsUnixExecutable);
+	bool SetFileReadOnlyFlag(const FString& FilePath, bool IsReadOnly);
+	bool SetFileCompressionFlag(const FString& FilePath, bool IsCompressed);
+	bool SetFileUnixExecutableFlag(const FString& FilePath, bool IsUnixExecutable);
+
 private:
 	FBuildPatchAppManifestRef NewManifest;
 	FBuildPatchAppManifestPtr OldManifest;
 	const FString VerifyDirectory;
 	const FString StagedFileDirectory;
 	const bool bUseStageDirectory;
-
-	FString SelectFullFilePath(const FString& BuildFile);
-	void SetupFileAttributes(const FString& FilePath, const FFileManifestData& FileManifest);
-	bool SetFileCompressionFlag(const FString& Filepath, bool bIsCompressed);
-	bool SetExecutableFlag(const FString& Filepath);
+	FBuildPatchProgress* BuildProgress;
 };
 
-FBuildPatchFileAttributesImpl::FBuildPatchFileAttributesImpl(const FBuildPatchAppManifestRef& InNewManifest, const FBuildPatchAppManifestPtr& InOldManifest, const FString& InVerifyDirectory, const FString& InStagedFileDirectory, bool bInUseStageDirectory)
+FBuildPatchFileAttributesImpl::FBuildPatchFileAttributesImpl(const FBuildPatchAppManifestRef& InNewManifest, const FBuildPatchAppManifestPtr& InOldManifest, const FString& InVerifyDirectory, const FString& InStagedFileDirectory, bool bInUseStageDirectory, FBuildPatchProgress* InBuildProgress)
 	: NewManifest(InNewManifest)
 	, OldManifest(InOldManifest)
 	, VerifyDirectory(InVerifyDirectory)
 	, StagedFileDirectory(InStagedFileDirectory)
 	, bUseStageDirectory(bInUseStageDirectory)
-{}
+	, BuildProgress(InBuildProgress)
+{
+	BuildProgress->SetStateProgress(EBuildPatchProgress::SettingAttributes, 0.0f);
+}
 
-bool FBuildPatchFileAttributesImpl::ApplyAttributes(bool bForce)
+bool FBuildPatchFileAttributesImpl::ApplyAttributes(bool)
 {
 	// We need to set attributes for all files in the new build that require it
-    {
-        TArray<FString> BuildFiles;
-        NewManifest->GetFileList(BuildFiles);
-        for (const FString& BuildFile : BuildFiles)
-        {
-            // Break if quitting
-            if (FBuildPatchInstallError::HasFatalError())
-            {
-                break;
-            }
-
-            const FFileManifestData* FileManifest = NewManifest->GetFileManifest(BuildFile);
-            if (FileManifest != nullptr)
-            {
-                // Apply
-                const bool bHasAttrib = FileManifest->bIsReadOnly || FileManifest->bIsCompressed || FileManifest->bIsUnixExecutable;
-                if (bHasAttrib || bForce)
-                {
-                    SetupFileAttributes(SelectFullFilePath(BuildFile), *FileManifest);
-                }
-            }
-        }
-    }
-
-	// We also need to check if any attributes have been removed, unless we forced anyway
-	if (OldManifest.IsValid() && !bForce)
+	TArray<FString> BuildFileList = NewManifest->GetBuildFileList();
+	BuildProgress->SetStateProgress(EBuildPatchProgress::SettingAttributes, 0.0f);
+	for (int32 BuildFileIdx = 0; BuildFileIdx < BuildFileList.Num() && !FBuildPatchInstallError::HasFatalError(); ++BuildFileIdx)
 	{
-		TArray<FString> BuildFiles;
-		OldManifest->GetFileList(BuildFiles);
-		for (const FString& BuildFile : BuildFiles)
+		const FString& BuildFile = BuildFileList[BuildFileIdx];
+		const FFileManifestData* FileManifest = NewManifest->GetFileManifest(BuildFile);
+		if (FileManifest != nullptr)
 		{
-			// Break if quitting
-			if (FBuildPatchInstallError::HasFatalError())
-			{
-				break;
-			}
-
-			const FFileManifestData* NewFileManifest = NewManifest->GetFileManifest(BuildFile);
-			const FFileManifestData* OldFileManifest = OldManifest->GetFileManifest(BuildFile);
-			if (NewFileManifest != nullptr)
-			{
-				const bool bAttribChanged = (OldFileManifest->bIsReadOnly && !NewFileManifest->bIsReadOnly) || (OldFileManifest->bIsCompressed && !NewFileManifest->bIsCompressed);
-				if (bAttribChanged)
-				{
-					SetupFileAttributes(SelectFullFilePath(BuildFile), *NewFileManifest);
-				}
-			}
+			SetupFileAttributes(SelectFullFilePath(BuildFile), *FileManifest);
 		}
+		BuildProgress->SetStateProgress(EBuildPatchProgress::SettingAttributes, BuildFileIdx / float(BuildFileList.Num()));
 	}
+	BuildProgress->SetStateProgress(EBuildPatchProgress::SettingAttributes, 1.0f);
 
 	// We don't fail on this step currently
 	return true;
@@ -105,18 +76,59 @@ FString FBuildPatchFileAttributesImpl::SelectFullFilePath(const FString& BuildFi
 
 void FBuildPatchFileAttributesImpl::SetupFileAttributes(const FString& FilePath, const FFileManifestData& FileManifest)
 {
-	// File must not be readonly to be able to set attributes
-	IPlatformFile::GetPlatformPhysical().SetReadOnly(*FilePath, false);
-	// Set correct attributes
-	SetFileCompressionFlag(FilePath, FileManifest.bIsCompressed);
-	if (!IPlatformFile::GetPlatformPhysical().SetReadOnly(*FilePath, FileManifest.bIsReadOnly))
+	bool FileExists;
+	bool IsReadOnly;
+	bool IsCompressed;
+	bool IsUnixExecutable;
+
+	// First check file attributes as it's much faster to read and do nothing
+	bool KnownAttributes = GetCurrentFileAttributes(FilePath, FileExists, IsReadOnly, IsCompressed, IsUnixExecutable);
+
+	// If we know the file is missing, skip out
+	if (KnownAttributes && !FileExists)
 	{
-		GLog->Logf(TEXT("BuildPatchServices: WARNING: Could not set readonly flag %s"), *FilePath);
+		return;
+	}
+
+	// Set compression attribute
+	if (!KnownAttributes || IsCompressed != FileManifest.bIsCompressed)
+	{
+		// Must make not readonly if required
+		if (!KnownAttributes || IsReadOnly)
+		{
+			IsReadOnly = false;
+			SetFileReadOnlyFlag(FilePath, false);
+		}
+		SetFileCompressionFlag(FilePath, FileManifest.bIsCompressed);
 	}
-	if (FileManifest.bIsUnixExecutable && !SetExecutableFlag(FilePath))
+
+	// Set executable attribute
+	if (!KnownAttributes || IsUnixExecutable != FileManifest.bIsUnixExecutable)
+	{
+		// Must make not readonly if required
+		if (!KnownAttributes || IsReadOnly)
+		{
+			IsReadOnly = false;
+			SetFileReadOnlyFlag(FilePath, false);
+		}
+		SetFileUnixExecutableFlag(FilePath, FileManifest.bIsUnixExecutable);
+	}
+
+	// Set readonly attribute
+	if (!KnownAttributes || IsReadOnly != FileManifest.bIsReadOnly)
+	{
+		SetFileReadOnlyFlag(FilePath, FileManifest.bIsReadOnly);
+	}
+}
+
+bool FBuildPatchFileAttributesImpl::SetFileReadOnlyFlag(const FString& FilePath, bool IsReadOnly)
+{
+	if (!IPlatformFile::GetPlatformPhysical().SetReadOnly(*FilePath, IsReadOnly))
 	{
-		GLog->Logf(TEXT("BuildPatchServices: WARNING: Could not set executable flag %s"), *FilePath);
+		GLog->Logf(TEXT("BuildPatchServices: WARNING: Could not set readonly flag %s"), *FilePath);
+		return false;
 	}
+	return true;
 }
 
 #if PLATFORM_WINDOWS
@@ -128,11 +140,34 @@ void FBuildPatchFileAttributesImpl::SetupFileAttributes(const FString& FilePath,
 #include <fileapi.h>
 #include <winioctl.h>
 
-bool FBuildPatchFileAttributesImpl::SetFileCompressionFlag(const FString& Filepath, bool bIsCompressed)
+bool FBuildPatchFileAttributesImpl::GetCurrentFileAttributes(const FString& FilePath, bool& OutFileExists, bool& OutIsReadonly, bool& OutIsCompressed, bool& OutIsUnixExecutable)
+{
+	DWORD FileAttributes = ::GetFileAttributes(*FilePath);
+	DWORD Error = ::GetLastError();
+	if (FileAttributes != INVALID_FILE_ATTRIBUTES)
+	{
+		OutFileExists = true;
+		OutIsReadonly = (FileAttributes & FILE_ATTRIBUTE_READONLY) == FILE_ATTRIBUTE_READONLY;
+		OutIsCompressed = (FileAttributes & FILE_ATTRIBUTE_COMPRESSED) == FILE_ATTRIBUTE_COMPRESSED;
+		OutIsUnixExecutable = false;
+		return true;
+	}
+	else if (Error == ERROR_PATH_NOT_FOUND || Error == ERROR_FILE_NOT_FOUND)
+	{
+		OutFileExists = false;
+		OutIsReadonly = false;
+		OutIsCompressed = false;
+		OutIsUnixExecutable = false;
+		return true;
+	}
+	return false;
+}
+
+bool FBuildPatchFileAttributesImpl::SetFileCompressionFlag(const FString& FilePath, bool bIsCompressed)
 {
 	// Get the file handle
 	HANDLE FileHandle = ::CreateFile(
-		*Filepath,								// Path to file
+		*FilePath,								// Path to file
 		GENERIC_READ | GENERIC_WRITE,			// Read and write access
 		FILE_SHARE_READ | FILE_SHARE_WRITE,		// Share access for DeviceIoControl
 		NULL,									// Default security
@@ -143,7 +178,7 @@ bool FBuildPatchFileAttributesImpl::SetFileCompressionFlag(const FString& Filepa
 	uint32 Error = ::GetLastError();
 	if (FileHandle == NULL || FileHandle == INVALID_HANDLE_VALUE)
 	{
-		GLog->Logf(TEXT("BuildPatchServices: WARNING: Could not open file to set compression flag %d Error:%d File:%s"), bIsCompressed, Error, *Filepath);
+		GLog->Logf(TEXT("BuildPatchServices: WARNING: Could not open file to set compression flag %d Error:%d File:%s"), bIsCompressed, Error, *FilePath);
 		return false;
 	}
 
@@ -152,21 +187,21 @@ bool FBuildPatchFileAttributesImpl::SetFileCompressionFlag(const FString& Filepa
 	uint16* pMessage = &Message;
 	DWORD Dummy = 0;
 	LPDWORD pDummy = &Dummy;
-	BOOL bSuccess = ::DeviceIoControl(
+	bool bSuccess = ::DeviceIoControl(
 		FileHandle,								// The file handle
 		FSCTL_SET_COMPRESSION,					// Control code
 		pMessage,								// Our message
-		sizeof(uint16),							//
+		sizeof(uint16),							// Our message size
 		NULL,									// Not used
 		0,										// Not used
 		pDummy,									// The value returned will be meaningless, but is required
 		NULL									// No overlap structure, we a running this synchronously
-		);
+		) == TRUE;
 	Error = ::GetLastError();
 	const bool bFileSystemUnsupported = Error == ERROR_INVALID_FUNCTION;
-	if (bSuccess == FALSE && bFileSystemUnsupported == false)
+	if (bSuccess == false && bFileSystemUnsupported == false)
 	{
-		GLog->Logf(TEXT("BuildPatchServices: WARNING: Could not set compression flag %d Error:%d File:%s"), bIsCompressed, Error, *Filepath);
+		GLog->Logf(TEXT("BuildPatchServices: WARNING: Could not set compression flag %d Error:%d File:%s"), bIsCompressed, Error, *FilePath);
 	}
 
 	// Close the open file handle
@@ -176,7 +211,7 @@ bool FBuildPatchFileAttributesImpl::SetFileCompressionFlag(const FString& Filepa
 	return bSuccess == TRUE || bFileSystemUnsupported;
 }
 
-bool FBuildPatchFileAttributesImpl::SetExecutableFlag(const FString& Filepath)
+bool FBuildPatchFileAttributesImpl::SetFileUnixExecutableFlag(const FString& FilePath, bool IsUnixExecutable)
 {
 	// Not implemented
 	return true;
@@ -185,38 +220,71 @@ bool FBuildPatchFileAttributesImpl::SetExecutableFlag(const FString& Filepath)
 #include "HideWindowsPlatformTypes.h"
 
 #elif PLATFORM_MAC
-bool FBuildPatchFileAttributesImpl::SetFileCompressionFlag(const FString& Filepath, bool bIsCompressed)
+bool FBuildPatchFileAttributesImpl::GetCurrentFileAttributes(const FString& FilePath, bool& OutFileExists, bool& OutIsReadonly, bool& OutIsCompressed, bool& OutIsUnixExecutable)
+{
+	struct stat FileInfo;
+	int32 Result = stat(TCHAR_TO_UTF8(*FilePath), &FileInfo);
+	int32 Error = errno;
+	if (Result == 0)
+	{
+		mode_t ExeFlags = S_IXUSR | S_IXGRP | S_IXOTH;
+		OutFileExists = true;
+		OutIsReadonly = (FileInfo.st_mode & S_IWUSR) == 0;
+		OutIsCompressed = false;
+		OutIsUnixExecutable = (FileInfo.st_mode & ExeFlags) == ExeFlags;
+		return true;
+	}
+	else if (Error == ENOTDIR || Error == ENOENT)
+	{
+		OutFileExists = false;
+		OutIsReadonly = false;
+		OutIsCompressed = false;
+		OutIsUnixExecutable = false;
+		return true;
+	}
+	return false;
+}
+
+bool FBuildPatchFileAttributesImpl::SetFileCompressionFlag(const FString& FilePath, bool bIsCompressed)
 {
 	// Not implemented
 	return true;
 }
 
-bool FBuildPatchFileAttributesImpl::SetExecutableFlag(const FString& Filepath)
+bool FBuildPatchFileAttributesImpl::SetFileUnixExecutableFlag(const FString& FilePath, bool IsUnixExecutable)
 {
 	bool bSuccess = false;
-	// Enable executable permission bit
+	// Set executable permission bit
 	struct stat FileInfo;
-	if (stat(TCHAR_TO_UTF8(*Filepath), &FileInfo) == 0)
+	if (stat(TCHAR_TO_UTF8(*FilePath), &FileInfo) == 0)
 	{
-		bSuccess = chmod(TCHAR_TO_UTF8(*Filepath), FileInfo.st_mode | S_IXUSR | S_IXGRP | S_IXOTH) == 0;
+		mode_t ExeFlags = S_IXUSR | S_IXGRP | S_IXOTH;
+		FileInfo.st_mode = IsUnixExecutable ? (FileInfo.st_mode | ExeFlags) : (FileInfo.st_mode & ~ExeFlags);
+		bSuccess = chmod(TCHAR_TO_UTF8(*FilePath), FileInfo.st_mode) == 0;
 	}
 	return bSuccess;
 }
 #elif PLATFORM_LINUX
-bool FBuildPatchFileAttributesImpl::SetFileCompressionFlag(const FString& Filepath, bool bIsCompressed)
+bool FBuildPatchFileAttributesImpl::GetCurrentFileAttributes(const FString& FilePath, bool& OutFileExists, bool& OutIsReadonly, bool& OutIsCompressed, bool& OutIsUnixExecutable)
+{
+	// Not implemented
+	return true;
+}
+
+bool FBuildPatchFileAttributesImpl::SetFileCompressionFlag(const FString& FilePath, bool bIsCompressed)
 {
 	// Not implemented
 	return true;
 }
 
-bool FBuildPatchFileAttributesImpl::SetExecutableFlag(const FString& Filepath)
+bool FBuildPatchFileAttributesImpl::SetFileUnixExecutableFlag(const FString& FilePath, bool IsUnixExecutable)
 {
 	// Not implemented
 	return true;
 }
 #endif
 
-TSharedRef<FBuildPatchFileAttributes> FBuildPatchFileAttributesFactory::Create(const FBuildPatchAppManifestRef& NewManifest, const FBuildPatchAppManifestPtr& OldManifest, const FString& VerifyDirectory, const FString& StagedFileDirectory)
+TSharedRef<FBuildPatchFileAttributes> FBuildPatchFileAttributesFactory::Create(const FBuildPatchAppManifestRef& NewManifest, const FBuildPatchAppManifestPtr& OldManifest, const FString& VerifyDirectory, const FString& StagedFileDirectory, FBuildPatchProgress* BuildProgress)
 {
-	return MakeShareable(new FBuildPatchFileAttributesImpl(NewManifest, OldManifest, VerifyDirectory, StagedFileDirectory, !StagedFileDirectory.IsEmpty()));
+	return MakeShareable(new FBuildPatchFileAttributesImpl(NewManifest, OldManifest, VerifyDirectory, StagedFileDirectory, !StagedFileDirectory.IsEmpty(), BuildProgress));
 }
diff --git a/Engine/Source/Runtime/Online/BuildPatchServices/Private/BuildPatchFileAttributes.h b/Engine/Source/Runtime/Online/BuildPatchServices/Private/BuildPatchFileAttributes.h
--- a/Engine/Source/Runtime/Online/BuildPatchServices/Private/BuildPatchFileAttributes.h
+++ b/Engine/Source/Runtime/Online/BuildPatchServices/Private/BuildPatchFileAttributes.h
@@ -30,5 +30,6 @@ public:
 	static TSharedRef<FBuildPatchFileAttributes> Create(const FBuildPatchAppManifestRef& NewManifest,
 	                                                    const FBuildPatchAppManifestPtr& OldManifest,
 	                                                    const FString& VerifyDirectory,
-	                                                    const FString& StagedFileDirectory);
+	                                                    const FString& StagedFileDirectory,
+	                                                    FBuildPatchProgress* BuildProgress);
 };
diff --git a/Engine/Source/Runtime/Online/BuildPatchServices/Private/BuildPatchInstaller.cpp b/Engine/Source/Runtime/Online/BuildPatchServices/Private/BuildPatchInstaller.cpp
--- a/Engine/Source/Runtime/Online/BuildPatchServices/Private/BuildPatchInstaller.cpp
+++ b/Engine/Source/Runtime/Online/BuildPatchServices/Private/BuildPatchInstaller.cpp
@@ -127,12 +127,14 @@ uint32 FBuildPatchInstaller::Run()
 		// Backup local changes then move generated files
 		bInstallSuccess = bInstallSuccess && RunBackupAndMove();
 
+		// There is no more potential for initialising
+		BuildProgress.SetStateProgress(EBuildPatchProgress::Initializing, 1.0f);
+
 		// Setup file attributes
 		bInstallSuccess = bInstallSuccess && RunFileAttributes(bIsRepairing);
 
 		// Run Verification
 		CorruptFiles.Empty();
-		BuildProgress.SetStateProgress(EBuildPatchProgress::Initializing, 1.0f);
 		bProcessSuccess = bInstallSuccess && RunVerification(CorruptFiles);
 
 		// Clean staging if INSTALL success
@@ -293,6 +295,7 @@ bool FBuildPatchInstaller::RunInstallation(TArray<FString>& CorruptFiles)
 		BuildProgress.SetStateWeight(EBuildPatchProgress::Downloading, 0.0f);
 		BuildProgress.SetStateWeight(EBuildPatchProgress::Installing, 0.0f);
 		BuildProgress.SetStateWeight(EBuildPatchProgress::MovingToInstall, 0.0f);
+		BuildProgress.SetStateWeight(EBuildPatchProgress::SettingAttributes, 0.2f);
 		BuildProgress.SetStateWeight(EBuildPatchProgress::BuildVerification, 1.0f);
 		// Mark all installation steps complete
 		BuildProgress.SetStateProgress(EBuildPatchProgress::Initializing, 1.0f);
@@ -359,10 +362,12 @@ bool FBuildPatchInstaller::RunInstallation(TArray<FString>& CorruptFiles)
 
 	// Setup some weightings for the progress tracking
 	const float NumRequiredChunksFloat = NumRequiredChunks;
+	const bool bHasFileAttributes = NewBuildManifest->HasFileAttributes();
+	const float AttributesWeight = bHasFileAttributes ? bIsRepairing ? 1.0f / 50.0f : 1.0f / 20.0f : 0.0f;
 	BuildProgress.SetStateWeight(EBuildPatchProgress::Downloading, NumRequiredChunksFloat > 0.0f ? InitialNumChunkDownloads / NumRequiredChunksFloat : 0.0f);
 	BuildProgress.SetStateWeight(EBuildPatchProgress::Installing, NumRequiredChunksFloat > 0.0f ? 0.1f + (InitialNumChunkConstructions / NumRequiredChunksFloat) : 0.0f);
 	BuildProgress.SetStateWeight(EBuildPatchProgress::MovingToInstall, NumFilesToConstruct > 0 ? 0.05f : 0.0f);
-	// A verify weight of 1 / 9 will make it 10% of the total progress
+	BuildProgress.SetStateWeight(EBuildPatchProgress::SettingAttributes, AttributesWeight);
 	BuildProgress.SetStateWeight(EBuildPatchProgress::BuildVerification, 1.1f / 9.0f);
 
 	// If this is a repair operation, start off with install and download complete
@@ -634,7 +639,7 @@ bool FBuildPatchInstaller::RunFileAttributes(bool bForce)
 	FString& OptionalStageDirectory = bShouldStageOnly ? InstallStagingDir : EmptyString;
 
 	// Construct the attributes class
-	auto Attributes = FBuildPatchFileAttributesFactory::Create(NewBuildManifest, CurrentBuildManifest, InstallDirectory, OptionalStageDirectory);
+	auto Attributes = FBuildPatchFileAttributesFactory::Create(NewBuildManifest, CurrentBuildManifest, InstallDirectory, OptionalStageDirectory, &BuildProgress);
 	Attributes->ApplyAttributes(bForce);
 
 	// We don't fail on this step currently
diff --git a/Engine/Source/Runtime/Online/BuildPatchServices/Private/BuildPatchManifest.cpp b/Engine/Source/Runtime/Online/BuildPatchServices/Private/BuildPatchManifest.cpp
--- a/Engine/Source/Runtime/Online/BuildPatchServices/Private/BuildPatchManifest.cpp
+++ b/Engine/Source/Runtime/Online/BuildPatchServices/Private/BuildPatchManifest.cpp
@@ -1492,6 +1492,19 @@ void FBuildPatchAppManifest::EnumerateChunkPartInventory(const TArray<FGuid>& Ch
 	}
 }
 
+bool FBuildPatchAppManifest::HasFileAttributes() const
+{
+	for (const auto& FileManifestPair : FileManifestLookup)
+	{
+		const FFileManifestData* FileManifest = FileManifestPair.Value;
+		if (FileManifest->bIsReadOnly || FileManifest->bIsUnixExecutable || FileManifest->bIsCompressed)
+		{
+			return true;
+		}
+	}
+	return false;
+}
+
 bool FBuildPatchAppManifest::IsSameAs(FBuildPatchAppManifestRef Other) const
 {
 	return this == &Other.Get() || (GetAppID() == Other->GetAppID() && GetAppName() == Other->GetAppName() && GetVersionString() == Other->GetVersionString());
diff --git a/Engine/Source/Runtime/Online/BuildPatchServices/Private/BuildPatchManifest.h b/Engine/Source/Runtime/Online/BuildPatchServices/Private/BuildPatchManifest.h
--- a/Engine/Source/Runtime/Online/BuildPatchServices/Private/BuildPatchManifest.h
+++ b/Engine/Source/Runtime/Online/BuildPatchServices/Private/BuildPatchManifest.h
@@ -571,6 +571,9 @@ public:
 	 */
 	void EnumerateChunkPartInventory(const TArray< FGuid >& ChunksRequired, TMap< FGuid, TArray< FFileChunkPart > >& ChunkPartsAvailable) const;
 
+	/** @return True if any files in this manifest have file attributes to be set */
+	bool HasFileAttributes() const;
+
 	/** @return True if this manifest is for the same build, i.e. same ID, Name, and Version */
 	bool IsSameAs(FBuildPatchAppManifestRef InstallManifest) const;
 public:
diff --git a/Engine/Source/Runtime/Online/BuildPatchServices/Private/BuildPatchProgress.cpp b/Engine/Source/Runtime/Online/BuildPatchServices/Private/BuildPatchProgress.cpp
--- a/Engine/Source/Runtime/Online/BuildPatchServices/Private/BuildPatchProgress.cpp
+++ b/Engine/Source/Runtime/Online/BuildPatchServices/Private/BuildPatchProgress.cpp
@@ -20,6 +20,7 @@ const FString& EBuildPatchProgress::ToString(const EBuildPatchProgress::Type& Pr
 	static const FString Recycling(TEXT("Recycling"));
 	static const FString Installing(TEXT("Installing"));
 	static const FString MovingToInstall(TEXT("MovingToInstall"));
+	static const FString SettingAttributes(TEXT("SettingAttributes"));
 	static const FString BuildVerification(TEXT("BuildVerification"));
 	static const FString CleanUp(TEXT("CleanUp"));
 	static const FString PrerequisitesInstall(TEXT("PrerequisitesInstall"));
@@ -40,6 +41,8 @@ const FString& EBuildPatchProgress::ToString(const EBuildPatchProgress::Type& Pr
 			return Installing;
 		case EBuildPatchProgress::MovingToInstall:
 			return MovingToInstall;
+		case EBuildPatchProgress::SettingAttributes:
+			return SettingAttributes;
 		case EBuildPatchProgress::BuildVerification:
 			return BuildVerification;
 		case EBuildPatchProgress::CleanUp:
@@ -63,7 +66,6 @@ const FText& EBuildPatchProgress::ToText(const EBuildPatchProgress::Type& Progre
 	static const FText Initializing = LOCTEXT("EBuildPatchProgress_Initialising", "Initializing");
 	static const FText Resuming = LOCTEXT("EBuildPatchProgress_Resuming", "Resuming");
 	static const FText Downloading = LOCTEXT("EBuildPatchProgress_Downloading", "Downloading");
-	static const FText Recycling = LOCTEXT("EBuildPatchProgress_Recycling", "Recycling Content");
 	static const FText Installing = LOCTEXT("EBuildPatchProgress_Installing", "Installing");
 	static const FText BuildVerification = LOCTEXT("EBuildPatchProgress_BuildVerification", "Verifying");
 	static const FText CleanUp = LOCTEXT("EBuildPatchProgress_CleanUp", "Cleaning up");
@@ -82,7 +84,8 @@ const FText& EBuildPatchProgress::ToText(const EBuildPatchProgress::Type& Progre
 		case EBuildPatchProgress::Downloading:
 			return Downloading;
 		case EBuildPatchProgress::Installing:
-		case EBuildPatchProgress::MovingToInstall: //User should see same as "Installing"
+		case EBuildPatchProgress::MovingToInstall:
+		case EBuildPatchProgress::SettingAttributes:
 			return Installing;
 		case EBuildPatchProgress::BuildVerification:
 			return BuildVerification;
@@ -290,6 +293,7 @@ const bool FBuildPatchProgress::bHasProgressValue[EBuildPatchProgress::NUM_PROGR
 	true,  // Downloading
 	true,  // Installing
 	true,  // MovingToInstall
+	true,  // SettingAttributes
 	true,  // BuildVerification
 	false, // CleanUp
 	false, // PrerequisitesInstall
@@ -305,6 +309,7 @@ const bool FBuildPatchProgress::bCountsTowardsProgress[EBuildPatchProgress::NUM_
 	true,  // Downloading
 	true,  // Installing
 	true,  // MovingToInstall
+	true,  // SettingAttributes
 	true,  // BuildVerification
 	false, // CleanUp
 	false, // PrerequisitesInstall
diff --git a/Engine/Source/Runtime/Online/BuildPatchServices/Private/BuildPatchProgress.h b/Engine/Source/Runtime/Online/BuildPatchServices/Private/BuildPatchProgress.h
--- a/Engine/Source/Runtime/Online/BuildPatchServices/Private/BuildPatchProgress.h
+++ b/Engine/Source/Runtime/Online/BuildPatchServices/Private/BuildPatchProgress.h
@@ -29,6 +29,9 @@ namespace EBuildPatchProgress
 		// The patch process is moving staged files to the install
 		MovingToInstall,
 
+		// The patch process is setting up attributes on the build
+		SettingAttributes,
+
 		// The patch process is verifying the build
 		BuildVerification,
 
