Merging using UE4-To-UE4-LauncherDev

diff --git a/Engine/Source/Runtime/Online/HTTP/Private/HttpManager.cpp b/Engine/Source/Runtime/Online/HTTP/Private/HttpManager.cpp
--- a/Engine/Source/Runtime/Online/HTTP/Private/HttpManager.cpp
+++ b/Engine/Source/Runtime/Online/HTTP/Private/HttpManager.cpp
@@ -19,6 +19,31 @@ FHttpManager::~FHttpManager()
 	
 }
 
+void FHttpManager::Flush(bool bShutdown)
+{
+	FScopeLock ScopeLock(&RequestLock);
+
+	if (bShutdown)
+	{
+		// Clear delegates since they may point to deleted instances
+		for (TArray<TSharedRef<IHttpRequest>>::TIterator It(Requests); It; ++It)
+		{
+			TSharedRef<IHttpRequest> Request = *It;
+			Request->OnProcessRequestComplete().Unbind();
+		}
+	}
+
+	// block until all active requests have completed
+	double LastTime = FPlatformTime::Seconds();
+	while (Requests.Num() > 0)
+	{
+		const double AppTime = FPlatformTime::Seconds();
+		Tick(AppTime - LastTime);
+		LastTime = AppTime;
+		FPlatformProcess::Sleep(0.5f);
+	}
+}
+
 bool FHttpManager::Tick(float DeltaSeconds)
 {
 	FScopeLock ScopeLock(&RequestLock);
diff --git a/Engine/Source/Runtime/Online/HTTP/Private/HttpModule.cpp b/Engine/Source/Runtime/Online/HTTP/Private/HttpModule.cpp
--- a/Engine/Source/Runtime/Online/HTTP/Private/HttpModule.cpp
+++ b/Engine/Source/Runtime/Online/HTTP/Private/HttpModule.cpp
@@ -53,6 +53,12 @@ void FHttpModule::StartupModule()
 
 void FHttpModule::ShutdownModule()
 {
+	if (HttpManager != nullptr)
+	{
+		// block on any http requests that have already been queued up
+		HttpManager->Flush(true);
+	}
+
 #if PLATFORM_WINDOWS
 
 	extern bool bUseCurl;
diff --git a/Engine/Source/Runtime/Online/HTTP/Public/HttpManager.h b/Engine/Source/Runtime/Online/HTTP/Public/HttpManager.h
--- a/Engine/Source/Runtime/Online/HTTP/Public/HttpManager.h
+++ b/Engine/Source/Runtime/Online/HTTP/Public/HttpManager.h
@@ -49,6 +49,13 @@ public:
 	*/
 	virtual bool IsValidRequest(const IHttpRequest* RequestPtr) const;
 
+	/**
+	 * Block until all pending requests are finished processing
+	 *
+	 * @param bShutdown true if final flush during shutdown
+	 */
+	virtual void Flush(bool bShutdown);
+
 	/**
 	 * FTicker callback
 	 *
@@ -56,8 +63,8 @@ public:
 	 *
 	 * @return false if no longer needs ticking
 	 */
-	bool Tick(float DeltaSeconds) override;
-
+	virtual bool Tick(float DeltaSeconds) override;
+	
 	/**
 	 * List all of the Http requests currently being processed
 	 *
