Change OSS warning into a normal log message, because it always fires if you load the editor for projects without one.

diff --git a/Engine/Source/Runtime/Online/OnlineSubsystem/Private/OnlineSubsystemModule.cpp b/Engine/Source/Runtime/Online/OnlineSubsystem/Private/OnlineSubsystemModule.cpp
--- a/Engine/Source/Runtime/Online/OnlineSubsystem/Private/OnlineSubsystemModule.cpp
+++ b/Engine/Source/Runtime/Online/OnlineSubsystem/Private/OnlineSubsystemModule.cpp
@@ -88,7 +88,7 @@ void FOnlineSubsystemModule::StartupModule()
 	}
 	else
 	{
-		UE_LOG(LogOnline, Warning, TEXT("No default platform service specified for OnlineSubsystem"));
+		UE_LOG(LogOnline, Log, TEXT("No default platform service specified for OnlineSubsystem"));
 	}
 }
 
