Steam: fix inability to load the bundled library on Linux (UE-12449).

- Fixed by [REDACTED_PERSON].

diff --git a/Engine/Source/Runtime/Online/OnlineSubsystemSteam/Private/OnlineSubsystemModuleSteam.cpp b/Engine/Source/Runtime/Online/OnlineSubsystemSteam/Private/OnlineSubsystemModuleSteam.cpp
--- a/Engine/Source/Runtime/Online/OnlineSubsystemSteam/Private/OnlineSubsystemModuleSteam.cpp
+++ b/Engine/Source/Runtime/Online/OnlineSubsystemSteam/Private/OnlineSubsystemModuleSteam.cpp
@@ -137,7 +137,7 @@ void FOnlineSubsystemSteamModule::LoadSteamModules()
 	{
 		// try bundled one
 		UE_LOG_ONLINE(Log, TEXT("Could not find system one, loading bundled libsteam_api.so."));
-		FString RootSteamPath = FPaths::EngineDir() / FString::Printf(TEXT("Binaries/ThirdParty/Steamworks/%s/Linux/"), STEAM_SDK_VER); 
+		FString RootSteamPath = FPaths::EngineDir() / FString::Printf(TEXT("Binaries/ThirdParty/Steamworks/%s/Linux/"), STEAM_SDK_VER_PATH);
 		SteamDLLHandle = FPlatformProcess::GetDllHandle(*(RootSteamPath + "libsteam_api.so"));
 	}
 #else
