Add ShowProfileUI functionality for steam

diff --git a/Engine/Source/Runtime/Online/OnlineSubsystemSteam/Private/OnlineExternalUIInterfaceSteam.cpp b/Engine/Source/Runtime/Online/OnlineSubsystemSteam/Private/OnlineExternalUIInterfaceSteam.cpp
--- a/Engine/Source/Runtime/Online/OnlineSubsystemSteam/Private/OnlineExternalUIInterfaceSteam.cpp
+++ b/Engine/Source/Runtime/Online/OnlineSubsystemSteam/Private/OnlineExternalUIInterfaceSteam.cpp
@@ -22,6 +22,14 @@ void FOnlineAsyncEventSteamExternalUITriggered::TriggerDelegates()
 	FOnlineAsyncEvent::TriggerDelegates();
 	IOnlineExternalUIPtr ExternalUIInterface = Subsystem->GetExternalUIInterface();
 	ExternalUIInterface->TriggerOnExternalUIChangeDelegates(bIsActive);
+
+	// Calling this to mimic behavior as close as possible with other platforms where this delegate is passed in (such as PS4/Xbox).
+	if (!bIsActive)
+	{
+		FOnlineExternalUISteamPtr ExternalUISteam = StaticCastSharedPtr<FOnlineExternalUISteam>(ExternalUIInterface);
+		ExternalUISteam->ProfileUIClosedDelegate.ExecuteIfBound();
+		ExternalUISteam->ProfileUIClosedDelegate.Unbind();
+	}
 }
 
 bool FOnlineExternalUISteam::ShowLoginUI(const int ControllerIndex, bool bShowOnlineOnly, const FOnLoginUIClosedDelegate& Delegate)
@@ -74,8 +82,10 @@ bool FOnlineExternalUISteam::ShowWebURL(const FString& WebURL)
 
 bool FOnlineExternalUISteam::ShowProfileUI( const FUniqueNetId& Requestor, const FUniqueNetId& Requestee, const FOnProfileUIClosedDelegate& Delegate )
 {
-	//@todo
-	return false;
+	SteamFriends()->ActivateGameOverlayToUser(TCHAR_TO_UTF8(TEXT("steamid")), (const FUniqueNetIdSteam&)Requestee);
+
+	ProfileUIClosedDelegate = Delegate;
+	return true;
 }
 
 bool FOnlineExternalUISteam::ShowAccountUpgradeUI(const FUniqueNetId& UniqueId)
diff --git a/Engine/Source/Runtime/Online/OnlineSubsystemSteam/Private/OnlineExternalUIInterfaceSteam.h b/Engine/Source/Runtime/Online/OnlineSubsystemSteam/Private/OnlineExternalUIInterfaceSteam.h
--- a/Engine/Source/Runtime/Online/OnlineSubsystemSteam/Private/OnlineExternalUIInterfaceSteam.h
+++ b/Engine/Source/Runtime/Online/OnlineSubsystemSteam/Private/OnlineExternalUIInterfaceSteam.h
@@ -68,6 +68,9 @@ PACKAGE_SCOPE:
 	{
 	}
 
+	/** Triggered when the steam overlay is closed if it was opened via ShowProfileUI. Delegate will be unbound after it is executed. */
+	FOnProfileUIClosedDelegate ProfileUIClosedDelegate;
+
 public:
 
 	virtual ~FOnlineExternalUISteam()
