Supporting SHA checks for matching chunks, instead of data checks.

diff --git a/Engine/Source/Runtime/Online/BuildPatchServices/Private/BuildPatchChunk.cpp b/Engine/Source/Runtime/Online/BuildPatchServices/Private/BuildPatchChunk.cpp
--- a/Engine/Source/Runtime/Online/BuildPatchServices/Private/BuildPatchChunk.cpp
+++ b/Engine/Source/Runtime/Online/BuildPatchServices/Private/BuildPatchChunk.cpp
@@ -34,6 +34,18 @@ FString FSHAHashData::ToString() const
 	return BytesToHex(Hash, FSHA1::DigestSize);
 }
 
+bool FSHAHashData::isZero() const
+{
+	for (const auto& digit : Hash)
+	{
+		if (digit != 0)
+		{
+			return false;
+		}
+	}
+	return true;
+}
+
 /* FChunkHeader implementation
 *****************************************************************************/
 FChunkHeader::FChunkHeader()
diff --git a/Engine/Source/Runtime/Online/BuildPatchServices/Private/BuildPatchChunk.h b/Engine/Source/Runtime/Online/BuildPatchServices/Private/BuildPatchChunk.h
--- a/Engine/Source/Runtime/Online/BuildPatchServices/Private/BuildPatchChunk.h
+++ b/Engine/Source/Runtime/Online/BuildPatchServices/Private/BuildPatchChunk.h
@@ -27,6 +27,7 @@ struct FSHAHashData
 	bool operator==(const FSHAHashData& Other) const;
 	bool operator!=(const FSHAHashData& Other) const;
 	FString ToString() const;
+	bool isZero() const;
 };
 
 static_assert(FSHA1::DigestSize == 20, "If this changes a lot of stuff here will break!");
diff --git a/Engine/Source/Runtime/Online/BuildPatchServices/Private/BuildPatchGeneration.cpp b/Engine/Source/Runtime/Online/BuildPatchServices/Private/BuildPatchGeneration.cpp
--- a/Engine/Source/Runtime/Online/BuildPatchServices/Private/BuildPatchGeneration.cpp
+++ b/Engine/Source/Runtime/Online/BuildPatchServices/Private/BuildPatchGeneration.cpp
@@ -331,9 +331,20 @@ const uint32 FBuildGenerationChunkCache::FChunkReader::BytesLeft()
 *****************************************************************************/
 FBuildGenerationChunkCache::FBuildGenerationChunkCache(const FDateTime& InDataAgeThreshold)
 	: DataAgeThreshold(InDataAgeThreshold)
+	, StatChunksInDataCache(nullptr)
+	, StatNumCacheLoads(nullptr)
+	, StatNumCacheBoots(nullptr)
 {
 }
 
+void FBuildGenerationChunkCache::SetStatsCollector(FStatsCollectorRef InStatsCollector)
+{
+	StatsCollector = InStatsCollector;
+	StatChunksInDataCache = StatsCollector->CreateStat(TEXT("Chunk Cache: Num Chunks"), EStatFormat::Value);
+	StatNumCacheLoads = StatsCollector->CreateStat(TEXT("Chunk Cache: Num Chunks Loaded"), EStatFormat::Value);
+	StatNumCacheBoots = StatsCollector->CreateStat(TEXT("Chunk Cache: Num Chunks Booted"), EStatFormat::Value);
+}
+
 TSharedRef< FBuildGenerationChunkCache::FChunkReader > FBuildGenerationChunkCache::GetChunkReader( const FString& ChunkFilePath )
 {
 	if( ChunkCache.Contains( ChunkFilePath ) == false )
@@ -356,10 +367,26 @@ TSharedRef< FBuildGenerationChunkCache::FChunkReader > FBuildGenerationChunkCach
 			ChunkCache.Remove( OldestAccessChunk );
 			delete BytesReadPerChunk[ OldestAccessChunk ];
 			BytesReadPerChunk.Remove( OldestAccessChunk );
+			if(StatChunksInDataCache != nullptr)
+			{
+				FStatsCollector::Accumulate(StatChunksInDataCache, -1);
+			}
+			if(StatNumCacheBoots != nullptr)
+			{
+				FStatsCollector::Accumulate(StatNumCacheBoots, 1);
+			}
 		}
 		// Add the chunk to cache
 		ChunkCache.Add( ChunkFilePath, MakeShareable( new FChunkFile( 1, true ) ) );
 		BytesReadPerChunk.Add( ChunkFilePath, new uint32( 0 ) );
+		if(StatChunksInDataCache != nullptr)
+		{
+			FStatsCollector::Accumulate(StatChunksInDataCache, 1);
+		}
+		if(StatNumCacheLoads != nullptr)
+		{
+			FStatsCollector::Accumulate(StatNumCacheLoads, 1);
+		}
 	}
 	return MakeShareable( new FChunkReader( ChunkFilePath, ChunkCache[ ChunkFilePath ], BytesReadPerChunk[ ChunkFilePath ], DataAgeThreshold ) );
 }
@@ -372,6 +399,10 @@ void FBuildGenerationChunkCache::Cleanup()
 		delete BytesReadPerChunkIt.Value();
 	}
 	BytesReadPerChunk.Empty( 0 );
+	if(StatChunksInDataCache != nullptr)
+	{
+		FStatsCollector::Set(StatChunksInDataCache, 0);
+	}
 }
 
 /* FBuildGenerationChunkCache system singleton setup
@@ -486,11 +517,16 @@ bool FBuildDataGenerator::GenerateChunksManifestFromDirectory( const FBuildPatch
 	// Create stat collector
 	FStatsCollectorRef StatsCollector = FStatsCollectorFactory::Create();
 
-	// Create a build streamer
-	FBuildStreamerRef BuildStream = FBuildStreamerFactory::Create(Settings.RootDirectory, Settings.IgnoreListFile, StatsCollector);
-
 	// Enumerate Chunks
 	FCloudEnumerationRef CloudEnumeration = FCloudEnumerationFactory::Create(FBuildPatchServicesModule::GetCloudDirectory());
+	// Force waiting on cloud enumeration for more accurate stats, this line can be removed when stats are not required.
+	CloudEnumeration->GetChunkInventory();
+
+	// Chunks matching
+	FDataMatcherRef DataMatcher = FDataMatcherFactory::Create(FBuildPatchServicesModule::GetCloudDirectory());
+
+	// Create a build streamer
+	FBuildStreamerRef BuildStream = FBuildStreamerFactory::Create(Settings.RootDirectory, Settings.IgnoreListFile, StatsCollector);
 
 	// Output to log for builder info
 	GLog->Logf(TEXT("Running Chunks Patch Generation for: %u:%s %s"), Settings.AppID, *Settings.AppName, *Settings.BuildVersion);
@@ -498,6 +534,7 @@ bool FBuildDataGenerator::GenerateChunksManifestFromDirectory( const FBuildPatch
 	// Create our chunk cache
 	const FDateTime Cutoff = Settings.bShouldHonorReuseThreshold ? FDateTime::UtcNow() - FTimespan::FromDays(Settings.DataAgeThreshold) : FDateTime::MinValue();
 	FBuildGenerationChunkCache::Init(Cutoff);
+	FBuildGenerationChunkCache::Get().SetStatsCollector(StatsCollector);
 
 	// Create a manifest builder
 	FManifestDetails ManifestDetails;
@@ -527,9 +564,6 @@ bool FBuildDataGenerator::GenerateChunksManifestFromDirectory( const FBuildPatch
 	}
 	FManifestBuilderRef ManifestBuilder = FManifestBuilderFactory::Create(ManifestDetails, BuildStream);
 
-	// Chunks matching
-	FDataMatcherRef DataMatcher = FDataMatcherFactory::Create(FBuildPatchServicesModule::GetCloudDirectory());
-
 	// Refers to how much data has been dequeued
 	//uint64 ProcessPos = 0;
 
@@ -732,7 +766,7 @@ bool FBuildDataGenerator::GenerateFilesManifestFromDirectory( const FBuildPatchS
 
 	// Save manifest into the cloud directory
 	FString JsonFilename = FBuildPatchServicesModule::GetCloudDirectory() / FDefaultValueHelper::RemoveWhitespaces(BuildManifest->Data->AppName + BuildManifest->Data->BuildVersion) + TEXT(".manifest");
-	BuildManifest->Data->ManifestFileVersion = EBuildPatchAppManifestVersion::GetLatestJsonVersion();
+	BuildManifest->Data->ManifestFileVersion = EBuildPatchAppManifestVersion::GetLatestFileDataVersion();
 	BuildManifest->SaveToFile(JsonFilename, false);
 
 	// Output to log for builder info
@@ -954,7 +988,7 @@ bool FBuildDataGenerator::SaveOutFileData(const FString& SourceFile, const FSHAH
 	bool bSuccess = false;
 	IFileManager& FileManager = IFileManager::Get();
 
-	const FString NewFilename = FBuildPatchUtils::GetFileNewFilename( EBuildPatchAppManifestVersion::GetLatestVersion(), FBuildPatchServicesModule::GetCloudDirectory(), FileGuid, FileHash );
+	const FString NewFilename = FBuildPatchUtils::GetFileNewFilename( EBuildPatchAppManifestVersion::GetLatestFileDataVersion(), FBuildPatchServicesModule::GetCloudDirectory(), FileGuid, FileHash );
 	bAlreadySaved = FPlatformFileManager::Get().GetPlatformFile().FileExists(*NewFilename);
 
 #if SAVE_OLD_FILEDATA_FILENAMES
diff --git a/Engine/Source/Runtime/Online/BuildPatchServices/Private/BuildPatchGeneration.h b/Engine/Source/Runtime/Online/BuildPatchServices/Private/BuildPatchGeneration.h
--- a/Engine/Source/Runtime/Online/BuildPatchServices/Private/BuildPatchGeneration.h
+++ b/Engine/Source/Runtime/Online/BuildPatchServices/Private/BuildPatchGeneration.h
@@ -181,7 +181,7 @@ private:
 
 private:
 	// The number of chunks to cache in memory
-	const static int32 NumChunksToCache = 256;
+	const static int32 NumChunksToCache = 1024;
 
 	FDateTime DataAgeThreshold;
 
@@ -191,10 +191,23 @@ private:
 	// The store for how much file data has been read for each chunk
 	TMap< FString, uint32* > BytesReadPerChunk;
 
+	// The stats collector class, and some stats variables
+	FStatsCollectorPtr StatsCollector;
+	volatile int64* StatChunksInDataCache;
+	volatile int64* StatNumCacheLoads;
+	volatile int64* StatNumCacheBoots;
+
 	// The singleton instance for this class
 	static TSharedPtr< FBuildGenerationChunkCache > SingletonInstance;
 
 public:
+
+	/**
+	 * Sets the stats collector, which will be used to record statistics from the point of settings
+	 * @param StatsCollector		Ref to the stats collector
+	 */
+	void SetStatsCollector(FStatsCollectorRef StatsCollector);
+
 	/**
 	 * Get a reader class for a chunk file. The reader class is the interface to chunk data which will either come from RAM or file.
 	 * @return		Shared ref to a reader class, in order to allow memory freeing, you should not keep hold of references when finished with chunk
diff --git a/Engine/Source/Runtime/Online/BuildPatchServices/Private/BuildPatchManifest.cpp b/Engine/Source/Runtime/Online/BuildPatchServices/Private/BuildPatchManifest.cpp
--- a/Engine/Source/Runtime/Online/BuildPatchServices/Private/BuildPatchManifest.cpp
+++ b/Engine/Source/Runtime/Online/BuildPatchServices/Private/BuildPatchManifest.cpp
@@ -8,9 +8,6 @@
 
 #define LOCTEXT_NAMESPACE "BuildPatchManifest"
 
-// Remove this when we are to enable creating compressed file data and related manifests
-#define ENABLE_NOCHUNKS_COMPRESSION 0
-
 // The manifest header magic codeword, for quick checking that the opened file is probably a manifest file.
 #define MANIFEST_HEADER_MAGIC		0x44BEC00C
 
@@ -35,6 +32,23 @@ FString ToStringBlob( const DataType& DataVal )
 	const void* AsBuffer = &DataVal;
 	return FString::FromBlob( static_cast<const uint8*>( AsBuffer ), sizeof( DataType ) );
 }
+template< typename DataType >
+bool FromHexString(const FString& HexString, DataType& ValueOut)
+{
+	void* AsBuffer = &ValueOut;
+	if (HexString.Len() == (sizeof(DataType)* 2))
+	{
+		HexToBytes(HexString, static_cast<uint8*>(AsBuffer));
+		return true;
+	}
+	return false;
+}
+template< typename DataType >
+FString ToHexString(const DataType& DataVal)
+{
+	const void* AsBuffer = &DataVal;
+	return BytesToHex(static_cast<const uint8*>(AsBuffer), sizeof(DataType));
+}
 
 /**
  * Helper functions to decide whether the passed in data is a JSON string we expect to deserialize a manifest from
@@ -57,16 +71,22 @@ bool BufferIsJsonManifest(const TArray<uint8>& DataInput)
 *****************************************************************************/
 EBuildPatchAppManifestVersion::Type EBuildPatchAppManifestVersion::GetLatestVersion()
 {
-#if ENABLE_NOCHUNKS_COMPRESSION
 	return static_cast<EBuildPatchAppManifestVersion::Type>(LatestPlusOne - 1);
-#else
-	return EBuildPatchAppManifestVersion::StoredAsCompressedUClass;
-#endif
 }
 
 EBuildPatchAppManifestVersion::Type EBuildPatchAppManifestVersion::GetLatestJsonVersion()
 {
-	return static_cast<EBuildPatchAppManifestVersion::Type>(StoredAsCompressedUClass - 1);
+	return GetLatestVersion();
+}
+
+EBuildPatchAppManifestVersion::Type EBuildPatchAppManifestVersion::GetLatestFileDataVersion()
+{
+	return static_cast<EBuildPatchAppManifestVersion::Type>(StoresChunkFileSizes);
+}
+
+EBuildPatchAppManifestVersion::Type EBuildPatchAppManifestVersion::GetLatestChunkDataVersion()
+{
+	return GetLatestVersion();
 }
 
 const FString& EBuildPatchAppManifestVersion::GetChunkSubdir(const EBuildPatchAppManifestVersion::Type ManifestVersion)
@@ -136,6 +156,7 @@ FCustomFieldData::FCustomFieldData(const FString& InKey, const FString& InValue)
 FChunkInfoData::FChunkInfoData()
 	: Guid()
 	, Hash(0)
+	, ShaHash()
 	, FileSize(0)
 	, GroupNumber(0)
 {
@@ -796,6 +817,15 @@ void FBuildPatchAppManifest::SerializeToJSON(FString& JSONOutput)
 			Writer->WriteValue(ChunkGuid.ToString(), ToStringBlob(ChunkHash));
 		}
 		Writer->WriteObjectEnd();
+		// Write chunk sha list
+		Writer->WriteObjectStart(TEXT("ChunkShaList"));
+		for (const auto& ChunkInfo : Data->ChunkList)
+		{
+			const FGuid& ChunkGuid = ChunkInfo.Guid;
+			const FSHAHashData& ChunkSha = ChunkInfo.ShaHash;
+			Writer->WriteValue(ChunkGuid.ToString(), ToHexString(ChunkSha));
+		}
+		Writer->WriteObjectEnd();
 		// Write data group list
 		Writer->WriteObjectStart(TEXT("DataGroupList"));
 		for (const auto& ChunkInfo : Data->ChunkList)
@@ -975,6 +1005,25 @@ bool FBuildPatchAppManifest::DeserializeFromJSON( const FString& JSONInput )
 		}
 	}
 
+	// Get the ChunkShaList (optional)
+	TSharedPtr<FJsonValue> JsonChunkShaList = JsonValueMap.FindRef(TEXT("ChunkShaList"));
+	if (JsonChunkShaList.IsValid())
+	{
+		TSharedPtr<FJsonObject> JsonChunkHashListObj = JsonChunkShaList->AsObject();
+		for (auto ChunkHashIt = JsonChunkHashListObj->Values.CreateConstIterator(); ChunkHashIt && bSuccess; ++ChunkHashIt)
+		{
+			FGuid ChunkGuid;
+			FSHAHashData ChunkSha;
+			bSuccess = bSuccess && FGuid::Parse(ChunkHashIt.Key(), ChunkGuid);
+			bSuccess = bSuccess && FromHexString(ChunkHashIt.Value()->AsString(), ChunkSha);
+			if (bSuccess && ChunkInfoLookup.Contains(ChunkGuid))
+			{
+				FChunkInfoData* ChunkInfoData = ChunkInfoLookup[ChunkGuid];
+				ChunkInfoData->ShaHash = ChunkSha;
+			}
+		}
+	}
+
 	// Get the DataGroupList
 	TSharedPtr<FJsonValue> JsonDataGroupList = JsonValueMap.FindRef(TEXT("DataGroupList"));
 	if (JsonDataGroupList.IsValid())
@@ -1243,6 +1292,17 @@ bool FBuildPatchAppManifest::GetChunkHash(const FGuid& ChunkGuid, uint64& OutHas
 	return false;
 }
 
+bool FBuildPatchAppManifest::GetChunkShaHash(const FGuid& ChunkGuid, FSHAHashData& OutHash) const
+{
+	const FChunkInfoData* const * ChunkInfo = ChunkInfoLookup.Find(ChunkGuid);
+	if (ChunkInfo != nullptr)
+	{
+		OutHash = (*ChunkInfo)->ShaHash;
+		return OutHash.isZero() == false;
+	}
+	return false;
+}
+
 bool FBuildPatchAppManifest::GetFileHash(const FGuid& FileGuid, FSHAHashData& OutHash) const
 {
 	const FString* const * FoundFilename = FileNameLookup.Find(FileGuid);
diff --git a/Engine/Source/Runtime/Online/BuildPatchServices/Private/BuildPatchManifest.h b/Engine/Source/Runtime/Online/BuildPatchServices/Private/BuildPatchManifest.h
--- a/Engine/Source/Runtime/Online/BuildPatchServices/Private/BuildPatchManifest.h
+++ b/Engine/Source/Runtime/Online/BuildPatchServices/Private/BuildPatchManifest.h
@@ -40,12 +40,13 @@ namespace EBuildPatchAppManifestVersion
 		StoresPrerequisitesInfo,
 		// Manifest stores chunk download sizes
 		StoresChunkFileSizes,
-		// Manifest is now stored using UObject serialization and compressed
+		// Manifest can optionally be stored using UObject serialization and compressed
 		StoredAsCompressedUClass,
-		// Added support for file data to be split into max individual part size, files now go to FilesV3
-		FileSplittingSupport,
-		// Added support for file data compression, (files still FilesV3)
-		FileCompressionSupport,
+		// These two features were removed and never used
+		UNUSED_0,
+		UNUSED_1,
+		// Manifest stores chunk data SHA1 hash to use in place of data compare, for faster generation
+		StoresChunkDataShaHashes,
 
 
 		// Always after the latest version, signifies the latest version plus 1 to allow initialization simplicity
@@ -60,9 +61,15 @@ namespace EBuildPatchAppManifestVersion
 	/** @return The last known manifest feature of this code base. Handy for manifest constructor */
 	Type GetLatestVersion();
 
-	/** @return The latest version of a manifest support by JSON serialization */
+	/** @return The latest version of a manifest supported by JSON serialization method */
 	Type GetLatestJsonVersion();
 
+	/** @return The latest version of a manifest supported by file data (nochunks) */
+	Type GetLatestFileDataVersion();
+
+	/** @return The latest version of a manifest supported by chunk data */
+	Type GetLatestChunkDataVersion();
+
 	/**
 	 * Get the chunk subdirectory for used for a specific manifest version, e.g. Chunks, ChunksV2 etc
 	 * @param ManifestVersion     The version of the manifest
@@ -128,6 +135,9 @@ struct FChunkInfoData
 	UPROPERTY()
 	uint64 Hash;
 
+	UPROPERTY()
+	FSHAHashData ShaHash;
+
 	UPROPERTY()
 	int64 FileSize;
 
@@ -507,6 +517,14 @@ public:
 	 */
 	bool GetChunkHash(const FGuid& ChunkGuid, uint64& OutHash) const;
 
+	/**
+	 * Gets the SHA1 hash for a given chunk
+	 * @param ChunkGuid		IN		The guid of the chunk to get hash for
+	 * @param OutHash		OUT		Receives the hash value if found
+	 * @return	true if we had the hash for this chunk
+	 */
+	bool GetChunkShaHash(const FGuid& ChunkGuid, FSHAHashData& OutHash) const;
+
 	/**
 	 * Gets the file hash for given file data
 	 * @param FileGuid		IN		The guid of the file data to get hash for
diff --git a/Engine/Source/Runtime/Online/BuildPatchServices/Private/Core/RingBuffer.h b/Engine/Source/Runtime/Online/BuildPatchServices/Private/Core/RingBuffer.h
--- a/Engine/Source/Runtime/Online/BuildPatchServices/Private/Core/RingBuffer.h
+++ b/Engine/Source/Runtime/Online/BuildPatchServices/Private/Core/RingBuffer.h
@@ -71,6 +71,12 @@ public:
 	 */
 	int32 SerialCompare(const DataType* SerialBuffer, uint32 CompareLen) const;
 
+	/**
+	 * Get the SHA1 hash for the data currently in the FIFO
+	 * @param OutHash	Receives the SHA hash
+	 */
+	void GetShaHash(FSHAHash& OutHash) const;
+
 	/**
 	 * Serializes the internal buffer into the given buffer
 	 * @param SerialBuffer	A preallocated buffer to copy data into
@@ -262,6 +268,20 @@ int32 TRingBuffer< DataType, BufferDataSize >::SerialCompare(const DataType* Ser
 	return FMemory::Memcmp(Data, &SerialBuffer[FirstPartLen], sizeof(DataType)* (CompareLen - FirstPartLen));
 }
 
+template< typename DataType, uint32 BufferDataSize >
+void TRingBuffer< DataType, BufferDataSize >::GetShaHash(FSHAHash& OutHash) const
+{
+	const uint32 FirstPartLen = FMath::Min(BufferDataSize - DataIndex, NumDataAvailable);
+	FSHA1 Sha;
+	Sha.Update(Data + BottomIndex(), sizeof(DataType) * FirstPartLen);
+	if (FirstPartLen < NumDataAvailable)
+	{
+		Sha.Update(Data, sizeof(DataType) * (NumDataAvailable - FirstPartLen));
+	}
+	Sha.Final();
+	Sha.GetHash(OutHash.Hash);
+}
+
 template< typename DataType, uint32 BufferDataSize >
 void TRingBuffer< DataType, BufferDataSize >::Serialize(DataType* SerialBuffer) const
 {
diff --git a/Engine/Source/Runtime/Online/BuildPatchServices/Private/Generation/CloudEnumeration.cpp b/Engine/Source/Runtime/Online/BuildPatchServices/Private/Generation/CloudEnumeration.cpp
--- a/Engine/Source/Runtime/Online/BuildPatchServices/Private/Generation/CloudEnumeration.cpp
+++ b/Engine/Source/Runtime/Online/BuildPatchServices/Private/Generation/CloudEnumeration.cpp
@@ -18,6 +18,7 @@ namespace BuildPatchServices
 		virtual TSet<FGuid> GetChunkSet(uint64 ChunkHash) const override;
 		virtual TMap<uint64, TSet<FGuid>> GetChunkInventory() const override;
 		virtual TMap<FGuid, int64> GetChunkFileSizes() const override;
+		virtual TMap<FGuid, FSHAHash> GetChunkShaHashes() const override;
 	private:
 		void EnumerateCloud();
 		void EnumerateManifestData(const FBuildPatchAppManifestRef& Manifest);
@@ -27,6 +28,7 @@ namespace BuildPatchServices
 		mutable FCriticalSection InventoryCS;
 		TMap<uint64, TSet<FGuid>> ChunkInventory;
 		TMap<FGuid, int64> ChunkFileSizes;
+		TMap<FGuid, FSHAHash> ChunkShaHashes;
 		TMap<FSHAHash, TSet<FGuid>> FileInventory;
 		uint64 NumChunksFound;
 		uint64 NumFilesFound;
@@ -67,6 +69,13 @@ namespace BuildPatchServices
 		return ChunkFileSizes;
 	}
 
+	TMap<FGuid, FSHAHash> FCloudEnumerationImpl::GetChunkShaHashes() const
+	{
+		FScopeLock ScopeLock(&InventoryCS);
+		Future.Wait();
+		return ChunkShaHashes;
+	}
+
 	void FCloudEnumerationImpl::EnumerateCloud()
 	{
 		FScopeLock ScopeLock(&InventoryCS);
@@ -120,9 +129,10 @@ namespace BuildPatchServices
 		Manifest->GetDataList(DataList);
 		if (!Manifest->IsFileDataManifest())
 		{
+			FSHAHashData DataShaHash;
+			uint64 DataChunkHash;
 			for (const auto& DataGuid : DataList)
 			{
-				uint64 DataChunkHash;
 				if (Manifest->GetChunkHash(DataGuid, DataChunkHash))
 				{
 					if (DataChunkHash != 0)
@@ -144,6 +154,11 @@ namespace BuildPatchServices
 				{
 					GLog->Logf(TEXT("FCloudEnumeration: WARNING: Missing chunk hash for %s in manifest %s %s"), *DataGuid.ToString(), *Manifest->GetAppName(), *Manifest->GetVersionString());
 				}
+				if (Manifest->GetChunkShaHash(DataGuid, DataShaHash))
+				{
+					FSHAHash& ShaHash = ChunkShaHashes.FindOrAdd(DataGuid);
+					FMemory::Memcpy(ShaHash.Hash, DataShaHash.Hash, FSHA1::DigestSize);
+				}
 			}
 		}
 		else
diff --git a/Engine/Source/Runtime/Online/BuildPatchServices/Private/Generation/CloudEnumeration.h b/Engine/Source/Runtime/Online/BuildPatchServices/Private/Generation/CloudEnumeration.h
--- a/Engine/Source/Runtime/Online/BuildPatchServices/Private/Generation/CloudEnumeration.h
+++ b/Engine/Source/Runtime/Online/BuildPatchServices/Private/Generation/CloudEnumeration.h
@@ -11,6 +11,7 @@ namespace BuildPatchServices
 		virtual TSet<FGuid> GetChunkSet(uint64 ChunkHash) const = 0;
 		virtual TMap<uint64, TSet<FGuid>> GetChunkInventory() const = 0;
 		virtual TMap<FGuid, int64> GetChunkFileSizes() const = 0;
+		virtual TMap<FGuid, FSHAHash> GetChunkShaHashes() const = 0;
 	};
 
 	typedef TSharedRef<FCloudEnumeration, ESPMode::ThreadSafe> FCloudEnumerationRef;
diff --git a/Engine/Source/Runtime/Online/BuildPatchServices/Private/Generation/DataScanner.cpp b/Engine/Source/Runtime/Online/BuildPatchServices/Private/Generation/DataScanner.cpp
--- a/Engine/Source/Runtime/Online/BuildPatchServices/Private/Generation/DataScanner.cpp
+++ b/Engine/Source/Runtime/Online/BuildPatchServices/Private/Generation/DataScanner.cpp
@@ -61,6 +61,8 @@ namespace BuildPatchServices
 	private:
 		FDataScanResult ScanData();
 		bool ProcessCurrentWindow();
+		bool FindExistingChunk(const TMap<uint64, TSet<FGuid>>& ChunkLookup, TMap<FGuid, FSHAHash>& ChunkShaHashes, uint64 ChunkHash, const FRollingHash<WindowSize>& ChunkBuffer, FGuid& OutMatchedChunk);
+		bool FindExistingChunk(const TMap<uint64, TSet<FGuid>>& ChunkLookup, TMap<FGuid, FSHAHash>& ChunkShaHashes, uint64 ChunkHash, const TArray<uint8>& ChunkBuffer, FGuid& OutMatchedChunk);
 
 	private:
 		const uint64 DataStartOffset;
@@ -71,6 +73,19 @@ namespace BuildPatchServices
 		FThreadSafeBool bIsComplete;
 		FThreadSafeBool bShouldAbort;
 		TFuture<FDataScanResult> FutureResult;
+		volatile int64* StatCreatedScanners;
+		volatile int64* StatRunningScanners;
+		volatile int64* StatCpuTime;
+		volatile int64* StatConsumeBytesTime;
+		volatile int64* StatFindMatchTime;
+		volatile int64* StatDataMatchTime;
+		volatile int64* StatChunkWriterTime;
+		volatile int64* StatHashCollisions;
+		volatile int64* StatChunkDataChecks;
+		volatile int64* StatChunkDataMatches;
+		volatile int64* StatMissingChunks;
+		volatile int64* StatMatchedData;
+		volatile int64* StatExtraData;
 
 	public:
 		static FThreadSafeCounter NumIncompleteScanners;
@@ -193,6 +208,21 @@ namespace BuildPatchServices
 		, bIsComplete(false)
 		, bShouldAbort(false)
 	{
+		// Create statistics
+		StatCreatedScanners = StatsCollector->CreateStat(TEXT("Scanner: Created Scanners"), EStatFormat::Value);
+		StatRunningScanners = StatsCollector->CreateStat(TEXT("Scanner: Running Scanners"), EStatFormat::Value);
+		StatCpuTime = StatsCollector->CreateStat(TEXT("Scanner: CPU Time"), EStatFormat::Timer);
+		StatConsumeBytesTime = StatsCollector->CreateStat(TEXT("Scanner: Consume Bytes Time"), EStatFormat::Timer);
+		StatFindMatchTime = StatsCollector->CreateStat(TEXT("Scanner: Find Match Time"), EStatFormat::Timer);
+		StatDataMatchTime = StatsCollector->CreateStat(TEXT("Scanner: Data Match Time"), EStatFormat::Timer);
+		StatChunkWriterTime = StatsCollector->CreateStat(TEXT("Scanner: Chunk Writer Time"), EStatFormat::Timer);
+		StatHashCollisions = StatsCollector->CreateStat(TEXT("Scanner: Hash Collisions"), EStatFormat::Value);
+		StatChunkDataChecks = StatsCollector->CreateStat(TEXT("Scanner: Chunk Data Checks"), EStatFormat::Value);
+		StatChunkDataMatches = StatsCollector->CreateStat(TEXT("Scanner: Chunk Data Matches"), EStatFormat::Value);
+		StatMissingChunks = StatsCollector->CreateStat(TEXT("Scanner: Missing Chunks"), EStatFormat::Value);
+		StatMatchedData = StatsCollector->CreateStat(TEXT("Scanner: Matched Data"), EStatFormat::DataSize);
+		StatExtraData = StatsCollector->CreateStat(TEXT("Scanner: Extra Data"), EStatFormat::DataSize);
+		// Queue thread
 		NumIncompleteScanners.Increment();
 		TFunction<FDataScanResult()> Task = [this]()
 		{
@@ -220,22 +250,103 @@ namespace BuildPatchServices
 		return MoveTemp(FutureResult.Get());
 	}
 
+	bool FDataScannerImpl::FindExistingChunk(const TMap<uint64, TSet<FGuid>>& ChunkLookup, TMap<FGuid, FSHAHash>& ChunkShaHashes, uint64 ChunkHash, const FRollingHash<WindowSize>& RollingHash, FGuid& OutMatchedChunk)
+	{
+		FStatsScopedTimer FindTimer(StatFindMatchTime);
+		bool bFoundChunkMatch = false;
+		if (ChunkLookup.Contains(ChunkHash))
+		{
+			FSHAHash ChunkSha;
+			RollingHash.GetWindowData().GetShaHash(ChunkSha);
+			for (FGuid& PotentialMatch : ChunkLookup.FindRef(ChunkHash))
+			{
+				// Use sha if we have it
+				if (ChunkShaHashes.Contains(PotentialMatch))
+				{
+					if(ChunkSha == ChunkShaHashes[PotentialMatch])
+					{
+						bFoundChunkMatch = true;
+						OutMatchedChunk = PotentialMatch;
+						break;
+					}
+				}
+				else
+				{
+					// Otherwise compare data
+					TArray<uint8> SerialBuffer;
+					FStatsScopedTimer DataMatchTimer(StatDataMatchTime);
+					FStatsCollector::Accumulate(StatChunkDataChecks, 1);
+					SerialBuffer.AddUninitialized(WindowSize);
+					RollingHash.GetWindowData().Serialize(SerialBuffer.GetData());
+					bool ChunkFound = false;
+					if (DataMatcher->CompareData(PotentialMatch, ChunkHash, SerialBuffer, ChunkFound))
+					{
+						FStatsCollector::Accumulate(StatChunkDataMatches, 1);
+						ChunkShaHashes.Add(PotentialMatch, ChunkSha);
+						bFoundChunkMatch = true;
+						OutMatchedChunk = PotentialMatch;
+						break;
+					}
+					else if(!ChunkFound)
+					{
+						FStatsCollector::Accumulate(StatMissingChunks, 1);
+					}
+				}
+				FStatsCollector::Accumulate(StatHashCollisions, 1);
+			}
+		}
+		return bFoundChunkMatch;
+	}
+
+	bool FDataScannerImpl::FindExistingChunk(const TMap<uint64, TSet<FGuid>>& ChunkLookup, TMap<FGuid, FSHAHash>& ChunkShaHashes, uint64 ChunkHash, const TArray<uint8>& ChunkBuffer, FGuid& OutMatchedChunk)
+	{
+		FStatsScopedTimer FindTimer(StatFindMatchTime);
+		bool bFoundChunkMatch = false;
+		if (ChunkLookup.Contains(ChunkHash))
+		{
+			FSHAHash ChunkSha;
+			FSHA1::HashBuffer(ChunkBuffer.GetData(), ChunkBuffer.Num(), ChunkSha.Hash);
+			for (FGuid& PotentialMatch : ChunkLookup.FindRef(ChunkHash))
+			{
+				// Use sha if we have it
+				if (ChunkShaHashes.Contains(PotentialMatch))
+				{
+					if(ChunkSha == ChunkShaHashes[PotentialMatch])
+					{
+						bFoundChunkMatch = true;
+						OutMatchedChunk = PotentialMatch;
+						break;
+					}
+				}
+				else
+				{
+					// Otherwise compare data
+					FStatsScopedTimer DataMatchTimer(StatDataMatchTime);
+					FStatsCollector::Accumulate(StatChunkDataChecks, 1);
+					bool ChunkFound = false;
+					if (DataMatcher->CompareData(PotentialMatch, ChunkHash, ChunkBuffer, ChunkFound))
+					{
+						FStatsCollector::Accumulate(StatChunkDataMatches, 1);
+						ChunkShaHashes.Add(PotentialMatch, ChunkSha);
+						bFoundChunkMatch = true;
+						OutMatchedChunk = PotentialMatch;
+						break;
+					}
+					else if (!ChunkFound)
+					{
+						FStatsCollector::Accumulate(StatMissingChunks, 1);
+					}
+				}
+				FStatsCollector::Accumulate(StatHashCollisions, 1);
+			}
+		}
+		return bFoundChunkMatch;
+	}
+
 	FDataScanResult FDataScannerImpl::ScanData()
 	{
 		// Count running scanners
 		FScopeCounter ScopeCounter(&NumRunningScanners);
-
-		// Create statistics
-		auto* StatCreatedScanners = StatsCollector->CreateStat(TEXT("Scanner: Created Scanners"), EStatFormat::Value);
-		auto* StatRunningScanners = StatsCollector->CreateStat(TEXT("Scanner: Running Scanners"), EStatFormat::Value);
-		auto* StatCpuTime = StatsCollector->CreateStat(TEXT("Scanner: CPU Time"), EStatFormat::Timer);
-		auto* StatConsumeBytesTime = StatsCollector->CreateStat(TEXT("Scanner: Consume Bytes Time"), EStatFormat::Timer);
-		auto* StatFindMatchTime = StatsCollector->CreateStat(TEXT("Scanner: Find Match Time"), EStatFormat::Timer);
-		auto* StatChunkWriterTime = StatsCollector->CreateStat(TEXT("Scanner: Chunk Writer Time"), EStatFormat::Timer);
-		auto* StatHashCollisions = StatsCollector->CreateStat(TEXT("Scanner: Hash Collisions"), EStatFormat::Value);
-		auto* StatMissingChunks = StatsCollector->CreateStat(TEXT("Scanner: Missing Chunks"), EStatFormat::Value);
-		auto* StatMatchedData = StatsCollector->CreateStat(TEXT("Scanner: Matched Data"), EStatFormat::DataSize);
-		auto* StatExtraData = StatsCollector->CreateStat(TEXT("Scanner: Extra Data"), EStatFormat::DataSize);
 		FStatsCollector::Accumulate(StatCreatedScanners, 1);
 		FStatsCollector::Accumulate(StatRunningScanners, 1);
 
@@ -254,8 +365,10 @@ namespace BuildPatchServices
 		// Get a copy of the chunk inventory
 		TMap<uint64, TSet<FGuid>> ChunkInventory = CloudEnumeration->GetChunkInventory();
 		TMap<FGuid, int64> ChunkFileSizes = CloudEnumeration->GetChunkFileSizes();
+		TMap<FGuid, FSHAHash> ChunkShaHashes = CloudEnumeration->GetChunkShaHashes();
 
 		// Loop over and process all data
+		FGuid MatchedChunk;
 		uint64 TempTimer;
 		uint64 CpuTimer;
 		FStatsCollector::AccumulateTimeBegin(CpuTimer);
@@ -265,7 +378,7 @@ namespace BuildPatchServices
 			const uint32 NumDataNeeded = RollingHash.GetNumDataNeeded();
 			if (NumDataNeeded > 0)
 			{
-				FStatsCollector::AccumulateTimeBegin(TempTimer);
+				FStatsScopedTimer ConsumeTimer(StatConsumeBytesTime);
 				uint32 NumConsumedBytes = 0;
 				if (idx < Data.Num())
 				{
@@ -282,44 +395,20 @@ namespace BuildPatchServices
 					PaddedZeros = Zeros.Num();
 				}
 				check(RollingHash.GetNumDataNeeded() == 0);
-				FStatsCollector::AccumulateTimeEnd(StatConsumeBytesTime, TempTimer);
 				continue;
 			}
 
 			const uint64 NumDataInWindow = WindowSize - PaddedZeros;
 			const uint64 WindowHash = RollingHash.GetWindowHash();
 			// Try find match
-			bool bFoundChunkMatch = false;
-			if (ChunkInventory.Contains(WindowHash))
-			{
-				bool ChunkFound = false;
-				FStatsCollector::AccumulateTimeBegin(TempTimer);
-				RollingHash.GetWindowData().Serialize(ChunkBuffer.GetData());
-				for (FGuid& PotentialMatch : ChunkInventory.FindRef(WindowHash))
-				{
-					if (DataMatcher->CompareData(PotentialMatch, WindowHash, ChunkBuffer, ChunkFound))
-					{
-						bFoundChunkMatch = true;
-						DataStructure.PushKnownChunk(PotentialMatch, NumDataInWindow);
-						FChunkInfo& ChunkInfo = ChunkInfoLookup.FindOrAdd(PotentialMatch);
-						ChunkInfo.Hash = WindowHash;
-						ChunkInfo.IsNew = false;
-						break;
-					}
-					else if(ChunkFound)
-					{
-						FStatsCollector::Accumulate(StatHashCollisions, 1);
-					}
-					else
-					{
-						FStatsCollector::Accumulate(StatMissingChunks, 1);
-					}
-				}
-				FStatsCollector::AccumulateTimeEnd(StatFindMatchTime, TempTimer);
-			}
-
-			if (bFoundChunkMatch)
+			if (FindExistingChunk(ChunkInventory, ChunkShaHashes, WindowHash, RollingHash, MatchedChunk))
 			{
+				// Push the chunk to the structure
+				DataStructure.PushKnownChunk(MatchedChunk, NumDataInWindow);
+				FChunkInfo& ChunkInfo = ChunkInfoLookup.FindOrAdd(MatchedChunk);
+				ChunkInfo.Hash = WindowHash;
+				ChunkInfo.ShaHash = ChunkShaHashes[MatchedChunk];
+				ChunkInfo.IsNew = false;
 				FStatsCollector::Accumulate(StatMatchedData, NumDataInWindow);
 				// Clear matched window
 				RollingHash.Clear();
@@ -338,45 +427,28 @@ namespace BuildPatchServices
 				if (NewChunkBuffer.Num() == WindowSize)
 				{
 					const uint64 NewChunkHash = FRollingHash<WindowSize>::GetHashForDataSet(NewChunkBuffer.GetData());
-					for (FGuid& PotentialMatch : ChunkInventory.FindRef(NewChunkHash))
+					if (FindExistingChunk(ChunkInventory, ChunkShaHashes, NewChunkHash, NewChunkBuffer, MatchedChunk))
 					{
-						bool ChunkFound = false;
-						FStatsCollector::AccumulateTimeBegin(TempTimer);
-						if (DataMatcher->CompareData(PotentialMatch, NewChunkHash, NewChunkBuffer, ChunkFound))
-						{
-							bFoundChunkMatch = true;
-							DataStructure.RemapCurrentChunk(PotentialMatch);
-							FChunkInfo& ChunkInfo = ChunkInfoLookup.FindOrAdd(PotentialMatch);
-							ChunkInfo.Hash = NewChunkHash;
-							ChunkInfo.IsNew = false;
-							break;
-						}
-						else if(ChunkFound)
-						{
-							FStatsCollector::Accumulate(StatHashCollisions, 1);
-						}
-						else
-						{
-							FStatsCollector::Accumulate(StatMissingChunks, 1);
-						}
-						FStatsCollector::AccumulateTimeEnd(StatFindMatchTime, TempTimer);
+						DataStructure.RemapCurrentChunk(MatchedChunk);
+						FChunkInfo& ChunkInfo = ChunkInfoLookup.FindOrAdd(MatchedChunk);
+						ChunkInfo.Hash = NewChunkHash;
+						ChunkInfo.ShaHash = ChunkShaHashes[MatchedChunk];
+						ChunkInfo.IsNew = false;
+						FStatsCollector::Accumulate(StatMatchedData, WindowSize);
 					}
-					if (!bFoundChunkMatch)
+					else
 					{
+						FStatsScopedTimer ChunkWriterTimer(StatChunkWriterTime);
 						const FGuid& NewChunkGuid = DataStructure.GetCurrentChunkId();
 						FStatsCollector::AccumulateTimeEnd(StatCpuTime, CpuTimer);
-						FStatsCollector::AccumulateTimeBegin(TempTimer);
 						ChunkWriter.QueueChunk(NewChunkBuffer.GetData(), NewChunkGuid, NewChunkHash);
-						FStatsCollector::AccumulateTimeEnd(StatChunkWriterTime, TempTimer);
-						FStatsCollector::Accumulate(StatExtraData, NewChunkBuffer.Num());
 						FStatsCollector::AccumulateTimeBegin(CpuTimer);
 						FChunkInfo& ChunkInfo = ChunkInfoLookup.FindOrAdd(NewChunkGuid);
 						ChunkInfo.Hash = NewChunkHash;
 						ChunkInfo.IsNew = true;
-					}
-					else
-					{
-						FStatsCollector::Accumulate(StatMatchedData, WindowSize);
+						FSHA1::HashBuffer(NewChunkBuffer.GetData(), NewChunkBuffer.Num(), ChunkInfo.ShaHash.Hash);
+						ChunkShaHashes.Add(NewChunkGuid, ChunkInfo.ShaHash);
+						FStatsCollector::Accumulate(StatExtraData, NewChunkBuffer.Num());
 					}
 					DataStructure.CompleteCurrentChunk();
 					NewChunkBuffer.Empty(WindowSize);
@@ -399,43 +471,30 @@ namespace BuildPatchServices
 		if (NewChunkBuffer.Num() > 0)
 		{
 			NewChunkBuffer.AddZeroed(WindowSize - NewChunkBuffer.Num());
-			bool bFoundChunkMatch = false;
 			const uint64 NewChunkHash = FRollingHash<WindowSize>::GetHashForDataSet(NewChunkBuffer.GetData());
-			for (FGuid& PotentialMatch : ChunkInventory.FindRef(NewChunkHash))
+			if (FindExistingChunk(ChunkInventory, ChunkShaHashes, NewChunkHash, NewChunkBuffer, MatchedChunk))
 			{
-				bool ChunkFound = false;
-				FStatsCollector::AccumulateTimeBegin(TempTimer);
-				if (DataMatcher->CompareData(PotentialMatch, NewChunkHash, NewChunkBuffer, ChunkFound))
-				{
-					bFoundChunkMatch = true;
-					DataStructure.RemapCurrentChunk(PotentialMatch);
-					FChunkInfo& ChunkInfo = ChunkInfoLookup.FindOrAdd(PotentialMatch);
-					ChunkInfo.Hash = NewChunkHash;
-					ChunkInfo.IsNew = false;
-					break;
-				}
-				else if(ChunkFound)
-				{
-					FStatsCollector::Accumulate(StatHashCollisions, 1);
-				}
-				else
-				{
-					FStatsCollector::Accumulate(StatMissingChunks, 1);
-				}
-				FStatsCollector::AccumulateTimeEnd(StatFindMatchTime, TempTimer);
+				// Setup chunk info for a match
+				DataStructure.RemapCurrentChunk(MatchedChunk);
+				FChunkInfo& ChunkInfo = ChunkInfoLookup.FindOrAdd(MatchedChunk);
+				ChunkInfo.Hash = NewChunkHash;
+				ChunkInfo.ShaHash = ChunkShaHashes[MatchedChunk];
+				ChunkInfo.IsNew = false;
 			}
-			if (!bFoundChunkMatch)
+			else
 			{
+				// Save the final chunk if no match
+				FStatsScopedTimer ChunkWriterTimer(StatChunkWriterTime);
 				const FGuid& NewChunkGuid = DataStructure.GetCurrentChunkId();
 				FStatsCollector::AccumulateTimeEnd(StatCpuTime, CpuTimer);
-				FStatsCollector::AccumulateTimeBegin(TempTimer);
 				ChunkWriter.QueueChunk(NewChunkBuffer.GetData(), NewChunkGuid, NewChunkHash);
-				FStatsCollector::AccumulateTimeEnd(StatChunkWriterTime, TempTimer);
-				FStatsCollector::Accumulate(StatExtraData, NewChunkBuffer.Num());
 				FStatsCollector::AccumulateTimeBegin(CpuTimer);
 				FChunkInfo& ChunkInfo = ChunkInfoLookup.FindOrAdd(NewChunkGuid);
 				ChunkInfo.Hash = NewChunkHash;
 				ChunkInfo.IsNew = true;
+				FSHA1::HashBuffer(NewChunkBuffer.GetData(), NewChunkBuffer.Num(), ChunkInfo.ShaHash.Hash);
+				ChunkShaHashes.Add(NewChunkGuid, ChunkInfo.ShaHash);
+				FStatsCollector::Accumulate(StatExtraData, NewChunkBuffer.Num());
 			}
 		}
 		FStatsCollector::AccumulateTimeEnd(StatCpuTime, CpuTimer);
diff --git a/Engine/Source/Runtime/Online/BuildPatchServices/Private/Generation/DataScanner.h b/Engine/Source/Runtime/Online/BuildPatchServices/Private/Generation/DataScanner.h
--- a/Engine/Source/Runtime/Online/BuildPatchServices/Private/Generation/DataScanner.h
+++ b/Engine/Source/Runtime/Online/BuildPatchServices/Private/Generation/DataScanner.h
@@ -25,10 +25,10 @@ namespace BuildPatchServices
 	{
 		// The hash for this chunk
 		uint64 Hash;
+		// The SHA hash for this chunk
+		FSHAHash ShaHash;
 		// The size of this chunk file
 		int64 ChunkFileSize;
-		// The group number this chunk is stored in
-		uint8 GroupNumber;
 		// Whether this chunk is new
 		bool IsNew;
 	};
diff --git a/Engine/Source/Runtime/Online/BuildPatchServices/Private/Generation/ManifestBuilder.cpp b/Engine/Source/Runtime/Online/BuildPatchServices/Private/Generation/ManifestBuilder.cpp
--- a/Engine/Source/Runtime/Online/BuildPatchServices/Private/Generation/ManifestBuilder.cpp
+++ b/Engine/Source/Runtime/Online/BuildPatchServices/Private/Generation/ManifestBuilder.cpp
@@ -304,6 +304,7 @@ namespace BuildPatchServices
 					auto& ChunkInfoData = Manifest->Data->ChunkList.Last();
 					ChunkInfoData.Guid = ChunkPart.Guid;
 					ChunkInfoData.Hash = ChunkInfo.Hash;
+					FMemory::Memcpy(ChunkInfoData.ShaHash.Hash, ChunkInfo.ShaHash.Hash, FSHA1::DigestSize);
 					ChunkInfoData.FileSize = ChunkInfo.ChunkFileSize;
 					ChunkInfoData.GroupNumber = FCrc::MemCrc32(&ChunkPart.Guid, sizeof(FGuid)) % 100;
 				}
diff --git a/Engine/Source/Runtime/Online/BuildPatchServices/Private/Generation/StatsCollector.h b/Engine/Source/Runtime/Online/BuildPatchServices/Private/Generation/StatsCollector.h
--- a/Engine/Source/Runtime/Online/BuildPatchServices/Private/Generation/StatsCollector.h
+++ b/Engine/Source/Runtime/Online/BuildPatchServices/Private/Generation/StatsCollector.h
@@ -42,4 +42,22 @@ namespace BuildPatchServices
 	public:
 		static FStatsCollectorRef Create();
 	};
+
+	class FStatsScopedTimer
+	{
+	public:
+		FStatsScopedTimer(volatile int64* InStat)
+			: Stat(InStat)
+		{
+			FStatsCollector::AccumulateTimeBegin(TempTime);
+		}
+		~FStatsScopedTimer()
+		{
+			FStatsCollector::AccumulateTimeEnd(Stat, TempTime);
+		}
+
+	private:
+		uint64 TempTime;
+		volatile int64* Stat;
+	};
 }
