Add GetPlayerDisplayName to TurnBasedBlueprintLibrary

diff --git a/Engine/Source/Runtime/Online/OnlineSubsystemUtils/Classes/TurnBasedBlueprintLibrary.h b/Engine/Source/Runtime/Online/OnlineSubsystemUtils/Classes/TurnBasedBlueprintLibrary.h
--- a/Engine/Source/Runtime/Online/OnlineSubsystemUtils/Classes/TurnBasedBlueprintLibrary.h
+++ b/Engine/Source/Runtime/Online/OnlineSubsystemUtils/Classes/TurnBasedBlueprintLibrary.h
@@ -19,4 +19,7 @@ class ONLINESUBSYSTEMUTILS_API UTurnBasedBlueprintLibrary : public UBlueprintFun
 
 	UFUNCTION(BlueprintCallable, meta = (WorldContext = "WorldContextObject"), Category = "Online|TurnBased")
 	static void RegisterTurnBasedMatchInterfaceObject(UObject* WorldContextObject, APlayerController* PlayerController, UObject* Object);
+
+	UFUNCTION(BlueprintCallable, meta = (WorldContext = "WorldContextObject"), Category = "Online|TurnBased")
+	static void GetPlayerDisplayName(UObject* WorldContextObject, APlayerController* PlayerController, FString MatchID, int32 PlayerIndex, /*out*/ FString& PlayerDisplayName);
 };
diff --git a/Engine/Source/Runtime/Online/OnlineSubsystemUtils/Private/TurnBasedBlueprintLibrary.cpp b/Engine/Source/Runtime/Online/OnlineSubsystemUtils/Private/TurnBasedBlueprintLibrary.cpp
--- a/Engine/Source/Runtime/Online/OnlineSubsystemUtils/Private/TurnBasedBlueprintLibrary.cpp
+++ b/Engine/Source/Runtime/Online/OnlineSubsystemUtils/Private/TurnBasedBlueprintLibrary.cpp
@@ -82,4 +82,36 @@ void UTurnBasedBlueprintLibrary::RegisterTurnBasedMatchInterfaceObject(UObject*
 			OnlineTurnBasedPtr->RegisterTurnBasedMatchInterfaceObject(Object);
 		}
 	}
+}
+
+void UTurnBasedBlueprintLibrary::GetPlayerDisplayName(UObject* WorldContextObject, APlayerController* PlayerController, FString MatchID, int32 PlayerIndex, /*out*/ FString& PlayerDisplayName)
+{
+	FOnlineSubsystemBPCallHelper Helper(TEXT("GetMyPlayerIndex"), GEngine->GetWorldFromContextObject(WorldContextObject));
+	Helper.QueryIDFromPlayerController(PlayerController);
+
+	if (Helper.IsValid())
+	{
+		IOnlineTurnBasedPtr OnlineTurnBasedPtr = Helper.OnlineSub->GetTurnBasedInterface();
+		if (OnlineTurnBasedPtr.IsValid())
+		{
+			FTurnBasedMatchPtr Match = OnlineTurnBasedPtr->GetMatchWithID(MatchID);
+			if (Match.IsValid())
+			{
+				if (Match->GetPlayerDisplayName(PlayerIndex, PlayerDisplayName) == false)
+				{
+					FString Message = FString::Printf(TEXT("Player index %d not within bounds of player array."), PlayerIndex);
+					FFrame::KismetExecutionMessage(*Message, ELogVerbosity::Warning);
+				}
+			}
+			else
+			{
+				FString Message = FString::Printf(TEXT("Match ID %s not found"), *MatchID);
+				FFrame::KismetExecutionMessage(*Message, ELogVerbosity::Warning);
+			}
+		}
+		else
+		{
+			FFrame::KismetExecutionMessage(TEXT("Turn Based Matches not supported by Online Subsystem"), ELogVerbosity::Warning);
+		}
+	}
 }
\ No newline at end of file
