Fixing potential cloud enumeration deadlock

diff --git a/Engine/Source/Runtime/Online/BuildPatchServices/Private/Generation/CloudEnumeration.cpp b/Engine/Source/Runtime/Online/BuildPatchServices/Private/Generation/CloudEnumeration.cpp
--- a/Engine/Source/Runtime/Online/BuildPatchServices/Private/Generation/CloudEnumeration.cpp
+++ b/Engine/Source/Runtime/Online/BuildPatchServices/Private/Generation/CloudEnumeration.cpp
@@ -50,30 +50,39 @@ namespace BuildPatchServices
 
 	TSet<FGuid> FCloudEnumerationImpl::GetChunkSet(uint64 ChunkHash) const
 	{
-		FScopeLock ScopeLock(&InventoryCS);
 		Future.Wait();
-		return ChunkInventory.FindRef(ChunkHash);
+		{
+			FScopeLock ScopeLock(&InventoryCS);
+			return ChunkInventory.FindRef(ChunkHash);
+		}
 	}
 
+
 	TMap<uint64, TSet<FGuid>> FCloudEnumerationImpl::GetChunkInventory() const
 	{
-		FScopeLock ScopeLock(&InventoryCS);
 		Future.Wait();
-		return ChunkInventory;
+		{
+			FScopeLock ScopeLock(&InventoryCS);
+			return ChunkInventory;
+		}
 	}
 
 	TMap<FGuid, int64> FCloudEnumerationImpl::GetChunkFileSizes() const
 	{
-		FScopeLock ScopeLock(&InventoryCS);
 		Future.Wait();
-		return ChunkFileSizes;
+		{
+			FScopeLock ScopeLock(&InventoryCS);
+			return ChunkFileSizes;
+		}
 	}
 
 	TMap<FGuid, FSHAHash> FCloudEnumerationImpl::GetChunkShaHashes() const
 	{
-		FScopeLock ScopeLock(&InventoryCS);
 		Future.Wait();
-		return ChunkShaHashes;
+		{
+			FScopeLock ScopeLock(&InventoryCS);
+			return ChunkShaHashes;
+		}
 	}
 
 	void FCloudEnumerationImpl::EnumerateCloud()
