Fix for non-unity builds for Android
#ue4
#android

diff --git a/Engine/Source/Runtime/Engine/Private/PhysicsEngine/PhysicsSerializer.cpp b/Engine/Source/Runtime/Engine/Private/PhysicsEngine/PhysicsSerializer.cpp
--- a/Engine/Source/Runtime/Engine/Private/PhysicsEngine/PhysicsSerializer.cpp
+++ b/Engine/Source/Runtime/Engine/Private/PhysicsEngine/PhysicsSerializer.cpp
@@ -8,6 +8,7 @@
 #include "PhysicsPublic.h"
 #include "PhysicsSerializer.h"
 #include "PhysDerivedData.h"
+#include "PhysXSupport.h"
 
 UPhysicsSerializer::UPhysicsSerializer(const FObjectInitializer& ObjectInitializer)
 : Super(ObjectInitializer)
diff --git a/Engine/Source/Runtime/Launch/Private/Android/AndroidEventManager.cpp b/Engine/Source/Runtime/Launch/Private/Android/AndroidEventManager.cpp
--- a/Engine/Source/Runtime/Launch/Private/Android/AndroidEventManager.cpp
+++ b/Engine/Source/Runtime/Launch/Private/Android/AndroidEventManager.cpp
@@ -7,6 +7,7 @@
 #include "CallbackDevice.h"
 #include <android/native_window.h> 
 #include <android/native_window_jni.h> 
+#include "IHeadMountedDisplay.h"
 
 
 DEFINE_LOG_CATEGORY(LogAndroidEvents);
@@ -86,7 +87,7 @@ void FAppEventManager::Tick()
 			bHaveGame = true;
 			break;
 
-		// window focus events that follow their own  heirarchy, and might or might not respect App main events heirarchy
+		// window focus events that follow their own hierarchy, and might or might not respect App main events hierarchy
 
 		case APP_EVENT_STATE_WINDOW_GAINED_FOCUS: 
 			bWindowInFocus = true;
diff --git a/Engine/Source/Runtime/Launch/Private/Android/LaunchAndroid.cpp b/Engine/Source/Runtime/Launch/Private/Android/LaunchAndroid.cpp
--- a/Engine/Source/Runtime/Launch/Private/Android/LaunchAndroid.cpp
+++ b/Engine/Source/Runtime/Launch/Private/Android/LaunchAndroid.cpp
@@ -15,6 +15,8 @@
 #include "AndroidWindow.h"
 #include <android/sensor.h>
 #include "Core.h"
+#include "AndroidApplication.h"
+#include "IHeadMountedDisplayModule.h"
 
 // Function pointer for retrieving joystick events
 // Function has been part of the OS since Honeycomb, but only appeared in the
diff --git a/Engine/Source/Runtime/Online/Android/OnlineSubsystemGooglePlay/Private/OnlineAsyncTaskGooglePlayLogin.cpp b/Engine/Source/Runtime/Online/Android/OnlineSubsystemGooglePlay/Private/OnlineAsyncTaskGooglePlayLogin.cpp
--- a/Engine/Source/Runtime/Online/Android/OnlineSubsystemGooglePlay/Private/OnlineAsyncTaskGooglePlayLogin.cpp
+++ b/Engine/Source/Runtime/Online/Android/OnlineSubsystemGooglePlay/Private/OnlineAsyncTaskGooglePlayLogin.cpp
@@ -29,9 +29,9 @@ void FOnlineAsyncTaskGooglePlayLogin::TriggerDelegates()
 
 void FOnlineAsyncTaskGooglePlayLogin::OnAuthActionFinished(gpg::AuthOperation InOp, gpg::AuthStatus InStatus)
 {
-	if (InOp == AuthOperation::SIGN_IN)
+	if (InOp == gpg::AuthOperation::SIGN_IN)
 	{
-		if (InStatus == AuthStatus::ERROR_NOT_AUTHORIZED && RetryCount > 0)
+		if (InStatus == gpg::AuthStatus::ERROR_NOT_AUTHORIZED && RetryCount > 0)
 		{
 			--RetryCount;
 			Subsystem->GetGameServices()->StartAuthorizationUI();
@@ -39,7 +39,7 @@ void FOnlineAsyncTaskGooglePlayLogin::OnAuthActionFinished(gpg::AuthOperation In
 		}
 		
 		Status = InStatus;
-		bWasSuccessful = Status == AuthStatus::VALID;
+		bWasSuccessful = Status == gpg::AuthStatus::VALID;
 		bIsComplete = true;
 	}
 }
diff --git a/Engine/Source/Runtime/Online/Android/OnlineSubsystemGooglePlay/Private/OnlineAsyncTaskGooglePlayQueryAchievements.cpp b/Engine/Source/Runtime/Online/Android/OnlineSubsystemGooglePlay/Private/OnlineAsyncTaskGooglePlayQueryAchievements.cpp
--- a/Engine/Source/Runtime/Online/Android/OnlineSubsystemGooglePlay/Private/OnlineAsyncTaskGooglePlayQueryAchievements.cpp
+++ b/Engine/Source/Runtime/Online/Android/OnlineSubsystemGooglePlay/Private/OnlineAsyncTaskGooglePlayQueryAchievements.cpp
@@ -37,8 +37,8 @@ void FOnlineAsyncTaskGooglePlayQueryAchievements::Tick()
 	// We're already running on the online thread. Using the blocking version of the API function
 	// won't block the main thread and simplifies things a bit.
 	// Try a 10 second timeout.
-	Response = Subsystem->GetGameServices()->Achievements().FetchAllBlocking(DataSource::CACHE_OR_NETWORK, Timeout(10000));
+	Response = Subsystem->GetGameServices()->Achievements().FetchAllBlocking(gpg::DataSource::CACHE_OR_NETWORK, gpg::Timeout(10000));
 	
-	bWasSuccessful = Response.status == ResponseStatus::VALID || Response.status == ResponseStatus::VALID_BUT_STALE;
+	bWasSuccessful = Response.status == gpg::ResponseStatus::VALID || Response.status == gpg::ResponseStatus::VALID_BUT_STALE;
 	bIsComplete = true;
 }
diff --git a/Engine/Source/Runtime/Online/Android/OnlineSubsystemGooglePlay/Private/OnlineSubsystemGooglePlayPrivatePCH.h b/Engine/Source/Runtime/Online/Android/OnlineSubsystemGooglePlay/Private/OnlineSubsystemGooglePlayPrivatePCH.h
--- a/Engine/Source/Runtime/Online/Android/OnlineSubsystemGooglePlay/Private/OnlineSubsystemGooglePlayPrivatePCH.h
+++ b/Engine/Source/Runtime/Online/Android/OnlineSubsystemGooglePlay/Private/OnlineSubsystemGooglePlayPrivatePCH.h
@@ -5,7 +5,8 @@
 /** FName declaration of Steam subsystem */
 #define GOOGLEPLAY_SUBSYSTEM FName(TEXT("GooglePlay"))
 
-#include "Core.h"
+#include "Engine.h"
+#include "CoreUObject.h"
 #include "SocketSubsystem.h"
 #include "ModuleManager.h"
 
diff --git a/Engine/Source/Runtime/Online/Android/OnlineSubsystemGooglePlay/Public/OnlineSubsystemGooglePlay.h b/Engine/Source/Runtime/Online/Android/OnlineSubsystemGooglePlay/Public/OnlineSubsystemGooglePlay.h
--- a/Engine/Source/Runtime/Online/Android/OnlineSubsystemGooglePlay/Public/OnlineSubsystemGooglePlay.h
+++ b/Engine/Source/Runtime/Online/Android/OnlineSubsystemGooglePlay/Public/OnlineSubsystemGooglePlay.h
@@ -8,6 +8,7 @@
 #include "OnlineAchievementsInterfaceGooglePlay.h"
 #include "OnlineLeaderboardInterfaceGooglePlay.h"
 #include "OnlineExternalUIInterfaceGooglePlay.h"
+#include "OnlineAsyncTaskManagerGooglePlay.h"
 #include "UniquePtr.h"
 
 #include "gpg/game_services.h"
diff --git a/Engine/Source/Runtime/OpenGLDrv/Private/Android/AndroidOpenGLPrivate.h b/Engine/Source/Runtime/OpenGLDrv/Private/Android/AndroidOpenGLPrivate.h
--- a/Engine/Source/Runtime/OpenGLDrv/Private/Android/AndroidOpenGLPrivate.h
+++ b/Engine/Source/Runtime/OpenGLDrv/Private/Android/AndroidOpenGLPrivate.h
@@ -5,6 +5,8 @@
 =============================================================================*/
 #pragma once
 
+#include "AndroidApplication.h"
+
 bool GAndroidGPUInfoReady = false;
 
 class FAndroidGPUInfo
