Fix bracket position causing "'Create' is deprecated" warning

diff --git a/Engine/Source/Runtime/Online/Android/OnlineSubsystemGooglePlay/Private/OnlineSubsystemGooglePlay.cpp b/Engine/Source/Runtime/Online/Android/OnlineSubsystemGooglePlay/Private/OnlineSubsystemGooglePlay.cpp
--- a/Engine/Source/Runtime/Online/Android/OnlineSubsystemGooglePlay/Private/OnlineSubsystemGooglePlay.cpp
+++ b/Engine/Source/Runtime/Online/Android/OnlineSubsystemGooglePlay/Private/OnlineSubsystemGooglePlay.cpp
@@ -116,7 +116,7 @@ bool FOnlineSubsystemGooglePlay::Init()
 	FPlatformMisc::LowLevelOutputDebugStringf(TEXT("FOnlineSubsystemAndroid::Init"));
 	
 	OnlineAsyncTaskThreadRunnable.Reset(new FOnlineAsyncTaskManagerGooglePlay);
-	OnlineAsyncTaskThread.Reset(FRunnableThread::Create(OnlineAsyncTaskThreadRunnable.Get(), *FString::Printf(TEXT("OnlineAsyncTaskThread %s")), *InstanceName.ToString()));
+	OnlineAsyncTaskThread.Reset(FRunnableThread::Create(OnlineAsyncTaskThreadRunnable.Get(), *FString::Printf(TEXT("OnlineAsyncTaskThread %s"), *InstanceName.ToString())));
 
 	IdentityInterface = MakeShareable(new FOnlineIdentityGooglePlay(this));
 	LeaderboardsInterface = MakeShareable(new FOnlineLeaderboardsGooglePlay(this));
