ShowInviteUI takes a session name now

diff --git a/Engine/Source/Runtime/Online/IOS/OnlineSubsystemIOS/Private/OnlineExternalUIInterfaceIOS.cpp b/Engine/Source/Runtime/Online/IOS/OnlineSubsystemIOS/Private/OnlineExternalUIInterfaceIOS.cpp
--- a/Engine/Source/Runtime/Online/IOS/OnlineSubsystemIOS/Private/OnlineExternalUIInterfaceIOS.cpp
+++ b/Engine/Source/Runtime/Online/IOS/OnlineSubsystemIOS/Private/OnlineExternalUIInterfaceIOS.cpp
@@ -17,7 +17,7 @@ bool FOnlineExternalUIIOS::ShowFriendsUI(int32 LocalUserNum)
 	return false;
 }
 
-bool FOnlineExternalUIIOS::ShowInviteUI(int32 LocalUserNum) 
+bool FOnlineExternalUIIOS::ShowInviteUI(int32 LocalUserNum, FName SessionMame)
 {
 	return false;
 }
diff --git a/Engine/Source/Runtime/Online/IOS/OnlineSubsystemIOS/Private/OnlineExternalUIInterfaceIOS.h b/Engine/Source/Runtime/Online/IOS/OnlineSubsystemIOS/Private/OnlineExternalUIInterfaceIOS.h
--- a/Engine/Source/Runtime/Online/IOS/OnlineSubsystemIOS/Private/OnlineExternalUIInterfaceIOS.h
+++ b/Engine/Source/Runtime/Online/IOS/OnlineSubsystemIOS/Private/OnlineExternalUIInterfaceIOS.h
@@ -16,7 +16,7 @@ public:
 	// Begin IOnlineExternalUI interface
 	virtual bool ShowLoginUI(const int ControllerIndex, bool bShowOnlineOnly, const FOnLoginUIClosedDelegate& Delegate) override;
 	virtual bool ShowFriendsUI(int32 LocalUserNum) override;
-	virtual bool ShowInviteUI(int32 LocalUserNum) override;
+	virtual bool ShowInviteUI(int32 LocalUserNum, FName SessionMame = GameSessionName) override;
 	virtual bool ShowAchievementsUI(int32 LocalUserNum) override;
 	virtual bool ShowLeaderboardUI(const FString& LeaderboardName) override;
 	virtual bool ShowWebURL(const FString& WebURL) override;
diff --git a/Engine/Source/Runtime/Online/OnlineSubsystem/Public/Interfaces/OnlineExternalUIInterface.h b/Engine/Source/Runtime/Online/OnlineSubsystem/Public/Interfaces/OnlineExternalUIInterface.h
--- a/Engine/Source/Runtime/Online/OnlineSubsystem/Public/Interfaces/OnlineExternalUIInterface.h
+++ b/Engine/Source/Runtime/Online/OnlineSubsystem/Public/Interfaces/OnlineExternalUIInterface.h
@@ -61,7 +61,7 @@ public:
 	 *
 	 * @return true if it was able to show the UI, false if it failed
 	 */
-	virtual bool ShowInviteUI(int32 LocalUserNum) = 0;
+	virtual bool ShowInviteUI(int32 LocalUserNum, FName SessionMame = GameSessionName) = 0;
 
 	/**
 	 *	Displays the UI that shows a user's list of achievements
diff --git a/Engine/Source/Runtime/Online/OnlineSubsystem/Public/Interfaces/OnlineSessionInterface.h b/Engine/Source/Runtime/Online/OnlineSubsystem/Public/Interfaces/OnlineSessionInterface.h
--- a/Engine/Source/Runtime/Online/OnlineSubsystem/Public/Interfaces/OnlineSessionInterface.h
+++ b/Engine/Source/Runtime/Online/OnlineSubsystem/Public/Interfaces/OnlineSessionInterface.h
@@ -4,12 +4,7 @@
 
 #include "OnlineSessionSettings.h"
 #include "OnlineDelegateMacros.h"
-
-const FName GameSessionName(TEXT("Game"));
-const FName PartySessionName(TEXT("Party"));
-
-const FName GamePort(TEXT("GamePort"));
-const FName BeaconPort(TEXT("BeaconPort"));
+#include "OnlineSubsystemTypes.h"
 
 /**
  * Delegate fired when a session create request has completed
diff --git a/Engine/Source/Runtime/Online/OnlineSubsystem/Public/OnlineSubsystemTypes.h b/Engine/Source/Runtime/Online/OnlineSubsystem/Public/OnlineSubsystemTypes.h
--- a/Engine/Source/Runtime/Online/OnlineSubsystem/Public/OnlineSubsystemTypes.h
+++ b/Engine/Source/Runtime/Online/OnlineSubsystem/Public/OnlineSubsystemTypes.h
@@ -32,6 +32,12 @@
 #define S_OK 0
 #endif
 
+const FName GameSessionName(TEXT("Game"));
+const FName PartySessionName(TEXT("Party"));
+
+const FName GamePort(TEXT("GamePort"));
+const FName BeaconPort(TEXT("BeaconPort"));
+
 /**
  * Generates a random nonce (number used once) of the desired length
  *
diff --git a/Engine/Source/Runtime/Online/OnlineSubsystemSteam/Private/OnlineExternalUIInterfaceSteam.cpp b/Engine/Source/Runtime/Online/OnlineSubsystemSteam/Private/OnlineExternalUIInterfaceSteam.cpp
--- a/Engine/Source/Runtime/Online/OnlineSubsystemSteam/Private/OnlineExternalUIInterfaceSteam.cpp
+++ b/Engine/Source/Runtime/Online/OnlineSubsystemSteam/Private/OnlineExternalUIInterfaceSteam.cpp
@@ -43,7 +43,7 @@ bool FOnlineExternalUISteam::ShowFriendsUI(int32 LocalUserNum)
 	return true;
 }
 
-bool FOnlineExternalUISteam::ShowInviteUI(int32 LocalUserNum)
+bool FOnlineExternalUISteam::ShowInviteUI(int32 LocalUserNum, FName SessionMame)
 {
 	IOnlineSessionPtr SessionInt = SteamSubsystem->GetSessionInterface();
 	if (SessionInt.IsValid() && SessionInt->HasPresenceSession())
diff --git a/Engine/Source/Runtime/Online/OnlineSubsystemSteam/Private/OnlineExternalUIInterfaceSteam.h b/Engine/Source/Runtime/Online/OnlineSubsystemSteam/Private/OnlineExternalUIInterfaceSteam.h
--- a/Engine/Source/Runtime/Online/OnlineSubsystemSteam/Private/OnlineExternalUIInterfaceSteam.h
+++ b/Engine/Source/Runtime/Online/OnlineSubsystemSteam/Private/OnlineExternalUIInterfaceSteam.h
@@ -80,7 +80,7 @@ public:
 	// IOnlineExternalUI
 	virtual bool ShowLoginUI(const int ControllerIndex, bool bShowOnlineOnly, const FOnLoginUIClosedDelegate& Delegate) override;
 	virtual bool ShowFriendsUI(int32 LocalUserNum) override;
-	virtual bool ShowInviteUI(int32 LocalUserNum) override;
+	virtual bool ShowInviteUI(int32 LocalUserNum, FName SessionMame = GameSessionName) override;
 	virtual bool ShowAchievementsUI(int32 LocalUserNum) override;
 	virtual bool ShowLeaderboardUI(const FString& LeaderboardName) override;
 	virtual bool ShowWebURL(const FString& WebURL) override;
