Fix friend status string to get the game status, instead of the "connect" property which is the command line to connect to that player's game.

diff --git a/Engine/Source/Runtime/Online/OnlineSubsystemSteam/Private/OnlineFriendsInterfaceSteam.cpp b/Engine/Source/Runtime/Online/OnlineSubsystemSteam/Private/OnlineFriendsInterfaceSteam.cpp
--- a/Engine/Source/Runtime/Online/OnlineSubsystemSteam/Private/OnlineFriendsInterfaceSteam.cpp
+++ b/Engine/Source/Runtime/Online/OnlineSubsystemSteam/Private/OnlineFriendsInterfaceSteam.cpp
@@ -214,7 +214,7 @@ void FOnlineAsyncTaskSteamReadFriendsList::Finalize()
 			const bool bIsPlayingAGame = SteamFriendsPtr->GetFriendGamePlayed(SteamPlayerId, &FriendGameInfo);
 			// Now fill in the friend info
 			Friend->AccountData.Add(TEXT("nickname"), NickName);
-			Friend->Presence.Status.StatusStr = UTF8_TO_TCHAR(SteamFriendsPtr->GetFriendRichPresence(SteamPlayerId,"connect"));
+			Friend->Presence.Status.StatusStr = UTF8_TO_TCHAR(SteamFriendsPtr->GetFriendRichPresence(SteamPlayerId,"status"));
 			FString JoinablePresenceString = UTF8_TO_TCHAR(SteamFriendsPtr->GetFriendRichPresence(SteamPlayerId,"Joinable"));
 			// Remote friend is responsible for updating their presence to have the joinable status
 			Friend->Presence.bIsJoinable = JoinablePresenceString == TEXT("true");
