[INTEGRATE]

Change 2612102 by [REDACTED_PERSON][REDACTED_HANDLE] on 2015/07/07 05:11
Change 2612143 by [REDACTED_PERSON][REDACTED_HANDLE] on 2015/07/07 06:55

diff --git a/Engine/Source/Runtime/Online/Android/OnlineSubsystemGooglePlay/Private/OnlineSubsystemGooglePlay.cpp b/Engine/Source/Runtime/Online/Android/OnlineSubsystemGooglePlay/Private/OnlineSubsystemGooglePlay.cpp
--- a/Engine/Source/Runtime/Online/Android/OnlineSubsystemGooglePlay/Private/OnlineSubsystemGooglePlay.cpp
+++ b/Engine/Source/Runtime/Online/Android/OnlineSubsystemGooglePlay/Private/OnlineSubsystemGooglePlay.cpp
@@ -68,11 +68,6 @@ IOnlineUserCloudPtr FOnlineSubsystemGooglePlay::GetUserCloudInterface() const
 	return nullptr;
 }
 
-IOnlineUserCloudPtr FOnlineSubsystemGooglePlay::GetUserCloudInterface(const FString& Key) const
-{
-	return nullptr;
-}
-
 IOnlineLeaderboardsPtr FOnlineSubsystemGooglePlay::GetLeaderboardsInterface() const
 {
 	return LeaderboardsInterface;
diff --git a/Engine/Source/Runtime/Online/Android/OnlineSubsystemGooglePlay/Public/OnlineSubsystemGooglePlay.h b/Engine/Source/Runtime/Online/Android/OnlineSubsystemGooglePlay/Public/OnlineSubsystemGooglePlay.h
--- a/Engine/Source/Runtime/Online/Android/OnlineSubsystemGooglePlay/Public/OnlineSubsystemGooglePlay.h
+++ b/Engine/Source/Runtime/Online/Android/OnlineSubsystemGooglePlay/Public/OnlineSubsystemGooglePlay.h
@@ -33,7 +33,6 @@ public:
 	virtual IOnlineGroupsPtr GetGroupsInterface() const override;
 	virtual IOnlineSharedCloudPtr GetSharedCloudInterface() const override;
 	virtual IOnlineUserCloudPtr GetUserCloudInterface() const override;
-	virtual IOnlineUserCloudPtr GetUserCloudInterface(const FString& Key) const override;
 	virtual IOnlineLeaderboardsPtr GetLeaderboardsInterface() const override;
 	virtual IOnlineVoicePtr GetVoiceInterface() const  override;
 	virtual IOnlineExternalUIPtr GetExternalUIInterface() const override;
diff --git a/Engine/Source/Runtime/Online/IOS/OnlineSubsystemIOS/Private/OnlineSubsystemIOS.cpp b/Engine/Source/Runtime/Online/IOS/OnlineSubsystemIOS/Private/OnlineSubsystemIOS.cpp
--- a/Engine/Source/Runtime/Online/IOS/OnlineSubsystemIOS/Private/OnlineSubsystemIOS.cpp
+++ b/Engine/Source/Runtime/Online/IOS/OnlineSubsystemIOS/Private/OnlineSubsystemIOS.cpp
@@ -38,11 +38,6 @@ IOnlineUserCloudPtr FOnlineSubsystemIOS::GetUserCloudInterface() const
 	return UserCloudInterface;
 }
 
-IOnlineUserCloudPtr FOnlineSubsystemIOS::GetUserCloudInterface(const FString& Key) const
-{
-	return nullptr;
-}
-
 IOnlineLeaderboardsPtr FOnlineSubsystemIOS::GetLeaderboardsInterface() const
 {
 	return LeaderboardsInterface;
@@ -207,4 +202,4 @@ bool FOnlineSubsystemIOS::IsInAppPurchasingEnabled()
 	bool bEnableIAP = false;
 	GConfig->GetBool(TEXT("OnlineSubsystemIOS.Store"), TEXT("bSupportsInAppPurchasing"), bEnableIAP, GEngineIni);
 	return bEnableIAP;
-}
\ No newline at end of file
+}
diff --git a/Engine/Source/Runtime/Online/IOS/OnlineSubsystemIOS/Public/OnlineSubsystemIOS.h b/Engine/Source/Runtime/Online/IOS/OnlineSubsystemIOS/Public/OnlineSubsystemIOS.h
--- a/Engine/Source/Runtime/Online/IOS/OnlineSubsystemIOS/Public/OnlineSubsystemIOS.h
+++ b/Engine/Source/Runtime/Online/IOS/OnlineSubsystemIOS/Public/OnlineSubsystemIOS.h
@@ -25,7 +25,6 @@ public:
 	virtual IOnlineGroupsPtr GetGroupsInterface() const override;
 	virtual IOnlineSharedCloudPtr GetSharedCloudInterface() const override;
 	virtual IOnlineUserCloudPtr GetUserCloudInterface() const override;
-	virtual IOnlineUserCloudPtr GetUserCloudInterface(const FString& Key) const override;
 	virtual IOnlineLeaderboardsPtr GetLeaderboardsInterface() const override;
 	virtual IOnlineVoicePtr GetVoiceInterface() const  override;
 	virtual IOnlineExternalUIPtr GetExternalUIInterface() const override;
diff --git a/Engine/Source/Runtime/Online/OnlineSubsystem/Public/OnlineSubsystem.h b/Engine/Source/Runtime/Online/OnlineSubsystem/Public/OnlineSubsystem.h
--- a/Engine/Source/Runtime/Online/OnlineSubsystem/Public/OnlineSubsystem.h
+++ b/Engine/Source/Runtime/Online/OnlineSubsystem/Public/OnlineSubsystem.h
@@ -179,13 +179,6 @@ public:
 	*/
 	virtual IOnlineUserCloudPtr GetUserCloudInterface() const = 0;
 
-	/**
-	* Get the interface for accessing user files in the cloud for a specific service
-	* @param Key   The key for the required user cloud interface
-	* @return      Interface pointer for the appropriate cloud service
-	*/
-	virtual IOnlineUserCloudPtr GetUserCloudInterface(const FString& Key) const = 0;
-
 	/**
 	 * Get the interface for accessing user entitlements
 	 * @return Interface pointer for the appropriate entitlements service
@@ -277,24 +270,24 @@ public:
 	virtual IOnlineChatPtr GetChatInterface() const = 0;
 
 	/**
-	 * Get the notification handler instance for this subsystem
-	 * @return Pointer for the appropriate notification handler
-	 */
+	* Get the notification handler instance for this subsystem
+	* @return Pointer for the appropriate notification handler
+	*/
 	FOnlineNotificationHandlerPtr GetOnlineNotificationHandler() const
 	{
 		return OnlineNotificationHandler;
 	}
 
 	/**
-	 * Get the interface for managing turn based multiplayer games
-	 * @return Interface pointer for the appropriate online user service
-	 */
+	* Get the interface for managing turn based multiplayer games
+	* @return Interface pointer for the appropriate online user service
+	*/
 	virtual IOnlineTurnBasedPtr GetTurnBasedInterface() const = 0;
 
 	/**
-	 * Get the transport manager instance for this subsystem
-	 * @return Pointer for the appropriate transport manager
-	 */
+	* Get the transport manager instance for this subsystem
+	* @return Pointer for the appropriate transport manager
+	*/
 	FOnlineNotificationTransportManagerPtr GetOnlineNotificationTransportManager() const
 	{
 		return OnlineNotificationTransportManager;
diff --git a/Engine/Source/Runtime/Online/OnlineSubsystemAmazon/Public/OnlineSubsystemAmazon.h b/Engine/Source/Runtime/Online/OnlineSubsystemAmazon/Public/OnlineSubsystemAmazon.h
--- a/Engine/Source/Runtime/Online/OnlineSubsystemAmazon/Public/OnlineSubsystemAmazon.h
+++ b/Engine/Source/Runtime/Online/OnlineSubsystemAmazon/Public/OnlineSubsystemAmazon.h
@@ -57,10 +57,6 @@ public:
 	{
 		return NULL;
 	}
-	virtual IOnlineUserCloudPtr GetUserCloudInterface(const FString& Key) const override
-	{
-		return NULL;
-	}
 	virtual IOnlineLeaderboardsPtr GetLeaderboardsInterface() const override
 	{
 		return NULL;
@@ -156,4 +152,4 @@ public:
 	bool IsEnabled();
 };
 
-typedef TSharedPtr<FOnlineSubsystemAmazon, ESPMode::ThreadSafe> FOnlineSubsystemAmazonPtr;
\ No newline at end of file
+typedef TSharedPtr<FOnlineSubsystemAmazon, ESPMode::ThreadSafe> FOnlineSubsystemAmazonPtr;
diff --git a/Engine/Source/Runtime/Online/OnlineSubsystemFacebook/Private/IOS/OnlineSubsystemFacebook.cpp b/Engine/Source/Runtime/Online/OnlineSubsystemFacebook/Private/IOS/OnlineSubsystemFacebook.cpp
--- a/Engine/Source/Runtime/Online/OnlineSubsystemFacebook/Private/IOS/OnlineSubsystemFacebook.cpp
+++ b/Engine/Source/Runtime/Online/OnlineSubsystemFacebook/Private/IOS/OnlineSubsystemFacebook.cpp
@@ -58,11 +58,6 @@ IOnlineUserCloudPtr FOnlineSubsystemFacebook::GetUserCloudInterface() const
 	return nullptr;
 }
 
-IOnlineUserCloudPtr FOnlineSubsystemFacebook::GetUserCloudInterface(const FString& Key) const
-{
-	return nullptr;
-}
-
 IOnlineLeaderboardsPtr FOnlineSubsystemFacebook::GetLeaderboardsInterface() const
 {
 	return nullptr;
@@ -212,4 +207,4 @@ bool FOnlineSubsystemFacebook::IsEnabled()
 #endif
 
 	return bEnableFacebookSupport;
-}
\ No newline at end of file
+}
diff --git a/Engine/Source/Runtime/Online/OnlineSubsystemFacebook/Private/Windows/OnlineSubsystemFacebook.cpp b/Engine/Source/Runtime/Online/OnlineSubsystemFacebook/Private/Windows/OnlineSubsystemFacebook.cpp
--- a/Engine/Source/Runtime/Online/OnlineSubsystemFacebook/Private/Windows/OnlineSubsystemFacebook.cpp
+++ b/Engine/Source/Runtime/Online/OnlineSubsystemFacebook/Private/Windows/OnlineSubsystemFacebook.cpp
@@ -36,11 +36,6 @@ IOnlineUserCloudPtr FOnlineSubsystemFacebook::GetUserCloudInterface() const
 	return nullptr;
 }
 
-IOnlineUserCloudPtr FOnlineSubsystemFacebook::GetUserCloudInterface(const FString& Key) const
-{
-	return nullptr;
-}
-
 IOnlineLeaderboardsPtr FOnlineSubsystemFacebook::GetLeaderboardsInterface() const
 {
 	return nullptr;
diff --git a/Engine/Source/Runtime/Online/OnlineSubsystemFacebook/Public/OnlineSubsystemFacebook.h b/Engine/Source/Runtime/Online/OnlineSubsystemFacebook/Public/OnlineSubsystemFacebook.h
--- a/Engine/Source/Runtime/Online/OnlineSubsystemFacebook/Public/OnlineSubsystemFacebook.h
+++ b/Engine/Source/Runtime/Online/OnlineSubsystemFacebook/Public/OnlineSubsystemFacebook.h
@@ -27,7 +27,6 @@ public:
 	virtual IOnlineGroupsPtr GetGroupsInterface() const override;
 	virtual IOnlineSharedCloudPtr GetSharedCloudInterface() const override;
 	virtual IOnlineUserCloudPtr GetUserCloudInterface() const override;
-	virtual IOnlineUserCloudPtr GetUserCloudInterface(const FString& Key) const override;
 	virtual IOnlineLeaderboardsPtr GetLeaderboardsInterface() const override;
 	virtual IOnlineVoicePtr GetVoiceInterface() const override;
 	virtual IOnlineExternalUIPtr GetExternalUIInterface() const override;
diff --git a/Engine/Source/Runtime/Online/OnlineSubsystemNull/Private/OnlineSubsystemNull.cpp b/Engine/Source/Runtime/Online/OnlineSubsystemNull/Private/OnlineSubsystemNull.cpp
--- a/Engine/Source/Runtime/Online/OnlineSubsystemNull/Private/OnlineSubsystemNull.cpp
+++ b/Engine/Source/Runtime/Online/OnlineSubsystemNull/Private/OnlineSubsystemNull.cpp
@@ -40,11 +40,6 @@ IOnlineUserCloudPtr FOnlineSubsystemNull::GetUserCloudInterface() const
 	return nullptr;
 }
 
-IOnlineUserCloudPtr FOnlineSubsystemNull::GetUserCloudInterface(const FString& Key) const
-{
-	return nullptr;
-}
-
 IOnlineEntitlementsPtr FOnlineSubsystemNull::GetEntitlementsInterface() const
 {
 	return nullptr;
diff --git a/Engine/Source/Runtime/Online/OnlineSubsystemNull/Public/OnlineSubsystemNull.h b/Engine/Source/Runtime/Online/OnlineSubsystemNull/Public/OnlineSubsystemNull.h
--- a/Engine/Source/Runtime/Online/OnlineSubsystemNull/Public/OnlineSubsystemNull.h
+++ b/Engine/Source/Runtime/Online/OnlineSubsystemNull/Public/OnlineSubsystemNull.h
@@ -38,7 +38,6 @@ public:
 	virtual IOnlineGroupsPtr GetGroupsInterface() const override;
 	virtual IOnlineSharedCloudPtr GetSharedCloudInterface() const override;
 	virtual IOnlineUserCloudPtr GetUserCloudInterface() const override;
-	virtual IOnlineUserCloudPtr GetUserCloudInterface(const FString& Key) const override;
 	virtual IOnlineEntitlementsPtr GetEntitlementsInterface() const override;
 	virtual IOnlineLeaderboardsPtr GetLeaderboardsInterface() const override;
 	virtual IOnlineVoicePtr GetVoiceInterface() const override;
diff --git a/Engine/Source/Runtime/Online/OnlineSubsystemSteam/Private/OnlineSubsystemSteam.cpp b/Engine/Source/Runtime/Online/OnlineSubsystemSteam/Private/OnlineSubsystemSteam.cpp
--- a/Engine/Source/Runtime/Online/OnlineSubsystemSteam/Private/OnlineSubsystemSteam.cpp
+++ b/Engine/Source/Runtime/Online/OnlineSubsystemSteam/Private/OnlineSubsystemSteam.cpp
@@ -184,11 +184,6 @@ IOnlineUserCloudPtr FOnlineSubsystemSteam::GetUserCloudInterface() const
 	return UserCloudInterface;
 }
 
-IOnlineUserCloudPtr FOnlineSubsystemSteam::GetUserCloudInterface(const FString& Key) const
-{
-	return nullptr;
-}
-
 IOnlineLeaderboardsPtr FOnlineSubsystemSteam::GetLeaderboardsInterface() const
 {
 	return LeaderboardsInterface;
diff --git a/Engine/Source/Runtime/Online/OnlineSubsystemSteam/Public/OnlineSubsystemSteam.h b/Engine/Source/Runtime/Online/OnlineSubsystemSteam/Public/OnlineSubsystemSteam.h
--- a/Engine/Source/Runtime/Online/OnlineSubsystemSteam/Public/OnlineSubsystemSteam.h
+++ b/Engine/Source/Runtime/Online/OnlineSubsystemSteam/Public/OnlineSubsystemSteam.h
@@ -210,7 +210,6 @@ public:
 	virtual IOnlineGroupsPtr GetGroupsInterface() const override;
 	virtual IOnlineSharedCloudPtr GetSharedCloudInterface() const override;
 	virtual IOnlineUserCloudPtr GetUserCloudInterface() const override;
-	virtual IOnlineUserCloudPtr GetUserCloudInterface(const FString& Key) const override;
 	virtual IOnlineLeaderboardsPtr GetLeaderboardsInterface() const override;
 	virtual IOnlineVoicePtr GetVoiceInterface() const override;
 	virtual IOnlineExternalUIPtr GetExternalUIInterface() const override;
