PR1299 - fix for RestorePurchases blueprint node name
#ios

diff --git a/Engine/Source/Runtime/Online/OnlineSubsystemUtils/Classes/InAppPurchaseRestoreCallbackProxy.h b/Engine/Source/Runtime/Online/OnlineSubsystemUtils/Classes/InAppPurchaseRestoreCallbackProxy.h
--- a/Engine/Source/Runtime/Online/OnlineSubsystemUtils/Classes/InAppPurchaseRestoreCallbackProxy.h
+++ b/Engine/Source/Runtime/Online/OnlineSubsystemUtils/Classes/InAppPurchaseRestoreCallbackProxy.h
@@ -20,7 +20,7 @@ class UInAppPurchaseRestoreCallbackProxy : public UObject
 	FInAppPurchaseRestoreResult OnFailure;
 
 	// Kicks off a transaction for the provided product identifier
-	UFUNCTION(BlueprintCallable, meta = (DisplayName="Make an In-App Purchase"), Category="Online|InAppPurchase")
+	UFUNCTION(BlueprintCallable, meta = (DisplayName="Restore In-App Purchases"), Category="Online|InAppPurchase")
 	static UInAppPurchaseRestoreCallbackProxy* CreateProxyObjectForInAppPurchaseRestore(class APlayerController* PlayerController);
 
 public:
