PR1286 - addition of Locale data and transaction id to products
#ios

diff --git a/Engine/Source/Runtime/Online/IOS/OnlineSubsystemIOS/Private/OnlineStoreInterfaceIOS.cpp b/Engine/Source/Runtime/Online/IOS/OnlineSubsystemIOS/Private/OnlineStoreInterfaceIOS.cpp
--- a/Engine/Source/Runtime/Online/IOS/OnlineSubsystemIOS/Private/OnlineStoreInterfaceIOS.cpp
+++ b/Engine/Source/Runtime/Online/IOS/OnlineSubsystemIOS/Private/OnlineStoreInterfaceIOS.cpp
@@ -133,6 +133,7 @@
 			}
 
 			StoreInterface->CachedPurchaseStateObject->ProvidedProductInformation.ReceiptData = ReceiptData;
+			StoreInterface->CachedPurchaseStateObject->ProvidedProductInformation.TransactionIdentifier = transaction.transactionIdentifier;
 			StoreInterface->CachedPurchaseStateObject->ReadState = EOnlineAsyncTaskState::Done;
 		}
      
@@ -459,6 +460,10 @@ void FOnlineStoreInterfaceIOS::ProcessProductsResponse( SKProductsResponse* Resp
 			PurchaseProductInfo.DisplayName = [Product localizedTitle];
 			PurchaseProductInfo.DisplayDescription = [Product localizedDescription];
 			PurchaseProductInfo.DisplayPrice = [numberFormatter stringFromNumber : Product.price];
+			PurchaseProductInfo.CurrencyCode = [Product.priceLocale objectForKey : NSLocaleCurrencyCode];
+			PurchaseProductInfo.CurrencySymbol = [Product.priceLocale objectForKey : NSLocaleCurrencySymbol];
+			PurchaseProductInfo.DecimalSeparator = [Product.priceLocale objectForKey : NSLocaleDecimalSeparator];
+			PurchaseProductInfo.GroupingSeparator = [Product.priceLocale objectForKey : NSLocaleGroupingSeparator];
 
 			[numberFormatter release];
 
@@ -507,6 +512,10 @@ void FOnlineStoreInterfaceIOS::ProcessProductsResponse( SKProductsResponse* Resp
 				NewProductInfo.DisplayName = [Product localizedTitle];
 				NewProductInfo.DisplayDescription = [Product localizedDescription];
 				NewProductInfo.DisplayPrice = [numberFormatter stringFromNumber : Product.price];
+				NewProductInfo.CurrencyCode = [Product.priceLocale objectForKey : NSLocaleCurrencyCode];
+				NewProductInfo.CurrencySymbol = [Product.priceLocale objectForKey : NSLocaleCurrencySymbol];
+				NewProductInfo.DecimalSeparator = [Product.priceLocale objectForKey : NSLocaleDecimalSeparator];
+				NewProductInfo.GroupingSeparator = [Product.priceLocale objectForKey : NSLocaleGroupingSeparator];
 
 				UE_LOG(LogOnline, Log, TEXT("\nProduct Identifier: %s, Name: %s, Description: %s, Price: %s\n"),
 					*NewProductInfo.Identifier,
diff --git a/Engine/Source/Runtime/Online/OnlineSubsystem/Public/Interfaces/OnlineStoreInterface.h b/Engine/Source/Runtime/Online/OnlineSubsystem/Public/Interfaces/OnlineStoreInterface.h
--- a/Engine/Source/Runtime/Online/OnlineSubsystem/Public/Interfaces/OnlineStoreInterface.h
+++ b/Engine/Source/Runtime/Online/OnlineSubsystem/Public/Interfaces/OnlineStoreInterface.h
@@ -89,6 +89,10 @@ struct FInAppPurchaseProductInfo
 	UPROPERTY(BlueprintReadOnly, Category = ProductInfo)
 	FString Identifier;
 
+	// the unique transaction identifier
+	UPROPERTY(BlueprintReadOnly, Category = ProductInfo)
+	FString TransactionIdentifier;
+
 	// The localized display name
 	UPROPERTY(BlueprintReadOnly, Category = ProductInfo)
 	FString DisplayName;
@@ -101,6 +105,22 @@ struct FInAppPurchaseProductInfo
 	UPROPERTY(BlueprintReadOnly, Category = ProductInfo)
 	FString DisplayPrice;
 
+	// The localized currency code of the price
+	UPROPERTY(BlueprintReadOnly, Category = ProductInfo)
+	FString CurrencyCode;
+
+	// The localized currency symbol of the price
+	UPROPERTY(BlueprintReadOnly, Category = ProductInfo)
+	FString CurrencySymbol;
+
+	// The localized decimal separator used in the price
+	UPROPERTY(BlueprintReadOnly, Category = ProductInfo)
+	FString DecimalSeparator;
+
+	// The localized grouping separator of the price
+	UPROPERTY(BlueprintReadOnly, Category = ProductInfo)
+	FString GroupingSeparator;
+
 	// The localized display price name
 	UPROPERTY(BlueprintReadOnly, Category = ProductInfo)
 	FString ReceiptData;
