Fix issue where a FOnlineAsyncTaskSteamStoreStats task could get stuck in the queue forever, if a call was made to WriteAchievements (for example) immediately after the steam callback set it to done but before the task in progress had a chance to check the result.

diff --git a/Engine/Source/Runtime/Online/OnlineSubsystemSteam/Private/OnlineAsyncTaskManagerSteam.cpp b/Engine/Source/Runtime/Online/OnlineSubsystemSteam/Private/OnlineAsyncTaskManagerSteam.cpp
--- a/Engine/Source/Runtime/Online/OnlineSubsystemSteam/Private/OnlineAsyncTaskManagerSteam.cpp
+++ b/Engine/Source/Runtime/Online/OnlineSubsystemSteam/Private/OnlineAsyncTaskManagerSteam.cpp
@@ -373,19 +373,6 @@ public:
 	{
 		return FString::Printf(TEXT("FOnlineAsyncEventSteamStatsReceived bWasSuccessful: %d User: %s Result: %s"), (StatsReceivedResult == k_EResultOK) ? 1 : 0, *UserId.ToDebugString(), *SteamResultString(StatsReceivedResult));	
 	}
-
-	/**
-	 * Give the async task a chance to marshal its data back to the game thread
-	 * Can only be called on the game thread by the async task manager
-	 */
-	virtual void Finalize() override
-	{
-		FOnlineAsyncEvent::Finalize();
-
-		// Add given user's stats to the cached list
-		FOnlineLeaderboardsSteamPtr Leaderboards = StaticCastSharedPtr<FOnlineLeaderboardsSteam>(Subsystem->GetLeaderboardsInterface());
-		Leaderboards->SetUserStatsState(UserId, (StatsReceivedResult == k_EResultOK) ? EOnlineAsyncTaskState::Done : EOnlineAsyncTaskState::Failed);
-	}
 };
 
 /**
@@ -473,7 +460,7 @@ public:
 	{
 		FOnlineAsyncEvent::Finalize();
 		FOnlineLeaderboardsSteamPtr Leaderboards = StaticCastSharedPtr<FOnlineLeaderboardsSteam>(Subsystem->GetLeaderboardsInterface());
-		Leaderboards->SetUserStatsStoreState(UserId, StatsStoredResult == k_EResultOK ? EOnlineAsyncTaskState::Done : EOnlineAsyncTaskState::Failed);
+		Leaderboards->UserStatsStoreStatsFinishedDelegate.ExecuteIfBound(StatsStoredResult == k_EResultOK ? EOnlineAsyncTaskState::Done : EOnlineAsyncTaskState::Failed);
 	}
 };
 
@@ -549,18 +536,6 @@ public:
 	{
 		return FString::Printf(TEXT("FOnlineAsyncEventSteamStatsUnloaded UserId: %s"), *UserId.ToDebugString());	
 	}
-
-	/**
-	 * Give the async task a chance to marshal its data back to the game thread
-	 * Can only be called on the game thread by the async task manager
-	 */
-	virtual void Finalize() override
-	{
-		FOnlineAsyncEvent::Finalize();
-		// Remove given user's stats from the cached list
-		FOnlineLeaderboardsSteamPtr Leaderboards = StaticCastSharedPtr<FOnlineLeaderboardsSteam>(Subsystem->GetLeaderboardsInterface());
-		Leaderboards->SetUserStatsState(UserId, EOnlineAsyncTaskState::NotStarted);
-	}
 };
 
 /**
diff --git a/Engine/Source/Runtime/Online/OnlineSubsystemSteam/Private/OnlineLeaderboardInterfaceSteam.cpp b/Engine/Source/Runtime/Online/OnlineSubsystemSteam/Private/OnlineLeaderboardInterfaceSteam.cpp
--- a/Engine/Source/Runtime/Online/OnlineSubsystemSteam/Private/OnlineLeaderboardInterfaceSteam.cpp
+++ b/Engine/Source/Runtime/Online/OnlineSubsystemSteam/Private/OnlineLeaderboardInterfaceSteam.cpp
@@ -194,10 +194,6 @@ public:
 		{
 			UE_LOG_ONLINE(Warning, TEXT("Failed to obtain steam user stats, user: %s error: unknown"), *UserId.ToDebugString());
 		}
-
-		// Add given user's stats to the cached list
-		FOnlineLeaderboardsSteamPtr Leaderboards = StaticCastSharedPtr<FOnlineLeaderboardsSteam>(Subsystem->GetLeaderboardsInterface());
-		Leaderboards->SetUserStatsState(UserId, bWasSuccessful ? EOnlineAsyncTaskState::Done : EOnlineAsyncTaskState::Failed);
 	}
 };
 
@@ -1143,34 +1139,31 @@ public:
 		if (!bInit)
 		{
 			OperationStarted();
+			bInit = true;
+			Leaderboards->UserStatsStoreStatsFinishedDelegate.BindRaw(this, &FOnlineAsyncTaskSteamStoreStats::OnUserStatsStoreStatsFinished);
 			if (!SteamUserStats()->StoreStats())
 			{
-				bIsComplete = true;
-				bWasSuccessful = false;
-				Leaderboards->SetUserStatsStoreState(UserId, EOnlineAsyncTaskState::Failed);
-				OperationFailed();
+				OnUserStatsStoreStatsFinished(EOnlineAsyncTaskState::Failed);
+				return;
 			}
-			bInit = true;
 		}
+	}
 
-		if (bInit)
-		{
-			EOnlineAsyncTaskState::Type AsyncState = Leaderboards->GetUserStatsStoreState(UserId);
-			if (AsyncState != EOnlineAsyncTaskState::InProgress)
-			{
-				bIsComplete = true;
-				bWasSuccessful = (AsyncState == EOnlineAsyncTaskState::Done) ? true : false;
-				Leaderboards->SetUserStatsStoreState(UserId, EOnlineAsyncTaskState::NotStarted);
+	void OnUserStatsStoreStatsFinished(EOnlineAsyncTaskState::Type State)
+	{
+		FOnlineLeaderboardsSteamPtr Leaderboards = StaticCastSharedPtr<FOnlineLeaderboardsSteam>(Subsystem->GetLeaderboardsInterface());
+		Leaderboards->UserStatsStoreStatsFinishedDelegate.Unbind();
 
-				if (bWasSuccessful)
-				{
-					OperationSucceeded();
-				}
-				else
-				{
-					OperationFailed();
-				}
-			}
+		bIsComplete = true;
+		bWasSuccessful = State == EOnlineAsyncTaskState::Done;
+
+		if (bWasSuccessful)
+		{
+			OperationSucceeded();
+		}
+		else
+		{
+			OperationFailed();
 		}
 	}
 };
@@ -1372,7 +1365,6 @@ bool FOnlineLeaderboardsSteam::WriteLeaderboards(const FName& SessionName, const
 
 void FOnlineLeaderboardsSteam::WriteAchievementsInternal(const FUniqueNetIdSteam& UserId, FOnlineAchievementsWriteRef& WriteObject, const FOnAchievementsWrittenDelegate& OnWriteFinishedDelegate)
 {
-	SetUserStatsStoreState(UserId, EOnlineAsyncTaskState::InProgress);
 	FOnlineAsyncTaskSteamWriteAchievements* NewTask = new FOnlineAsyncTaskSteamWriteAchievements(SteamSubsystem, UserId, WriteObject, OnWriteFinishedDelegate);
 	SteamSubsystem->QueueAsyncTask(NewTask);
 }
@@ -1380,7 +1372,6 @@ void FOnlineLeaderboardsSteam::WriteAchievementsInternal(const FUniqueNetIdSteam
 bool FOnlineLeaderboardsSteam::FlushLeaderboards(const FName& SessionName)
 {
 	FUniqueNetIdSteam UserId(SteamUser()->GetSteamID());
-	SetUserStatsStoreState(UserId, EOnlineAsyncTaskState::InProgress);
 	FOnlineAsyncTaskSteamFlushLeaderboards* NewTask = new FOnlineAsyncTaskSteamFlushLeaderboards(SteamSubsystem, SessionName, UserId);
 	SteamSubsystem->QueueAsyncTask(NewTask);
 	return true;
@@ -1447,82 +1438,8 @@ void FOnlineLeaderboardsSteam::FindLeaderboard(const FName& LeaderboardName)
 
 void FOnlineLeaderboardsSteam::CacheCurrentUsersStats()
 {
-	ISteamUser* SteamUserPtr = SteamUser();
-	if (SteamUserPtr != NULL &&
-		SteamUserPtr->BLoggedOn())
-	{
-		FUniqueNetIdSteam UserId(SteamUserPtr->GetSteamID());
-		SetUserStatsState(UserId, EOnlineAsyncTaskState::InProgress);
+	// Triggers a Steam event async to let us know when the stats are available
+	SteamUserStats()->RequestCurrentStats();
 
-		// Triggers a Steam event async to let us know when the stats are available
-		SteamUserStats()->RequestCurrentStats();
-	}
 	// TODO ONLINE - Need to validate this on user login and invalidate on user logout
 }
-
-EOnlineAsyncTaskState::Type FOnlineLeaderboardsSteam::GetUserStatsState(const FUniqueNetIdSteam& UserId)
-{
-	for (int32 UserIdx = 0; UserIdx < UserStatsReceivedState.Num(); UserIdx++)
-	{
-		if (UserStatsReceivedState[UserIdx].UserId == UserId)
-		{
-			return UserStatsReceivedState[UserIdx].StatsState;
-		}
-	}
-
-	// Default to not started when not found
-	return EOnlineAsyncTaskState::NotStarted;
-}
-
-void FOnlineLeaderboardsSteam::SetUserStatsState(const FUniqueNetIdSteam& UserId, EOnlineAsyncTaskState::Type NewState)
-{
-	// Try to update an existing user
-	for (int32 UserIdx = 0; UserIdx < UserStatsReceivedState.Num(); UserIdx++)
-	{
-		if (UserStatsReceivedState[UserIdx].UserId == UserId)
-		{
-			UserStatsReceivedState[UserIdx].StatsState = NewState;
-			return;
-		}
-	}
-
-	// Create a new user entry
-	new (UserStatsReceivedState) FUserStatsStateSteam(UserId, NewState);
-}
-
-EOnlineAsyncTaskState::Type FOnlineLeaderboardsSteam::GetUserStatsStoreState(const FUniqueNetIdSteam& UserId)
-{
-	FScopeLock ScopeLock(&UserStatsStoredLock);
-	for (int32 UserIdx = 0; UserIdx < UserStatsStoredState.Num(); UserIdx++)
-	{
-		if (UserStatsStoredState[UserIdx].UserId == UserId)
-		{
-			return UserStatsStoredState[UserIdx].StatsState;
-		}
-	}
-
-	// Default to not started when not found
-	return EOnlineAsyncTaskState::NotStarted;
-}
-
-void FOnlineLeaderboardsSteam::SetUserStatsStoreState(const FUniqueNetIdSteam& UserId, EOnlineAsyncTaskState::Type NewState)
-{
-	FScopeLock ScopeLock(&UserStatsStoredLock);
-
-	// Try to update an existing user
-	for (int32 UserIdx = 0; UserIdx < UserStatsStoredState.Num(); UserIdx++)
-	{
-		if (UserStatsStoredState[UserIdx].UserId == UserId)
-		{
-			UserStatsStoredState[UserIdx].StatsState = NewState;
-			return;
-		}
-	}
-
-	// Create a new user entry
-	new (UserStatsStoredState) FUserStatsStateSteam(UserId, NewState);
-}
-
-
-
-
diff --git a/Engine/Source/Runtime/Online/OnlineSubsystemSteam/Private/OnlineLeaderboardInterfaceSteam.h b/Engine/Source/Runtime/Online/OnlineSubsystemSteam/Private/OnlineLeaderboardInterfaceSteam.h
--- a/Engine/Source/Runtime/Online/OnlineSubsystemSteam/Private/OnlineLeaderboardInterfaceSteam.h
+++ b/Engine/Source/Runtime/Online/OnlineSubsystemSteam/Private/OnlineLeaderboardInterfaceSteam.h
@@ -33,6 +33,8 @@ public:
 	}
 };
 
+DECLARE_DELEGATE_OneParam(FOnSteamUserStatsStoreStatsFinished, EOnlineAsyncTaskState::Type);
+
 /**
  * Interface definition for the online services leaderboard services 
  */
@@ -54,12 +56,10 @@ PACKAGE_SCOPE:
 
 	/** Critical section for thread safe operation of the leaderboard metadata */
 	FCriticalSection LeaderboardMetadataLock;
-	/** All known stats received/removed requests from Steam */
-	TArray<FUserStatsStateSteam> UserStatsReceivedState;
 	/** Critical section for thread safe operation of the stats stored state */
 	FCriticalSection UserStatsStoredLock;
-	/** All known stats received/removed requests from Steam */
-	TArray<FUserStatsStateSteam> UserStatsStoredState;
+	/** Called when the event from SteamUserStats()->StoreStats() finishing is triggered */
+	FOnSteamUserStatsStoreStatsFinished UserStatsStoreStatsFinishedDelegate;
 
 	FOnlineLeaderboardsSteam(FOnlineSubsystemSteam* InSteamSubsystem) :
 		SteamSubsystem(InSteamSubsystem)
@@ -113,22 +113,6 @@ PACKAGE_SCOPE:
 	 */
 	void SetUserStatsState(const FUniqueNetIdSteam& UserId, EOnlineAsyncTaskState::Type NewState);
 
-	/**
-	 *	Get the stored stats state for a given user
-	 * 
-	 * @param User user to check if stats have been stored
-	 * @return State of the stats store for a given user, or NotStarted if user isn't found
-	 */
-	EOnlineAsyncTaskState::Type GetUserStatsStoreState(const FUniqueNetIdSteam& UserId);
-
-	/**
-	 *	Set the received stats state for a given user (game thread only)
-	 * 
-	 * @param User user to set state of stats stored
-	 * @param NewState state of the stats store for a given user
-	 */
-	void SetUserStatsStoreState(const FUniqueNetIdSteam& UserId, EOnlineAsyncTaskState::Type NewState);
-
 	/**
 	 * Commits any changes in the online stats cache to the permanent storage (internal helper for FOnlineAchievementsSteam::WriteAchievements)
 	 *
