Improvements to Android Google Play Games and iOS Game Center login, UE-4190 and UEPLAT-618:
FOnlineSubsystemGooglePlay and FOnlineSubsystemIOS no longer try to automatically connect during engine startup. This was causing timing-related problems and didn't allow the game to control when the login UI appears. For these subsystems you will now have to call ShowLoginUI explicitly to initiate a connection with the service.
Implemented FOnlineExternalUIGooglePlay::ShowLoginUI and FOnlineExternalUIIOS::ShowLoginUI, which will try to connect to the respective service if not currently connected, showing the platform's login UI if necessary.
Added a new blueprint node "Show External Login UI" to expose the ShowLoginUI C++ function
Deprecated the Connect to Service node (use show external login UI instead)
Implemented FOnlineIdentityGooglePlay::Logout. Note iOS apps do not support programmatic logout from within the game, so the iOS version of this function is unsupported.
Added a new blueprint node "Logout" to expose the Logout C++ function.
Added a new node "Is Logged In" to expose the Online Subsystem login status to blueprints.
Fixed the implementation of FOnlineIdentityGooglePlay::GetLoginStatus.

diff --git a/Engine/Source/Runtime/Engine/Classes/Kismet/KismetSystemLibrary.h b/Engine/Source/Runtime/Engine/Classes/Kismet/KismetSystemLibrary.h
--- a/Engine/Source/Runtime/Engine/Classes/Kismet/KismetSystemLibrary.h
+++ b/Engine/Source/Runtime/Engine/Classes/Kismet/KismetSystemLibrary.h
@@ -1395,6 +1395,14 @@ class ENGINE_API UKismetSystemLibrary : public UBlueprintFunctionLibrary
 	UFUNCTION(BlueprintCallable, Category = "Utilities|Platform")
 	static void ShowPlatformSpecificAchievementsScreen(class APlayerController* SpecificPlayer);
 
+	/**
+	 * Returns whether the player is logged in to the currently active online subsystem.
+	 *
+	 * @param Player Specific player's login status to get. May not be supported on all platforms. If null, defaults to the player with ControllerId 0.
+	 */
+	UFUNCTION(BlueprintPure, Category = "Online")
+	static bool IsLoggedIn(APlayerController* SpecificPlayer);
+
 	/**
 	 * Allows or inhibits screensaver
 	 * @param	bAllowScreenSaver		If false, don't allow screensaver if possible, otherwise allow default behavior
diff --git a/Engine/Source/Runtime/Engine/Private/KismetSystemLibrary.cpp b/Engine/Source/Runtime/Engine/Private/KismetSystemLibrary.cpp
--- a/Engine/Source/Runtime/Engine/Private/KismetSystemLibrary.cpp
+++ b/Engine/Source/Runtime/Engine/Private/KismetSystemLibrary.cpp
@@ -14,6 +14,9 @@
 #include "Kismet/GameplayStatics.h"
 #include "SlateCore.h"
 #include "Engine/StreamableManager.h"
+#include "OnlineSubsystemTypes.h"
+#include "OnlineSubsystemUtils.h"
+#include "OnlineIdentityInterface.h"
 
 //////////////////////////////////////////////////////////////////////////
 // UKismetSystemLibrary
@@ -3074,6 +3077,26 @@ void UKismetSystemLibrary::ShowPlatformSpecificAchievementsScreen(class APlayerC
 		ExternalUI->ShowAchievementsUI(LocalUserNum);
 	}
 }
+bool UKismetSystemLibrary::IsLoggedIn(APlayerController* SpecificPlayer)
+{
+	IOnlineIdentityPtr Identity = Online::GetIdentityInterface();
+	
+	if (!Identity.IsValid())
+	{
+		return false;
+	}
+	
+	int LocalUserNum = 0;
+	if (SpecificPlayer != nullptr)
+	{
+		ULocalPlayer* LocalPlayer = Cast<ULocalPlayer>(SpecificPlayer->Player);
+		if(LocalPlayer)
+		{
+			LocalUserNum = LocalPlayer->GetControllerId();
+		}
+	}
+	return Identity->GetLoginStatus(LocalUserNum) == ELoginStatus::LoggedIn;
+}
 
 void UKismetSystemLibrary::SetStructurePropertyByName(UObject* Object, FName PropertyName, const FGenericStruct& Value)
 {
diff --git a/Engine/Source/Runtime/Online/Android/OnlineSubsystemGooglePlay/Private/OnlineAsyncTaskGooglePlayAuthAction.cpp b/Engine/Source/Runtime/Online/Android/OnlineSubsystemGooglePlay/Private/OnlineAsyncTaskGooglePlayAuthAction.cpp
new file mode 100644
--- /dev/null
+++ b/Engine/Source/Runtime/Online/Android/OnlineSubsystemGooglePlay/Private/OnlineAsyncTaskGooglePlayAuthAction.cpp
@@ -0,0 +1,24 @@
+// Copyright 1998-2015 Epic Games, Inc. All Rights Reserved.
+
+#include "OnlineSubsystemGooglePlayPrivatePCH.h"
+#include "OnlineAsyncTaskGooglePlayAuthAction.h"
+
+#include "gpg/builder.h"
+#include "gpg/debug.h"
+
+FOnlineAsyncTaskGooglePlayAuthAction::FOnlineAsyncTaskGooglePlayAuthAction(
+	FOnlineSubsystemGooglePlay* InSubsystem)
+	: FOnlineAsyncTaskBasic(InSubsystem)
+	, bInit(false)
+{
+	check(InSubsystem);
+}
+
+void FOnlineAsyncTaskGooglePlayAuthAction::Tick()
+{
+	if ( !bInit)
+	{
+		Start_OnTaskThread();
+		bInit = true;
+	}
+}
diff --git a/Engine/Source/Runtime/Online/Android/OnlineSubsystemGooglePlay/Private/OnlineAsyncTaskGooglePlayAuthAction.h b/Engine/Source/Runtime/Online/Android/OnlineSubsystemGooglePlay/Private/OnlineAsyncTaskGooglePlayAuthAction.h
new file mode 100644
--- /dev/null
+++ b/Engine/Source/Runtime/Online/Android/OnlineSubsystemGooglePlay/Private/OnlineAsyncTaskGooglePlayAuthAction.h
@@ -0,0 +1,48 @@
+// Copyright 1998-2015 Epic Games, Inc. All Rights Reserved.
+
+#pragma once
+
+#include "OnlineAsyncTaskManager.h"
+#include "OnlineSubsystemGooglePlayPackage.h"
+
+#include "gpg/status.h"
+#include "gpg/types.h"
+
+class FOnlineSubsystemGooglePlay;
+
+class FOnlineAsyncTaskGooglePlayAuthAction : public FOnlineAsyncTaskBasic<FOnlineSubsystemGooglePlay>
+{
+public:
+	/**
+	 * Constructor.
+	 *
+	 * @param InSubsystem a pointer to the owning subsystem
+	 * @param InPlayerId index of the player who's logging in
+	 */
+	FOnlineAsyncTaskGooglePlayAuthAction(FOnlineSubsystemGooglePlay* InSubsystem);
+
+	// FOnlineAsyncTask
+	virtual void Tick() override;
+
+	// FOnlineAsyncItem
+	virtual FString ToString() const override { return TEXT("AuthAction"); }
+
+private:
+	/**
+	 * The OnAuthActionFinished callback is handled globally in FOnlineSubsystemGooglePlay.
+	 * The subsystem is responsible for forwarding the call to any pending auth task.
+	 *
+	 * @param InOp Forwarded from Google's callback, indicates whether this was a sign-in or a sign-out
+	 * @param InStatus Forwarded from Google's callback, indicates whether the operation succeeded
+	 */
+	virtual void OnAuthActionFinished(gpg::AuthOperation InOp, gpg::AuthStatus InStatus) = 0;
+
+	/**
+	 * This function will be called once during the task's first tick on the online thread.
+	 * Usually this is the best place to make the system API call to start the async operation.
+	 */
+	virtual void Start_OnTaskThread() {}
+
+	/** Tracks whether or not this task has ticked at least once */
+	bool bInit;
+};
diff --git a/Engine/Source/Runtime/Online/Android/OnlineSubsystemGooglePlay/Private/OnlineAsyncTaskGooglePlayLogin.cpp b/Engine/Source/Runtime/Online/Android/OnlineSubsystemGooglePlay/Private/OnlineAsyncTaskGooglePlayLogin.cpp
--- a/Engine/Source/Runtime/Online/Android/OnlineSubsystemGooglePlay/Private/OnlineAsyncTaskGooglePlayLogin.cpp
+++ b/Engine/Source/Runtime/Online/Android/OnlineSubsystemGooglePlay/Private/OnlineAsyncTaskGooglePlayLogin.cpp
@@ -3,46 +3,73 @@
 #include "OnlineSubsystemGooglePlayPrivatePCH.h"
 #include "OnlineAsyncTaskGooglePlayLogin.h"
 
-namespace
-{
-	const int MaxRetries = 1;
-}
+#include "gpg/builder.h"
+#include "gpg/debug.h"
 
 FOnlineAsyncTaskGooglePlayLogin::FOnlineAsyncTaskGooglePlayLogin(
 	FOnlineSubsystemGooglePlay* InSubsystem,
-	int InPlayerId)
-	: FOnlineAsyncTaskBasic(InSubsystem)
+	int InPlayerId,
+	const FOnCompletedDelegate& InDelegate)
+	: FOnlineAsyncTaskGooglePlayAuthAction(InSubsystem)
 	, PlayerId(InPlayerId)
 	, Status(gpg::AuthStatus::ERROR_NOT_AUTHORIZED)
-	, RetryCount(MaxRetries)
+	, Delegate(InDelegate)
 {
 }
 
+void FOnlineAsyncTaskGooglePlayLogin::Start_OnTaskThread()
+{
+	// If we haven't created a GameServices object yet, do so.
+	if (Subsystem->GameServicesPtr.get() == nullptr)
+	{
+		// Store the Subsystem pointer locally so that the OnAuthActionFinished lambda can capture it
+		FOnlineSubsystemGooglePlay* LocalSubsystem = Subsystem;
+
+		Subsystem->GameServicesPtr = GameServices::Builder()
+			.SetDefaultOnLog(LogLevel::VERBOSE)
+			.SetOnAuthActionStarted([](AuthOperation Op) {
+				UE_LOG(LogOnline, Log, TEXT("GPG OnAuthActionStarted: %s"), *FString(DebugString(Op).c_str()));
+			})
+			.SetOnAuthActionFinished([LocalSubsystem](AuthOperation Op, AuthStatus LocalStatus) {
+				UE_LOG(LogOnline, Log, TEXT("GPG OnAuthActionFinished: %s, AuthStatus: %s"),
+					*FString(DebugString(Op).c_str()),
+					*FString(DebugString(LocalStatus).c_str()));
+				LocalSubsystem->OnAuthActionFinished(Op, LocalStatus);
+			})
+			.Create(Subsystem->PlatformConfiguration);
+	}
+	else if(Subsystem->GameServicesPtr->IsAuthorized())
+	{
+		// We have a GameServices object and the user is authorized, nothing else to do.
+		Status = AuthStatus::VALID;
+		bWasSuccessful = true;
+		bIsComplete = true;
+	}
+	else
+	{
+		// We have created the GameServices object but the user isn't authorized.
+		bWasSuccessful = false;
+		bIsComplete = true;
+	}
+}
+
 void FOnlineAsyncTaskGooglePlayLogin::Finalize()
 {
-	Subsystem->GetIdentityGooglePlay()->OnLoginCompleted(PlayerId, Status);
+	// Async task manager owns the task and is responsible for cleaning it up.
+	Subsystem->CurrentLoginTask = nullptr;
 }
 
 void FOnlineAsyncTaskGooglePlayLogin::TriggerDelegates()
 {
+	Delegate.ExecuteIfBound();
 }
 
 void FOnlineAsyncTaskGooglePlayLogin::OnAuthActionFinished(gpg::AuthOperation InOp, gpg::AuthStatus InStatus)
 {
 	if (InOp == gpg::AuthOperation::SIGN_IN)
 	{
-		if (InStatus == gpg::AuthStatus::ERROR_NOT_AUTHORIZED && RetryCount > 0)
-		{
-			--RetryCount;
-			Subsystem->GetGameServices()->StartAuthorizationUI();
-			return;
-		}
-		
 		Status = InStatus;
 		bWasSuccessful = Status == gpg::AuthStatus::VALID;
 		bIsComplete = true;
 	}
 }
-
-
-
diff --git a/Engine/Source/Runtime/Online/Android/OnlineSubsystemGooglePlay/Private/OnlineAsyncTaskGooglePlayLogin.h b/Engine/Source/Runtime/Online/Android/OnlineSubsystemGooglePlay/Private/OnlineAsyncTaskGooglePlayLogin.h
--- a/Engine/Source/Runtime/Online/Android/OnlineSubsystemGooglePlay/Private/OnlineAsyncTaskGooglePlayLogin.h
+++ b/Engine/Source/Runtime/Online/Android/OnlineSubsystemGooglePlay/Private/OnlineAsyncTaskGooglePlayLogin.h
@@ -3,6 +3,7 @@
 #pragma once
 
 #include "OnlineAsyncTaskManager.h"
+#include "OnlineAsyncTaskGooglePlayAuthAction.h"
 #include "OnlineSubsystemGooglePlayPackage.h"
 
 #include "gpg/status.h"
@@ -10,34 +11,34 @@
 
 class FOnlineSubsystemGooglePlay;
 
-class FOnlineAsyncTaskGooglePlayLogin : public FOnlineAsyncTaskBasic<FOnlineSubsystemGooglePlay>
+class FOnlineAsyncTaskGooglePlayLogin : public FOnlineAsyncTaskGooglePlayAuthAction
 {
 public:
+	/** Delegate fired upon completion. */
+	DECLARE_DELEGATE(FOnCompletedDelegate);
+
 	/**
 	 * Constructor.
 	 *
 	 * @param InSubsystem a pointer to the owning subsysetm
 	 * @param InPlayerId index of the player who's logging in
 	 */
-	FOnlineAsyncTaskGooglePlayLogin(FOnlineSubsystemGooglePlay* InSubsystem, int InPlayerId);
+	FOnlineAsyncTaskGooglePlayLogin(FOnlineSubsystemGooglePlay* InSubsystem, int InPlayerId, const FOnCompletedDelegate& InDelegate);
 
 	// FOnlineAsyncItem
 	virtual FString ToString() const override { return TEXT("Login"); }
 	virtual void Finalize() override;
 	virtual void TriggerDelegates() override;
 
-PACKAGE_SCOPE:
-	/**
-	 * The OnAuthActionFinished callback is handled globally in FOnlineSubsystemGooglePlay.
-	 * The subsystem is responsible for forwarding the call to any pending login task.
-	 *
-	 * @param InOp Forwarded from Google's callback, indicates whether this was a sign-in or a sign-out
-	 * @param InStatus Forwarded from Google's callback, indicates whether the operation succeeded
-	 */
-	void OnAuthActionFinished(gpg::AuthOperation InOp, gpg::AuthStatus InStatus);
-
 private:
+	/** The subsystem is the only class that should be calling OnAuthActionFinished */
+	friend class FOnlineSubsystemGooglePlay;
+
+	// FOnlineAsyncTaskGooglePlayAuthAction
+	virtual void OnAuthActionFinished(gpg::AuthOperation InOp, gpg::AuthStatus InStatus) override;
+	virtual void Start_OnTaskThread() override;
+
 	int PlayerId;
 	gpg::AuthStatus Status;
-	int RetryCount;
+	FOnCompletedDelegate Delegate;
 };
diff --git a/Engine/Source/Runtime/Online/Android/OnlineSubsystemGooglePlay/Private/OnlineAsyncTaskGooglePlayLogout.cpp b/Engine/Source/Runtime/Online/Android/OnlineSubsystemGooglePlay/Private/OnlineAsyncTaskGooglePlayLogout.cpp
new file mode 100644
--- /dev/null
+++ b/Engine/Source/Runtime/Online/Android/OnlineSubsystemGooglePlay/Private/OnlineAsyncTaskGooglePlayLogout.cpp
@@ -0,0 +1,50 @@
+// Copyright 1998-2015 Epic Games, Inc. All Rights Reserved.
+
+#include "OnlineSubsystemGooglePlayPrivatePCH.h"
+#include "OnlineAsyncTaskGooglePlayLogout.h"
+
+#include "gpg/game_services.h"
+
+FOnlineAsyncTaskGooglePlayLogout::FOnlineAsyncTaskGooglePlayLogout(
+	FOnlineSubsystemGooglePlay* InSubsystem,
+	int32 InPlayerId)
+	: FOnlineAsyncTaskGooglePlayAuthAction(InSubsystem)
+	, PlayerId(InPlayerId)
+	, Status(gpg::AuthStatus::ERROR_NOT_AUTHORIZED)
+{
+}
+
+void FOnlineAsyncTaskGooglePlayLogout::Start_OnTaskThread()
+{
+	// If we haven't created a GameServices object yet, do so.
+	if (Subsystem->GameServicesPtr.get() == nullptr)
+	{
+		UE_LOG(LogOnline, Log, TEXT("FOnlineAsyncTaskGooglePlayLogout::Start_OnTaskThread: GameServicesPtr is null"));
+		bWasSuccessful = false;
+		bIsComplete = true;
+		return;
+	}
+
+	Subsystem->GetGameServices()->SignOut();
+}
+
+void FOnlineAsyncTaskGooglePlayLogout::Finalize()
+{
+	// Async task manager owns the task and is responsible for cleaning it up.
+	Subsystem->CurrentLogoutTask = nullptr;
+}
+
+void FOnlineAsyncTaskGooglePlayLogout::TriggerDelegates()
+{
+	Subsystem->GetIdentityInterface()->TriggerOnLogoutCompleteDelegates(PlayerId, bWasSuccessful);
+}
+
+void FOnlineAsyncTaskGooglePlayLogout::OnAuthActionFinished(gpg::AuthOperation InOp, gpg::AuthStatus InStatus)
+{
+	if (InOp == gpg::AuthOperation::SIGN_OUT)
+	{
+		Status = InStatus;
+		bWasSuccessful = true;
+		bIsComplete = true;
+	}
+}
diff --git a/Engine/Source/Runtime/Online/Android/OnlineSubsystemGooglePlay/Private/OnlineAsyncTaskGooglePlayLogout.h b/Engine/Source/Runtime/Online/Android/OnlineSubsystemGooglePlay/Private/OnlineAsyncTaskGooglePlayLogout.h
new file mode 100644
--- /dev/null
+++ b/Engine/Source/Runtime/Online/Android/OnlineSubsystemGooglePlay/Private/OnlineAsyncTaskGooglePlayLogout.h
@@ -0,0 +1,40 @@
+// Copyright 1998-2015 Epic Games, Inc. All Rights Reserved.
+
+#pragma once
+
+#include "OnlineAsyncTaskManager.h"
+#include "OnlineAsyncTaskGooglePlayAuthAction.h"
+#include "OnlineSubsystemGooglePlayPackage.h"
+
+#include "gpg/status.h"
+#include "gpg/types.h"
+
+class FOnlineSubsystemGooglePlay;
+
+class FOnlineAsyncTaskGooglePlayLogout : public FOnlineAsyncTaskGooglePlayAuthAction
+{
+public:
+	/**
+	 * Constructor.
+	 *
+	 * @param InSubsystem a pointer to the owning subsysetm
+	 * @param InPlayerId index of the player who's logging in
+	 */
+	FOnlineAsyncTaskGooglePlayLogout(FOnlineSubsystemGooglePlay* InSubsystem, int32 InPlayerId);
+
+	// FOnlineAsyncItem
+	virtual FString ToString() const override { return TEXT("Logout"); }
+	virtual void Finalize() override;
+	virtual void TriggerDelegates() override;
+
+private:
+	/** The subsystem is the only class that should be calling OnAuthActionFinished */
+	friend class FOnlineSubsystemGooglePlay;
+
+	// FOnlineAsyncTaskGooglePlayAuthAction
+	virtual void OnAuthActionFinished(gpg::AuthOperation InOp, gpg::AuthStatus InStatus) override;
+	virtual void Start_OnTaskThread() override;
+
+	int32 PlayerId;
+	gpg::AuthStatus Status;
+};
diff --git a/Engine/Source/Runtime/Online/Android/OnlineSubsystemGooglePlay/Private/OnlineAsyncTaskGooglePlayShowLoginUI.cpp b/Engine/Source/Runtime/Online/Android/OnlineSubsystemGooglePlay/Private/OnlineAsyncTaskGooglePlayShowLoginUI.cpp
new file mode 100644
--- /dev/null
+++ b/Engine/Source/Runtime/Online/Android/OnlineSubsystemGooglePlay/Private/OnlineAsyncTaskGooglePlayShowLoginUI.cpp
@@ -0,0 +1,62 @@
+// Copyright 1998-2015 Epic Games, Inc. All Rights Reserved.
+
+#include "OnlineSubsystemGooglePlayPrivatePCH.h"
+#include "OnlineAsyncTaskGooglePlayShowLoginUI.h"
+
+FOnlineAsyncTaskGooglePlayShowLoginUI::FOnlineAsyncTaskGooglePlayShowLoginUI(
+	FOnlineSubsystemGooglePlay* InSubsystem,
+	int InPlayerId,
+	const IOnlineExternalUI::FOnLoginUIClosedDelegate& InDelegate)
+	: FOnlineAsyncTaskGooglePlayAuthAction(InSubsystem)
+	, PlayerId(InPlayerId)
+	, Delegate(InDelegate)
+{
+}
+
+void FOnlineAsyncTaskGooglePlayShowLoginUI::Start_OnTaskThread()
+{	
+	if(Subsystem->GetGameServices() == nullptr)
+	{
+		UE_LOG(LogOnline, Log, TEXT("FOnlineAsyncTaskGooglePlayShowLoginUI: GameServicesPtr is null."));
+		bWasSuccessful = false;
+		bIsComplete = true;
+		return;
+	}
+
+	if(Subsystem->GetGameServices()->IsAuthorized())
+	{
+		// User is already authorized, nothing to do.
+		bWasSuccessful = true;
+		bIsComplete = true;
+	}
+
+	// The user isn't authorized, show the sign-in UI.
+	Subsystem->GetGameServices()->StartAuthorizationUI();
+}
+
+void FOnlineAsyncTaskGooglePlayShowLoginUI::Finalize()
+{
+	// Async task manager owns the task and is responsible for cleaning it up.
+	Subsystem->CurrentShowLoginUITask = nullptr;
+}
+
+void FOnlineAsyncTaskGooglePlayShowLoginUI::TriggerDelegates()
+{
+	TSharedPtr<FUniqueNetId> UserId;
+
+	if (bWasSuccessful)
+	{
+		UserId = MakeShareable(new FUniqueNetIdString());
+	}
+
+	Delegate.ExecuteIfBound(UserId, PlayerId);
+}
+
+void FOnlineAsyncTaskGooglePlayShowLoginUI::OnAuthActionFinished(gpg::AuthOperation InOp, gpg::AuthStatus InStatus)
+{
+	if (InOp == AuthOperation::SIGN_IN)
+	{
+		bWasSuccessful = InStatus == AuthStatus::VALID;
+		bIsComplete = true;
+	}
+}
diff --git a/Engine/Source/Runtime/Online/Android/OnlineSubsystemGooglePlay/Private/OnlineAsyncTaskGooglePlayShowLoginUI.h b/Engine/Source/Runtime/Online/Android/OnlineSubsystemGooglePlay/Private/OnlineAsyncTaskGooglePlayShowLoginUI.h
new file mode 100644
--- /dev/null
+++ b/Engine/Source/Runtime/Online/Android/OnlineSubsystemGooglePlay/Private/OnlineAsyncTaskGooglePlayShowLoginUI.h
@@ -0,0 +1,40 @@
+// Copyright 1998-2015 Epic Games, Inc. All Rights Reserved.
+
+#pragma once
+
+#include "OnlineAsyncTaskManager.h"
+#include "OnlineExternalUIInterface.h"
+#include "OnlineAsyncTaskGooglePlayAuthAction.h"
+#include "OnlineSubsystemGooglePlayPackage.h"
+
+#include "gpg/types.h"
+
+class FOnlineSubsystemGooglePlay;
+
+class FOnlineAsyncTaskGooglePlayShowLoginUI : public FOnlineAsyncTaskGooglePlayAuthAction
+{
+public:
+	/**
+	 * Constructor.
+	 *
+	 * @param InSubsystem a pointer to the owning subsystem
+	 * @param InPlayerId index of the player who's logging in
+	 */
+	FOnlineAsyncTaskGooglePlayShowLoginUI(FOnlineSubsystemGooglePlay* InSubsystem, int InPlayerId, const IOnlineExternalUI::FOnLoginUIClosedDelegate& InDelegate);
+
+	// FOnlineAsyncItem
+	virtual FString ToString() const override { return TEXT("ShowLoginUI"); }
+	virtual void Finalize() override;
+	virtual void TriggerDelegates() override;
+
+private:
+	/** The subsystem is the only class that should be calling OnAuthActionFinished */
+	friend class FOnlineSubsystemGooglePlay;
+
+	// FOnlineAsyncTaskGooglePlayAuthAction
+	virtual void OnAuthActionFinished(gpg::AuthOperation InOp, gpg::AuthStatus InStatus) override;
+	virtual void Start_OnTaskThread() override;
+
+	int PlayerId;
+	IOnlineExternalUI::FOnLoginUIClosedDelegate Delegate;
+};
diff --git a/Engine/Source/Runtime/Online/Android/OnlineSubsystemGooglePlay/Private/OnlineExternalUIInterfaceGooglePlay.cpp b/Engine/Source/Runtime/Online/Android/OnlineSubsystemGooglePlay/Private/OnlineExternalUIInterfaceGooglePlay.cpp
--- a/Engine/Source/Runtime/Online/Android/OnlineSubsystemGooglePlay/Private/OnlineExternalUIInterfaceGooglePlay.cpp
+++ b/Engine/Source/Runtime/Online/Android/OnlineSubsystemGooglePlay/Private/OnlineExternalUIInterfaceGooglePlay.cpp
@@ -15,7 +15,8 @@ FOnlineExternalUIGooglePlay::FOnlineExternalUIGooglePlay(FOnlineSubsystemGoogleP
 
 bool FOnlineExternalUIGooglePlay::ShowLoginUI(const int ControllerIndex, bool bShowOnlineOnly, const FOnLoginUIClosedDelegate& Delegate)
 {
-	return false;
+	Subsystem->StartShowLoginUITask(ControllerIndex, Delegate);	
+	return true;
 }
 
 bool FOnlineExternalUIGooglePlay::ShowFriendsUI(int32 LocalUserNum)
@@ -35,6 +36,11 @@ bool FOnlineExternalUIGooglePlay::ShowAchievementsUI(int32 LocalUserNum)
 		return false;
 	}
 
+	if (!Subsystem->GetGameServices()->IsAuthorized())
+	{
+		return false;
+	}
+
 	Subsystem->GetGameServices()->Achievements().ShowAllUI();
 	return true;
 }
@@ -46,6 +52,11 @@ bool FOnlineExternalUIGooglePlay::ShowLeaderboardUI(const FString& LeaderboardNa
 		return false;
 	}
 
+	if (!Subsystem->GetGameServices()->IsAuthorized())
+	{
+		return false;
+	}
+
 	auto DefaultSettings = GetDefault<UAndroidRuntimeSettings>();
 
 	for(const auto& Mapping : DefaultSettings->LeaderboardMap)
diff --git a/Engine/Source/Runtime/Online/Android/OnlineSubsystemGooglePlay/Private/OnlineIdentityInterfaceGooglePlay.cpp b/Engine/Source/Runtime/Online/Android/OnlineSubsystemGooglePlay/Private/OnlineIdentityInterfaceGooglePlay.cpp
--- a/Engine/Source/Runtime/Online/Android/OnlineSubsystemGooglePlay/Private/OnlineIdentityInterfaceGooglePlay.cpp
+++ b/Engine/Source/Runtime/Online/Android/OnlineSubsystemGooglePlay/Private/OnlineIdentityInterfaceGooglePlay.cpp
@@ -14,6 +14,7 @@ FOnlineIdentityGooglePlay::FOnlineIdentityGooglePlay(FOnlineSubsystemGooglePlay*
 
 {
 	UE_LOG(LogOnline, Display, TEXT("FOnlineIdentityAndroid::FOnlineIdentityAndroid()"));
+	check(MainSubsystem != nullptr);
 	PendingConnectRequest.ConnectionInterface = this;
 }
 
@@ -65,8 +66,8 @@ bool FOnlineIdentityGooglePlay::Login(int32 LocalUserNum, const FOnlineAccountCr
 
 bool FOnlineIdentityGooglePlay::Logout(int32 LocalUserNum)
 {
-	TriggerOnLogoutCompleteDelegates(LocalUserNum, false);
-	return false;
+	MainSubsystem->StartLogoutTask(LocalUserNum);
+	return true;
 }
 
 
@@ -78,37 +79,23 @@ bool FOnlineIdentityGooglePlay::AutoLogin(int32 LocalUserNum)
 
 ELoginStatus::Type FOnlineIdentityGooglePlay::GetLoginStatus(int32 LocalUserNum) const
 {
-	ELoginStatus::Type LoginStatus = ELoginStatus::NotLoggedIn;
-
-	if(LocalUserNum < MAX_LOCAL_PLAYERS && bLoggedIn)
-	{
-		UE_LOG(LogOnline, Display, TEXT("FOnlineIdentityAndroid::GetLoginStatus - Logged in"));
-		LoginStatus = ELoginStatus::LoggedIn;
-	}
-	else
+	if (LocalUserNum < MAX_LOCAL_PLAYERS && MainSubsystem->GetGameServices() != nullptr && MainSubsystem->GetGameServices()->IsAuthorized())
 	{
-		UE_LOG(LogOnline, Display, TEXT("FOnlineIdentityAndroid::GetLoginStatus - Not logged in"));
+		return ELoginStatus::LoggedIn;
 	}
 
-	return LoginStatus;
+	return ELoginStatus::NotLoggedIn;
 }
 
 
 ELoginStatus::Type FOnlineIdentityGooglePlay::GetLoginStatus(const FUniqueNetId& UserId) const
 {
-	ELoginStatus::Type LoginStatus = ELoginStatus::NotLoggedIn;
-
-	if(bLoggedIn)
-	{
-		UE_LOG(LogOnline, Display, TEXT("FOnlineIdentityAndroid::GetLoginStatus - Logged in"));
-		LoginStatus = ELoginStatus::LoggedIn;
-	}
-	else
+	if (MainSubsystem->GetGameServices() != nullptr && MainSubsystem->GetGameServices()->IsAuthorized())
 	{
-		UE_LOG(LogOnline, Display, TEXT("FOnlineIdentityAndroid::GetLoginStatus - Not logged in"));
+		return ELoginStatus::LoggedIn;
 	}
 
-	return LoginStatus;
+	return ELoginStatus::NotLoggedIn;
 }
 
 
diff --git a/Engine/Source/Runtime/Online/Android/OnlineSubsystemGooglePlay/Private/OnlineSubsystemGooglePlay.cpp b/Engine/Source/Runtime/Online/Android/OnlineSubsystemGooglePlay/Private/OnlineSubsystemGooglePlay.cpp
--- a/Engine/Source/Runtime/Online/Android/OnlineSubsystemGooglePlay/Private/OnlineSubsystemGooglePlay.cpp
+++ b/Engine/Source/Runtime/Online/Android/OnlineSubsystemGooglePlay/Private/OnlineSubsystemGooglePlay.cpp
@@ -7,11 +7,11 @@
 #include "Android/AndroidJNI.h"
 #include "OnlineAsyncTaskManagerGooglePlay.h"
 #include "OnlineAsyncTaskGooglePlayLogin.h"
+#include "OnlineAsyncTaskGooglePlayLogout.h"
 
 #include <android_native_app_glue.h>
 
 #include "gpg/android_initialization.h"
-#include "gpg/android_platform_configuration.h"
 #include "gpg/builder.h"
 #include "gpg/debug.h"
 #include "gpg/android_support.h"
@@ -24,6 +24,8 @@ FOnlineSubsystemGooglePlay::FOnlineSubsystemGooglePlay()
 	, AchievementsInterface(nullptr)
 	, StoreInterface(nullptr)
 	, CurrentLoginTask(nullptr)
+	, CurrentShowLoginUITask(nullptr)
+	, CurrentLogoutTask(nullptr)
 {
 
 }
@@ -128,54 +130,10 @@ bool FOnlineSubsystemGooglePlay::Init()
 	AndroidInitialization::android_main(GNativeAndroidApp);
 
 	extern jobject GJavaGlobalNativeActivity;
-	AndroidPlatformConfiguration PlatformConfiguration;
 	PlatformConfiguration.SetActivity(GNativeAndroidApp->activity->clazz);
 
-	// Queue up a task for the login so that other tasks execute after it.
-	check(CurrentLoginTask == nullptr);
-	CurrentLoginTask = new FOnlineAsyncTaskGooglePlayLogin(this, 0);
-	QueueAsyncTask(CurrentLoginTask);
-
-	WaitingForLogin = true;
-	WaitForLostFocus = false;
-
 	OnActivityResultDelegateHandle = FJavaWrapper::OnActivityResultDelegate.AddRaw(this, &FOnlineSubsystemGooglePlay::OnActivityResult);
 
-	// Create() returns a std::unqiue_ptr, but we convert it to a TUniquePtr.
-	GameServicesPtr.Reset( GameServices::Builder()
-		.SetDefaultOnLog(LogLevel::VERBOSE)
-		.SetOnAuthActionStarted([](AuthOperation op) {
-			UE_LOG(LogOnline, Log, TEXT("GPG sign in started"));
-		})
-		.SetOnAuthActionFinished([this](AuthOperation Op, AuthStatus Status) {
-			OnAuthActionFinished(Op, Status);
-
-			if (Op == AuthOperation::SIGN_IN)
-			{
-				if (Status != AuthStatus::VALID)
-				{
-					WaitForLostFocus = true;
-				}
-				WaitingForLogin = false;
-			}
-		})
-		.Create(PlatformConfiguration).release() );
-
-	// Wait for GooglePlay to complete login task
-	while (WaitingForLogin)
-	{
-		FPlatformProcess::Sleep(0.01f);
-	}
-	if (WaitForLostFocus)
-	{
-		extern bool WaitForAndroidLoseFocusEvent(double TimeoutSeconds);
-		if (!WaitForAndroidLoseFocusEvent(3.0))
-		{
-			FPlatformMisc::LowLevelOutputDebugString(TEXT("FOnlineSubsystemAndroid::Init - timed out"));
-			OnAuthActionFinished(AuthOperation::SIGN_IN, AuthStatus::ERROR_NOT_AUTHORIZED);
-		}
-	}
-
 	return true;
 }
 
@@ -251,6 +209,58 @@ bool FOnlineSubsystemGooglePlay::IsInAppPurchasingEnabled()
 	return bEnabledIAP;
 }
 
+void FOnlineSubsystemGooglePlay::StartShowLoginUITask(int PlayerId, const IOnlineExternalUI::FOnLoginUIClosedDelegate& Delegate)
+{
+	if (AreAnyAsyncLoginTasksRunning())
+	{
+		UE_LOG(LogOnline, Log, TEXT("FOnlineSubsystemGooglePlay::StartShowLoginUITask: An asynchronous login task is already running."));
+		Delegate.ExecuteIfBound(nullptr, PlayerId);
+		return;
+	}
+
+	if (GameServicesPtr.get() == nullptr)
+	{
+		// This is likely the first login attempt during this run. Attempt to create the
+		// GameServices object, which will automatically start a "silent" login attempt.
+		// If that succeeds, there's no need to show the login UI explicitly. If it fails,
+		// we'll call ShowAuthorizationUI.
+		
+		auto TheDelegate = FOnlineAsyncTaskGooglePlayLogin::FOnCompletedDelegate::CreateLambda([this, PlayerId, Delegate]()
+		{
+			 StartShowLoginUITask_Internal(PlayerId, Delegate);
+		});
+
+		CurrentLoginTask = new FOnlineAsyncTaskGooglePlayLogin(this, PlayerId, TheDelegate);
+		QueueAsyncTask(CurrentLoginTask);
+	}
+	else
+	{
+		// We already have a GameServices object, so we can directly go to ShowAuthorizationUI.
+		StartShowLoginUITask_Internal(PlayerId, Delegate);
+	}
+}
+
+void FOnlineSubsystemGooglePlay::StartLogoutTask(int32 LocalUserNum)
+{
+	if (CurrentLogoutTask != nullptr)
+	{
+		UE_LOG(LogOnline, Log, TEXT("FOnlineSubsystemGooglePlay::StartLogoutTask: A logout task is already in progress."));
+		IdentityInterface->TriggerOnLogoutCompleteDelegates(LocalUserNum, false);
+		return;
+	}
+
+	CurrentLogoutTask = new FOnlineAsyncTaskGooglePlayLogout(this, LocalUserNum);
+	QueueAsyncTask(CurrentLogoutTask);
+}
+
+void FOnlineSubsystemGooglePlay::StartShowLoginUITask_Internal(int PlayerId, const IOnlineExternalUI::FOnLoginUIClosedDelegate& Delegate)
+{
+	check(!AreAnyAsyncLoginTasksRunning());
+
+	CurrentShowLoginUITask = new FOnlineAsyncTaskGooglePlayShowLoginUI(this, PlayerId, Delegate);
+	QueueAsyncTask(CurrentShowLoginUITask);
+}
+
 void FOnlineSubsystemGooglePlay::QueueAsyncTask(FOnlineAsyncTask* AsyncTask)
 {
 	check(OnlineAsyncTaskThreadRunnable);
@@ -259,24 +269,28 @@ void FOnlineSubsystemGooglePlay::QueueAsyncTask(FOnlineAsyncTask* AsyncTask)
 
 void FOnlineSubsystemGooglePlay::OnAuthActionFinished(AuthOperation Op, AuthStatus Status)
 {
-	FString OpString(DebugString(Op).c_str());
-	FString StatusString(DebugString(Status).c_str());
-	UE_LOG(LogOnline, Log, TEXT("FOnlineSubsystemGooglePlay::OnAuthActionFinished, Op: %s, Status: %s"), *OpString, *StatusString);
-
 	if (Op == AuthOperation::SIGN_IN)
 	{
-		if (CurrentLoginTask == nullptr)
+		if (CurrentLoginTask != nullptr)
 		{
-			UE_LOG(LogOnline, Warning, TEXT("FOnlineSubsystemGooglePlay::OnAuthActionFinished: sign-in operation with a null CurrentLoginTask"));
-			return;
-		}
+			// Only one login task should be active at a time
+			check(CurrentShowLoginUITask == nullptr);
 
-		CurrentLoginTask->OnAuthActionFinished(Op, Status);
+			CurrentLoginTask->OnAuthActionFinished(Op, Status);
+		}
+		else if(CurrentShowLoginUITask != nullptr)
+		{
+			// Only one login task should be active at a time
+			check(CurrentLoginTask == nullptr);
 
-		if (CurrentLoginTask->bIsComplete)
+			CurrentShowLoginUITask->OnAuthActionFinished(Op, Status);
+		}
+	}
+	else if (Op == AuthOperation::SIGN_OUT)
+	{
+		if (CurrentLogoutTask != nullptr)
 		{
-			// Async task manager owns the task now and is responsible for cleaning it up.
-			CurrentLoginTask = nullptr;
+			CurrentLogoutTask->OnAuthActionFinished(Op, Status);
 		}
 	}
 }
diff --git a/Engine/Source/Runtime/Online/Android/OnlineSubsystemGooglePlay/Public/OnlineSubsystemGooglePlay.h b/Engine/Source/Runtime/Online/Android/OnlineSubsystemGooglePlay/Public/OnlineSubsystemGooglePlay.h
--- a/Engine/Source/Runtime/Online/Android/OnlineSubsystemGooglePlay/Public/OnlineSubsystemGooglePlay.h
+++ b/Engine/Source/Runtime/Online/Android/OnlineSubsystemGooglePlay/Public/OnlineSubsystemGooglePlay.h
@@ -12,8 +12,11 @@
 #include "UniquePtr.h"
 
 #include "gpg/game_services.h"
+#include "gpg/android_platform_configuration.h"
 
 class FOnlineAsyncTaskGooglePlayLogin;
+class FOnlineAsyncTaskGooglePlayShowLoginUI;
+class FOnlineAsyncTaskGooglePlayLogout;
 class FRunnableThread;
 
 /**
@@ -83,7 +86,7 @@ PACKAGE_SCOPE:
 	void QueueAsyncTask(class FOnlineAsyncTask* AsyncTask);
 
 	/** Returns a pointer to the Google API entry point */
-	gpg::GameServices* GetGameServices() const { return GameServicesPtr.Get(); }
+	gpg::GameServices* GetGameServices() const { return GameServicesPtr.get(); }
 
 	/** Utility function, useful for Google APIs that take a std::string but we only have an FString */
 	static std::string ConvertFStringToStdString(const FString& InString);
@@ -100,6 +103,15 @@ PACKAGE_SCOPE:
 	*/
 	bool IsInAppPurchasingEnabled();
 
+	/** Returns true if there are any async login tasks currently being tracked. */
+	bool AreAnyAsyncLoginTasksRunning() const { return CurrentLoginTask != nullptr || CurrentShowLoginUITask != nullptr; }
+
+	/** Start a ShowLoginUI async task. Creates the GameServices object first if necessary. */
+	void StartShowLoginUITask(int PlayerId, const IOnlineExternalUI::FOnLoginUIClosedDelegate& Delegate);
+
+	/** Start a logout task if one isn't already in progress. */
+	void StartLogoutTask(int32 LocalUserNum);
+
 private:
 
 	/** Google callback when auth is complete */
@@ -108,6 +120,9 @@ private:
 	/** Android callback when an activity is finished */
 	void OnActivityResult(JNIEnv *env, jobject thiz, jobject activity, jint requestCode, jint resultCode, jobject data);
 
+	/** Start a ShowLoginUI async task. */
+	void StartShowLoginUITask_Internal(int PlayerId, const IOnlineExternalUI::FOnLoginUIClosedDelegate& Delegate);
+
 	/** Online async task runnable */
 	TUniquePtr<class FOnlineAsyncTaskManagerGooglePlay> OnlineAsyncTaskThreadRunnable;
 
@@ -129,13 +144,43 @@ private:
 	FOnlineExternalUIGooglePlayPtr ExternalUIInterface;
 
 	/** Pointer to the main entry point for the Google API */
-	TUniquePtr<gpg::GameServices> GameServicesPtr;
+	std::unique_ptr<gpg::GameServices> GameServicesPtr;
 
-	/** Track the current login task (if any) so callbacks can notify it */
+	// Since the OnAuthActionFinished callback is "global" (within the scope of one GameServices object),
+	// we track the async tasks that depend on this callback to notify us that they're finished.
+	// Only one of these pointers should be non-null at any given time.
+
+	/**
+	 * Track the current login task (if any) so callbacks can notify it.
+	 * Still owned by the AsyncTaskManager, do not delete via this pointer!
+	 **/
 	FOnlineAsyncTaskGooglePlayLogin* CurrentLoginTask;
 
+	/**
+	 * Track the current ShowLoginUI task (if any)
+	 * Still owned by the AsyncTaskManager, do not delete via this pointer!
+	 **/
+	FOnlineAsyncTaskGooglePlayShowLoginUI* CurrentShowLoginUITask;
+
+	/**
+	 * Track the current Logout task
+	 * Still owned by the AsyncTaskManager, do not delete via this pointer!
+	 */
+	FOnlineAsyncTaskGooglePlayLogout* CurrentLogoutTask;
+
 	/** Handle of registered delegate for onActivityResult */
 	FDelegateHandle OnActivityResultDelegateHandle;
+
+	/** This task needs to be able to set the GameServicesPtr and clear CurrentLoginTask*/
+	friend class FOnlineAsyncTaskGooglePlayLogin;
+
+	/** This task needs to be able to clear CurrentShowLoginUITask */
+	friend class FOnlineAsyncTaskGooglePlayShowLoginUI;
+
+	/** This task needs to be able to clear CurrentLogoutTask */
+	friend class FOnlineAsyncTaskGooglePlayLogout;
+
+	gpg::AndroidPlatformConfiguration PlatformConfiguration;
 };
 
 typedef TSharedPtr<FOnlineSubsystemGooglePlay, ESPMode::ThreadSafe> FOnlineSubsystemGooglePlayPtr;
diff --git a/Engine/Source/Runtime/Online/IOS/OnlineSubsystemIOS/Private/OnlineExternalUIInterfaceIOS.cpp b/Engine/Source/Runtime/Online/IOS/OnlineSubsystemIOS/Private/OnlineExternalUIInterfaceIOS.cpp
--- a/Engine/Source/Runtime/Online/IOS/OnlineSubsystemIOS/Private/OnlineExternalUIInterfaceIOS.cpp
+++ b/Engine/Source/Runtime/Online/IOS/OnlineSubsystemIOS/Private/OnlineExternalUIInterfaceIOS.cpp
@@ -2,14 +2,89 @@
 
 // Module includes
 #include "OnlineSubsystemIOSPrivatePCH.h"
+#include "OnlineIdentityInterfaceIOS.h"
 
-FOnlineExternalUIIOS::FOnlineExternalUIIOS()
+FOnlineExternalUIIOS::FOnlineExternalUIIOS(FOnlineSubsystemIOS* InSubsystem)
+	: Subsystem(InSubsystem)
 {
+	check(InSubsystem != nullptr);
 }
 
 bool FOnlineExternalUIIOS::ShowLoginUI(const int ControllerIndex, bool bShowOnlineOnly, const FOnLoginUIClosedDelegate& Delegate)
 {
-	return false;
+	FOnlineIdentityIOS* IdentityInterface = static_cast<FOnlineIdentityIOS*>(Subsystem->GetIdentityInterface().Get());
+	
+	check(IdentityInterface != nullptr);
+	
+	if (IdentityInterface->GetLocalGameCenterUser() == nullptr)
+	{
+		UE_LOG(LogOnline, Log, TEXT("Game Center localPlayer is null."));
+		Delegate.ExecuteIfBound(nullptr, ControllerIndex);
+		return true;
+	}
+	
+	if (IdentityInterface->GetLocalGameCenterUser().isAuthenticated)
+	{
+		Delegate.ExecuteIfBound(IdentityInterface->GetLocalPlayerUniqueId(), ControllerIndex);
+		return true;
+	}
+	
+	// Not authenticated, set a handler
+	
+	// Copy the delegate so that the block can still access it when it runs.
+	FOnLoginUIClosedDelegate CopiedDelegate = Delegate;
+	
+	// Trigger the login event on the main thread.
+	dispatch_async(dispatch_get_main_queue(), ^
+	{
+		[IdentityInterface->GetLocalGameCenterUser() setAuthenticateHandler:(^(UIViewController* viewcontroller, NSError *error)
+		{
+			UE_LOG(LogOnline, Log, TEXT("Game Center authenticateHandler. viewController = %p"), viewcontroller);
+			
+			// The login process has completed.
+			if (viewcontroller == nil)
+			{
+				FString ErrorMessage;
+				TSharedPtr<FUniqueNetIdString> UniqueNetId;
+
+				if (IdentityInterface->GetLocalGameCenterUser().isAuthenticated == YES)
+				{
+					// Trigger the completion delegate on the game thread.
+					const FString PlayerId(IdentityInterface->GetLocalGameCenterUser().playerID);
+					UniqueNetId = MakeShareable(new FUniqueNetIdString(PlayerId));
+					UE_LOG(LogOnline, Log, TEXT("The user %s has logged into Game Center"), *PlayerId);
+				}
+				else
+				{
+					ErrorMessage = TEXT("The user could not be authenticated by Game Center");
+					UE_LOG(LogOnline, Log, TEXT("%s"), *ErrorMessage);
+				}
+
+				if (error)
+				{
+					NSString *errstr = [error localizedDescription];
+					UE_LOG(LogOnline, Warning, TEXT("Game Center login has failed. %s]"), *FString(errstr));
+				}
+
+				// Report back to the game thread whether this succeeded.
+				[FIOSAsyncTask CreateTaskWithBlock : ^ bool(void)
+				{
+					IdentityInterface->SetLocalPlayerUniqueId(UniqueNetId);
+					CopiedDelegate.ExecuteIfBound(UniqueNetId, ControllerIndex);
+
+					return true;
+				}];
+			}
+			else
+			{
+				// Game Center has provided a view controller for us to login, we present it.
+				[[IOSAppDelegate GetDelegate].IOSController 
+					presentViewController:viewcontroller animated:YES completion:nil];
+			}
+		})];
+	});
+	
+	return true;
 }
 
 bool FOnlineExternalUIIOS::ShowFriendsUI(int32 LocalUserNum)
diff --git a/Engine/Source/Runtime/Online/IOS/OnlineSubsystemIOS/Private/OnlineExternalUIInterfaceIOS.h b/Engine/Source/Runtime/Online/IOS/OnlineSubsystemIOS/Private/OnlineExternalUIInterfaceIOS.h
--- a/Engine/Source/Runtime/Online/IOS/OnlineSubsystemIOS/Private/OnlineExternalUIInterfaceIOS.h
+++ b/Engine/Source/Runtime/Online/IOS/OnlineSubsystemIOS/Private/OnlineExternalUIInterfaceIOS.h
@@ -5,11 +5,13 @@
 #include "OnlineExternalUIInterface.h"
 #include "OnlineSubsystemIOSTypes.h"
 
+class FOnlineSubsystemIOS;
+
 class FOnlineExternalUIIOS : public IOnlineExternalUI
 {
 PACKAGE_SCOPE:
 
-	FOnlineExternalUIIOS();
+	FOnlineExternalUIIOS(FOnlineSubsystemIOS* InSubsystem);
 
 public:
 
@@ -23,6 +25,9 @@ public:
 	virtual bool ShowProfileUI(const FUniqueNetId& Requestor, const FUniqueNetId& Requestee, const FOnProfileUIClosedDelegate& Delegate) override;
 	virtual bool ShowAccountUpgradeUI(const FUniqueNetId& UniqueId) override;
 	// End IOnlineExternalUI interface
+	
+private:
+	FOnlineSubsystemIOS* Subsystem;
 };
 
 typedef TSharedPtr<FOnlineExternalUIIOS, ESPMode::ThreadSafe> FOnlineExternalUIIOSPtr;
\ No newline at end of file
diff --git a/Engine/Source/Runtime/Online/IOS/OnlineSubsystemIOS/Private/OnlineIdentityInterfaceIOS.cpp b/Engine/Source/Runtime/Online/IOS/OnlineSubsystemIOS/Private/OnlineIdentityInterfaceIOS.cpp
--- a/Engine/Source/Runtime/Online/IOS/OnlineSubsystemIOS/Private/OnlineIdentityInterfaceIOS.cpp
+++ b/Engine/Source/Runtime/Online/IOS/OnlineSubsystemIOS/Private/OnlineIdentityInterfaceIOS.cpp
@@ -5,7 +5,16 @@
 FOnlineIdentityIOS::FOnlineIdentityIOS() :
 	UniqueNetId( NULL )
 {
-	Login( 0, FOnlineAccountCredentials() );
+}
+
+TSharedPtr<FUniqueNetIdString> FOnlineIdentityIOS::GetLocalPlayerUniqueId() const
+{
+	return UniqueNetId;
+}
+
+void FOnlineIdentityIOS::SetLocalPlayerUniqueId(const TSharedPtr<FUniqueNetIdString>& UniqueId)
+{
+	UniqueNetId = UniqueId;
 }
 
 TSharedPtr<FUserOnlineAccount> FOnlineIdentityIOS::GetUserAccount(const FUniqueNetId& UserId) const
@@ -26,6 +35,10 @@ bool FOnlineIdentityIOS::Login(int32 LocalUserNum, const FOnlineAccountCredentia
 {
 	bool bStartedLogin = false;
 
+	// Since the iOS login code may show a UI, ShowLoginUI is a better fit here. Also, note that the ConnectToService blueprint
+	// node that calls Login is deprecated (there's a new ShowExternalLoginUI node meant to replace it).
+	UE_LOG(LogOnline, Warning, TEXT("Using the IOnlineIdentity::Login function on iOS is not recommended. Please use IOnlineExternalUI::ShowLoginUI instead."));
+
 	// Was the login handled by Game Center
 	if( GetLocalGameCenterUser() && 
 		GetLocalGameCenterUser().isAuthenticated )
@@ -104,8 +117,8 @@ bool FOnlineIdentityIOS::Login(int32 LocalUserNum, const FOnlineAccountCredentia
 
 bool FOnlineIdentityIOS::Logout(int32 LocalUserNum)
 {
-	TriggerOnLogoutCompleteDelegates(LocalUserNum, true);
-	return false;
+	TriggerOnLogoutCompleteDelegates(LocalUserNum, false);
+	return true;
 }
 
 bool FOnlineIdentityIOS::AutoLogin(int32 LocalUserNum)
diff --git a/Engine/Source/Runtime/Online/IOS/OnlineSubsystemIOS/Private/OnlineIdentityInterfaceIOS.h b/Engine/Source/Runtime/Online/IOS/OnlineSubsystemIOS/Private/OnlineIdentityInterfaceIOS.h
--- a/Engine/Source/Runtime/Online/IOS/OnlineSubsystemIOS/Private/OnlineIdentityInterfaceIOS.h
+++ b/Engine/Source/Runtime/Online/IOS/OnlineSubsystemIOS/Private/OnlineIdentityInterfaceIOS.h
@@ -24,6 +24,9 @@ PACKAGE_SCOPE:
 	 */
 	FOnlineIdentityIOS();	
 
+	TSharedPtr<FUniqueNetIdString> GetLocalPlayerUniqueId() const;
+
+	void SetLocalPlayerUniqueId(const TSharedPtr<FUniqueNetIdString>& UniqueId);
 
 public:
 
diff --git a/Engine/Source/Runtime/Online/IOS/OnlineSubsystemIOS/Private/OnlineSubsystemIOS.cpp b/Engine/Source/Runtime/Online/IOS/OnlineSubsystemIOS/Private/OnlineSubsystemIOS.cpp
--- a/Engine/Source/Runtime/Online/IOS/OnlineSubsystemIOS/Private/OnlineSubsystemIOS.cpp
+++ b/Engine/Source/Runtime/Online/IOS/OnlineSubsystemIOS/Private/OnlineSubsystemIOS.cpp
@@ -141,7 +141,7 @@ bool FOnlineSubsystemIOS::Init()
 		FriendsInterface = MakeShareable(new FOnlineFriendsIOS(this));
 		LeaderboardsInterface = MakeShareable(new FOnlineLeaderboardsIOS(this));
 		AchievementsInterface = MakeShareable(new FOnlineAchievementsIOS(this));
-		ExternalUIInterface = MakeShareable(new FOnlineExternalUIIOS());
+		ExternalUIInterface = MakeShareable(new FOnlineExternalUIIOS(this));
         TurnBasedInterface = MakeShareable(new FOnlineTurnBasedIOS());
         UserCloudInterface = MakeShareable(new FOnlineUserCloudInterfaceIOS());
         SharedCloudInterface = MakeShareable(new FOnlineSharedCloudInterfaceIOS());
diff --git a/Engine/Source/Runtime/Online/OnlineSubsystemUtils/Classes/ConnectionCallbackProxy.h b/Engine/Source/Runtime/Online/OnlineSubsystemUtils/Classes/ConnectionCallbackProxy.h
--- a/Engine/Source/Runtime/Online/OnlineSubsystemUtils/Classes/ConnectionCallbackProxy.h
+++ b/Engine/Source/Runtime/Online/OnlineSubsystemUtils/Classes/ConnectionCallbackProxy.h
@@ -22,7 +22,7 @@ class UConnectionCallbackProxy : public UOnlineBlueprintCallProxyBase
 	FOnlineConnectionResult OnFailure;
 
 	// Connects to an online service such as Google Play
-	UFUNCTION(BlueprintCallable, meta = (BlueprintInternalUseOnly = "true", WorldContext="WorldContextObject"), Category = "Online|Achievements")
+	UFUNCTION(BlueprintCallable, meta = (DeprecatedFunction, DeprecationMessage="Please use Show External Login UI instead", BlueprintInternalUseOnly = "true", WorldContext="WorldContextObject"), Category = "Online|Achievements")
 		static UConnectionCallbackProxy* ConnectToService(UObject* WorldContextObject, class APlayerController* PlayerController);
 
 	// UOnlineBlueprintCallProxyBase interface
diff --git a/Engine/Source/Runtime/Online/OnlineSubsystemUtils/Classes/LogoutCallbackProxy.h b/Engine/Source/Runtime/Online/OnlineSubsystemUtils/Classes/LogoutCallbackProxy.h
new file mode 100644
--- /dev/null
+++ b/Engine/Source/Runtime/Online/OnlineSubsystemUtils/Classes/LogoutCallbackProxy.h
@@ -0,0 +1,47 @@
+// Copyright 1998-2015 Epic Games, Inc. All Rights Reserved.
+
+#pragma once
+
+#include "Kismet/BlueprintAsyncActionBase.h"
+#include "LogoutCallbackProxy.generated.h"
+
+class APlayerController;
+
+DECLARE_DYNAMIC_MULTICAST_DELEGATE_OneParam(FOnlineLogoutResult, APlayerController*, PlayerController);
+
+UCLASS(MinimalAPI)
+class ULogoutCallbackProxy : public UBlueprintAsyncActionBase
+{
+	GENERATED_BODY()
+
+public:
+	ULogoutCallbackProxy(const FObjectInitializer& ObjectInitializer);
+
+	// Called when the logout completed successfully
+	UPROPERTY(BlueprintAssignable)
+	FOnlineLogoutResult OnSuccess;
+
+	// Called when the logout completed unsuccessfully
+	UPROPERTY(BlueprintAssignable)
+	FOnlineLogoutResult OnFailure;
+
+	// Logs out of the online service
+	UFUNCTION(BlueprintCallable, meta = (BlueprintInternalUseOnly = "true", WorldContext="WorldContextObject"), Category = "Online")
+	static ULogoutCallbackProxy* Logout(UObject* WorldContextObject, class APlayerController* PlayerController);
+
+	// UBlueprintAsyncActionBase interface
+	virtual void Activate() override;
+	// End of UBlueprintAsyncActionBase interface
+
+private:
+	// Internal callback when the login UI closes, calls out to the public success/failure callbacks
+	void OnLogoutCompleted(int LocalUserNum, bool bWasSuccessful);
+
+	// The player controller triggering things
+	TWeakObjectPtr<APlayerController> PlayerControllerWeakPtr;
+
+	FDelegateHandle OnLogoutCompleteDelegateHandle;
+
+	// The world context object in which this call is taking place
+	UObject* WorldContextObject;
+};
\ No newline at end of file
diff --git a/Engine/Source/Runtime/Online/OnlineSubsystemUtils/Classes/ShowLoginUICallbackProxy.h b/Engine/Source/Runtime/Online/OnlineSubsystemUtils/Classes/ShowLoginUICallbackProxy.h
new file mode 100644
--- /dev/null
+++ b/Engine/Source/Runtime/Online/OnlineSubsystemUtils/Classes/ShowLoginUICallbackProxy.h
@@ -0,0 +1,45 @@
+// Copyright 1998-2015 Epic Games, Inc. All Rights Reserved.
+
+#pragma once
+
+#include "Kismet/BlueprintAsyncActionBase.h"
+#include "ShowLoginUICallbackProxy.generated.h"
+
+class APlayerController;
+
+DECLARE_DYNAMIC_MULTICAST_DELEGATE_OneParam(FOnlineShowLoginUIResult, APlayerController*, PlayerController);
+
+UCLASS(MinimalAPI)
+class UShowLoginUICallbackProxy : public UBlueprintAsyncActionBase
+{
+	GENERATED_BODY()
+
+public:
+	UShowLoginUICallbackProxy(const FObjectInitializer& ObjectInitializer);
+
+	// Called when there is a successful query
+	UPROPERTY(BlueprintAssignable)
+	FOnlineShowLoginUIResult OnSuccess;
+
+	// Called when there is an unsuccessful query
+	UPROPERTY(BlueprintAssignable)
+	FOnlineShowLoginUIResult OnFailure;
+
+	// Shows the login UI for the currently active online subsystem, if the subsystem supports a login UI.
+	UFUNCTION(BlueprintCallable, meta = (BlueprintInternalUseOnly = "true", WorldContext="WorldContextObject"), Category = "Online")
+	static UShowLoginUICallbackProxy* ShowExternalLoginUI(UObject* WorldContextObject, class APlayerController* InPlayerController);
+
+	// UBlueprintAsyncActionBase interface
+	virtual void Activate() override;
+	// End of UBlueprintAsyncActionBase interface
+
+private:
+	// Internal callback when the login UI closes, calls out to the public success/failure callbacks
+	void OnShowLoginUICompleted(TSharedPtr<const FUniqueNetId> UniqueId, int LocalUserNum);
+
+	// The player controller triggering things
+	TWeakObjectPtr<APlayerController> PlayerControllerWeakPtr;
+
+	// The world context object in which this call is taking place
+	UObject* WorldContextObject;
+};
\ No newline at end of file
diff --git a/Engine/Source/Runtime/Online/OnlineSubsystemUtils/Private/LogoutCallbackProxy.cpp b/Engine/Source/Runtime/Online/OnlineSubsystemUtils/Private/LogoutCallbackProxy.cpp
new file mode 100644
--- /dev/null
+++ b/Engine/Source/Runtime/Online/OnlineSubsystemUtils/Private/LogoutCallbackProxy.cpp
@@ -0,0 +1,91 @@
+// Copyright 1998-2015 Epic Games, Inc. All Rights Reserved.
+
+#include "OnlineSubsystemUtilsPrivatePCH.h"
+#include "OnlineExternalUIInterface.h"
+
+ULogoutCallbackProxy::ULogoutCallbackProxy(const FObjectInitializer& ObjectInitializer)
+	: Super(ObjectInitializer)
+	, WorldContextObject(nullptr)
+{
+}
+
+ULogoutCallbackProxy* ULogoutCallbackProxy::Logout(UObject* WorldContextObject, class APlayerController* PlayerController)
+{
+	ULogoutCallbackProxy* Proxy = NewObject<ULogoutCallbackProxy>();
+	Proxy->PlayerControllerWeakPtr = PlayerController;
+	Proxy->WorldContextObject = WorldContextObject;
+	return Proxy;
+}
+
+void ULogoutCallbackProxy::Activate()
+{
+	if (!PlayerControllerWeakPtr.IsValid())
+	{
+		FFrame::KismetExecutionMessage(TEXT("A player controller must be provided in order to show the external login UI."), ELogVerbosity::Warning);
+		OnFailure.Broadcast(PlayerControllerWeakPtr.Get());
+		return;
+	}
+
+	FOnlineSubsystemBPCallHelper Helper(TEXT("Logout"), GEngine->GetWorldFromContextObject(WorldContextObject));
+	Helper.QueryIDFromPlayerController(PlayerControllerWeakPtr.Get());
+
+	if (Helper.OnlineSub == nullptr)
+	{
+		OnFailure.Broadcast(PlayerControllerWeakPtr.Get());
+		return;
+	}
+		
+	IOnlineIdentityPtr OnlineIdentity = Helper.OnlineSub->GetIdentityInterface();
+	if (!OnlineIdentity.IsValid())
+	{
+		FFrame::KismetExecutionMessage(TEXT("Logout: identity functionality not supported by the current online subsystem"), ELogVerbosity::Warning);
+		OnFailure.Broadcast(PlayerControllerWeakPtr.Get());
+		return;
+	}
+
+	const ULocalPlayer* LocalPlayer = Cast<ULocalPlayer>(PlayerControllerWeakPtr->Player);
+
+	if (LocalPlayer == nullptr)
+	{
+		FFrame::KismetExecutionMessage(TEXT("Only local players can log out"), ELogVerbosity::Warning);
+		OnFailure.Broadcast(PlayerControllerWeakPtr.Get());
+		return;
+	}
+	
+	const int32 ControllerId = LocalPlayer->GetControllerId();
+
+	if (!OnlineIdentity->OnLogoutCompleteDelegates[ControllerId].IsBoundToObject(this))
+	{
+		OnLogoutCompleteDelegateHandle = OnlineIdentity->AddOnLogoutCompleteDelegate_Handle(ControllerId,
+			FOnLogoutCompleteDelegate::CreateUObject(this, &ULogoutCallbackProxy::OnLogoutCompleted));
+		OnlineIdentity->Logout(ControllerId);
+	}
+	else
+	{
+		// already trying to logout!
+	}
+}
+
+void ULogoutCallbackProxy::OnLogoutCompleted(int LocalPlayerNum, bool bWasSuccessful)
+{
+	FOnlineSubsystemBPCallHelper Helper(TEXT("Logout"), GEngine->GetWorldFromContextObject(WorldContextObject));
+	Helper.QueryIDFromPlayerController(PlayerControllerWeakPtr.Get());
+
+	if (Helper.IsValid())
+	{
+		IOnlineIdentityPtr OnlineIdentity = Helper.OnlineSub->GetIdentityInterface();
+		if (OnlineIdentity.IsValid())
+		{
+			OnlineIdentity->ClearOnLogoutCompleteDelegate_Handle(LocalPlayerNum, OnLogoutCompleteDelegateHandle);
+		}
+	}
+
+	if (bWasSuccessful)
+	{
+		OnSuccess.Broadcast(PlayerControllerWeakPtr.Get());
+	}
+	else
+	{
+		OnFailure.Broadcast(PlayerControllerWeakPtr.Get());
+	}
+}
diff --git a/Engine/Source/Runtime/Online/OnlineSubsystemUtils/Private/ShowLoginUICallbackProxy.cpp b/Engine/Source/Runtime/Online/OnlineSubsystemUtils/Private/ShowLoginUICallbackProxy.cpp
new file mode 100644
--- /dev/null
+++ b/Engine/Source/Runtime/Online/OnlineSubsystemUtils/Private/ShowLoginUICallbackProxy.cpp
@@ -0,0 +1,91 @@
+// Copyright 1998-2015 Epic Games, Inc. All Rights Reserved.
+
+#include "OnlineSubsystemUtilsPrivatePCH.h"
+#include "OnlineExternalUIInterface.h"
+
+UShowLoginUICallbackProxy::UShowLoginUICallbackProxy(const FObjectInitializer& ObjectInitializer)
+	: Super(ObjectInitializer)
+	, WorldContextObject(nullptr)
+{
+}
+
+UShowLoginUICallbackProxy* UShowLoginUICallbackProxy::ShowExternalLoginUI(UObject* WorldContextObject, class APlayerController* InPlayerController)
+{
+	UShowLoginUICallbackProxy* Proxy = NewObject<UShowLoginUICallbackProxy>();
+	Proxy->PlayerControllerWeakPtr = InPlayerController;
+	Proxy->WorldContextObject = WorldContextObject;
+	return Proxy;
+}
+
+void UShowLoginUICallbackProxy::Activate()
+{
+	if (!PlayerControllerWeakPtr.IsValid())
+	{
+		FFrame::KismetExecutionMessage(TEXT("A player controller must be provided in order to show the external login UI."), ELogVerbosity::Warning);
+		OnFailure.Broadcast(PlayerControllerWeakPtr.Get());
+		return;
+	}
+
+	FOnlineSubsystemBPCallHelper Helper(TEXT("ShowLoginUI"), GEngine->GetWorldFromContextObject(WorldContextObject));
+
+	if (Helper.OnlineSub == nullptr)
+	{
+		OnFailure.Broadcast(PlayerControllerWeakPtr.Get());
+		return;
+	}
+		
+	IOnlineExternalUIPtr OnlineExternalUI = Helper.OnlineSub->GetExternalUIInterface();
+	if (!OnlineExternalUI.IsValid())
+	{
+		FFrame::KismetExecutionMessage(TEXT("External UI not supported by the current online subsystem"), ELogVerbosity::Warning);
+		OnFailure.Broadcast(PlayerControllerWeakPtr.Get());
+		return;
+	}
+
+	const ULocalPlayer* LocalPlayer = Cast<ULocalPlayer>(PlayerControllerWeakPtr->Player);
+
+	if (LocalPlayer == nullptr)
+	{
+		FFrame::KismetExecutionMessage(TEXT("Can only show login UI for local players"), ELogVerbosity::Warning);
+		OnFailure.Broadcast(PlayerControllerWeakPtr.Get());
+		return;
+	}
+		
+	const bool bWaitForDelegate = OnlineExternalUI->ShowLoginUI(LocalPlayer->GetControllerId(), false,
+		IOnlineExternalUI::FOnLoginUIClosedDelegate::CreateUObject(this, &UShowLoginUICallbackProxy::OnShowLoginUICompleted));
+
+	if (!bWaitForDelegate)
+	{
+		FFrame::KismetExecutionMessage(TEXT("The online subsystem couldn't show its login UI"), ELogVerbosity::Log);
+		OnFailure.Broadcast(PlayerControllerWeakPtr.Get());
+	}
+}
+
+void UShowLoginUICallbackProxy::OnShowLoginUICompleted(TSharedPtr<const FUniqueNetId> UniqueId, int LocalPlayerNum)
+{
+	// Update the cached unique ID for the local player and the player state.
+	APlayerController* PlayerController = PlayerControllerWeakPtr.Get();
+	
+	if (PlayerController != nullptr)
+	{
+		ULocalPlayer* LocalPlayer = PlayerController->GetLocalPlayer();
+		if (LocalPlayer != nullptr)
+		{
+			LocalPlayer->SetCachedUniqueNetId(UniqueId);
+		}
+		
+		if (PlayerController->PlayerState != nullptr)
+		{
+			PlayerController->PlayerState->SetUniqueId(UniqueId);
+		}
+	}
+
+	if (UniqueId.IsValid())
+	{
+		OnSuccess.Broadcast(PlayerControllerWeakPtr.Get());
+	}
+	else
+	{
+		OnFailure.Broadcast(PlayerControllerWeakPtr.Get());
+	}
+}
