Prevent a crash when trying to join a session with null SessionInfo when using OnlineSubsystemNull. Fixes UE-15898.

diff --git a/Engine/Source/Runtime/Online/OnlineSubsystemNull/Private/OnlineSessionInterfaceNull.cpp b/Engine/Source/Runtime/Online/OnlineSubsystemNull/Private/OnlineSessionInterfaceNull.cpp
--- a/Engine/Source/Runtime/Online/OnlineSubsystemNull/Private/OnlineSessionInterfaceNull.cpp
+++ b/Engine/Source/Runtime/Online/OnlineSubsystemNull/Private/OnlineSessionInterfaceNull.cpp
@@ -661,10 +661,12 @@ bool FOnlineSessionNull::SendSessionInviteToFriends(const FUniqueNetId& LocalUse
 
 uint32 FOnlineSessionNull::JoinLANSession(int32 PlayerNum, FNamedOnlineSession* Session, const FOnlineSession* SearchSession)
 {
+	check(Session != nullptr);
+
 	uint32 Result = E_FAIL;
 	Session->SessionState = EOnlineSessionState::Pending;
 
-	if (Session->SessionInfo.IsValid())
+	if (Session->SessionInfo.IsValid() && SearchSession != nullptr && SearchSession->SessionInfo.IsValid())
 	{
 		// Copy the session info over
 		const FOnlineSessionInfoNull* SearchSessionInfo = (const FOnlineSessionInfoNull*)SearchSession->SessionInfo.Get();
