#xmpp
- removed some default xmpp endpoints

diff --git a/Engine/Source/Runtime/Online/XMPP/Private/XmppModule.cpp b/Engine/Source/Runtime/Online/XMPP/Private/XmppModule.cpp
--- a/Engine/Source/Runtime/Online/XMPP/Private/XmppModule.cpp
+++ b/Engine/Source/Runtime/Online/XMPP/Private/XmppModule.cpp
@@ -59,16 +59,8 @@ bool FXmppModule::HandleXmppCommand( const TCHAR* Cmd, FOutputDevice& Ar )
 	{
 		FParse::Value(FCommandLine::Get(), TEXT("MCPCONFIG="), ConfigOverride);
 	}
-	if (ConfigOverride.Equals(TEXT("ci"), ESearchCase::IgnoreCase))
-	{
-		XmppServer.ServerAddr = TEXT("xmpp-service-ci.ol.epicgames.net");
-		XmppServer.Domain = TEXT("ci.ol.epicgames.net");
-	}
-	else
-	{
-		XmppServer.ServerAddr = TEXT("xmpp-service-gamedev.ol.epicgames.net");
-		XmppServer.Domain = TEXT("gamedev.ol.epicgames.net");
-	}
+	XmppServer.ServerAddr = TEXT("127.0.0.1");
+	XmppServer.Domain = TEXT("localhost.net");
 	XmppServer.ServerPort = 5222;
 
 	if (FParse::Command(&Cmd, TEXT("Test")))
