Fix non-unity compile error.

diff --git a/Engine/Source/Runtime/Online/HTTP/Private/Windows/HttpWinInet.cpp b/Engine/Source/Runtime/Online/HTTP/Private/Windows/HttpWinInet.cpp
--- a/Engine/Source/Runtime/Online/HTTP/Private/Windows/HttpWinInet.cpp
+++ b/Engine/Source/Runtime/Online/HTTP/Private/Windows/HttpWinInet.cpp
@@ -2,6 +2,7 @@
 
 #include "HttpPrivatePCH.h"
 #include "HttpWinInet.h"
+#include "EngineVersion.h"
 
 bool FWinInetConnection::bStaticConnectionInitialized = false;
 
