[INTEGRATE] merge ue4 <-ue4-Dev cl 2624191

diff --git a/Engine/Source/Runtime/Online/OnlineSubsystem/Public/OnlineSubsystemTypes.h b/Engine/Source/Runtime/Online/OnlineSubsystem/Public/OnlineSubsystemTypes.h
--- a/Engine/Source/Runtime/Online/OnlineSubsystem/Public/OnlineSubsystemTypes.h
+++ b/Engine/Source/Runtime/Online/OnlineSubsystem/Public/OnlineSubsystemTypes.h
@@ -1074,6 +1074,10 @@ public:
 	 * @return Any additional auth data associated with a registered user
 	 */
 	virtual bool GetAuthAttribute(const FString& AttrName, FString& OutAttrValue) const = 0;
+	/** 
+	 * @return True, if the data has been changed
+	 */
+	virtual bool SetUserAttribute(const FString& AttrName, const FString& AttrValue) = 0;
 };
 
 /** 
diff --git a/Engine/Source/Runtime/Online/OnlineSubsystemAmazon/Private/OnlineIdentityAmazon.cpp b/Engine/Source/Runtime/Online/OnlineSubsystemAmazon/Private/OnlineIdentityAmazon.cpp
--- a/Engine/Source/Runtime/Online/OnlineSubsystemAmazon/Private/OnlineIdentityAmazon.cpp
+++ b/Engine/Source/Runtime/Online/OnlineSubsystemAmazon/Private/OnlineIdentityAmazon.cpp
@@ -25,6 +25,11 @@ bool FUserOnlineAccountAmazon::GetUserAttribute(const FString& AttrName, FString
 	return false;
 }
 
+bool FUserOnlineAccountAmazon::SetUserAttribute(const FString& AttrName, const FString& AttrValue)
+{
+	return false;
+}
+
 FString FUserOnlineAccountAmazon::GetAccessToken() const
 {
 	return AuthTicket;
diff --git a/Engine/Source/Runtime/Online/OnlineSubsystemAmazon/Private/OnlineIdentityAmazon.h b/Engine/Source/Runtime/Online/OnlineSubsystemAmazon/Private/OnlineIdentityAmazon.h
--- a/Engine/Source/Runtime/Online/OnlineSubsystemAmazon/Private/OnlineIdentityAmazon.h
+++ b/Engine/Source/Runtime/Online/OnlineSubsystemAmazon/Private/OnlineIdentityAmazon.h
@@ -20,6 +20,7 @@ public:
 	virtual FString GetRealName() const override;
 	virtual FString GetDisplayName() const override;
 	virtual bool GetUserAttribute(const FString& AttrName, FString& OutAttrValue) const override;
+	virtual bool SetUserAttribute(const FString& AttrName, const FString& AttrValue) override;
 
 	// FUserOnlineAccount
 
diff --git a/Engine/Source/Runtime/Online/OnlineSubsystemFacebook/Private/IOS/OnlineIdentityFacebook.cpp b/Engine/Source/Runtime/Online/OnlineSubsystemFacebook/Private/IOS/OnlineIdentityFacebook.cpp
--- a/Engine/Source/Runtime/Online/OnlineSubsystemFacebook/Private/IOS/OnlineIdentityFacebook.cpp
+++ b/Engine/Source/Runtime/Online/OnlineSubsystemFacebook/Private/IOS/OnlineIdentityFacebook.cpp
@@ -45,6 +45,11 @@ bool FUserOnlineAccountFacebook::GetUserAttribute(const FString& AttrName, FStri
 	return false;
 }
 
+bool FUserOnlineAccountFacebook::SetUserAttribute(const FString& AttrName, const FString& AttrValue)
+{
+	return false;
+}
+
 FString FUserOnlineAccountFacebook::GetAccessToken() const
 {
 	return AuthTicket;
diff --git a/Engine/Source/Runtime/Online/OnlineSubsystemFacebook/Private/IOS/OnlineIdentityFacebook.h b/Engine/Source/Runtime/Online/OnlineSubsystemFacebook/Private/IOS/OnlineIdentityFacebook.h
--- a/Engine/Source/Runtime/Online/OnlineSubsystemFacebook/Private/IOS/OnlineIdentityFacebook.h
+++ b/Engine/Source/Runtime/Online/OnlineSubsystemFacebook/Private/IOS/OnlineIdentityFacebook.h
@@ -21,6 +21,7 @@ public:
 	virtual FString GetRealName() const override;
 	virtual FString GetDisplayName() const override;
 	virtual bool GetUserAttribute(const FString& AttrName, FString& OutAttrValue) const override;
+	virtual bool SetUserAttribute(const FString& AttrName, const FString& AttrValue) override;
 
 	// FUserOnlineAccount
 
diff --git a/Engine/Source/Runtime/Online/OnlineSubsystemFacebook/Private/Windows/OnlineIdentityFacebook.cpp b/Engine/Source/Runtime/Online/OnlineSubsystemFacebook/Private/Windows/OnlineIdentityFacebook.cpp
--- a/Engine/Source/Runtime/Online/OnlineSubsystemFacebook/Private/Windows/OnlineIdentityFacebook.cpp
+++ b/Engine/Source/Runtime/Online/OnlineSubsystemFacebook/Private/Windows/OnlineIdentityFacebook.cpp
@@ -27,6 +27,11 @@ bool FUserOnlineAccountFacebook::GetUserAttribute(const FString& AttrName, FStri
 	return false;
 }
 
+bool FUserOnlineAccountFacebook::SetUserAttribute(const FString& AttrName, const FString& AttrValue)
+{
+	return false;
+}
+
 FString FUserOnlineAccountFacebook::GetAccessToken() const
 {
 	return AuthTicket;
diff --git a/Engine/Source/Runtime/Online/OnlineSubsystemFacebook/Private/Windows/OnlineIdentityFacebook.h b/Engine/Source/Runtime/Online/OnlineSubsystemFacebook/Private/Windows/OnlineIdentityFacebook.h
--- a/Engine/Source/Runtime/Online/OnlineSubsystemFacebook/Private/Windows/OnlineIdentityFacebook.h
+++ b/Engine/Source/Runtime/Online/OnlineSubsystemFacebook/Private/Windows/OnlineIdentityFacebook.h
@@ -20,6 +20,7 @@ public:
 	virtual FString GetRealName() const override;
 	virtual FString GetDisplayName() const override;
 	virtual bool GetUserAttribute(const FString& AttrName, FString& OutAttrValue) const override;
+	virtual bool SetUserAttribute(const FString& AttrName, const FString& AttrValue) override;
 
 	// FUserOnlineAccount
 
diff --git a/Engine/Source/Runtime/Online/OnlineSubsystemNull/Private/OnlineIdentityNull.cpp b/Engine/Source/Runtime/Online/OnlineSubsystemNull/Private/OnlineIdentityNull.cpp
--- a/Engine/Source/Runtime/Online/OnlineSubsystemNull/Private/OnlineIdentityNull.cpp
+++ b/Engine/Source/Runtime/Online/OnlineSubsystemNull/Private/OnlineIdentityNull.cpp
@@ -27,6 +27,17 @@ bool FUserOnlineAccountNull::GetUserAttribute(const FString& AttrName, FString&
 	return false;
 }
 
+bool FUserOnlineAccountNull::SetUserAttribute(const FString& AttrName, const FString& AttrValue)
+{
+	const FString* FoundAttr = UserAttributes.Find(AttrName);
+	if (FoundAttr == NULL || *FoundAttr != AttrValue)
+	{
+		UserAttributes.Add(AttrName, AttrValue);
+		return true;
+	}
+	return false;
+}
+
 inline FString GenerateRandomUserId(int32 LocalUserNum)
 {
 	FString HostName;
diff --git a/Engine/Source/Runtime/Online/OnlineSubsystemNull/Private/OnlineIdentityNull.h b/Engine/Source/Runtime/Online/OnlineSubsystemNull/Private/OnlineIdentityNull.h
--- a/Engine/Source/Runtime/Online/OnlineSubsystemNull/Private/OnlineIdentityNull.h
+++ b/Engine/Source/Runtime/Online/OnlineSubsystemNull/Private/OnlineIdentityNull.h
@@ -19,6 +19,7 @@ public:
 	virtual FString GetRealName() const override { return TEXT("DummyRealName"); }
 	virtual FString GetDisplayName() const override  { return TEXT("DummyDisplayName"); }
 	virtual bool GetUserAttribute(const FString& AttrName, FString& OutAttrValue) const override;
+	virtual bool SetUserAttribute(const FString& AttrName, const FString& AttrValue) override;
 
 	// FUserOnlineAccount
 
