Merging using UE4-To-UE4-LauncherDev

2627874 #xmpp
- hook into a few socket status callbacks to detect SSL errors during login

2627876 #online,mcp,xmpp
- disable xmpp connection except for the default OSS instance

2627877 #online,identity,mcp
- avoid issues with double login complete callbacks due to relying on xmpploginchanged instead of xmpplogincomplete delegates for detecting xmpp login completion

diff --git a/Engine/Source/Runtime/Online/XMPP/Private/XmppJingle/XmppConnectionJingle.cpp b/Engine/Source/Runtime/Online/XMPP/Private/XmppJingle/XmppConnectionJingle.cpp
--- a/Engine/Source/Runtime/Online/XMPP/Private/XmppJingle/XmppConnectionJingle.cpp
+++ b/Engine/Source/Runtime/Online/XMPP/Private/XmppJingle/XmppConnectionJingle.cpp
@@ -105,10 +105,16 @@ public:
 				{
 					XmppThread->Restart();
 				}
+
+				XmppSocket = new buzz::XmppSocket(Connection.ClientSettings.use_tls());
+				XmppSocket->SignalError.connect(this, &FXmppConnectionPumpThread::OnSocketError);
+				XmppSocket->SignalClosed.connect(this, &FXmppConnectionPumpThread::OnSocketClosed);
+				XmppSocket->SignalCloseEvent.connect(this, &FXmppConnectionPumpThread::OnSslClosed);
+
 				// kick off login task
 				XmppPump->DoLogin(
 					Connection.ClientSettings, 
-					new buzz::XmppSocket(Connection.ClientSettings.use_tls()), 
+					XmppSocket, 
 					NULL);
 
 				Connection.HandlePumpStarting(XmppPump);
@@ -151,6 +157,7 @@ public:
 
 		delete XmppPump;
 		XmppPump = NULL;
+		XmppSocket = NULL;
 	}
 
 private:
@@ -198,42 +205,70 @@ private:
 	}
 
 	// callbacks
+
+	void OnSocketError()
+	{
+		UE_LOG(LogXmpp, VeryVerbose, TEXT("OnSocketError state=%d"), (int32)XmppSocket->state());
+		UE_LOG(LogXmpp, VeryVerbose, TEXT("OnSocketError error=%d"), (int32)XmppSocket->error());
+		UE_LOG(LogXmpp, VeryVerbose, TEXT("OnSocketError winsock=%d"), XmppSocket->GetError());
+	}
+
+	void OnSocketClosed()
+	{
+		UE_LOG(LogXmpp, VeryVerbose, TEXT("OnSocketClosed state=%d"), (int32)XmppSocket->state());
+		UE_LOG(LogXmpp, VeryVerbose, TEXT("OnSocketClosed error=%d"), (int32)XmppSocket->error());
+		UE_LOG(LogXmpp, VeryVerbose, TEXT("OnSocketClosed winsock=%d"), XmppSocket->GetError());
+	}
+
+	void OnSslClosed(int Error)
+	{
+		UE_LOG(LogXmpp, VeryVerbose, TEXT("OnSslClosed ERROR=%d"), Error);
+		UE_LOG(LogXmpp, VeryVerbose, TEXT("OnSslClosed state=%d"), (int32)XmppSocket->state());
+		UE_LOG(LogXmpp, VeryVerbose, TEXT("OnSslClosed error=%d"), (int32)XmppSocket->error());
+		UE_LOG(LogXmpp, VeryVerbose, TEXT("OnSslClosed winsock=%d"), XmppSocket->GetError());
+
+		if (LoginState == ELoginProgress::ProcessingLogin)
+		{
+			OnSignalStateChange(buzz::XmppEngine::STATE_CLOSED);
+		}
+	}
+
 	void OnSignalStateChange(buzz::XmppEngine::State State)
 	{
 		switch (State)
 		{
-		case buzz::XmppEngine::STATE_START:
-			UE_LOG(LogXmpp, Verbose, TEXT("STATE_START"));
-			break;
-		case buzz::XmppEngine::STATE_OPENING:
-			UE_LOG(LogXmpp, Verbose, TEXT("STATE_OPENING"));
-			break;
-		case buzz::XmppEngine::STATE_OPEN:
-		{
-			UE_LOG(LogXmpp, Verbose, TEXT("STATE_OPEN"));
+			case buzz::XmppEngine::STATE_START:
+				UE_LOG(LogXmpp, Verbose, TEXT("STATE_START"));
+				break;
+			case buzz::XmppEngine::STATE_OPENING:
+				UE_LOG(LogXmpp, Verbose, TEXT("STATE_OPENING"));
+				break;
+			case buzz::XmppEngine::STATE_OPEN:
+			{
+				UE_LOG(LogXmpp, Verbose, TEXT("STATE_OPEN"));
 
-			Connection.HandleLoginChange(LoginState, ELoginProgress::LoggedIn);
-			LoginState = ELoginProgress::LoggedIn;
+				Connection.HandleLoginChange(LoginState, ELoginProgress::LoggedIn);
+				LoginState = ELoginProgress::LoggedIn;
 
-			StartServerPing();
-		}
+				StartServerPing();
+			}
 			break;
-		case buzz::XmppEngine::STATE_CLOSED:
-		{
-			UE_LOG(LogXmpp, Verbose, TEXT("STATE_CLOSED"));
-
-			if (LoginState != ELoginProgress::LoggedIn)
+			case buzz::XmppEngine::STATE_CLOSED:
 			{
-				LogError(TEXT("log-in"));
-			}
+				UE_LOG(LogXmpp, Verbose, TEXT("STATE_CLOSED"));
 
-			StopServerPing();
-			Connection.HandlePumpQuitting(XmppPump);
-			XmppThread->Quit();
+				if (LoginState != ELoginProgress::LoggedIn)
+				{
+					LogError(TEXT("log-in"));
+				}
 
-			Connection.HandleLoginChange(LoginState, ELoginProgress::LoggedOut);
-			LoginState = ELoginProgress::LoggedOut;
-		}
+				StopServerPing();
+				Connection.HandlePumpQuitting(XmppPump);
+				XmppThread->Quit();
+
+				Connection.HandleLoginChange(LoginState, ELoginProgress::LoggedOut);
+				LoginState = ELoginProgress::LoggedOut;
+			}
 			break;
 		}
 	}
@@ -300,6 +335,8 @@ private:
 	buzz::XmppPump* XmppPump;
 	/** thread used by xmpp. set to this current thread */
 	rtc::Thread* XmppThread;
+	/** socket for the connection */
+	buzz::XmppSocket* XmppSocket;
 
 	/** steps during login/logout */
 	ELoginProgress::Type LoginState;
@@ -654,4 +691,4 @@ const FXmppUserJid& FXmppConnectionJingle::GetUserJid() const
 	return UserJid;
 }
 
-#endif //WITH_XMPP_JINGLE
+#endif //WITH_XMPP_JINGLE
\ No newline at end of file
