Android non-unity compile fixes

diff --git a/Engine/Source/Runtime/Online/Android/OnlineSubsystemGooglePlay/Private/OnlineAsyncTaskGooglePlayLogin.cpp b/Engine/Source/Runtime/Online/Android/OnlineSubsystemGooglePlay/Private/OnlineAsyncTaskGooglePlayLogin.cpp
--- a/Engine/Source/Runtime/Online/Android/OnlineSubsystemGooglePlay/Private/OnlineAsyncTaskGooglePlayLogin.cpp
+++ b/Engine/Source/Runtime/Online/Android/OnlineSubsystemGooglePlay/Private/OnlineAsyncTaskGooglePlayLogin.cpp
@@ -5,6 +5,7 @@
 
 #include "gpg/builder.h"
 #include "gpg/debug.h"
+#include "gpg/types.h"
 
 FOnlineAsyncTaskGooglePlayLogin::FOnlineAsyncTaskGooglePlayLogin(
 	FOnlineSubsystemGooglePlay* InSubsystem,
@@ -25,12 +26,12 @@ void FOnlineAsyncTaskGooglePlayLogin::Start_OnTaskThread()
 		// Store the Subsystem pointer locally so that the OnAuthActionFinished lambda can capture it
 		FOnlineSubsystemGooglePlay* LocalSubsystem = Subsystem;
 
-		Subsystem->GameServicesPtr = GameServices::Builder()
-			.SetDefaultOnLog(LogLevel::VERBOSE)
-			.SetOnAuthActionStarted([](AuthOperation Op) {
+		Subsystem->GameServicesPtr = gpg::GameServices::Builder()
+			.SetDefaultOnLog(gpg::LogLevel::VERBOSE)
+			.SetOnAuthActionStarted([](gpg::AuthOperation Op) {
 				UE_LOG(LogOnline, Log, TEXT("GPG OnAuthActionStarted: %s"), *FString(DebugString(Op).c_str()));
 			})
-			.SetOnAuthActionFinished([LocalSubsystem](AuthOperation Op, AuthStatus LocalStatus) {
+			.SetOnAuthActionFinished([LocalSubsystem](gpg::AuthOperation Op, gpg::AuthStatus LocalStatus) {
 				UE_LOG(LogOnline, Log, TEXT("GPG OnAuthActionFinished: %s, AuthStatus: %s"),
 					*FString(DebugString(Op).c_str()),
 					*FString(DebugString(LocalStatus).c_str()));
@@ -41,7 +42,7 @@ void FOnlineAsyncTaskGooglePlayLogin::Start_OnTaskThread()
 	else if(Subsystem->GameServicesPtr->IsAuthorized())
 	{
 		// We have a GameServices object and the user is authorized, nothing else to do.
-		Status = AuthStatus::VALID;
+		Status = gpg::AuthStatus::VALID;
 		bWasSuccessful = true;
 		bIsComplete = true;
 	}
diff --git a/Engine/Source/Runtime/Online/Android/OnlineSubsystemGooglePlay/Private/OnlineAsyncTaskGooglePlayShowLoginUI.cpp b/Engine/Source/Runtime/Online/Android/OnlineSubsystemGooglePlay/Private/OnlineAsyncTaskGooglePlayShowLoginUI.cpp
--- a/Engine/Source/Runtime/Online/Android/OnlineSubsystemGooglePlay/Private/OnlineAsyncTaskGooglePlayShowLoginUI.cpp
+++ b/Engine/Source/Runtime/Online/Android/OnlineSubsystemGooglePlay/Private/OnlineAsyncTaskGooglePlayShowLoginUI.cpp
@@ -54,9 +54,9 @@ void FOnlineAsyncTaskGooglePlayShowLoginUI::TriggerDelegates()
 
 void FOnlineAsyncTaskGooglePlayShowLoginUI::OnAuthActionFinished(gpg::AuthOperation InOp, gpg::AuthStatus InStatus)
 {
-	if (InOp == AuthOperation::SIGN_IN)
+	if (InOp == gpg::AuthOperation::SIGN_IN)
 	{
-		bWasSuccessful = InStatus == AuthStatus::VALID;
+		bWasSuccessful = InStatus == gpg::AuthStatus::VALID;
 		bIsComplete = true;
 	}
 }
diff --git a/Engine/Source/Runtime/Online/Android/OnlineSubsystemGooglePlay/Private/OnlineSubsystemGooglePlay.cpp b/Engine/Source/Runtime/Online/Android/OnlineSubsystemGooglePlay/Private/OnlineSubsystemGooglePlay.cpp
--- a/Engine/Source/Runtime/Online/Android/OnlineSubsystemGooglePlay/Private/OnlineSubsystemGooglePlay.cpp
+++ b/Engine/Source/Runtime/Online/Android/OnlineSubsystemGooglePlay/Private/OnlineSubsystemGooglePlay.cpp
@@ -8,6 +8,7 @@
 #include "OnlineAsyncTaskManagerGooglePlay.h"
 #include "OnlineAsyncTaskGooglePlayLogin.h"
 #include "OnlineAsyncTaskGooglePlayLogout.h"
+#include "OnlineAsyncTaskGooglePlayShowLoginUI.h"
 
 #include <android_native_app_glue.h>
 
