Merging using UE4-To-UE4-LauncherDev

#online,purchase,store
- added new interface IOnlinePurchase for handling IAP transactions
- added new interface IOnlineStoreV2 for retrieving purchasable products/offers from the store

WIP - not implemented for any online platform yet

diff --git a/Engine/Source/Runtime/Online/Android/OnlineSubsystemGooglePlay/Public/OnlineSubsystemGooglePlay.h b/Engine/Source/Runtime/Online/Android/OnlineSubsystemGooglePlay/Public/OnlineSubsystemGooglePlay.h
--- a/Engine/Source/Runtime/Online/Android/OnlineSubsystemGooglePlay/Public/OnlineSubsystemGooglePlay.h
+++ b/Engine/Source/Runtime/Online/Android/OnlineSubsystemGooglePlay/Public/OnlineSubsystemGooglePlay.h
@@ -44,16 +44,18 @@ public:
 	virtual IOnlineTitleFilePtr GetTitleFileInterface() const override;
 	virtual IOnlineEntitlementsPtr GetEntitlementsInterface() const override;
 	virtual IOnlineStorePtr GetStoreInterface() const override;
-	virtual IOnlineEventsPtr GetEventsInterface() const override { return NULL; }
-	virtual IOnlineMessagePtr GetMessageInterface() const override { return NULL; }
-	virtual IOnlineSharingPtr GetSharingInterface() const override { return NULL; }
-	virtual IOnlineUserPtr GetUserInterface() const override { return NULL; }
+	virtual IOnlineStoreV2Ptr GetStoreV2Interface() const override { return nullptr; }
+	virtual IOnlinePurchasePtr GetPurchaseInterface() const override { return nullptr; }
+	virtual IOnlineEventsPtr GetEventsInterface() const override { return nullptr; }
+	virtual IOnlineMessagePtr GetMessageInterface() const override { return nullptr; }
+	virtual IOnlineSharingPtr GetSharingInterface() const override { return nullptr; }
+	virtual IOnlineUserPtr GetUserInterface() const override { return nullptr; }
 	virtual IOnlineAchievementsPtr GetAchievementsInterface() const override;
-	virtual IOnlinePresencePtr GetPresenceInterface() const override { return NULL; }
-	virtual IOnlineChatPtr GetChatInterface() const override { return NULL; }
-	virtual IOnlineTurnBasedPtr GetTurnBasedInterface() const override { return NULL; }
+	virtual IOnlinePresencePtr GetPresenceInterface() const override { return nullptr; }
+	virtual IOnlineChatPtr GetChatInterface() const override { return nullptr; }
+	virtual IOnlineTurnBasedPtr GetTurnBasedInterface() const override { return nullptr; }
 
-	virtual class UObject* GetNamedInterface(FName InterfaceName) override { return NULL; }
+	virtual class UObject* GetNamedInterface(FName InterfaceName) override { return nullptr; }
 	virtual void SetNamedInterface(FName InterfaceName, class UObject* NewInterface) override {}
 	virtual bool IsDedicated() const override { return false; }
 	virtual bool IsServer() const override { return true; }
diff --git a/Engine/Source/Runtime/Online/IOS/OnlineSubsystemIOS/Public/OnlineSubsystemIOS.h b/Engine/Source/Runtime/Online/IOS/OnlineSubsystemIOS/Public/OnlineSubsystemIOS.h
--- a/Engine/Source/Runtime/Online/IOS/OnlineSubsystemIOS/Public/OnlineSubsystemIOS.h
+++ b/Engine/Source/Runtime/Online/IOS/OnlineSubsystemIOS/Public/OnlineSubsystemIOS.h
@@ -33,6 +33,8 @@ public:
 	virtual IOnlineTitleFilePtr GetTitleFileInterface() const override;
 	virtual IOnlineEntitlementsPtr GetEntitlementsInterface() const override;
 	virtual IOnlineStorePtr GetStoreInterface() const override;
+	virtual IOnlineStoreV2Ptr GetStoreV2Interface() const override { return nullptr; }
+	virtual IOnlinePurchasePtr GetPurchaseInterface() const override { return nullptr; }
 	virtual IOnlineEventsPtr GetEventsInterface() const override;
 	virtual IOnlineAchievementsPtr GetAchievementsInterface() const override;
 	virtual IOnlineSharingPtr GetSharingInterface() const override;
diff --git a/Engine/Source/Runtime/Online/OnlineSubsystem/Public/Interfaces/OnlineEntitlementsInterface.h b/Engine/Source/Runtime/Online/OnlineSubsystem/Public/Interfaces/OnlineEntitlementsInterface.h
--- a/Engine/Source/Runtime/Online/OnlineSubsystem/Public/Interfaces/OnlineEntitlementsInterface.h
+++ b/Engine/Source/Runtime/Online/OnlineSubsystem/Public/Interfaces/OnlineEntitlementsInterface.h
@@ -22,6 +22,18 @@ struct FOnlineEntitlement
 	FString ItemId;
 	/**  Namespace of the entitlement */
 	FString Namespace;
+	/** True if the entitlement is a consumable */
+	bool bIsConsumable;
+	/** Number of uses still available for a consumable */
+	int32 RemainingCount;
+	/** Number of prior uses for a consumable */
+	int32 ConsumedCount;
+
+	FOnlineEntitlement()
+		: bIsConsumable(false)
+		, RemainingCount(1)
+		, ConsumedCount(0)
+	{}
 
 	/**
 	 * @return Any additional data associated with the entitlement
@@ -49,7 +61,7 @@ struct FOnlineEntitlement
  * @param Namespace optional namespace that was used to query. Empty means all entitlements were queried
  * @param Error string representing the error condition
  */
-DECLARE_MULTICAST_DELEGATE_FourParams(FOnQueryEntitlementsComplete, bool, const FUniqueNetId&, const FString&, const FString&);
+DECLARE_MULTICAST_DELEGATE_FourParams(FOnQueryEntitlementsComplete, bool /*bWasSuccessful*/, const FUniqueNetId& /*UserId*/, const FString& /*Namespace*/, const FString& /*Error*/);
 typedef FOnQueryEntitlementsComplete::FDelegate FOnQueryEntitlementsCompleteDelegate;
 
 /**
@@ -105,7 +117,7 @@ public:
 	 * @param Namespace optional namespace that was used to query. Empty means all entitlements were queried
 	 * @param Error string representing the error condition
 	 */
-	DEFINE_ONLINE_DELEGATE_FOUR_PARAM(OnQueryEntitlementsComplete, bool, const FUniqueNetId&, const FString&, const FString&);
+	DEFINE_ONLINE_DELEGATE_FOUR_PARAM(OnQueryEntitlementsComplete, bool /*bWasSuccessful*/, const FUniqueNetId& /*UserId*/, const FString& /*Namespace*/, const FString& /*Error*/);
 };
 
 typedef TSharedPtr<IOnlineEntitlements, ESPMode::ThreadSafe> IOnlineEntitlementsPtr;
diff --git a/Engine/Source/Runtime/Online/OnlineSubsystem/Public/Interfaces/OnlinePurchaseInterface.h b/Engine/Source/Runtime/Online/OnlineSubsystem/Public/Interfaces/OnlinePurchaseInterface.h
new file mode 100644
--- /dev/null
+++ b/Engine/Source/Runtime/Online/OnlineSubsystem/Public/Interfaces/OnlinePurchaseInterface.h
@@ -0,0 +1,155 @@
+// Copyright 1998-2015 Epic Games, Inc. All Rights Reserved.
+
+#pragma once
+#include "OnlineDelegateMacros.h"
+#include "OnlineStoreInterfaceV2.h"
+#include "OnlineEntitlementsInterface.h"
+
+/**
+ * Info needed for checkout
+ */
+class FPurchaseCheckoutRequest
+{
+public:
+	/**
+	 * Add a offer entry for purchase
+	 *
+	 * @param OfferId id of offer to be purchased
+	 * @param Quantity number to purchase
+	 */
+	void AddPurchaseOffer(const FUniqueOfferId& OfferId, int32 Quantity);
+
+	/**
+	 * Single offer entry for purchase 
+	 */
+	struct FPurchaseOfferEntry
+	{
+		FUniqueOfferId OfferId;
+		int32 Quantity;
+	};
+	/** List of offers being purchased */
+	TArray<FPurchaseOfferEntry> PurchaseOffers;
+};
+
+/**
+ * State of a purchase transaction
+ */
+namespace EPurchaseTransactionState
+{
+	enum Type
+	{
+		/** processing has not started on the purchase */
+		NotStarted,
+		/** currently processing the purchase */
+		Processing,
+		/** purchase completed successfully */
+		Purchased,
+		/** purchase completed but failed */
+		Failed,
+		/** purchase canceled by user */
+		Canceled,
+		/** prior purchase that has been restored */
+		Restored
+	};
+}
+
+/**
+ * Receipt result from checkout
+ */
+class FPurchaseReceipt
+{
+public:
+	/** Unique Id for this transaction/order */
+	FString TransactionId;
+	/** Current state of the purchase */
+	EPurchaseTransactionState::Type TransactionState;
+	/** List of offers being purchased */
+	TArray<FUniqueOfferId> OfferIds;
+	/** The entitlements obtained from the completed purchase */
+	TArray<FUniqueEntitlementId> EntitlementIds;
+	/** Receipt validation data */
+	FString ReceiptValidation;
+};
+
+/**
+ * Info needed for code redemption
+ */
+class FRedeemCodeRequest
+{
+public:
+	/** Code to redeem */
+	FString Code;
+};
+
+/**
+ * Delegate called when checkout process completes
+ */
+DECLARE_DELEGATE_FourParams(FOnPurchaseCheckoutComplete, const FUniqueNetId& /*UserId*/, bool /*bWasSuccessful*/, const TSharedRef<FPurchaseReceipt>& /*Receipt*/, const FString& /*Error*/);
+
+/**
+ * Delegate called when code redemption process completes
+ */
+DECLARE_DELEGATE_FourParams(FOnPurchaseRedeemCodeComplete, const FUniqueNetId& /*UserId*/, bool /*bWasSuccessful*/, const TSharedRef<FPurchaseReceipt>& /*Receipt*/, const FString& /*Error*/);
+
+/**
+ * Delegate called when query receipt process completes
+ */
+DECLARE_DELEGATE_ThreeParams(FOnQueryReceiptsComplete, const FUniqueNetId& /*UserId*/, bool /*bWasSuccessful*/, const FString& /*Error*/);
+
+/**
+ *	IOnlinePurchase - Interface for IAP (In App Purchases) services
+ */
+class IOnlinePurchase
+{
+public:
+
+	/**
+	 * Determine if user is allowed to purchase from store 
+	 *
+	 * @param UserId user initiating the request
+	 *
+	 * @return true if user can make a purchase
+	 */
+	virtual bool IsAllowedToPurchase(const FUniqueNetId& UserId) = 0;
+
+	/**
+	 * Initiate the checkout process for purchasing offers via payment
+	 *
+	 * @param UserId user initiating the request
+	 * @param CheckoutRequest info needed for the checkout request
+	 * @param Delegate completion callback
+	 *
+	 * @return true if user can make a purchase
+	 */
+	virtual bool Checkout(const FUniqueNetId& UserId, const FPurchaseCheckoutRequest& CheckoutRequest, const FOnPurchaseCheckoutComplete& Delegate = FOnPurchaseCheckoutComplete()) = 0;
+
+	/**
+	 * Initiate the checkout process for obtaining offers via code redemption
+	 *
+	 * @param UserId user initiating the request
+	 * @param RedeemCodeRequest info needed for the redeem request
+	 * @param Delegate completion callback
+	 *
+	 * @return true if user can make a purchase
+	 */
+	virtual bool RedeemCode(const FUniqueNetId& UserId, const FRedeemCodeRequest& RedeemCodeRequest, const FOnPurchaseRedeemCodeComplete& Delegate = FOnPurchaseRedeemCodeComplete()) = 0;
+
+	/**
+	 * Query for all of the user's receipts from prior purchases
+	 *
+	 * @param UserId user initiating the request
+	 * @param Delegate completion callback
+	 *
+	 * @return true if user can make a purchase
+	 */
+	virtual bool QueryReceipts(const FUniqueNetId& UserId, const FOnQueryReceiptsComplete& Delegate = FOnQueryReceiptsComplete()) = 0;
+
+	/**
+	 * Get list of cached receipts for user (includes transactions currently being processed)
+	 *
+	 * @param UserId user initiating the request
+	 * @param OutReceipts [out] list of receipts for the user 
+	 */
+	virtual void GetReceipts(const FUniqueNetId& UserId, TArray<FPurchaseReceipt>& OutReceipts) const = 0;
+	
+};
\ No newline at end of file
diff --git a/Engine/Source/Runtime/Online/OnlineSubsystem/Public/Interfaces/OnlineStoreInterfaceV2.h b/Engine/Source/Runtime/Online/OnlineSubsystem/Public/Interfaces/OnlineStoreInterfaceV2.h
new file mode 100644
--- /dev/null
+++ b/Engine/Source/Runtime/Online/OnlineSubsystem/Public/Interfaces/OnlineStoreInterfaceV2.h
@@ -0,0 +1,138 @@
+// Copyright 1998-2015 Epic Games, Inc. All Rights Reserved.
+
+#pragma once
+#include "OnlineDelegateMacros.h"
+
+typedef FString FUniqueOfferId;
+typedef FString FUniqueCategoryId;
+
+/**
+ * Offer entry for display from online store
+ */
+class FOnlineStoreOffer
+{
+public:
+	/** Unique offer identifier */
+	FUniqueOfferId OfferId;
+	/** Title for display */
+	FText Title;	
+	/** Short description for display */
+	FText Description;
+	/** Full description for display */
+	FText LongDescription;
+	/** Price in text for display */
+	FText DisplayPrice;
+	/** Date the offer was released */
+	FDateTime ReleaseDate;
+	/** @return True if offer can be purchased */
+	virtual bool IsPurchaseable() const 
+	{ 
+		return true;
+	}
+};
+typedef TSharedRef<FOnlineStoreOffer> FOnlineStoreOfferRef;
+
+/**
+ * Category used to organize offers in the online store
+ */
+class FOnlineStoreCategory
+{
+public:
+	/** Unique identifier for this category */
+	FUniqueCategoryId Id;
+	/** Description for display */
+	FText Description;
+	/** List of optional sub categories */
+	TArray<FOnlineStoreCategory> SubCategories;
+};
+
+/**
+ * Filter for querying a subset of offers from the online store
+ */
+class FOnlineStoreFilter
+{
+public:
+	/** Namespace for identifying an entire store */
+	FString Namespace;
+	/** Keyword strings to match when filtering items/offers */
+	TArray<FString> Keywords;
+	/** Category paths to match when filtering offers */
+	TArray<FOnlineStoreCategory> IncludeCategories;
+	/** Category paths to exclude when filtering offers */
+	TArray<FOnlineStoreCategory> ExcludeCategories;
+	/** Currency to use for viewing catalog */
+	FString CurrencyCode;
+};
+
+/**
+ * Delegate called when available online categories have been queried
+ */
+DECLARE_DELEGATE_ThreeParams(FOnQueryOnlineStoreCategoriesComplete, const FUniqueNetId& /*UserId*/, bool /*bWasSuccessful*/, const FString& /*Error*/);
+
+/**
+ * Delegate called when online store query completes
+ */
+DECLARE_DELEGATE_FourParams(FOnQueryOnlineStoreOffersComplete, const FUniqueNetId& /*UserId*/, bool /*bWasSuccessful*/, const TArray<FUniqueOfferId>& /*OfferIds*/, const FString& /*Error*/);
+
+/**
+ *	Access to available offers for purchase
+ */
+class IOnlineStoreV2
+{
+public:
+
+	/**
+	 * Query for available store categories
+	 *
+	 * @param UserId user initiating the request
+	 * @param Delegate completion callback
+	 *
+	 * @return true if async operation started
+	 */
+	virtual bool QueryCategories(const FUniqueNetId& UserId, const FOnQueryOnlineStoreCategoriesComplete& Delegate = FOnQueryOnlineStoreCategoriesComplete()) = 0;
+
+	/**
+	 * Get currently cached store categories
+	 *
+	 * @param OutCategories [out] list of categories previously queried
+	 */
+	virtual void GetCategories(TArray<FOnlineStoreCategory>& OutCategories) const = 0;
+
+	/**
+	 * Query for available store offers using a filter
+	 *
+	 * @param UserId user initiating the request
+	 * @param Filter only return offers matching the filter
+	 * @param Delegate completion callback
+	 *
+	 * @return true if async operation started
+	 */
+	virtual bool QueryOffersByFilter(const FUniqueNetId& UserId, const FOnlineStoreFilter& Filter, const FOnQueryOnlineStoreOffersComplete& Delegate = FOnQueryOnlineStoreOffersComplete()) = 0;
+
+	/**
+	 * Query for available store offers matching the given ids
+	 *
+	 * @param UserId user initiating the request
+	 * @param OfferIds only return offers matching these ids
+	 * @param Delegate completion callback
+	 *
+	 * @return true if async operation started
+	 */
+	virtual bool QueryOffersById(const FUniqueNetId& UserId, const TArray<FUniqueOfferId>& OfferIds, const FOnQueryOnlineStoreOffersComplete& Delegate = FOnQueryOnlineStoreOffersComplete()) = 0;
+
+	/**
+	 * Get currently cached store offers
+	 *
+	 * @param OutOffers [out] list of offers previously queried
+	 */
+	virtual void GetOffers(TArray<FOnlineStoreOfferRef>& OutOffers) const = 0;
+
+	/**
+	 * Get currently cached store offer entry
+	 *
+	 * @param OfferId id of offer to find
+	 *
+	 * @return offer if found or null
+	 */
+	virtual TSharedPtr<FOnlineStoreOffer> GetOffer(const FUniqueOfferId& OfferId) const = 0;
+};
diff --git a/Engine/Source/Runtime/Online/OnlineSubsystem/Public/OnlineSubsystem.h b/Engine/Source/Runtime/Online/OnlineSubsystem/Public/OnlineSubsystem.h
--- a/Engine/Source/Runtime/Online/OnlineSubsystem/Public/OnlineSubsystem.h
+++ b/Engine/Source/Runtime/Online/OnlineSubsystem/Public/OnlineSubsystem.h
@@ -45,6 +45,8 @@ typedef TSharedPtr<class IOnlineTime, ESPMode::ThreadSafe> IOnlineTimePtr;
 typedef TSharedPtr<class IOnlineIdentity, ESPMode::ThreadSafe> IOnlineIdentityPtr;
 typedef TSharedPtr<class IOnlineTitleFile, ESPMode::ThreadSafe> IOnlineTitleFilePtr;
 typedef TSharedPtr<class IOnlineStore, ESPMode::ThreadSafe> IOnlineStorePtr;
+typedef TSharedPtr<class IOnlineStoreV2, ESPMode::ThreadSafe> IOnlineStoreV2Ptr;
+typedef TSharedPtr<class IOnlinePurchase, ESPMode::ThreadSafe> IOnlinePurchasePtr;
 typedef TSharedPtr<class IOnlineEvents, ESPMode::ThreadSafe> IOnlineEventsPtr;
 typedef TSharedPtr<class IOnlineAchievements, ESPMode::ThreadSafe> IOnlineAchievementsPtr;
 typedef TSharedPtr<class IOnlineSharing, ESPMode::ThreadSafe> IOnlineSharingPtr;
@@ -226,6 +228,18 @@ public:
 	 * @return Interface pointer for the appropriate online store service
 	 */
 	virtual IOnlineStorePtr GetStoreInterface() const = 0;
+
+	/** 
+	 * Get the interface for accessing an online store
+	 * @return Interface pointer for the appropriate online store service
+	 */
+	virtual IOnlineStoreV2Ptr GetStoreV2Interface() const = 0;
+
+	/** 
+	 * Get the interface for purchasing 
+	 * @return Interface pointer for the appropriate purchase service
+	 */
+	virtual IOnlinePurchasePtr GetPurchaseInterface() const = 0;
 	
 	/** 
 	 * Get the interface for accessing online achievements
diff --git a/Engine/Source/Runtime/Online/OnlineSubsystemAmazon/Public/OnlineSubsystemAmazon.h b/Engine/Source/Runtime/Online/OnlineSubsystemAmazon/Public/OnlineSubsystemAmazon.h
--- a/Engine/Source/Runtime/Online/OnlineSubsystemAmazon/Public/OnlineSubsystemAmazon.h
+++ b/Engine/Source/Runtime/Online/OnlineSubsystemAmazon/Public/OnlineSubsystemAmazon.h
@@ -33,102 +33,30 @@ PACKAGE_SCOPE:
 public:
 	// IOnlineSubsystem
 
-	virtual IOnlineSessionPtr GetSessionInterface() const override
-	{
-		return NULL;
-	}
-	virtual IOnlineFriendsPtr GetFriendsInterface() const override
-	{
-		return NULL;
-	}
-	virtual IOnlinePartyPtr GetPartyInterface() const override
-	{
-		return NULL;
-	}
-	virtual IOnlineGroupsPtr GetGroupsInterface() const override
-	{
-		return nullptr;
-	}
-	virtual IOnlineSharedCloudPtr GetSharedCloudInterface() const override
-	{
-		return NULL;
-	}
-	virtual IOnlineUserCloudPtr GetUserCloudInterface() const override
-	{
-		return NULL;
-	}
-	virtual IOnlineLeaderboardsPtr GetLeaderboardsInterface() const override
-	{
-		return NULL;
-	}
-	virtual IOnlineVoicePtr GetVoiceInterface() const override
-	{
-		return NULL;
-	}
-	virtual IOnlineExternalUIPtr GetExternalUIInterface() const override	
-	{
-		return NULL;
-	}
-	virtual IOnlineTimePtr GetTimeInterface() const override
-	{
-		return NULL;
-	}
-	virtual IOnlineTitleFilePtr GetTitleFileInterface() const override
-	{
-		return NULL;
-	}
-
-	virtual IOnlineEntitlementsPtr GetEntitlementsInterface() const override
-	{
-		return NULL;
-	}
-
+	virtual IOnlineSessionPtr GetSessionInterface() const override { return nullptr; }
+	virtual IOnlineFriendsPtr GetFriendsInterface() const override { return nullptr; }
+	virtual IOnlinePartyPtr GetPartyInterface() const override { return nullptr; }
+	virtual IOnlineGroupsPtr GetGroupsInterface() const override { return nullptr; }
+	virtual IOnlineSharedCloudPtr GetSharedCloudInterface() const override { return nullptr; }
+	virtual IOnlineUserCloudPtr GetUserCloudInterface() const override { return nullptr; }
+	virtual IOnlineLeaderboardsPtr GetLeaderboardsInterface() const override { return nullptr; }
+	virtual IOnlineVoicePtr GetVoiceInterface() const override { return nullptr; }
+	virtual IOnlineExternalUIPtr GetExternalUIInterface() const override { return nullptr; }
+	virtual IOnlineTimePtr GetTimeInterface() const override { return nullptr; }
+	virtual IOnlineTitleFilePtr GetTitleFileInterface() const override { return nullptr; }
+	virtual IOnlineEntitlementsPtr GetEntitlementsInterface() const override { return nullptr; }
 	virtual IOnlineIdentityPtr GetIdentityInterface() const override;
-
-	virtual IOnlineStorePtr GetStoreInterface() const override
-	{
-		return NULL;
-	}
-
-	virtual IOnlineEventsPtr GetEventsInterface() const override
-	{
-		return NULL;
-	}
-
-	virtual IOnlineAchievementsPtr GetAchievementsInterface() const override
-	{
-		return NULL;
-	}
-
-	virtual IOnlineSharingPtr GetSharingInterface() const override
-	{
-		return NULL;
-	}
-
-	virtual IOnlineUserPtr GetUserInterface() const override
-	{
-		return NULL;
-	}
-
-	virtual IOnlineMessagePtr GetMessageInterface() const override
-	{
-		return NULL;
-	}
-
-	virtual IOnlinePresencePtr GetPresenceInterface() const override
-	{
-		return NULL;
-	}
-
-	virtual IOnlineChatPtr GetChatInterface() const override
-	{
-		return NULL;
-	}
-
-    virtual IOnlineTurnBasedPtr GetTurnBasedInterface() const override
-    {
-        return NULL;
-    }
+	virtual IOnlineStorePtr GetStoreInterface() const override { return nullptr; }
+	virtual IOnlineStoreV2Ptr GetStoreV2Interface() const override { return nullptr; }
+	virtual IOnlinePurchasePtr GetPurchaseInterface() const override { return nullptr; }
+	virtual IOnlineEventsPtr GetEventsInterface() const override { return nullptr; }
+	virtual IOnlineAchievementsPtr GetAchievementsInterface() const override { return nullptr; }
+	virtual IOnlineSharingPtr GetSharingInterface() const override { return nullptr; }
+	virtual IOnlineUserPtr GetUserInterface() const override { return nullptr; }
+	virtual IOnlineMessagePtr GetMessageInterface() const override { return nullptr; }
+	virtual IOnlinePresencePtr GetPresenceInterface() const override { return nullptr; }
+	virtual IOnlineChatPtr GetChatInterface() const override { return nullptr; }
+	virtual IOnlineTurnBasedPtr GetTurnBasedInterface() const override { return nullptr; }
 	
 	virtual bool Init() override;
 	virtual bool Shutdown() override;
diff --git a/Engine/Source/Runtime/Online/OnlineSubsystemFacebook/Public/OnlineSubsystemFacebook.h b/Engine/Source/Runtime/Online/OnlineSubsystemFacebook/Public/OnlineSubsystemFacebook.h
--- a/Engine/Source/Runtime/Online/OnlineSubsystemFacebook/Public/OnlineSubsystemFacebook.h
+++ b/Engine/Source/Runtime/Online/OnlineSubsystemFacebook/Public/OnlineSubsystemFacebook.h
@@ -35,6 +35,8 @@ public:
 	virtual IOnlineTitleFilePtr GetTitleFileInterface() const override;
 	virtual IOnlineEntitlementsPtr GetEntitlementsInterface() const override;
 	virtual IOnlineStorePtr GetStoreInterface() const override;
+	virtual IOnlineStoreV2Ptr GetStoreV2Interface() const override { return nullptr; }
+	virtual IOnlinePurchasePtr GetPurchaseInterface() const override { return nullptr; }
 	virtual IOnlineEventsPtr GetEventsInterface() const override;
 	virtual IOnlineAchievementsPtr GetAchievementsInterface() const override;
 	virtual IOnlineSharingPtr GetSharingInterface() const override;
diff --git a/Engine/Source/Runtime/Online/OnlineSubsystemNull/Public/OnlineSubsystemNull.h b/Engine/Source/Runtime/Online/OnlineSubsystemNull/Public/OnlineSubsystemNull.h
--- a/Engine/Source/Runtime/Online/OnlineSubsystemNull/Public/OnlineSubsystemNull.h
+++ b/Engine/Source/Runtime/Online/OnlineSubsystemNull/Public/OnlineSubsystemNull.h
@@ -46,6 +46,8 @@ public:
 	virtual IOnlineIdentityPtr GetIdentityInterface() const override;
 	virtual IOnlineTitleFilePtr GetTitleFileInterface() const override;
 	virtual IOnlineStorePtr GetStoreInterface() const override;
+	virtual IOnlineStoreV2Ptr GetStoreV2Interface() const override { return nullptr; }
+	virtual IOnlinePurchasePtr GetPurchaseInterface() const override { return nullptr; }
 	virtual IOnlineEventsPtr GetEventsInterface() const override;
 	virtual IOnlineAchievementsPtr GetAchievementsInterface() const override;
 	virtual IOnlineSharingPtr GetSharingInterface() const override;
@@ -77,23 +79,23 @@ PACKAGE_SCOPE:
 	/** Only the factory makes instances */
 	FOnlineSubsystemNull(FName InInstanceName) :
 		FOnlineSubsystemImpl(InInstanceName),
-		SessionInterface(NULL),
-		VoiceInterface(NULL),
-		LeaderboardsInterface(NULL),
-		IdentityInterface(NULL),
-		AchievementsInterface(NULL),
-		OnlineAsyncTaskThreadRunnable(NULL),
-		OnlineAsyncTaskThread(NULL)
+		SessionInterface(nullptr),
+		VoiceInterface(nullptr),
+		LeaderboardsInterface(nullptr),
+		IdentityInterface(nullptr),
+		AchievementsInterface(nullptr),
+		OnlineAsyncTaskThreadRunnable(nullptr),
+		OnlineAsyncTaskThread(nullptr)
 	{}
 
 	FOnlineSubsystemNull() :
-		SessionInterface(NULL),
-		VoiceInterface(NULL),
-		LeaderboardsInterface(NULL),
-		IdentityInterface(NULL),
-		AchievementsInterface(NULL),
-		OnlineAsyncTaskThreadRunnable(NULL),
-		OnlineAsyncTaskThread(NULL)
+		SessionInterface(nullptr),
+		VoiceInterface(nullptr),
+		LeaderboardsInterface(nullptr),
+		IdentityInterface(nullptr),
+		AchievementsInterface(nullptr),
+		OnlineAsyncTaskThreadRunnable(nullptr),
+		OnlineAsyncTaskThread(nullptr)
 	{}
 
 private:
diff --git a/Engine/Source/Runtime/Online/OnlineSubsystemSteam/Public/OnlineSubsystemSteam.h b/Engine/Source/Runtime/Online/OnlineSubsystemSteam/Public/OnlineSubsystemSteam.h
--- a/Engine/Source/Runtime/Online/OnlineSubsystemSteam/Public/OnlineSubsystemSteam.h
+++ b/Engine/Source/Runtime/Online/OnlineSubsystemSteam/Public/OnlineSubsystemSteam.h
@@ -218,6 +218,8 @@ public:
 	virtual IOnlineTitleFilePtr GetTitleFileInterface() const override;
 	virtual IOnlineEntitlementsPtr GetEntitlementsInterface() const override;
 	virtual IOnlineStorePtr GetStoreInterface() const override;
+	virtual IOnlineStoreV2Ptr GetStoreV2Interface() const override { return nullptr; }
+	virtual IOnlinePurchasePtr GetPurchaseInterface() const override { return nullptr; }
 	virtual IOnlineEventsPtr GetEventsInterface() const override;
 	virtual IOnlineAchievementsPtr GetAchievementsInterface() const override;
 	virtual IOnlineSharingPtr GetSharingInterface() const override;
