fix for Clodukit being enabled on iOS7
UE-16973
#ios

diff --git a/Engine/Source/Programs/UnrealBuildTool/IOS/UEBuildIOS.cs b/Engine/Source/Programs/UnrealBuildTool/IOS/UEBuildIOS.cs
--- a/Engine/Source/Programs/UnrealBuildTool/IOS/UEBuildIOS.cs
+++ b/Engine/Source/Programs/UnrealBuildTool/IOS/UEBuildIOS.cs
@@ -304,12 +304,6 @@ namespace UnrealBuildTool
 
             InBuildTarget.GlobalLinkEnvironment.Config.AdditionalFrameworks.Add( new UEBuildFramework( "GameKit" ) );
             InBuildTarget.GlobalLinkEnvironment.Config.AdditionalFrameworks.Add( new UEBuildFramework( "StoreKit" ) );
-			InBuildTarget.GlobalLinkEnvironment.Config.AdditionalFrameworks.Add( new UEBuildFramework( "CloudKit" ) );
-
-			if (RunTimeIOSVersion != "6.1")
-			{
-				InBuildTarget.GlobalLinkEnvironment.Config.AdditionalFrameworks.Add (new UEBuildFramework ("MultipeerConnectivity"));
-			}
         }
 
         /**
diff --git a/Engine/Source/Runtime/Online/IOS/OnlineSubsystemIOS/OnlineSubsystemIOS.Build.cs b/Engine/Source/Runtime/Online/IOS/OnlineSubsystemIOS/OnlineSubsystemIOS.Build.cs
--- a/Engine/Source/Runtime/Online/IOS/OnlineSubsystemIOS/OnlineSubsystemIOS.Build.cs
+++ b/Engine/Source/Runtime/Online/IOS/OnlineSubsystemIOS/OnlineSubsystemIOS.Build.cs
@@ -24,5 +24,8 @@ public class OnlineSubsystemIOS : ModuleRules
                 "Http",
             }
             );
+
+		PublicWeakFrameworks.Add("Cloudkit");
+		PublicWeakFrameworks.Add("MultipeerConnectivity");
 	}
 }
diff --git a/Engine/Source/Runtime/Online/IOS/OnlineSubsystemIOS/Private/OnlineUserCloudInterfaceIOS.cpp b/Engine/Source/Runtime/Online/IOS/OnlineSubsystemIOS/Private/OnlineUserCloudInterfaceIOS.cpp
--- a/Engine/Source/Runtime/Online/IOS/OnlineSubsystemIOS/Private/OnlineUserCloudInterfaceIOS.cpp
+++ b/Engine/Source/Runtime/Online/IOS/OnlineSubsystemIOS/Private/OnlineUserCloudInterfaceIOS.cpp
@@ -12,81 +12,98 @@
 
 -(IOSCloudStorage*)init:(bool)registerHandler
 {
-	// get the current iCloud ubiquity token
-	iCloudToken = [NSFileManager defaultManager].ubiquityIdentityToken;
-
-	// register for iCloud change notifications
-	if (registerHandler)
+#ifdef __IPHONE_8_0
+	if ([CKContainer class])
 	{
-		[[NSNotificationCenter defaultCenter] addObserver:self selector : @selector(iCloudAccountAvailabilityChanged:) name: NSUbiquityIdentityDidChangeNotification object : nil];
-	}
+		// get the current iCloud ubiquity token
+		iCloudToken = [NSFileManager defaultManager].ubiquityIdentityToken;
 
-	CloudContainer = [CKContainer defaultContainer];
-	SharedDatabase = [CloudContainer publicCloudDatabase];
-//	if (iCloudToken != nil)
-	{
+		// register for iCloud change notifications
+		if (registerHandler)
+		{
+			[[NSNotificationCenter defaultCenter] addObserver:self selector : @selector(iCloudAccountAvailabilityChanged:) name: NSUbiquityIdentityDidChangeNotification object : nil];
+		}
+
+		CloudContainer = [CKContainer defaultContainer];
+		SharedDatabase = [CloudContainer publicCloudDatabase];
 		UserDatabase = [CloudContainer privateCloudDatabase];
 	}
-//	else
-//	{
-//		UserDatabase = nil;
-//	}
-
+#endif
 	return self;
 }
 
 -(bool)readFile:(NSString*)fileName sharedDB:(bool)shared completionHandler:(void(^)(CKRecord* record, NSError* error))handler
 {
-	CKDatabase* DB = shared ? SharedDatabase : UserDatabase;
-	if (DB != nil)
+#ifdef __IPHONE_8_0
+	if ([CKDatabase class])
 	{
-		CKRecordID* recordId = [[CKRecordID alloc] initWithRecordName:fileName];
-		[DB fetchRecordWithID : recordId completionHandler : handler];
-		return true;
+		CKDatabase* DB = shared ? SharedDatabase : UserDatabase;
+		if (DB != nil)
+		{
+			CKRecordID* recordId = [[CKRecordID alloc] initWithRecordName:fileName];
+			[DB fetchRecordWithID : recordId completionHandler : handler];
+			return true;
+		}
 	}
+#endif
 	return false;
 }
 
 -(bool)writeFile:(NSString*)fileName contents:(NSData*)fileContents sharedDB:(bool)shared completionHandler:(void(^)(CKRecord* record, NSError* error))handler
 {
-	CKDatabase* DB = shared ? SharedDatabase : UserDatabase;
-	if (DB != nil)
+#ifdef __IPHONE_8_0
+	if ([CKDatabase class])
 	{
-		CKRecordID* recordId = [[CKRecordID alloc] initWithRecordName:fileName];
-		CKRecord* record = [[CKRecord alloc] initWithRecordType:@"file" recordID: recordId];
-		record[@"contents"] = fileContents;
-		[DB saveRecord : record completionHandler : handler];
-		return true;
+		CKDatabase* DB = shared ? SharedDatabase : UserDatabase;
+		if (DB != nil)
+		{
+			CKRecordID* recordId = [[CKRecordID alloc] initWithRecordName:fileName];
+			CKRecord* record = [[CKRecord alloc] initWithRecordType:@"file" recordID: recordId];
+			record[@"contents"] = fileContents;
+			[DB saveRecord : record completionHandler : handler];
+			return true;
+		}
 	}
+#endif
 	return false;
 }
 
 -(bool)deleteFile:(NSString*)fileName sharedDB:(bool)shared completionHandler:(void(^)(CKRecordID* record, NSError* error))handler
 {
-	CKDatabase* DB = shared ? SharedDatabase : UserDatabase;
-	if (DB != nil)
+#ifdef __IPHONE_8_0
+	if ([CKDatabase class])
 	{
-		CKRecordID* recordId = [[CKRecordID alloc] initWithRecordName:fileName];
-		[DB deleteRecordWithID: recordId completionHandler : handler];
-		return true;
+		CKDatabase* DB = shared ? SharedDatabase : UserDatabase;
+		if (DB != nil)
+		{
+			CKRecordID* recordId = [[CKRecordID alloc] initWithRecordName:fileName];
+			[DB deleteRecordWithID : recordId completionHandler : handler];
+			return true;
+		}
 	}
+#endif
 	return false;
 }
 
 -(bool)query:(bool)shared fetchHandler:(void(^)(CKRecord* record))fetch completionHandler:(void(^)(CKQueryCursor* record, NSError* error))complete
 {
-	CKDatabase* DB = shared ? SharedDatabase : UserDatabase;
-	if (DB != nil)
+#ifdef __IPHONE_8_0
+	if ([CKDatabase class])
 	{
-		CKQuery* query = [[CKQuery alloc] initWithRecordType:@"file" predicate:[NSPredicate predicateWithFormat : @"TRUEPREDICATE"]];
-		CKQueryOperation* queryOp = [[CKQueryOperation alloc] initWithQuery:query];
-		queryOp.desiredKeys = @[@"record.recordID.recordName"];
-		queryOp.recordFetchedBlock = fetch;
-		queryOp.queryCompletionBlock = complete;
-		queryOp.resultsLimit = CKQueryOperationMaximumResults;
-		[DB addOperation : queryOp];
-		return true;
+		CKDatabase* DB = shared ? SharedDatabase : UserDatabase;
+		if (DB != nil)
+		{
+			CKQuery* query = [[CKQuery alloc] initWithRecordType:@"file" predicate:[NSPredicate predicateWithFormat : @"TRUEPREDICATE"]];
+			CKQueryOperation* queryOp = [[CKQueryOperation alloc] initWithQuery:query];
+			queryOp.desiredKeys = @[@"record.recordID.recordName"];
+			queryOp.recordFetchedBlock = fetch;
+			queryOp.queryCompletionBlock = complete;
+			queryOp.resultsLimit = CKQueryOperationMaximumResults;
+			[DB addOperation : queryOp];
+			return true;
+		}
 	}
+#endif
 	return false;
 }
 
@@ -108,7 +125,7 @@
 
 	if (theStorage == nil)
 	{
-		theStorage = [[IOSCloudStorage alloc] init: true];
+		theStorage = [[IOSCloudStorage alloc] init:true];
 	}
 	return theStorage;
 }
@@ -243,7 +260,8 @@ bool FOnlineUserCloudInterfaceIOS::ClearFile(const FUniqueNetId& UserId, const F
 
 void FOnlineUserCloudInterfaceIOS::EnumerateUserFiles(const FUniqueNetId& UserId)
 {
-    MetaDataState = EOnlineAsyncTaskState::InProgress;
+#ifdef __IPHONE_8_0
+	MetaDataState = EOnlineAsyncTaskState::InProgress;
 	if ([[IOSCloudStorage cloudStorage] query:false fetchHandler:^(CKRecord* record)
 	{
          FString FileName = record.recordID.recordName;
@@ -269,6 +287,7 @@ void FOnlineUserCloudInterfaceIOS::EnumerateUserFiles(const FUniqueNetId& UserId
     {
         TriggerOnEnumerateUserFilesCompleteDelegates(false, UserId);
     }
+#endif
 }
 
 void FOnlineUserCloudInterfaceIOS::GetUserFileList(const FUniqueNetId& UserId, TArray<FCloudFileHeader>& UserFiles)
@@ -282,6 +301,7 @@ void FOnlineUserCloudInterfaceIOS::GetUserFileList(const FUniqueNetId& UserId, T
 
 bool FOnlineUserCloudInterfaceIOS::ReadUserFile(const FUniqueNetId& UserId, const FString& FileName)
 {
+#ifdef __IPHONE_8_0
 	FCloudFile* CloudFile = GetCloudFile(FileName);
 	if (CloudFile)
 	{
@@ -314,11 +334,13 @@ bool FOnlineUserCloudInterfaceIOS::ReadUserFile(const FUniqueNetId& UserId, cons
     {
         TriggerOnReadUserFileCompleteDelegates(false, UserId, FileName);
     }
+#endif
 	return false;
 }
 
 bool FOnlineUserCloudInterfaceIOS::WriteUserFile(const FUniqueNetId& UserId, const FString& FileName, TArray<uint8>& FileContents)
 {
+#ifdef __IPHONE_8_0
 	FCloudFile* CloudFile = GetCloudFile(FileName, true);
 	if (CloudFile)
 	{
@@ -350,6 +372,7 @@ bool FOnlineUserCloudInterfaceIOS::WriteUserFile(const FUniqueNetId& UserId, con
     {
         TriggerOnWriteUserFileCompleteDelegates(false, UserId, FileName);
     }
+#endif
 	return false;
 }
 
@@ -361,6 +384,7 @@ void FOnlineUserCloudInterfaceIOS::CancelWriteUserFile(const FUniqueNetId& UserI
 
 bool FOnlineUserCloudInterfaceIOS::DeleteUserFile(const FUniqueNetId& UserId, const FString& FileName, bool bShouldCloudDelete, bool bShouldLocallyDelete)
 {
+#ifdef __IPHONE_8_0
 	FCloudFile* CloudFile = GetCloudFile(FileName);
 	if (CloudFile)
 	{
@@ -399,6 +423,7 @@ bool FOnlineUserCloudInterfaceIOS::DeleteUserFile(const FUniqueNetId& UserId, co
     {
         TriggerOnDeleteUserFileCompleteDelegates(false, UserId, FileName);
     }
+#endif
 	return false;
 }
 
diff --git a/Engine/Source/Runtime/Online/IOS/OnlineSubsystemIOS/Private/OnlineUserCloudInterfaceIOS.h b/Engine/Source/Runtime/Online/IOS/OnlineSubsystemIOS/Private/OnlineUserCloudInterfaceIOS.h
--- a/Engine/Source/Runtime/Online/IOS/OnlineSubsystemIOS/Private/OnlineUserCloudInterfaceIOS.h
+++ b/Engine/Source/Runtime/Online/IOS/OnlineSubsystemIOS/Private/OnlineUserCloudInterfaceIOS.h
@@ -3,6 +3,8 @@
 #pragma once
 
 #include "OnlineUserCloudInterface.h"
+
+#ifdef __IPHONE_8_0
 #import <CloudKit/CloudKit.h>
 
 @class IOSCloudStorage;
@@ -25,6 +27,7 @@
 +(IOSCloudStorage*)cloudStorage;
 
 @end
+#endif
 
 /**
 *	FOnlineUserCloudInterfaceIOS - Implementation of user cloud storage for IOS
