[INTEGRATE] merge ue4<-ue4-launcherDev cl 2638806

diff --git a/Engine/Source/Runtime/Online/BuildPatchServices/Private/BuildPatchError.cpp b/Engine/Source/Runtime/Online/BuildPatchServices/Private/BuildPatchError.cpp
--- a/Engine/Source/Runtime/Online/BuildPatchServices/Private/BuildPatchError.cpp
+++ b/Engine/Source/Runtime/Online/BuildPatchServices/Private/BuildPatchError.cpp
@@ -40,7 +40,8 @@ bool FBuildPatchInstallError::IsNoRetryError()
 	return ErrorState == EBuildPatchInstallError::DownloadError
 		|| ErrorState == EBuildPatchInstallError::MoveFileToInstall
 		|| ErrorState == EBuildPatchInstallError::InitializationError
-		|| ErrorState == EBuildPatchInstallError::PathLengthExceeded;
+		|| ErrorState == EBuildPatchInstallError::PathLengthExceeded
+		|| ErrorState == EBuildPatchInstallError::OutOfDiskSpace;
 }
 
 FString FBuildPatchInstallError::GetErrorString()
@@ -61,6 +62,7 @@ const FText FBuildPatchInstallError::GetErrorText()
 	static const FText FileConstructionFail( LOCTEXT( "BuildPatchInstallError_FileConstructionFail", "Please try again." ) );
 	static const FText MoveFileToInstall( LOCTEXT( "BuildPatchInstallError_MoveFileToInstall", "Please check your running processes." ) );
 	static const FText PathLengthExceeded(LOCTEXT("BuildPatchInstallError_PathLengthExceeded", "Please specify a shorter install location."));
+	static const FText OutOfDiskSpace(LOCTEXT("BuildPatchInstallError_OutOfDiskSpace", "Please free up some disk space and try again."));
 
 	FScopeLock ScopeLock( &ThreadLock );
 	FText ErrorText = FText::GetEmpty();
@@ -70,6 +72,7 @@ const FText FBuildPatchInstallError::GetErrorText()
 		case EBuildPatchInstallError::FileConstructionFail: ErrorText = FileConstructionFail; break;
 		case EBuildPatchInstallError::MoveFileToInstall: ErrorText = MoveFileToInstall; break;
 		case EBuildPatchInstallError::PathLengthExceeded: ErrorText = PathLengthExceeded; break;
+		case EBuildPatchInstallError::OutOfDiskSpace: ErrorText = OutOfDiskSpace; break;
 	}
 	return FText::Format(LOCTEXT("BuildPatchInstallLongError", "{0} {1}"), GetShortErrorText(), ErrorText);
 }
@@ -87,6 +90,7 @@ const FText& FBuildPatchInstallError::GetShortErrorText()
 	static const FText PrerequisiteError(LOCTEXT("BuildPatchInstallShortError_PrerequisiteError", "Prerequisites install failed."));
 	static const FText InitializationError(LOCTEXT("BuildPatchInstallShortError_InitializationError", "The installer failed to initialize."));
 	static const FText PathLengthExceeded(LOCTEXT("BuildPatchInstallShortError_PathLengthExceeded", "Maximum path length exceeded."));
+	static const FText OutOfDiskSpace(LOCTEXT("BuildPatchInstallShortError_OutOfDiskSpace", "Not enough disk space available."));
 	static const FText InvalidOrMax(LOCTEXT("BuildPatchInstallShortError_InvalidOrMax", "An unknown error ocurred."));
 
 	FScopeLock ScopeLock(&ThreadLock);
@@ -103,6 +107,7 @@ const FText& FBuildPatchInstallError::GetShortErrorText()
 		case EBuildPatchInstallError::PrerequisiteError: return PrerequisiteError;
 		case EBuildPatchInstallError::InitializationError: return InitializationError;
 		case EBuildPatchInstallError::PathLengthExceeded: return PathLengthExceeded;
+		case EBuildPatchInstallError::OutOfDiskSpace: return OutOfDiskSpace;
 		default: return InvalidOrMax;
 	}
 }
@@ -143,6 +148,7 @@ const FString& FBuildPatchInstallError::ToString( const EBuildPatchInstallError:
 	static const FString PrerequisiteError( "EBuildPatchInstallError::PrerequisiteError" );
 	static const FString InitializationError("EBuildPatchInstallError::InitializationError");
 	static const FString PathLengthExceeded("EBuildPatchInstallError::PathLengthExceeded");
+	static const FString OutOfDiskSpace("EBuildPatchInstallError::OutOfDiskSpace");
 	static const FString InvalidOrMax( "EBuildPatchInstallError::InvalidOrMax" );
 
 	FScopeLock ScopeLock( &ThreadLock );
@@ -159,6 +165,7 @@ const FString& FBuildPatchInstallError::ToString( const EBuildPatchInstallError:
 		case EBuildPatchInstallError::PrerequisiteError: return PrerequisiteError;
 		case EBuildPatchInstallError::InitializationError: return InitializationError;
 		case EBuildPatchInstallError::PathLengthExceeded: return PathLengthExceeded;
+		case EBuildPatchInstallError::OutOfDiskSpace: return OutOfDiskSpace;
 		default: return InvalidOrMax;
 	}
 }
diff --git a/Engine/Source/Runtime/Online/BuildPatchServices/Private/BuildPatchError.h b/Engine/Source/Runtime/Online/BuildPatchServices/Private/BuildPatchError.h
--- a/Engine/Source/Runtime/Online/BuildPatchServices/Private/BuildPatchError.h
+++ b/Engine/Source/Runtime/Online/BuildPatchServices/Private/BuildPatchError.h
@@ -47,7 +47,10 @@ namespace EBuildPatchInstallError
 		InitializationError = 9,
 
 		// An error occurred creating a file due to excessive path length
-		PathLengthExceeded = 10
+		PathLengthExceeded = 10,
+
+		// An error occurred creating a file due to their not being enough space left on the disk
+		OutOfDiskSpace = 11
 	};
 };
 
diff --git a/Engine/Source/Runtime/Online/BuildPatchServices/Private/BuildPatchFileConstructor.cpp b/Engine/Source/Runtime/Online/BuildPatchServices/Private/BuildPatchFileConstructor.cpp
--- a/Engine/Source/Runtime/Online/BuildPatchServices/Private/BuildPatchFileConstructor.cpp
+++ b/Engine/Source/Runtime/Online/BuildPatchServices/Private/BuildPatchFileConstructor.cpp
@@ -459,11 +459,33 @@ bool FBuildPatchFileConstructor::ConstructFileFromChunks( const FString& Filenam
 		}
 		else
 		{
-			FBuildPatchAnalytics::RecordConstructionError( Filename, FPlatformMisc::GetLastError(), TEXT( "Could Not Create File" ) );
-			ErrorString = TEXT( "Could not create new file " );
-			ErrorString += Filename;
-			GWarn->Logf( TEXT( "BuildPatchFileConstructor: ERROR: %s" ), *ErrorString );
-			FBuildPatchInstallError::SetFatalError( EBuildPatchInstallError::FileConstructionFail, ErrorString );
+			// Check drive space
+			bool bError = false;
+			uint64 TotalSize = 0;
+			uint64 AvailableSpace = 0;
+			if (FPlatformMisc::GetDiskTotalAndFreeSpace(InstallDirectory, TotalSize, AvailableSpace))
+			{
+				const int64 DriveSpace = AvailableSpace;
+				const int64 RequiredSpace = FileManifest->GetFileSize();
+				if (DriveSpace < RequiredSpace)
+				{
+					bError = true;
+					FBuildPatchAnalytics::RecordConstructionError(Filename, FPlatformMisc::GetLastError(), TEXT("Not Enough Disk Space"));
+					ErrorString = TEXT("Not enough disk space for new file ");
+					ErrorString += Filename;
+					GWarn->Logf(TEXT("BuildPatchFileConstructor: ERROR: %s"), *ErrorString);
+					FBuildPatchInstallError::SetFatalError(EBuildPatchInstallError::OutOfDiskSpace, ErrorString);
+				}
+			}
+
+			if (!bError)
+			{
+				FBuildPatchAnalytics::RecordConstructionError(Filename, FPlatformMisc::GetLastError(), TEXT("Could Not Create File"));
+				ErrorString = TEXT("Could not create new file ");
+				ErrorString += Filename;
+				GWarn->Logf(TEXT("BuildPatchFileConstructor: ERROR: %s"), *ErrorString);
+				FBuildPatchInstallError::SetFatalError(EBuildPatchInstallError::FileConstructionFail, ErrorString);
+			}
 		}
 	}
 	else
diff --git a/Engine/Source/Runtime/Online/BuildPatchServices/Private/BuildPatchInstaller.cpp b/Engine/Source/Runtime/Online/BuildPatchServices/Private/BuildPatchInstaller.cpp
--- a/Engine/Source/Runtime/Online/BuildPatchServices/Private/BuildPatchInstaller.cpp
+++ b/Engine/Source/Runtime/Online/BuildPatchServices/Private/BuildPatchInstaller.cpp
@@ -329,6 +329,21 @@ bool FBuildPatchInstaller::RunInstallation(TArray<FString>& CorruptFiles)
 		}
 	}
 
+	// Check drive space
+	uint64 TotalSize = 0;
+	uint64 AvailableSpace = 0;
+	if (FPlatformMisc::GetDiskTotalAndFreeSpace(InstallDirectory, TotalSize, AvailableSpace))
+	{
+		const int64 DriveSpace = AvailableSpace;
+		const int64 RequiredSpace = NewBuildManifest->GetFileSize(FilesToConstruct);
+		if (DriveSpace < RequiredSpace)
+		{
+			GWarn->Logf(TEXT("BuildPatchServices: ERROR: Could not begin install due to their not being enough HDD space. Needs %db, Free %db"), RequiredSpace, DriveSpace);
+			FBuildPatchInstallError::SetFatalError(EBuildPatchInstallError::OutOfDiskSpace);
+			return false;
+		}
+	}
+
 	// Create the downloader
 	FBuildPatchDownloader::Create(DataStagingDir, NewBuildManifest, &BuildProgress);
 
