Merging part of CL#2641230 LauncherDev->Main
change summary:
- Added new module to drive purchase flow UI.
- Started to flesh out new interface IOnlinePurchase for handling IAP transactions.

diff --git a/Engine/Source/Runtime/Online/OnlineSubsystem/Public/Interfaces/OnlinePurchaseInterface.h b/Engine/Source/Runtime/Online/OnlineSubsystem/Public/Interfaces/OnlinePurchaseInterface.h
--- a/Engine/Source/Runtime/Online/OnlineSubsystem/Public/Interfaces/OnlinePurchaseInterface.h
+++ b/Engine/Source/Runtime/Online/OnlineSubsystem/Public/Interfaces/OnlinePurchaseInterface.h
@@ -14,19 +14,29 @@ public:
 	/**
 	 * Add a offer entry for purchase
 	 *
-	 * @param OfferId id of offer to be purchased
-	 * @param Quantity number to purchase
+	 * @param InNamespace namespace of offer to be purchased
+	 * @param InOfferId id of offer to be purchased
+	 * @param InQuantity number to purchase
 	 */
-	void AddPurchaseOffer(const FUniqueOfferId& OfferId, int32 Quantity);
-
+	void AddPurchaseOffer(const FOfferNamespace& InNamespace, const FUniqueOfferId& InOfferId, int32 InQuantity)
+	{ 
+		PurchaseOffers.Add(FPurchaseOfferEntry(InNamespace, InOfferId, InQuantity)); 
+	}
 	/**
 	 * Single offer entry for purchase 
 	 */
 	struct FPurchaseOfferEntry
 	{
+		FPurchaseOfferEntry(const FOfferNamespace& InOfferNamespace, const FUniqueOfferId& InOfferId, int32 InQuantity)
+			: OfferNamespace(InOfferNamespace)
+			, OfferId(InOfferId)
+			, Quantity(InQuantity)
+		{ }
+
+		FOfferNamespace OfferNamespace;
 		FUniqueOfferId OfferId;
 		int32 Quantity;
-	};
+	};	
 	/** List of offers being purchased */
 	TArray<FPurchaseOfferEntry> PurchaseOffers;
 };
@@ -58,17 +68,46 @@ namespace EPurchaseTransactionState
  */
 class FPurchaseReceipt
 {
+public:
+	/**
+	 * Add a offer entry that has been purchased
+	 *
+	 * @param InNamespace of the offer that has been purchased
+	 * @param InOfferId id of offer that has been purchased
+	 * @param InQuantity number purchased
+	 */
+	void AddReceiptOffer(const FOfferNamespace& InNamespace, const FUniqueOfferId& InOfferId, int32 InQuantity)
+	{ 
+		ReceiptOffers.Add(FReceiptOfferEntry(InNamespace, InOfferId, InQuantity)); 
+	}
+
 public:
 	/** Unique Id for this transaction/order */
 	FString TransactionId;
 	/** Current state of the purchase */
 	EPurchaseTransactionState::Type TransactionState;
-	/** List of offers being purchased */
-	TArray<FUniqueOfferId> OfferIds;
 	/** The entitlements obtained from the completed purchase */
 	TArray<FUniqueEntitlementId> EntitlementIds;
 	/** Receipt validation data */
 	FString ReceiptValidation;
+
+	/**
+	 * Single purchased offer offer
+	 */
+	struct FReceiptOfferEntry
+	{
+		FReceiptOfferEntry(const FOfferNamespace& InNamespace, const FUniqueOfferId& InOfferId, int32 InQuantity)
+			: Namespace(InNamespace)
+			, OfferId(InOfferId)
+			, Quantity(InQuantity)
+		{ }
+
+		FOfferNamespace Namespace;
+		FUniqueOfferId OfferId;
+		int32 Quantity;
+	};
+	/** List of offers that were purchased */
+	TArray<FReceiptOfferEntry> ReceiptOffers;
 };
 
 /**
diff --git a/Engine/Source/Runtime/Online/OnlineSubsystem/Public/Interfaces/OnlineStoreInterfaceV2.h b/Engine/Source/Runtime/Online/OnlineSubsystem/Public/Interfaces/OnlineStoreInterfaceV2.h
--- a/Engine/Source/Runtime/Online/OnlineSubsystem/Public/Interfaces/OnlineStoreInterfaceV2.h
+++ b/Engine/Source/Runtime/Online/OnlineSubsystem/Public/Interfaces/OnlineStoreInterfaceV2.h
@@ -4,6 +4,7 @@
 #include "OnlineDelegateMacros.h"
 
 typedef FString FUniqueOfferId;
+typedef FString FOfferNamespace;
 typedef FString FUniqueCategoryId;
 
 /**
