Merging using UE4-To-UE4-LauncherDev

#http,curl
- disable printing of "Authorization" header info to log

diff --git a/Engine/Source/Runtime/Online/HTTP/Private/Curl/CurlHttp.cpp b/Engine/Source/Runtime/Online/HTTP/Private/Curl/CurlHttp.cpp
--- a/Engine/Source/Runtime/Online/HTTP/Private/Curl/CurlHttp.cpp
+++ b/Engine/Source/Runtime/Online/HTTP/Private/Curl/CurlHttp.cpp
@@ -566,7 +566,10 @@ bool FCurlHttpRequest::StartRequest()
 	const int32 NumAllHeaders = AllHeaders.Num();
 	for (int32 Idx = 0; Idx < NumAllHeaders; ++Idx)
 	{
-		UE_LOG(LogHttp, Verbose, TEXT("%p: Adding header '%s'"), this, *AllHeaders[Idx] );
+		if (!AllHeaders[Idx].Contains(TEXT("Authorization")))
+		{
+			UE_LOG(LogHttp, Verbose, TEXT("%p: Adding header '%s'"), this, *AllHeaders[Idx]);
+		}
 		HeaderList = curl_slist_append(HeaderList, TCHAR_TO_UTF8(*AllHeaders[Idx]));
 	}
 
