Merging CL# 2644212, 2644586 LauncherDev->Main

#OnlinePurchase
- Rearranging fields in FPurchaseReceipt
- Changing web purchase flow api

diff --git a/Engine/Source/Runtime/Online/OnlineSubsystem/Public/Interfaces/OnlinePurchaseInterface.h b/Engine/Source/Runtime/Online/OnlineSubsystem/Public/Interfaces/OnlinePurchaseInterface.h
--- a/Engine/Source/Runtime/Online/OnlineSubsystem/Public/Interfaces/OnlinePurchaseInterface.h
+++ b/Engine/Source/Runtime/Online/OnlineSubsystem/Public/Interfaces/OnlinePurchaseInterface.h
@@ -69,6 +69,34 @@ namespace EPurchaseTransactionState
 class FPurchaseReceipt
 {
 public:
+	struct FLineItemInfo
+	{
+		/** unique identifier representing this purchased item (the specific instance owned by this account) */
+		FUniqueEntitlementId UniqueId;
+
+		/** platform-specific opaque validation info (required to verify UniqueId belongs to this account) */
+		FString ValidationInfo;
+	};
+
+	/**
+	 * Single purchased offer offer
+	 */
+	struct FReceiptOfferEntry
+	{
+		FReceiptOfferEntry(const FOfferNamespace& InNamespace, const FUniqueOfferId& InOfferId, int32 InQuantity)
+			: Namespace(InNamespace)
+			, OfferId(InOfferId)
+			, Quantity(InQuantity)
+		{ }
+
+		FOfferNamespace Namespace;
+		FUniqueOfferId OfferId;
+		int32 Quantity;
+
+		/** Information about the individual items purchased */
+		TArray<FLineItemInfo> LineItems;
+	};
+
 	/**
 	 * Add a offer entry that has been purchased
 	 *
@@ -81,31 +109,17 @@ public:
 		ReceiptOffers.Add(FReceiptOfferEntry(InNamespace, InOfferId, InQuantity)); 
 	}
 
+	void AddReceiptOffer(const FReceiptOfferEntry& ReceiptOffer)
+	{ 
+		ReceiptOffers.Add(ReceiptOffer); 
+	}
+
 public:
 	/** Unique Id for this transaction/order */
 	FString TransactionId;
 	/** Current state of the purchase */
 	EPurchaseTransactionState::Type TransactionState;
-	/** The entitlements obtained from the completed purchase */
-	TArray<FUniqueEntitlementId> EntitlementIds;
-	/** Receipt validation data */
-	FString ReceiptValidation;
 
-	/**
-	 * Single purchased offer offer
-	 */
-	struct FReceiptOfferEntry
-	{
-		FReceiptOfferEntry(const FOfferNamespace& InNamespace, const FUniqueOfferId& InOfferId, int32 InQuantity)
-			: Namespace(InNamespace)
-			, OfferId(InOfferId)
-			, Quantity(InQuantity)
-		{ }
-
-		FOfferNamespace Namespace;
-		FUniqueOfferId OfferId;
-		int32 Quantity;
-	};
 	/** List of offers that were purchased */
 	TArray<FReceiptOfferEntry> ReceiptOffers;
 };
