[INTEGRATE] Add Image format detection using file magic headers and adding ETag support to ImageDownload module.
Merging CL#2643429 and CL#2645040 using UE4-To-UE4-LauncherDev

diff --git a/Engine/Source/Developer/ImageWrapper/Private/ImageWrapperModule.cpp b/Engine/Source/Developer/ImageWrapper/Private/ImageWrapperModule.cpp
--- a/Engine/Source/Developer/ImageWrapper/Private/ImageWrapperModule.cpp
+++ b/Engine/Source/Developer/ImageWrapper/Private/ImageWrapperModule.cpp
@@ -6,6 +6,33 @@
 DEFINE_LOG_CATEGORY(LogImageWrapper);
 
 
+namespace {
+
+	static const uint8 IMAGE_MAGIC_PNG[]  = {0x89, 0x50, 0x4E, 0x47, 0x0D, 0x0A, 0x1A, 0x0A};
+	static const uint8 IMAGE_MAGIC_JPEG[] = {0xFF, 0xD8, 0xFF};
+	static const uint8 IMAGE_MAGIC_BMP[]  = {0x42, 0x4D};
+	static const uint8 IMAGE_MAGIC_ICO[]  = {0x00, 0x00, 0x01, 0x00};
+	static const uint8 IMAGE_MAGIC_EXR[]  = {0x76, 0x2F, 0x31, 0x01};
+	static const uint8 IMAGE_MAGIC_ICNS[] = {0x69, 0x63, 0x6E, 0x73};
+
+	/** Internal helper function to verify image signature. */
+	template <int32 MagicCount> bool StartsWith(const uint8* Content, int32 ContentSize, const uint8 (&Magic)[MagicCount])
+	{
+		if (ContentSize < MagicCount)
+		{
+			return false;
+		}
+		for (int32 I = 0; I < MagicCount; ++I)
+		{
+			if ( Content[I] != Magic[I] )
+			{
+				return false;
+			}
+		}
+		return true;
+	}
+}
+
 /**
  * Image Wrapper module
  */
@@ -59,6 +86,37 @@ public:
 		return MakeShareable(ImageWrapper);
 	}
 
+	virtual EImageFormat::Type DetectImageFormat( const void* CompressedData, int32 CompressedSize ) override
+	{
+		EImageFormat::Type Format = EImageFormat::Invalid;
+		if (StartsWith((uint8*) CompressedData, CompressedSize, IMAGE_MAGIC_PNG))
+		{
+			Format = EImageFormat::PNG;
+		}
+		else if (StartsWith((uint8*) CompressedData, CompressedSize, IMAGE_MAGIC_JPEG))
+		{
+			Format = EImageFormat::JPEG; // @Todo: Should we detect grayscale vs non-grayscale?
+		}
+		else if (StartsWith((uint8*) CompressedData, CompressedSize, IMAGE_MAGIC_BMP))
+		{
+			Format = EImageFormat::BMP;
+		}
+		else if (StartsWith((uint8*) CompressedData, CompressedSize, IMAGE_MAGIC_ICO))
+		{
+			Format = EImageFormat::ICO;
+		}
+		else if (StartsWith((uint8*) CompressedData, CompressedSize, IMAGE_MAGIC_EXR))
+		{
+			Format = EImageFormat::EXR;
+		}
+		else if (StartsWith((uint8*) CompressedData, CompressedSize, IMAGE_MAGIC_ICNS))
+		{
+			Format = EImageFormat::ICNS;
+		}
+
+		return Format;
+	}
+
 public:
 
 	// IModuleInterface interface
diff --git a/Engine/Source/Developer/ImageWrapper/Public/Interfaces/IImageWrapper.h b/Engine/Source/Developer/ImageWrapper/Public/Interfaces/IImageWrapper.h
--- a/Engine/Source/Developer/ImageWrapper/Public/Interfaces/IImageWrapper.h
+++ b/Engine/Source/Developer/ImageWrapper/Public/Interfaces/IImageWrapper.h
@@ -10,8 +10,11 @@ namespace EImageFormat
 	 */
 	enum Type
 	{
+		/** Invalid or unrecognized format */
+		Invalid = -1,
+
 		/** Portable Network Graphics */
-		PNG,
+		PNG = 0,
 
 		/** Joint Photographic Experts Group */
 		JPEG,
diff --git a/Engine/Source/Developer/ImageWrapper/Public/Interfaces/IImageWrapperModule.h b/Engine/Source/Developer/ImageWrapper/Public/Interfaces/IImageWrapperModule.h
--- a/Engine/Source/Developer/ImageWrapper/Public/Interfaces/IImageWrapperModule.h
+++ b/Engine/Source/Developer/ImageWrapper/Public/Interfaces/IImageWrapperModule.h
@@ -19,6 +19,16 @@ public:
 	 */
 	virtual IImageWrapperPtr CreateImageWrapper( const EImageFormat::Type InFormat ) = 0;
 
+	/**
+	 * Detect image format by looking at the first few bytes of the compressed image data.
+	 * You can call this method as soon as you have 8-16 bytes of compressed file content available.
+	 *
+	 * @param InCompressedData - The raw image header.
+	 * @param InCompressedSize - The size of InCompressedData.
+	 * @return the detected format or EImageFormat::Invalid if the method could not detect the image format.
+	 */
+	virtual EImageFormat::Type DetectImageFormat( const void* InCompressedData, int32 InCompressedSize) = 0;
+
 public:
 
 	/**
diff --git a/Engine/Source/Runtime/Online/ImageDownload/Private/WebImage.cpp b/Engine/Source/Runtime/Online/ImageDownload/Private/WebImage.cpp
--- a/Engine/Source/Runtime/Online/ImageDownload/Private/WebImage.cpp
+++ b/Engine/Source/Runtime/Online/ImageDownload/Private/WebImage.cpp
@@ -3,11 +3,6 @@
 #include "ImageDownloadPCH.h"
 #include "WebImage.h"
 
-static const FString IMAGE_MIME_PNG(TEXT("image/png"));
-static const FString IMAGE_MIME_PNG_ALT(TEXT("image/x-png"));
-static const FString IMAGE_MIME_JPEG(TEXT("image/jpeg"));
-static const FString IMAGE_MIME_BMP(TEXT("image/bmp"));
-
 FWebImage::FWebImage()
 : StandInBrush(FCoreStyle::Get().GetDefaultBrush())
 , bDownloadSuccess(false)
@@ -19,7 +14,7 @@ TAttribute< const FSlateBrush* > FWebImage::Attr() const
 	return TAttribute< const FSlateBrush* >(AsShared(), &FWebImage::GetBrush);
 }
 
-bool FWebImage::BeginDownload(const FString& UrlIn, const FOnImageDownloaded& DownloadCb)
+bool FWebImage::BeginDownload(const FString& UrlIn, const TOptional<FString>& StandInETag, const FOnImageDownloaded& DownloadCb)
 {
 	CancelDownload();
 
@@ -30,9 +25,14 @@ bool FWebImage::BeginDownload(const FString& UrlIn, const FOnImageDownloaded& Do
 	TSharedRef<IHttpRequest> HttpRequest = FHttpModule::Get().CreateRequest();
 	HttpRequest->SetVerb(TEXT("GET"));
 	HttpRequest->SetURL(Url);
-	HttpRequest->SetHeader(TEXT("Accept"), FString::Printf(TEXT("%s, %s, %s; q=0.8, %s; q=0.5"), *IMAGE_MIME_PNG, *IMAGE_MIME_PNG_ALT, *IMAGE_MIME_JPEG, *IMAGE_MIME_BMP));
+	HttpRequest->SetHeader(TEXT("Accept"), TEXT("image/png, image/x-png, image/jpeg; q=0.8, image/vnd.microsoft.icon, image/x-icon, image/bmp, image/*; q=0.5, image/webp; q=0.0"));
 	HttpRequest->OnProcessRequestComplete().BindSP(this, &FWebImage::HttpRequestComplete, DownloadCb);
 
+	if (StandInETag.IsSet())
+	{
+		HttpRequest->SetHeader(TEXT("If-None-Match"), *StandInETag.GetValue());
+	}
+
 	// queue the request
 	PendingRequest = HttpRequest;
 	if (!HttpRequest->ProcessRequest())
@@ -73,50 +73,47 @@ bool FWebImage::ProcessHttpResponse(const FString& RequestUrl, FHttpResponsePtr
 		return false;
 	}
 
+	ETag = HttpResponse->GetHeader("ETag");
+
 	// check status code
 	int32 StatusCode = HttpResponse->GetResponseCode();
 	if (StatusCode / 100 != 2)
 	{
+		if ( StatusCode == 304)
+		{
+			// Not modified means that the image is identical to the placeholder image.
+			return true;
+		}
 		UE_LOG(LogImageDownload, Error, TEXT("Image Download: HTTP response %d. url=%s"), StatusCode, *RequestUrl);
 		return false;
 	}
 
-	// get the content type and figure out the image format from it
-	FString ContentType = HttpResponse->GetContentType();
-	EImageFormat::Type ImageFormat;
-	if (ContentType.Equals(IMAGE_MIME_PNG, ESearchCase::IgnoreCase) || ContentType.Equals(IMAGE_MIME_PNG_ALT, ESearchCase::IgnoreCase))
-	{
-		ImageFormat = EImageFormat::PNG;
-	}
-	else if (ContentType.Equals(IMAGE_MIME_JPEG, ESearchCase::IgnoreCase))
-	{
-		ImageFormat = EImageFormat::JPEG;
-	}
-	else if (ContentType.Equals(IMAGE_MIME_BMP, ESearchCase::IgnoreCase))
-	{
-		ImageFormat = EImageFormat::BMP;
-	}
-	else
+	// build an image wrapper for this type
+	static const FName MODULE_IMAGE_WRAPPER("ImageWrapper");
+	IImageWrapperModule& ImageWrapperModule = FModuleManager::LoadModuleChecked<IImageWrapperModule>(MODULE_IMAGE_WRAPPER);
+
+	// Look at the signature of the downloaded image to detect image type. (and ignore the content type header except for error reporting)
+	const TArray<uint8>& Content = HttpResponse->GetContent();
+	EImageFormat::Type ImageFormat = ImageWrapperModule.DetectImageFormat(Content.GetData(), Content.Num());
+
+	if (ImageFormat == EImageFormat::Invalid)
 	{
-		UE_LOG(LogImageDownload, Error, TEXT("Image Download: Invalid Content-Type '%s'. url=%s"), *ContentType, *RequestUrl);
+		FString ContentType = HttpResponse->GetContentType();
+		UE_LOG(LogImageDownload, Error, TEXT("Image Download: Could not recognize file type of image downloaded from url %s, server-reported content type: %s"), *RequestUrl, *ContentType);
 		return false;
 	}
 
-	// build an image wrapper for this type
-	static const FName MODULE_IMAGE_WRAPPER("ImageWrapper");
-	IImageWrapperModule& ImageWrapperModule = FModuleManager::LoadModuleChecked<IImageWrapperModule>(MODULE_IMAGE_WRAPPER);
 	IImageWrapperPtr ImageWrapper = ImageWrapperModule.CreateImageWrapper(ImageFormat);
 	if (!ImageWrapper.IsValid())
 	{
-		UE_LOG(LogImageDownload, Error, TEXT("Image Download: Unable to make image wrapper for %s"), *ContentType);
+		UE_LOG(LogImageDownload, Error, TEXT("Image Download: Unable to make image wrapper for image format %d"), ImageFormat);
 		return false;
 	}
 
 	// parse the content
-	const TArray<uint8>& Content = HttpResponse->GetContent();
 	if (!ImageWrapper->SetCompressed(Content.GetData(), Content.Num()))
 	{
-		UE_LOG(LogImageDownload, Error, TEXT("Image Download: Unable to parse %s from %s"), *ContentType, *RequestUrl);
+		UE_LOG(LogImageDownload, Error, TEXT("Image Download: Unable to parse image format %d from %s"), ImageFormat, *RequestUrl);
 		return false;
 	}
 
@@ -124,7 +121,7 @@ bool FWebImage::ProcessHttpResponse(const FString& RequestUrl, FHttpResponsePtr
 	const TArray<uint8>* RawImageData = nullptr;
 	if (!ImageWrapper->GetRaw(ERGBFormat::BGRA, 8, RawImageData) || RawImageData == nullptr)
 	{
-		UE_LOG(LogImageDownload, Error, TEXT("Image Download: Unable to convert %s to BGRA 8"), *ContentType);
+		UE_LOG(LogImageDownload, Error, TEXT("Image Download: Unable to convert image format %d to BGRA 8"), ImageFormat);
 		return false;
 	}
 
diff --git a/Engine/Source/Runtime/Online/ImageDownload/Private/WebImageCache.cpp b/Engine/Source/Runtime/Online/ImageDownload/Private/WebImageCache.cpp
--- a/Engine/Source/Runtime/Online/ImageDownload/Private/WebImageCache.cpp
+++ b/Engine/Source/Runtime/Online/ImageDownload/Private/WebImageCache.cpp
@@ -8,8 +8,24 @@ FWebImageCache::FWebImageCache()
 {
 }
 
-TSharedRef<const FWebImage> FWebImageCache::Download(const FString& Url)
+TSharedRef<const FWebImage> FWebImageCache::Download(const FString& Url, const TOptional<FString>& DefaultImageUrl)
 {
+	TAttribute<const FSlateBrush*> StandInBrush;
+	TOptional<FString> StandInEtag;
+	// If the optional DefaultImageUrl is set, use its brush as the stand-in image, falling back to the default brush if not downloaded yet.
+	if (DefaultImageUrl.IsSet())
+	{
+		TSharedRef<const FWebImage> DefaultImage = Download(DefaultImageUrl.GetValue());
+		StandInBrush = DefaultImage->Attr();
+
+		// If already downloaded, the etag will prevent further downloads, if the current URLs contents are identical to the default image.
+		StandInEtag = DefaultImage->GetETag();
+	}
+	else
+	{
+		StandInBrush = DefaultStandInBrush;
+	}
+
 	// canonicalize URL (we don't currently have code to do this so just treat the URL as opaque)
 	const FString& CanonicalUrl = Url;
 
@@ -31,8 +47,8 @@ TSharedRef<const FWebImage> FWebImageCache::Download(const FString& Url)
 		// if it is done and we failed, and it's being requested again, queue up another try
 		if (ImagePtr->DidDownloadFail())
 		{
-			ImagePtr->SetStandInBrush(DefaultStandInBrush);
-			ImagePtr->BeginDownload(ImagePtr->GetUrl());
+			ImagePtr->SetStandInBrush(StandInBrush);
+			ImagePtr->BeginDownload(ImagePtr->GetUrl(), StandInEtag);
 		}
 
 		// return the image ptr
@@ -43,8 +59,8 @@ TSharedRef<const FWebImage> FWebImageCache::Download(const FString& Url)
 
 	// make a new one
 	TSharedRef<FWebImage> WebImage = MakeShareable(new FWebImage());
-	WebImage->SetStandInBrush(DefaultStandInBrush);
-	WebImage->BeginDownload(CanonicalUrl);
+	WebImage->SetStandInBrush(StandInBrush);
+	WebImage->BeginDownload(CanonicalUrl, StandInEtag);
 
 	// add it to the cache
 	StrongRefCache.Add(WebImage);
diff --git a/Engine/Source/Runtime/Online/ImageDownload/Public/WebImage.h b/Engine/Source/Runtime/Online/ImageDownload/Public/WebImage.h
--- a/Engine/Source/Runtime/Online/ImageDownload/Public/WebImage.h
+++ b/Engine/Source/Runtime/Online/ImageDownload/Public/WebImage.h
@@ -24,7 +24,7 @@ typedef TSharedPtr<class IHttpResponse, ESPMode::ThreadSafe> FHttpResponsePtr;
  *   ImageCache.SetStandInBrush(FStyle::Get()->GetBrush("Foo"));
  *
  *   SNew(SImage)
- *   .Image(ImageCache.Download(Url).Attr())
+ *   .Image(ImageCache.Download(Url)->Attr())
  */
 class IMAGEDOWNLOAD_API FWebImage 
 	: public TSharedFromThis<FWebImage>
@@ -42,7 +42,13 @@ public:
 	FORCEINLINE FWebImage& SetStandInBrush(TAttribute<const FSlateBrush*> StandInBrushIn) { StandInBrush = StandInBrushIn; DownloadedBrush.Reset(); return *this; }
 
 	/** Begin downloading an image. This will automatically set the current brush to the downloaded image when it completes (if successful) */
-	bool BeginDownload(const FString& Url, const FOnImageDownloaded& DownloadCallback = FOnImageDownloaded());
+	bool BeginDownload(const FString& InUrl, const TOptional<FString>& StandInETag = TOptional<FString>(), const FOnImageDownloaded& DownloadCallback = FOnImageDownloaded());
+
+	/** Begin downloading an image. */
+	FORCEINLINE bool BeginDownload(const FString& InUrl, const FOnImageDownloaded& DownloadCallback)
+	{
+		return BeginDownload(InUrl, TOptional<FString>(), DownloadCallback);
+	}
 
 	/** Cancel any download in progress */
 	void CancelDownload();
@@ -70,6 +76,9 @@ public:
 	/** What URL was requested */
 	FORCEINLINE const FString& GetUrl() const { return Url; }
 
+	/** What is the ETag of the downloaded resource */
+	FORCEINLINE const TOptional<FString>& GetETag() const { return ETag; }
+
 private:
 	/** request complete callback */
 	void HttpRequestComplete(FHttpRequestPtr HttpRequest, FHttpResponsePtr HttpResponse, bool bSucceeded, FOnImageDownloaded DownloadCb);
@@ -93,4 +102,7 @@ private:
 
 	/** When did the download complete */
 	FDateTime DownloadTimeUtc;
+
+	/** The ETag of the downloaded image */
+	TOptional<FString> ETag;
 };
\ No newline at end of file
diff --git a/Engine/Source/Runtime/Online/ImageDownload/Public/WebImageCache.h b/Engine/Source/Runtime/Online/ImageDownload/Public/WebImageCache.h
--- a/Engine/Source/Runtime/Online/ImageDownload/Public/WebImageCache.h
+++ b/Engine/Source/Runtime/Online/ImageDownload/Public/WebImageCache.h
@@ -16,7 +16,7 @@ public:
 	FWebImageCache();
 
 	/** Find or create a WebImage object for this URL (you probably just want to call ->Attr() on this) */
-	TSharedRef<const FWebImage> Download(const FString& Url);
+	TSharedRef<const FWebImage> Download(const FString& Url, const TOptional<FString>& DefaultImageUrl = TOptional<FString>());
 
 	/** Set the brush that will be returned until the download completes (only affects future downloads). */
 	FORCEINLINE void SetDefaultStandInBrush(TAttribute<const FSlateBrush*> StandInBrushIn) { DefaultStandInBrush = StandInBrushIn; }
