#UE4 - fixed integration compile error

diff --git a/Engine/Source/Runtime/Online/ImageDownload/Private/WebImage.cpp b/Engine/Source/Runtime/Online/ImageDownload/Private/WebImage.cpp
--- a/Engine/Source/Runtime/Online/ImageDownload/Private/WebImage.cpp
+++ b/Engine/Source/Runtime/Online/ImageDownload/Private/WebImage.cpp
@@ -106,14 +106,14 @@ bool FWebImage::ProcessHttpResponse(const FString& RequestUrl, FHttpResponsePtr
 	IImageWrapperPtr ImageWrapper = ImageWrapperModule.CreateImageWrapper(ImageFormat);
 	if (!ImageWrapper.IsValid())
 	{
-		UE_LOG(LogImageDownload, Error, TEXT("Image Download: Unable to make image wrapper for image format %d"), ImageFormat);
+		UE_LOG(LogImageDownload, Error, TEXT("Image Download: Unable to make image wrapper for image format %d"), (int32)ImageFormat);
 		return false;
 	}
 
 	// parse the content
 	if (!ImageWrapper->SetCompressed(Content.GetData(), Content.Num()))
 	{
-		UE_LOG(LogImageDownload, Error, TEXT("Image Download: Unable to parse image format %d from %s"), ImageFormat, *RequestUrl);
+		UE_LOG(LogImageDownload, Error, TEXT("Image Download: Unable to parse image format %d from %s"), (int32)ImageFormat, *RequestUrl);
 		return false;
 	}
 
@@ -121,7 +121,7 @@ bool FWebImage::ProcessHttpResponse(const FString& RequestUrl, FHttpResponsePtr
 	const TArray<uint8>* RawImageData = nullptr;
 	if (!ImageWrapper->GetRaw(ERGBFormat::BGRA, 8, RawImageData) || RawImageData == nullptr)
 	{
-		UE_LOG(LogImageDownload, Error, TEXT("Image Download: Unable to convert image format %d to BGRA 8"), ImageFormat);
+		UE_LOG(LogImageDownload, Error, TEXT("Image Download: Unable to convert image format %d to BGRA 8"), (int32)ImageFormat);
 		return false;
 	}
 
