#UE4 - moved all http header files into header module
- headers in core and implementations in http were working but only coincidentally
- no way to export any of the platforms if needed because CORE_API was correct for generic platform but would have link errors because all platforms were in HTTP_API
- fixed up platform http calls to reference http module

diff --git a/Engine/Source/Developer/SuperSearch/Private/SSuperSearch.cpp b/Engine/Source/Developer/SuperSearch/Private/SSuperSearch.cpp
--- a/Engine/Source/Developer/SuperSearch/Private/SSuperSearch.cpp
+++ b/Engine/Source/Developer/SuperSearch/Private/SSuperSearch.cpp
@@ -316,9 +316,9 @@ void SSuperSearchBox::OnTextChanged(const FText& InText)
 			TSharedRef<class IHttpRequest> HttpRequest = FHttpModule::Get().CreateRequest();
 
 			// build the url
-				FString UrlEncodedString = FGenericPlatformHttp::UrlEncode(InText.ToString());	//we need to url encode for special characters (especially other languages)
+			FString UrlEncodedString = FPlatformHttp::UrlEncode(InText.ToString());	//we need to url encode for special characters (especially other languages)
 			// use partial response to only look at items and things we care about for items right now (title,link and label name)
-				const FText QueryURL = FText::Format(FText::FromString("https://www.googleapis.com/customsearch/v1?key=AIzaSyCMGfdDaSfjqv5zYoS0mTJnOT3e9MURWkU&cx=009868829633250020713:y7tfd8hlcgg&fields=items(title,link,labels/name)&q={0}+less:forums"), FText::FromString(UrlEncodedString));
+			const FText QueryURL = FText::Format(FText::FromString("https://www.googleapis.com/customsearch/v1?key=AIzaSyCMGfdDaSfjqv5zYoS0mTJnOT3e9MURWkU&cx=009868829633250020713:y7tfd8hlcgg&fields=items(title,link,labels/name)&q={0}+less:forums"), FText::FromString(UrlEncodedString));
 
 			//save http request into map to ensure correct ordering
 			FText & Query = RequestQueryMap.FindOrAdd(HttpRequest);
diff --git a/Engine/Source/Runtime/Analytics/AnalyticsET/Private/AnalyticsET.cpp b/Engine/Source/Runtime/Analytics/AnalyticsET/Private/AnalyticsET.cpp
--- a/Engine/Source/Runtime/Analytics/AnalyticsET/Private/AnalyticsET.cpp
+++ b/Engine/Source/Runtime/Analytics/AnalyticsET/Private/AnalyticsET.cpp
@@ -310,10 +310,10 @@ void FAnalyticsProviderET::FlushEvents()
 		JsonWriter->Close();
 
 		FString URLPath = FString::Printf(TEXT("CollectData.1?SessionID=%s&AppID=%s&AppVersion=%s&UserID=%s"),
-			*FGenericPlatformHttp::UrlEncode(SessionID),
-			*FGenericPlatformHttp::UrlEncode(APIKey),
-			*FGenericPlatformHttp::UrlEncode(AppVersion),
-			*FGenericPlatformHttp::UrlEncode(UserID));
+			*FPlatformHttp::UrlEncode(SessionID),
+			*FPlatformHttp::UrlEncode(APIKey),
+			*FPlatformHttp::UrlEncode(AppVersion),
+			*FPlatformHttp::UrlEncode(UserID));
 
 		// Recreate the URLPath for logging because we do not want to escape the parameters when logging.
 		// We cannot simply UrlEncode the entire Path after logging it because UrlEncode(Params) != UrlEncode(Param1) & UrlEncode(Param2) ...
diff --git a/Engine/Source/Runtime/Core/Private/GenericPlatform/GenericPlatformHttp.cpp b/Engine/Source/Runtime/Core/Private/GenericPlatform/GenericPlatformHttp.cpp
deleted file mode 100644
diff --git a/Engine/Source/Runtime/Core/Public/Android/AndroidHttp.h b/Engine/Source/Runtime/Core/Public/Android/AndroidHttp.h
deleted file mode 100644
diff --git a/Engine/Source/Runtime/Core/Public/Android/AndroidIncludes.h b/Engine/Source/Runtime/Core/Public/Android/AndroidIncludes.h
--- a/Engine/Source/Runtime/Core/Public/Android/AndroidIncludes.h
+++ b/Engine/Source/Runtime/Core/Public/Android/AndroidIncludes.h
@@ -21,7 +21,6 @@
 #include "Android/AndroidSplash.h"
 #include "Android/AndroidFile.h"
 #include "Android/AndroidSurvey.h"
-#include "Android/AndroidHttp.h"
 #include "Android/AndroidAffinity.h"
 
 // include platform properties and typedef it for the runtime
diff --git a/Engine/Source/Runtime/Core/Public/GenericPlatform/GenericPlatformHttp.h b/Engine/Source/Runtime/Core/Public/GenericPlatform/GenericPlatformHttp.h
deleted file mode 100644
diff --git a/Engine/Source/Runtime/Core/Public/HAL/PlatformHttp.h b/Engine/Source/Runtime/Core/Public/HAL/PlatformHttp.h
deleted file mode 100644
diff --git a/Engine/Source/Runtime/Core/Public/HAL/PlatformIncludes.h b/Engine/Source/Runtime/Core/Public/HAL/PlatformIncludes.h
--- a/Engine/Source/Runtime/Core/Public/HAL/PlatformIncludes.h
+++ b/Engine/Source/Runtime/Core/Public/HAL/PlatformIncludes.h
@@ -27,7 +27,6 @@
 #include "GenericPlatform/GenericPlatformSplash.h"
 #include "GenericPlatform/GenericPlatformFile.h"
 #include "GenericPlatform/GenericPlatformSurvey.h"
-#include "GenericPlatform/GenericPlatformHttp.h"
 #include "GenericPlatform/GenericPlatformAffinity.h"
 #include "GenericPlatform/GenericPlatformFramePacer.h"
 
diff --git a/Engine/Source/Runtime/Core/Public/HTML5/HTML5PlatformHttp.h b/Engine/Source/Runtime/Core/Public/HTML5/HTML5PlatformHttp.h
deleted file mode 100644
diff --git a/Engine/Source/Runtime/Core/Public/HTML5/HTML5PlatformIncludes.h b/Engine/Source/Runtime/Core/Public/HTML5/HTML5PlatformIncludes.h
--- a/Engine/Source/Runtime/Core/Public/HTML5/HTML5PlatformIncludes.h
+++ b/Engine/Source/Runtime/Core/Public/HTML5/HTML5PlatformIncludes.h
@@ -20,7 +20,6 @@
 #include "HTML5/HTML5PlatformTLS.h"
 #include "HTML5/HTML5PlatformSplash.h"
 #include "HTML5/HTML5PlatformSurvey.h"
-#include "HTML5/HTML5PlatformHttp.h"
 
 typedef FGenericPlatformRHIFramePacer FPlatformRHIFramePacer;
 
diff --git a/Engine/Source/Runtime/Core/Public/IOS/IOSPlatformHttp.h b/Engine/Source/Runtime/Core/Public/IOS/IOSPlatformHttp.h
deleted file mode 100644
diff --git a/Engine/Source/Runtime/Core/Public/IOS/IOSPlatformIncludes.h b/Engine/Source/Runtime/Core/Public/IOS/IOSPlatformIncludes.h
--- a/Engine/Source/Runtime/Core/Public/IOS/IOSPlatformIncludes.h
+++ b/Engine/Source/Runtime/Core/Public/IOS/IOSPlatformIncludes.h
@@ -22,7 +22,6 @@
 #include "Apple/ApplePlatformFile.h"
 #include "IOS/IOSPlatformSurvey.h"
 #include "IOSAsyncTask.h"
-#include "IOS/IOSPlatformHttp.h"
 #include "IOS/IOSPlatformFramePacer.h"
 #import "IOSPlatformString.h"
 
diff --git a/Engine/Source/Runtime/Core/Public/Linux/LinuxPlatformHttp.h b/Engine/Source/Runtime/Core/Public/Linux/LinuxPlatformHttp.h
deleted file mode 100644
diff --git a/Engine/Source/Runtime/Core/Public/Linux/LinuxPlatformIncludes.h b/Engine/Source/Runtime/Core/Public/Linux/LinuxPlatformIncludes.h
--- a/Engine/Source/Runtime/Core/Public/Linux/LinuxPlatformIncludes.h
+++ b/Engine/Source/Runtime/Core/Public/Linux/LinuxPlatformIncludes.h
@@ -25,7 +25,6 @@
 #include "Linux/LinuxPlatformSplash.h"
 #include "Linux/LinuxPlatformFile.h"
 #include "Linux/LinuxPlatformSurvey.h"
-#include "Linux/LinuxPlatformHttp.h"
 
 typedef FGenericPlatformRHIFramePacer FPlatformRHIFramePacer;
 
diff --git a/Engine/Source/Runtime/Core/Public/Mac/MacPlatformHttp.h b/Engine/Source/Runtime/Core/Public/Mac/MacPlatformHttp.h
deleted file mode 100644
diff --git a/Engine/Source/Runtime/Core/Public/Mac/MacPlatformIncludes.h b/Engine/Source/Runtime/Core/Public/Mac/MacPlatformIncludes.h
--- a/Engine/Source/Runtime/Core/Public/Mac/MacPlatformIncludes.h
+++ b/Engine/Source/Runtime/Core/Public/Mac/MacPlatformIncludes.h
@@ -25,7 +25,6 @@
 #include "Mac/MacPlatformSplash.h"
 #include "Apple/ApplePlatformFile.h"
 #include "Mac/MacPlatformSurvey.h"
-#include "Mac/MacPlatformHttp.h"
 
 typedef FGenericPlatformRHIFramePacer FPlatformRHIFramePacer;
 
diff --git a/Engine/Source/Runtime/Core/Public/WinRT/WinRTHttp.h b/Engine/Source/Runtime/Core/Public/WinRT/WinRTHttp.h
deleted file mode 100644
diff --git a/Engine/Source/Runtime/Core/Public/WinRT/WinRTIncludes.h b/Engine/Source/Runtime/Core/Public/WinRT/WinRTIncludes.h
--- a/Engine/Source/Runtime/Core/Public/WinRT/WinRTIncludes.h
+++ b/Engine/Source/Runtime/Core/Public/WinRT/WinRTIncludes.h
@@ -22,7 +22,6 @@
 #include "WinRT/WinRTTLS.h"
 #include "WinRT/WinRTSplash.h"
 #include "WinRT/WinRTSurvey.h"
-#include "WinRT/WinRTHttp.h"
 
 typedef FGenericPlatformRHIFramePacer FPlatformRHIFramePacer;
 
diff --git a/Engine/Source/Runtime/Core/Public/Windows/WindowsPlatformHttp.h b/Engine/Source/Runtime/Core/Public/Windows/WindowsPlatformHttp.h
deleted file mode 100644
diff --git a/Engine/Source/Runtime/Core/Public/Windows/WindowsPlatformIncludes.h b/Engine/Source/Runtime/Core/Public/Windows/WindowsPlatformIncludes.h
--- a/Engine/Source/Runtime/Core/Public/Windows/WindowsPlatformIncludes.h
+++ b/Engine/Source/Runtime/Core/Public/Windows/WindowsPlatformIncludes.h
@@ -23,7 +23,6 @@
 #include "Windows/WindowsPlatformSplash.h"
 #include "Windows/WindowsPlatformFile.h"
 #include "Windows/WindowsPlatformSurvey.h"
-#include "Windows/WindowsPlatformHttp.h"
 
 
 typedef FGenericPlatformRHIFramePacer FPlatformRHIFramePacer;
diff --git a/Engine/Source/Runtime/Online/HTTP/Private/GenericPlatform/GenericPlatformHttp.cpp b/Engine/Source/Runtime/Online/HTTP/Private/GenericPlatform/GenericPlatformHttp.cpp
new file mode 100644
--- /dev/null
+++ b/Engine/Source/Runtime/Online/HTTP/Private/GenericPlatform/GenericPlatformHttp.cpp
@@ -0,0 +1,172 @@
+// Copyright 1998-2015 Epic Games, Inc. All Rights Reserved.
+
+#include "HTTPPrivatePCH.h"
+#include "Runtime/Online/HTTP/Public/Interfaces/IHttpBase.h"
+#include "Runtime/Online/HTTP/Public/Interfaces/IHttpRequest.h"
+#include "Runtime/Online/HTTP/Public/Interfaces/IHttpResponse.h"
+#include "StringConv.h"
+
+
+/**
+ * A generic http request
+ */
+class FGenericPlatformHttpRequest : public IHttpRequest
+{
+public:
+	// IHttpBase
+	virtual FString GetURL() override { return TEXT(""); }
+	virtual FString GetURLParameter(const FString& ParameterName) override { return TEXT(""); }
+	virtual FString GetHeader(const FString& HeaderName) override { return TEXT(""); }
+	virtual TArray<FString> GetAllHeaders() override { return TArray<FString>(); }
+	virtual FString GetContentType() override { return TEXT(""); }
+	virtual int32 GetContentLength() override { return 0; }
+	virtual const TArray<uint8>& GetContent() override { static TArray<uint8> Temp; return Temp; }
+
+	// IHttpRequest
+	virtual FString GetVerb() override { return TEXT(""); }
+	virtual void SetVerb(const FString& Verb) override {}
+	virtual void SetURL(const FString& URL) override {}
+	virtual void SetContent(const TArray<uint8>& ContentPayload) override {}
+	virtual void SetContentAsString(const FString& ContentString) override {}
+	virtual void SetHeader(const FString& HeaderName, const FString& HeaderValue) override {}
+	virtual bool ProcessRequest() override { return false; }
+	virtual FHttpRequestCompleteDelegate& OnProcessRequestComplete() override { static FHttpRequestCompleteDelegate RequestCompleteDelegate; return RequestCompleteDelegate; }
+	virtual FHttpRequestProgressDelegate& OnRequestProgress() override { static FHttpRequestProgressDelegate RequestProgressDelegate; return RequestProgressDelegate; }
+	virtual void CancelRequest() override {}
+	virtual EHttpRequestStatus::Type GetStatus() override { return EHttpRequestStatus::NotStarted; }
+	virtual const FHttpResponsePtr GetResponse() const override { return nullptr; }
+	virtual void Tick(float DeltaSeconds) override {}
+	virtual float GetElapsedTime() override { return 0.0f; }
+};
+
+
+void FGenericPlatformHttp::Init()
+{
+}
+
+
+void FGenericPlatformHttp::Shutdown()
+{
+}
+
+
+IHttpRequest* FGenericPlatformHttp::ConstructRequest()
+{
+	return new FGenericPlatformHttpRequest();
+}
+
+
+static bool IsAllowedChar(UTF8CHAR LookupChar)
+{
+	static char AllowedChars[] = "ABCDEFGHIJKLMNOPQRSTUVWXYZabcdefghijklmnopqrstuvwxyz0123456789-_.~";
+	static bool bTableFilled = false;
+	static bool AllowedTable[256] = { false };
+
+	if (!bTableFilled)
+	{
+		for (int32 Idx = 0; Idx < ARRAY_COUNT(AllowedChars) - 1; ++Idx)	// -1 to avoid trailing 0
+		{
+			uint8 AllowedCharIdx = static_cast<uint8>(AllowedChars[Idx]);
+			check(AllowedCharIdx < ARRAY_COUNT(AllowedTable));
+			AllowedTable[AllowedCharIdx] = true;
+		}
+
+		bTableFilled = true;
+	}
+
+	return AllowedTable[LookupChar];
+}
+
+
+FString FGenericPlatformHttp::UrlEncode(const FString &UnencodedString)
+{
+	FTCHARToUTF8 Converter(*UnencodedString);	//url encoding must be encoded over each utf-8 byte
+	const UTF8CHAR* UTF8Data = (UTF8CHAR*) Converter.Get();	//converter uses ANSI instead of UTF8CHAR - not sure why - but other code seems to just do this cast. In this case it really doesn't matter
+	FString EncodedString = TEXT("");
+	
+	TCHAR Buffer[2] = { 0, 0 };
+
+	for (int32 ByteIdx = 0, Length = Converter.Length(); ByteIdx < Length; ++ByteIdx)
+	{
+		UTF8CHAR ByteToEncode = UTF8Data[ByteIdx];
+		
+		if (IsAllowedChar(ByteToEncode))
+		{
+			Buffer[0] = ByteToEncode;
+			FString TmpString = Buffer;
+			EncodedString += TmpString;
+		}
+		else if (ByteToEncode != '\0')
+		{
+			EncodedString += TEXT("%");
+			EncodedString += FString::Printf(TEXT("%.2X"), ByteToEncode);
+		}
+	}
+	return EncodedString;
+}
+
+FString FGenericPlatformHttp::UrlDecode(const FString &EncodedString)
+{
+	FTCHARToUTF8 Converter(*EncodedString);	
+	const UTF8CHAR* UTF8Data = (UTF8CHAR*)Converter.Get();	
+	
+	TArray<ANSICHAR> Data;
+	Data.Reserve(EncodedString.Len());
+
+	for (int32 CharIdx = 0; CharIdx < Converter.Length();)
+	{
+		if (UTF8Data[CharIdx] == '%')
+		{
+			int32 Value = 0;
+			if (UTF8Data[CharIdx + 1] == 'u')
+			{
+				if (CharIdx + 6 <= Converter.Length())
+				{
+					// Treat all %uXXXX as code point
+					Value = FParse::HexDigit(UTF8Data[CharIdx + 2]) << 12;
+					Value += FParse::HexDigit(UTF8Data[CharIdx + 3]) << 8;
+					Value += FParse::HexDigit(UTF8Data[CharIdx + 4]) << 4;
+					Value += FParse::HexDigit(UTF8Data[CharIdx + 5]);
+					CharIdx += 6;
+
+					ANSICHAR Buffer[8] = { 0 };
+					ANSICHAR* BufferPtr = Buffer;
+					int32 len = ARRAY_COUNT(Buffer);
+					FTCHARToUTF8_Convert::utf8fromcodepoint(Value, &BufferPtr, &len);
+
+					Data.Append(Buffer, BufferPtr - Buffer);
+				}
+				else
+				{
+					// Not enough in the buffer for valid decoding, skip it
+					CharIdx++;
+					continue;
+				}
+			}
+			else if(CharIdx + 3 <= Converter.Length())
+			{
+				// Treat all %XX as straight byte
+				Value = FParse::HexDigit(UTF8Data[CharIdx + 1]) << 4;
+				Value += FParse::HexDigit(UTF8Data[CharIdx + 2]);
+				CharIdx += 3;
+				Data.Add((ANSICHAR)(Value));
+			}
+			else
+			{
+				// Not enough in the buffer for valid decoding, skip it
+				CharIdx++;
+				continue;
+			}
+		}
+		else
+		{
+			// Non escaped characters
+			Data.Add(UTF8Data[CharIdx]);
+			CharIdx++;
+		}
+	}
+
+	Data.Add(0);
+	FUTF8ToTCHAR DecodeConvert(Data.GetData());
+	return DecodeConvert.Get();
+}
\ No newline at end of file
diff --git a/Engine/Source/Runtime/Online/HTTP/Private/HttpPrivatePCH.h b/Engine/Source/Runtime/Online/HTTP/Private/HttpPrivatePCH.h
--- a/Engine/Source/Runtime/Online/HTTP/Private/HttpPrivatePCH.h
+++ b/Engine/Source/Runtime/Online/HTTP/Private/HttpPrivatePCH.h
@@ -5,3 +5,6 @@
 #include "Core.h"
 #include "Http.h"
 
+#include "GenericPlatform/GenericPlatformHttp.h"
+#include "PlatformHttp.h"
+
diff --git a/Engine/Source/Runtime/Online/HTTP/Public/Android/AndroidHttp.h b/Engine/Source/Runtime/Online/HTTP/Public/Android/AndroidHttp.h
new file mode 100644
--- /dev/null
+++ b/Engine/Source/Runtime/Online/HTTP/Public/Android/AndroidHttp.h
@@ -0,0 +1,39 @@
+// Copyright 1998-2015 Epic Games, Inc. All Rights Reserved.
+
+#pragma once
+#include "GenericPlatform/GenericPlatformHttp.h"
+
+/**
+ * Platform specific Http implementations
+ */
+class FAndroidPlatformHttp : public FGenericPlatformHttp
+{
+public:
+
+	/**
+	 * Platform initialization step
+	 */
+	static void Init();
+
+	/**
+	 * Creates a platform-specific HTTP manager.
+	 *
+	 * @return NULL if default implementation is to be used
+	 */
+	static FHttpManager* CreatePlatformHttpManager();
+
+	/**
+	 * Platform shutdown step
+	 */
+	static void Shutdown();
+
+	/**
+	 * Creates a new Http request instance for the current platform
+	 *
+	 * @return request object
+	 */
+	static IHttpRequest* ConstructRequest();
+};
+
+
+typedef FAndroidPlatformHttp FPlatformHttp;
\ No newline at end of file
diff --git a/Engine/Source/Runtime/Online/HTTP/Public/GenericPlatform/GenericPlatformHttp.h b/Engine/Source/Runtime/Online/HTTP/Public/GenericPlatform/GenericPlatformHttp.h
new file mode 100644
--- /dev/null
+++ b/Engine/Source/Runtime/Online/HTTP/Public/GenericPlatform/GenericPlatformHttp.h
@@ -0,0 +1,60 @@
+// Copyright 1998-2015 Epic Games, Inc. All Rights Reserved.
+
+
+#pragma once
+#include "HAL/Platform.h"
+
+class FHttpManager;
+class FString;
+class IHttpRequest;
+
+/**
+ * Platform specific Http implementations
+ */
+class HTTP_API FGenericPlatformHttp
+{
+public:
+
+	/**
+	 * Platform initialization step
+	 */
+	static void Init();
+
+	/**
+	 * Creates a platform-specific HTTP manager.
+	 *
+	 * @return nullptr if default implementation is to be used
+	 */
+	static FHttpManager* CreatePlatformHttpManager()
+	{
+		return nullptr;
+	}
+
+	/**
+	 * Platform shutdown step
+	 */
+	static void Shutdown();
+
+	/**
+	 * Creates a new Http request instance for the current platform
+	 *
+	 * @return request object
+	 */
+	static IHttpRequest* ConstructRequest();
+
+	/**
+	 * Returns a percent-encoded version of the passed in string
+	 *
+	 * @param UnencodedString The unencoded string to convert to percent-encoding
+	 * @return The percent-encoded string
+	 */
+	static FString UrlEncode(const FString& UnencodedString);
+
+	/**
+	 * Returns a decoded version of the percent-encoded passed in string
+	 *
+	 * @param EncodedString The percent encoded string to convert to string
+	 * @return The decoded string
+	 */
+	static FString UrlDecode(const FString& EncodedString);
+};
diff --git a/Engine/Source/Runtime/Online/HTTP/Public/HTML5/HTML5PlatformHttp.h b/Engine/Source/Runtime/Online/HTTP/Public/HTML5/HTML5PlatformHttp.h
new file mode 100644
--- /dev/null
+++ b/Engine/Source/Runtime/Online/HTTP/Public/HTML5/HTML5PlatformHttp.h
@@ -0,0 +1,42 @@
+// Copyright 1998-2015 Epic Games, Inc. All Rights Reserved.
+
+#pragma once
+#include "GenericPlatform/GenericPlatformHttp.h"
+
+/**
+ * Platform specific Http implementations
+ */
+class FHTML5PlatformHttp : public FGenericPlatformHttp
+{
+public:
+
+	/**
+	 * Platform initialization step
+	 */
+	static void Init();
+
+	/**
+	 * Creates a platform-specific HTTP manager.
+	 *
+	 * @return NULL if default implementation is to be used
+	 */
+	static FHttpManager* CreatePlatformHttpManager()
+	{
+		return nullptr;
+	}
+
+	/**
+	 * Platform shutdown step
+	 */
+	static void Shutdown();
+
+	/**
+	 * Creates a new Http request instance for the current platform
+	 *
+	 * @return request object
+	 */
+	static IHttpRequest* ConstructRequest();
+};
+
+
+typedef FHTML5PlatformHttp FPlatformHttp;
\ No newline at end of file
diff --git a/Engine/Source/Runtime/Online/HTTP/Public/Http.h b/Engine/Source/Runtime/Online/HTTP/Public/Http.h
--- a/Engine/Source/Runtime/Online/HTTP/Public/Http.h
+++ b/Engine/Source/Runtime/Online/HTTP/Public/Http.h
@@ -9,5 +9,6 @@
 
 // Module includes
 #include "HttpModule.h"
+#include "PlatformHttp.h"
 
 HTTP_API DECLARE_LOG_CATEGORY_EXTERN(LogHttp, Warning, All);
\ No newline at end of file
diff --git a/Engine/Source/Runtime/Online/HTTP/Public/HttpManager.h b/Engine/Source/Runtime/Online/HTTP/Public/HttpManager.h
--- a/Engine/Source/Runtime/Online/HTTP/Public/HttpManager.h
+++ b/Engine/Source/Runtime/Online/HTTP/Public/HttpManager.h
@@ -5,6 +5,8 @@
 #include "Ticker.h"
 #include "HttpPackage.h"
 
+class IHttpRequest;
+
 /**
  * Manages Http request that are currently being processed
  */
diff --git a/Engine/Source/Runtime/Online/HTTP/Public/HttpModule.h b/Engine/Source/Runtime/Online/HTTP/Public/HttpModule.h
--- a/Engine/Source/Runtime/Online/HTTP/Public/HttpModule.h
+++ b/Engine/Source/Runtime/Online/HTTP/Public/HttpModule.h
@@ -5,6 +5,8 @@
 #include "ModuleManager.h"
 #include "HttpManager.h"
 
+class IHttpRequest;
+
 /**
  * Module for Http request implementations
  * Use FHttpFactory to create a new Http request
diff --git a/Engine/Source/Runtime/Online/HTTP/Public/IOS/IOSPlatformHttp.h b/Engine/Source/Runtime/Online/HTTP/Public/IOS/IOSPlatformHttp.h
new file mode 100644
--- /dev/null
+++ b/Engine/Source/Runtime/Online/HTTP/Public/IOS/IOSPlatformHttp.h
@@ -0,0 +1,42 @@
+// Copyright 1998-2015 Epic Games, Inc. All Rights Reserved.
+
+#pragma once
+#include "GenericPlatform/GenericPlatformHttp.h"
+
+/**
+ * Platform specific Http implementations
+ */
+class FIOSPlatformHttp : public FGenericPlatformHttp
+{
+public:
+
+	/**
+	 * Platform initialization step
+	 */
+	static void Init();
+
+	/**
+	 * Creates a platform-specific HTTP manager.
+	 *
+	 * @return NULL if default implementation is to be used
+	 */
+	static FHttpManager* CreatePlatformHttpManager()
+	{
+		return nullptr;
+	}
+
+	/**
+	 * Platform shutdown step
+	 */
+	static void Shutdown();
+
+	/**
+	 * Creates a new Http request instance for the current platform
+	 *
+	 * @return request object
+	 */
+	static IHttpRequest* ConstructRequest();
+};
+
+
+typedef FIOSPlatformHttp FPlatformHttp;
\ No newline at end of file
diff --git a/Engine/Source/Runtime/Online/HTTP/Public/Linux/LinuxPlatformHttp.h b/Engine/Source/Runtime/Online/HTTP/Public/Linux/LinuxPlatformHttp.h
new file mode 100644
--- /dev/null
+++ b/Engine/Source/Runtime/Online/HTTP/Public/Linux/LinuxPlatformHttp.h
@@ -0,0 +1,39 @@
+// Copyright 1998-2015 Epic Games, Inc. All Rights Reserved.
+
+#pragma once
+#include "GenericPlatform/GenericPlatformHttp.h"
+
+/**
+ * Platform specific Http implementations
+ */
+class FLinuxPlatformHttp : public FGenericPlatformHttp
+{
+public:
+
+	/**
+	 * Platform initialization step
+	 */
+	static void Init();
+
+	/**
+	 * Creates a platform-specific HTTP manager.
+	 *
+	 * @return NULL if default implementation is to be used
+	 */
+	static FHttpManager* CreatePlatformHttpManager();
+
+	/**
+	 * Platform shutdown step
+	 */
+	static void Shutdown();
+
+	/**
+	 * Creates a new Http request instance for the current platform
+	 *
+	 * @return request object
+	 */
+	static IHttpRequest* ConstructRequest();
+};
+
+
+typedef FLinuxPlatformHttp FPlatformHttp;
diff --git a/Engine/Source/Runtime/Online/HTTP/Public/Mac/MacPlatformHttp.h b/Engine/Source/Runtime/Online/HTTP/Public/Mac/MacPlatformHttp.h
new file mode 100644
--- /dev/null
+++ b/Engine/Source/Runtime/Online/HTTP/Public/Mac/MacPlatformHttp.h
@@ -0,0 +1,42 @@
+// Copyright 1998-2015 Epic Games, Inc. All Rights Reserved.
+
+#pragma once
+#include "GenericPlatform/GenericPlatformHttp.h"
+
+/**
+ * Platform specific Http implementations
+ */
+class FMacPlatformHttp : public FGenericPlatformHttp
+{
+public:
+
+	/**
+	 * Platform initialization step
+	 */
+	static void Init();
+
+	/**
+	 * Creates a platform-specific HTTP manager.
+	 *
+	 * @return NULL if default implementation is to be used
+	 */
+	static FHttpManager* CreatePlatformHttpManager()
+	{
+		return nullptr;
+	}
+
+	/**
+	 * Platform shutdown step
+	 */
+	static void Shutdown();
+
+	/**
+	 * Creates a new Http request instance for the current platform
+	 *
+	 * @return request object
+	 */
+	static IHttpRequest* ConstructRequest();
+};
+
+
+typedef FMacPlatformHttp FPlatformHttp;
\ No newline at end of file
diff --git a/Engine/Source/Runtime/Online/HTTP/Public/PlatformHttp.h b/Engine/Source/Runtime/Online/HTTP/Public/PlatformHttp.h
new file mode 100644
--- /dev/null
+++ b/Engine/Source/Runtime/Online/HTTP/Public/PlatformHttp.h
@@ -0,0 +1,24 @@
+// Copyright 1998-2015 Epic Games, Inc. All Rights Reserved.
+#pragma once
+
+#include "HAL/Platform.h"
+
+#if PLATFORM_WINDOWS
+#include "Windows/WindowsPlatformHttp.h"
+#elif PLATFORM_PS4
+#include "PS4/PS4PlatformHttp.h"
+#elif PLATFORM_XBOXONE
+#include "XboxOne/XboxOneHttp.h"
+#elif PLATFORM_MAC
+#include "Mac/MacPlatformHttp.h"
+#elif PLATFORM_IOS
+#include "IOS/IOSPlatformHttp.h"
+#elif PLATFORM_ANDROID
+#include "Android/AndroidHttp.h"
+#elif PLATFORM_WINRT
+#include "WinRT/WinRTHttp.h"
+#elif PLATFORM_HTML5
+#include "HTML5/HTML5PlatformHttp.h"
+#elif PLATFORM_LINUX
+#include "Linux/LinuxPlatformHttp.h"
+#endif
diff --git a/Engine/Source/Runtime/Online/HTTP/Public/WinRT/WinRTHttp.h b/Engine/Source/Runtime/Online/HTTP/Public/WinRT/WinRTHttp.h
new file mode 100644
--- /dev/null
+++ b/Engine/Source/Runtime/Online/HTTP/Public/WinRT/WinRTHttp.h
@@ -0,0 +1,42 @@
+// Copyright 1998-2015 Epic Games, Inc. All Rights Reserved.
+
+#pragma once
+#include "GenericPlatform/GenericPlatformHttp.h"
+
+/**
+ * Platform specific Http implementations
+ */
+class FWinRTPlatformHttp : public FGenericPlatformHttp
+{
+public:
+
+	/**
+	 * Platform initialization step
+	 */
+	static void Init();
+
+	/**
+	 * Creates a platform-specific HTTP manager.
+	 *
+	 * @return NULL if default implementation is to be used
+	 */
+	static FHttpManager* CreatePlatformHttpManager()
+	{
+		return nullptr;
+	}
+
+	/**
+	 * Platform shutdown step
+	 */
+	static void Shutdown();
+
+	/**
+	 * Creates a new Http request instance for the current platform
+	 *
+	 * @return request object
+	 */
+	static IHttpRequest* ConstructRequest();
+};
+
+
+typedef FWinRTPlatformHttp FPlatformHttp;
\ No newline at end of file
diff --git a/Engine/Source/Runtime/Online/HTTP/Public/Windows/WindowsPlatformHttp.h b/Engine/Source/Runtime/Online/HTTP/Public/Windows/WindowsPlatformHttp.h
new file mode 100644
--- /dev/null
+++ b/Engine/Source/Runtime/Online/HTTP/Public/Windows/WindowsPlatformHttp.h
@@ -0,0 +1,37 @@
+// Copyright 1998-2015 Epic Games, Inc. All Rights Reserved.
+
+#pragma once
+
+#include "GenericPlatform/GenericPlatformHttp.h"
+
+
+/**
+ * Platform specific HTTP implementations.
+ */
+class FWindowsPlatformHttp : public FGenericPlatformHttp
+{
+public:
+
+	/** Platform initialization step. */
+	static void Init();
+
+	/**
+	 * Creates a platform-specific HTTP manager.
+	 *
+	 * @return nullptr if default implementation is to be used.
+	 */
+	static FHttpManager* CreatePlatformHttpManager();
+
+	/** Platform shutdown step. */
+	static void Shutdown();
+
+	/**
+	 * Creates a new HTTP request instance for the current platform.
+	 *
+	 * @return The request object.
+	 */
+	static IHttpRequest* ConstructRequest();
+};
+
+
+typedef FWindowsPlatformHttp FPlatformHttp;
diff --git a/Engine/Source/Runtime/PerfCounters/PerfCounters.Build.cs b/Engine/Source/Runtime/PerfCounters/PerfCounters.Build.cs
--- a/Engine/Source/Runtime/PerfCounters/PerfCounters.Build.cs
+++ b/Engine/Source/Runtime/PerfCounters/PerfCounters.Build.cs
@@ -12,6 +12,7 @@ public class PerfCounters : ModuleRules
 				"CoreUObject",
 				"Json",
 				"Sockets",
+				"HTTP"
 			}
         );
     }
diff --git a/Engine/Source/Runtime/PerfCounters/Private/PerfCounters.cpp b/Engine/Source/Runtime/PerfCounters/Private/PerfCounters.cpp
--- a/Engine/Source/Runtime/PerfCounters/Private/PerfCounters.cpp
+++ b/Engine/Source/Runtime/PerfCounters/Private/PerfCounters.cpp
@@ -3,6 +3,7 @@
 #include "PerfCounters.h"
 #include "SocketSubsystem.h"
 #include "Sockets.h"
+#include "Http.h"
 
 #define JSON_ARRAY_NAME					TEXT("PerfCounters")
 #define JSON_PERFCOUNTER_NAME			TEXT("Name")
@@ -251,7 +252,7 @@ bool FPerfCounters::ProcessRequest(uint8* Buffer, int32 BufferLen, FResponse& Re
 			else if (Tokens[1].StartsWith(TEXT("/exec?c=")))
 			{
 				FString ExecCmd = Tokens[1].Mid(8);
-				FString ExecCmdDecoded = FGenericPlatformHttp::UrlDecode(ExecCmd);
+				FString ExecCmdDecoded = FPlatformHttp::UrlDecode(ExecCmd);
 
 				FStringOutputDevice StringOutDevice;
 				StringOutDevice.SetAutoEmitLineTerminator(true);
