Fix compilation on the case-sensitive fs.

diff --git a/Engine/Source/Runtime/Online/HTTP/Private/GenericPlatform/GenericPlatformHttp.cpp b/Engine/Source/Runtime/Online/HTTP/Private/GenericPlatform/GenericPlatformHttp.cpp
--- a/Engine/Source/Runtime/Online/HTTP/Private/GenericPlatform/GenericPlatformHttp.cpp
+++ b/Engine/Source/Runtime/Online/HTTP/Private/GenericPlatform/GenericPlatformHttp.cpp
@@ -1,6 +1,6 @@
 // Copyright 1998-2015 Epic Games, Inc. All Rights Reserved.
 
-#include "HTTPPrivatePCH.h"
+#include "HttpPrivatePCH.h"
 #include "Runtime/Online/HTTP/Public/Interfaces/IHttpBase.h"
 #include "Runtime/Online/HTTP/Public/Interfaces/IHttpRequest.h"
 #include "Runtime/Online/HTTP/Public/Interfaces/IHttpResponse.h"
