Adding a Log Category for BuildPatchServices, will be incorporated instead of GLog and GWarn moving forwards.

Merging using UE4-To-UE4-LauncherDev REVERSE

diff --git a/Engine/Source/Runtime/Online/BuildPatchServices/Private/BuildPatchServicesModule.cpp b/Engine/Source/Runtime/Online/BuildPatchServices/Private/BuildPatchServicesModule.cpp
--- a/Engine/Source/Runtime/Online/BuildPatchServices/Private/BuildPatchServicesModule.cpp
+++ b/Engine/Source/Runtime/Online/BuildPatchServices/Private/BuildPatchServicesModule.cpp
@@ -6,7 +6,8 @@
 
 #include "BuildPatchServicesPrivatePCH.h"
 
-IMPLEMENT_MODULE( FBuildPatchServicesModule, BuildPatchServices );
+DEFINE_LOG_CATEGORY(LogBuildPatchServices);
+IMPLEMENT_MODULE(FBuildPatchServicesModule, BuildPatchServices);
 
 /* FBuildPatchInstallationInfo implementation
 *****************************************************************************/
diff --git a/Engine/Source/Runtime/Online/BuildPatchServices/Private/BuildPatchServicesPrivatePCH.h b/Engine/Source/Runtime/Online/BuildPatchServices/Private/BuildPatchServicesPrivatePCH.h
--- a/Engine/Source/Runtime/Online/BuildPatchServices/Private/BuildPatchServicesPrivatePCH.h
+++ b/Engine/Source/Runtime/Online/BuildPatchServices/Private/BuildPatchServicesPrivatePCH.h
@@ -53,3 +53,5 @@
 #include "BuildPatchGeneration.h"
 #include "BuildPatchUtil.h"
 #include "BuildPatchAnalytics.h"
+
+DECLARE_LOG_CATEGORY_EXTERN(LogBuildPatchServices, Log, All);
