Implementing BuildPatchTool's ability to add 'tags' to each file in a build. This data will be used in forthcoming selective download feature.

Merging using UE4-To-UE4-LauncherDev REVERSE

diff --git a/Engine/Source/Runtime/Online/BuildPatchServices/Private/BuildPatchGeneration.cpp b/Engine/Source/Runtime/Online/BuildPatchServices/Private/BuildPatchGeneration.cpp
--- a/Engine/Source/Runtime/Online/BuildPatchServices/Private/BuildPatchGeneration.cpp
+++ b/Engine/Source/Runtime/Online/BuildPatchServices/Private/BuildPatchGeneration.cpp
@@ -11,6 +11,7 @@
 #include "Generation/BuildStreamer.h"
 #include "Generation/CloudEnumeration.h"
 #include "Generation/ManifestBuilder.h"
+#include "Generation/FileAttributesParser.h"
 
 using namespace BuildPatchServices;
 
@@ -461,50 +462,6 @@ static void AddCustomFieldsToBuildManifest(const TMap<FString, FVariant>& Custom
 	}
 }
 
-static void FileAttributesMetaToMap(const FString& AttributesList, TMap<FString, FFileAttributes>& FileAttributesMap)
-{
-	const TCHAR Quote = TEXT('\"');
-	const TCHAR EOFile = TEXT('\0');
-	const TCHAR EOLine = TEXT('\n');
-
-	const TCHAR* CharPtr = *AttributesList;
-	while (*CharPtr != EOFile)
-	{
-		// Parse filename
-		while (*CharPtr != Quote && *CharPtr != EOFile){ ++CharPtr; }
-		if (*CharPtr == EOFile)
-		{
-			break;
-		}
-		const TCHAR* FilenameStart = ++CharPtr;
-		while (*CharPtr != Quote && *CharPtr != EOFile && *CharPtr != EOLine){ ++CharPtr; }
-		checkf(*CharPtr != EOFile, TEXT("Attributes File List: Unexpected end of file before next quote! Pos:%d"), CharPtr - *AttributesList);
-		checkf(*CharPtr != EOLine, TEXT("Attributes File List: Unexpected end of line before next quote! Pos:%d"), CharPtr - *AttributesList);
-		const TCHAR* FilenameEnd = CharPtr++;
-		// Parse keywords
-		while (*CharPtr != Quote && *CharPtr != EOFile && *CharPtr != EOLine){ ++CharPtr; }
-		checkf(*CharPtr != Quote, TEXT("Attributes File List: Unexpected Quote before end of keywords! Pos:%d"), CharPtr - *AttributesList);
-		const TCHAR* EndOfLine = CharPtr;
-		// Grab info
-		FString Filename = FString(FilenameEnd - FilenameStart, FilenameStart).Replace(TEXT("\\"), TEXT("/"));
-		FString Keywords(EndOfLine - FilenameEnd, FilenameEnd);
-		FFileAttributes FileAttributes;
-		if (Keywords.Contains(TEXT("readonly")))
-		{
-			FileAttributes.bReadOnly = true;
-		}
-		if (Keywords.Contains(TEXT("compressed")))
-		{
-			FileAttributes.bCompressed = true;
-		}
-		if (Keywords.Contains(TEXT("executable")))
-		{
-			FileAttributes.bUnixExecutable = true;
-		}
-		FileAttributesMap.Add(MoveTemp(Filename), FileAttributes);
-	}
-}
-
 /* FBuildDataGenerator implementation
 *****************************************************************************/
 bool FBuildDataGenerator::GenerateChunksManifestFromDirectory( const FBuildPatchSettings& Settings )
@@ -514,11 +471,36 @@ bool FBuildDataGenerator::GenerateChunksManifestFromDirectory( const FBuildPatch
 	// Take the build CS
 	FScopeLock SingleConcurrentBuild( &SingleConcurrentBuildCS );
 
+	// Create a manifest details
+	FManifestDetails ManifestDetails;
+	ManifestDetails.bIsFileData = false;
+	ManifestDetails.AppId = Settings.AppID;
+	ManifestDetails.AppName = Settings.AppName;
+	ManifestDetails.BuildVersion = Settings.BuildVersion;
+	ManifestDetails.LaunchExe = Settings.LaunchExe;
+	ManifestDetails.LaunchCommand = Settings.LaunchCommand;
+	ManifestDetails.PrereqName = Settings.PrereqName;
+	ManifestDetails.PrereqPath = Settings.PrereqPath;
+	ManifestDetails.PrereqArgs = Settings.PrereqArgs;
+	ManifestDetails.CustomFields = Settings.CustomFields;
+
+	// Load the required file attributes
+	if(!Settings.AttributeListFile.IsEmpty())
+	{
+		FFileAttributesParserRef FileAttributesParser = FFileAttributesParserFactory::Create();
+		if(!FileAttributesParser->ParseFileAttributes(Settings.AttributeListFile, ManifestDetails.FileAttributesMap))
+		{
+			UE_LOG(LogBuildPatchServices, Error, TEXT("Attributes list file did not parse %s"), *Settings.AttributeListFile);
+			return false;
+		}
+	}
+
 	// Create stat collector
 	FStatsCollectorRef StatsCollector = FStatsCollectorFactory::Create();
 
 	// Enumerate Chunks
 	FCloudEnumerationRef CloudEnumeration = FCloudEnumerationFactory::Create(FBuildPatchServicesModule::GetCloudDirectory());
+
 	// Force waiting on cloud enumeration for more accurate stats, this line can be removed when stats are not required.
 	CloudEnumeration->GetChunkInventory();
 
@@ -536,37 +518,9 @@ bool FBuildDataGenerator::GenerateChunksManifestFromDirectory( const FBuildPatch
 	FBuildGenerationChunkCache::Init(Cutoff);
 	FBuildGenerationChunkCache::Get().SetStatsCollector(StatsCollector);
 
-	// Create a manifest builder
-	FManifestDetails ManifestDetails;
-	ManifestDetails.bIsFileData = false;
-	ManifestDetails.AppId = Settings.AppID;
-	ManifestDetails.AppName = Settings.AppName;
-	ManifestDetails.BuildVersion = Settings.BuildVersion;
-	ManifestDetails.LaunchExe = Settings.LaunchExe;
-	ManifestDetails.LaunchCommand = Settings.LaunchCommand;
-	ManifestDetails.PrereqName = Settings.PrereqName;
-	ManifestDetails.PrereqPath = Settings.PrereqPath;
-	ManifestDetails.PrereqArgs = Settings.PrereqArgs;
-	ManifestDetails.CustomFields = Settings.CustomFields;
-	// Get the file attributes
-	FString AttributesList;
-	if (Settings.AttributeListFile.Len() > 0)
-	{
-		FFileHelper::LoadFileToString(AttributesList, *Settings.AttributeListFile);
-		if (!AttributesList.IsEmpty())
-		{
-			FileAttributesMetaToMap(AttributesList, ManifestDetails.FileAttributesMap);
-		}
-		else
-		{
-			GLog->Logf(TEXT("WARNING: Attributes list file empty"));
-		}
-	}
+	// Create the manifest builder
 	FManifestBuilderRef ManifestBuilder = FManifestBuilderFactory::Create(ManifestDetails, BuildStream);
 
-	// Refers to how much data has been dequeued
-	//uint64 ProcessPos = 0;
-
 	// Used to store data read lengths
 	uint32 ReadLen = 0;
 
@@ -659,16 +613,13 @@ bool FBuildDataGenerator::GenerateFilesManifestFromDirectory( const FBuildPatchS
 	// Get the file attributes
 	FString AttributesList;
 	TMap<FString, FFileAttributes> FileAttributesMap;
-	if (Settings.AttributeListFile.Len() > 0)
+	if(!Settings.AttributeListFile.IsEmpty())
 	{
-		FFileHelper::LoadFileToString(AttributesList, *Settings.AttributeListFile);
-		if (!AttributesList.IsEmpty())
-		{
-			FileAttributesMetaToMap(AttributesList, FileAttributesMap);
-		}
-		else
+		FFileAttributesParserRef FileAttributesParser = FFileAttributesParserFactory::Create();
+		if(!FileAttributesParser->ParseFileAttributes(Settings.AttributeListFile, FileAttributesMap))
 		{
-			GLog->Logf(TEXT("WARNING: Attributes list file empty"));
+			UE_LOG(LogBuildPatchServices, Error, TEXT("Attributes list file did not parse %s"), *Settings.AttributeListFile);
+			return false;
 		}
 	}
 
@@ -752,6 +703,7 @@ bool FBuildDataGenerator::GenerateFilesManifestFromDirectory( const FBuildPatchS
 			FFileManifestData& FileManifest = *BuildManifest->FileManifestLookup[Filename];
 			FileManifest.bIsReadOnly = Attributes.bReadOnly;
 			FileManifest.bIsCompressed = Attributes.bCompressed;
+			FileManifest.InstallTags = Attributes.InstallTags.Array();
 			// Only overwrite unix exe if true
 			if (Attributes.bUnixExecutable)
 			{
diff --git a/Engine/Source/Runtime/Online/BuildPatchServices/Private/BuildPatchGeneration.h b/Engine/Source/Runtime/Online/BuildPatchServices/Private/BuildPatchGeneration.h
--- a/Engine/Source/Runtime/Online/BuildPatchServices/Private/BuildPatchGeneration.h
+++ b/Engine/Source/Runtime/Online/BuildPatchServices/Private/BuildPatchGeneration.h
@@ -355,6 +355,7 @@ struct FFileAttributes
 	bool bReadOnly;
 	bool bCompressed;
 	bool bUnixExecutable;
+	TSet<FString> InstallTags;
 	FFileAttributes();
 };
 
diff --git a/Engine/Source/Runtime/Online/BuildPatchServices/Private/BuildPatchManifest.cpp b/Engine/Source/Runtime/Online/BuildPatchServices/Private/BuildPatchManifest.cpp
--- a/Engine/Source/Runtime/Online/BuildPatchServices/Private/BuildPatchManifest.cpp
+++ b/Engine/Source/Runtime/Online/BuildPatchServices/Private/BuildPatchManifest.cpp
@@ -804,6 +804,17 @@ void FBuildPatchAppManifest::SerializeToJSON(FString& JSONOutput)
 					}
 					Writer->WriteArrayEnd();
 				}
+				if (FileManifest.InstallTags.Num() > 0)
+				{
+					Writer->WriteArrayStart(TEXT("InstallTags"));
+					{
+						for (const auto& InstallTag : FileManifest.InstallTags)
+						{
+							Writer->WriteValue(InstallTag);
+						}
+					}
+					Writer->WriteArrayEnd();
+				}
 			}
 			Writer->WriteObjectEnd();
 		}
@@ -957,6 +968,14 @@ bool FBuildPatchAppManifest::DeserializeFromJSON( const FString& JSONInput )
 				bSuccess = bSuccess && FromStringBlob(JsonChunkPart->GetStringField(TEXT("Size")), FileChunkPart.Size);
 				AllDataGuids.Add(FileChunkPart.Guid);
 			}
+			if (JsonFileManifest->HasTypedField<EJson::Array>(TEXT("InstallTags")))
+			{
+				TArray<TSharedPtr<FJsonValue>> JsonInstallTagsArray = JsonFileManifest->GetArrayField(TEXT("InstallTags"));
+				for (auto JsonInstallTagIt = JsonInstallTagsArray.CreateConstIterator(); JsonInstallTagIt && bSuccess; ++JsonInstallTagIt)
+				{
+					FileManifest.InstallTags.Add((*JsonInstallTagIt)->AsString());
+				}
+			}
 			FileManifest.bIsUnixExecutable = JsonFileManifest->HasField(TEXT("bIsUnixExecutable")) && JsonFileManifest->GetBoolField(TEXT("bIsUnixExecutable"));
 			FileManifest.bIsReadOnly = JsonFileManifest->HasField(TEXT("bIsReadOnly")) && JsonFileManifest->GetBoolField(TEXT("bIsReadOnly"));
 			FileManifest.bIsCompressed = JsonFileManifest->HasField(TEXT("bIsCompressed")) && JsonFileManifest->GetBoolField(TEXT("bIsCompressed"));
diff --git a/Engine/Source/Runtime/Online/BuildPatchServices/Private/BuildPatchManifest.h b/Engine/Source/Runtime/Online/BuildPatchServices/Private/BuildPatchManifest.h
--- a/Engine/Source/Runtime/Online/BuildPatchServices/Private/BuildPatchManifest.h
+++ b/Engine/Source/Runtime/Online/BuildPatchServices/Private/BuildPatchManifest.h
@@ -187,6 +187,9 @@ struct FFileManifestData
 	UPROPERTY()
 	TArray<FChunkPartData> FileChunkParts;
 
+	UPROPERTY()
+	TArray<FString> InstallTags;
+
 	UPROPERTY()
 	bool bIsUnixExecutable;
 
diff --git a/Engine/Source/Runtime/Online/BuildPatchServices/Private/Generation/FileAttributesParser.cpp b/Engine/Source/Runtime/Online/BuildPatchServices/Private/Generation/FileAttributesParser.cpp
new file mode 100644
--- /dev/null
+++ b/Engine/Source/Runtime/Online/BuildPatchServices/Private/Generation/FileAttributesParser.cpp
@@ -0,0 +1,156 @@
+// Copyright 1998-2015 Epic Games, Inc. All Rights Reserved.
+
+#if WITH_BUILDPATCHGENERATION
+
+#include "BuildPatchServicesPrivatePCH.h"
+#include "FileAttributesParser.h"
+
+namespace BuildPatchServices
+{
+	class FFileAttributesParserImpl
+		: public FFileAttributesParser
+	{
+		typedef void(*SetKeyValueAttributeFunction)(FFileAttributes&, FString);
+
+	public:
+		FFileAttributesParserImpl(IFileManager* FileManager);
+		virtual ~FFileAttributesParserImpl();
+
+		virtual bool ParseFileAttributes(const FString& MetaFilename, TMap<FString, FFileAttributes>& FileAttributes) override;
+
+	private:
+		bool FileAttributesMetaToMap(const FString& AttributesList, TMap<FString, FFileAttributes>& FileAttributesMap);
+
+	private:
+		IFileManager* FileManager;
+		TMap<FString, SetKeyValueAttributeFunction> AttributeSetters;
+	};
+
+	FFileAttributesParserImpl::FFileAttributesParserImpl(IFileManager* InFileManager)
+		: FileManager(InFileManager)
+	{
+		AttributeSetters.Add(TEXT("readonly"),   [](FFileAttributes& Attributes, FString Value){ Attributes.bReadOnly = true; });
+		AttributeSetters.Add(TEXT("compressed"), [](FFileAttributes& Attributes, FString Value){ Attributes.bCompressed = true; });
+		AttributeSetters.Add(TEXT("executable"), [](FFileAttributes& Attributes, FString Value){ Attributes.bUnixExecutable = true; });
+		AttributeSetters.Add(TEXT("tag"),        [](FFileAttributes& Attributes, FString Value){ Attributes.InstallTags.Add(MoveTemp(Value)); });
+	}
+
+	FFileAttributesParserImpl::~FFileAttributesParserImpl()
+	{
+	}
+
+	bool FFileAttributesParserImpl::ParseFileAttributes(const FString& MetaFilename, TMap<FString, FFileAttributes>& FileAttributes)
+	{
+		FArchive* Reader = FileManager->CreateFileReader(*MetaFilename);
+		if (Reader != nullptr)
+		{
+			TArray<uint8> FileData;
+			FileData.AddUninitialized(Reader->TotalSize());
+			Reader->Serialize(FileData.GetData(), FileData.Num());
+			bool Successful = Reader->Close();
+			delete Reader;
+			if (Successful)
+			{
+				FString FileDataString;
+				FFileHelper::BufferToString(FileDataString, FileData.GetData(), FileData.Num());
+				return FileAttributesMetaToMap(FileDataString, FileAttributes);
+			}
+			else
+			{
+				UE_LOG(LogBuildPatchServices, Error, TEXT("FFileAttributesParser: Error occured reading meta file %s"), *MetaFilename);
+			}
+		}
+		else
+		{
+			UE_LOG(LogBuildPatchServices, Error, TEXT("FFileAttributesParser: Could not open meta file %s"), *MetaFilename);
+		}
+
+		return false;
+	}
+
+	bool FFileAttributesParserImpl::FileAttributesMetaToMap(const FString& AttributesList, TMap<FString, FFileAttributes>& FileAttributesMap)
+	{
+		const TCHAR Quote = TEXT('\"');
+		const TCHAR EOFile = TEXT('\0');
+		const TCHAR EOLine = TEXT('\n');
+
+		bool Successful = true;
+		bool FoundFilename = false;
+
+		const TCHAR* CharPtr = *AttributesList;
+		while (*CharPtr != EOFile)
+		{
+			// Parse filename
+			while (*CharPtr != Quote && *CharPtr != EOFile){ ++CharPtr; }
+			if (*CharPtr == EOFile)
+			{
+				if (!FoundFilename)
+				{
+					UE_LOG(LogBuildPatchServices, Error, TEXT("FFileAttributesParser: Did not find opening quote for filename!"));
+					return false;
+				}
+				break;
+			}
+			const TCHAR* FilenameStart = ++CharPtr;
+			while (*CharPtr != Quote && *CharPtr != EOFile && *CharPtr != EOLine){ ++CharPtr; }
+			// Check we didn't run out of file
+			if (*CharPtr == EOFile)
+			{
+				UE_LOG(LogBuildPatchServices, Error, TEXT("FFileAttributesParser: Unexpected end of file before next quote! Pos:%d"), CharPtr - *AttributesList);
+				return false;
+			}
+			// Check we didn't run out of line
+			if(*CharPtr == EOLine)
+			{
+				UE_LOG(LogBuildPatchServices, Error, TEXT("FFileAttributesParser: Unexpected end of line before next quote! Pos:%d"), CharPtr - *AttributesList);
+				return false;
+			}
+			// Save positions
+			const TCHAR* FilenameEnd = CharPtr++;
+			const TCHAR* AttributesStart = CharPtr;
+			// Parse keywords
+			while (*CharPtr != Quote && *CharPtr != EOFile && *CharPtr != EOLine){ ++CharPtr; }
+			// Check we hit the end of the line or file, another quote it wrong
+			if (*CharPtr == Quote)
+			{
+				UE_LOG(LogBuildPatchServices, Error, TEXT("FFileAttributesParser: Unexpected Quote before end of keywords! Pos:%d"), CharPtr - *AttributesList);
+				return false;
+			}
+			FoundFilename = true;
+			// Save position
+			const TCHAR* EndOfLine = CharPtr;
+			// Grab info
+			FString Filename = FString(FilenameEnd - FilenameStart, FilenameStart).Replace(TEXT("\\"), TEXT("/"));
+			FFileAttributes& FileAttributes = FileAttributesMap.FindOrAdd(Filename);
+			TArray<FString> AttributeParamsArray;
+			FString AttributeParams(EndOfLine - AttributesStart, AttributesStart);
+			AttributeParams.ParseIntoArrayWS(AttributeParamsArray);
+			for (const FString& AttributeParam : AttributeParamsArray)
+			{
+				FString Key, Value;
+				if(!AttributeParam.Split(TEXT(":"), &Key, &Value))
+				{
+					Key = AttributeParam;
+				}
+				if (AttributeSetters.Contains(Key))
+				{
+					AttributeSetters[Key](FileAttributes, MoveTemp(Value));
+				}
+				else
+				{
+					UE_LOG(LogBuildPatchServices, Error, TEXT("FFileAttributesParser: Unrecognised attribute %s for %s"), *AttributeParam, *Filename);
+					Successful = false;
+				}
+			}
+		}
+
+		return Successful;
+	}
+
+	FFileAttributesParserRef FFileAttributesParserFactory::Create(IFileManager* FileManager)
+	{
+		return MakeShareable(new FFileAttributesParserImpl(FileManager));
+	}
+}
+
+#endif
diff --git a/Engine/Source/Runtime/Online/BuildPatchServices/Private/Generation/FileAttributesParser.h b/Engine/Source/Runtime/Online/BuildPatchServices/Private/Generation/FileAttributesParser.h
new file mode 100644
--- /dev/null
+++ b/Engine/Source/Runtime/Online/BuildPatchServices/Private/Generation/FileAttributesParser.h
@@ -0,0 +1,30 @@
+// Copyright 1998-2015 Epic Games, Inc. All Rights Reserved.
+#pragma once
+
+#if WITH_BUILDPATCHGENERATION
+
+namespace BuildPatchServices
+{
+	class FFileAttributesParser
+	{
+	public:
+		/**
+		 * Loads in the file attributes meta file and populates the given map.
+		 * @param   MetaFilename        The amount of data to attempt to retrieve.
+		 * @param   FileAttributes      The map to populate with the attributes
+		 * @return  True if the file existed and parsed successfully
+		 */
+		virtual bool ParseFileAttributes(const FString& MetaFilename, TMap<FString, FFileAttributes>& FileAttributes) = 0;
+	};
+
+	typedef TSharedRef<FFileAttributesParser, ESPMode::ThreadSafe> FFileAttributesParserRef;
+	typedef TSharedPtr<FFileAttributesParser, ESPMode::ThreadSafe> FFileAttributesParserPtr;
+
+	class FFileAttributesParserFactory
+	{
+	public:
+		static FFileAttributesParserRef Create(IFileManager* FileManager = &IFileManager::Get());
+	};
+}
+
+#endif
diff --git a/Engine/Source/Runtime/Online/BuildPatchServices/Private/Generation/ManifestBuilder.cpp b/Engine/Source/Runtime/Online/BuildPatchServices/Private/Generation/ManifestBuilder.cpp
--- a/Engine/Source/Runtime/Online/BuildPatchServices/Private/Generation/ManifestBuilder.cpp
+++ b/Engine/Source/Runtime/Online/BuildPatchServices/Private/Generation/ManifestBuilder.cpp
@@ -260,6 +260,7 @@ namespace BuildPatchServices
 						FileBuilder.FileManifest->SymlinkTarget = FileSpan.SymlinkTarget;
 						FileBuilder.FileManifest->bIsReadOnly = Attributes.bReadOnly;
 						FileBuilder.FileManifest->bIsCompressed = Attributes.bCompressed;
+						FileBuilder.FileManifest->InstallTags = Attributes.InstallTags.Array();
 						FileBuilder.FileManifest->Init();
 						check(FileBuilder.FileManifest->GetFileSize() == FileBuilder.FileSpan.Size);
 						FileBuilder.FileManifest = nullptr;
