Converting BuildPatchServices::FFileAttributesParser to use IPlatformFile for consistency

Merging using UE4-To-UE4-LauncherDev REVERSE

diff --git a/Engine/Source/Runtime/Online/BuildPatchServices/Private/Generation/FileAttributesParser.cpp b/Engine/Source/Runtime/Online/BuildPatchServices/Private/Generation/FileAttributesParser.cpp
--- a/Engine/Source/Runtime/Online/BuildPatchServices/Private/Generation/FileAttributesParser.cpp
+++ b/Engine/Source/Runtime/Online/BuildPatchServices/Private/Generation/FileAttributesParser.cpp
@@ -13,7 +13,7 @@ namespace BuildPatchServices
 		typedef void(*SetKeyValueAttributeFunction)(FFileAttributes&, FString);
 
 	public:
-		FFileAttributesParserImpl(IFileManager* FileManager);
+		FFileAttributesParserImpl(IPlatformFile& PlatformFile);
 		virtual ~FFileAttributesParserImpl();
 
 		virtual bool ParseFileAttributes(const FString& MetaFilename, TMap<FString, FFileAttributes>& FileAttributes) override;
@@ -22,12 +22,12 @@ namespace BuildPatchServices
 		bool FileAttributesMetaToMap(const FString& AttributesList, TMap<FString, FFileAttributes>& FileAttributesMap);
 
 	private:
-		IFileManager* FileManager;
+		IPlatformFile& PlatformFile;
 		TMap<FString, SetKeyValueAttributeFunction> AttributeSetters;
 	};
 
-	FFileAttributesParserImpl::FFileAttributesParserImpl(IFileManager* InFileManager)
-		: FileManager(InFileManager)
+	FFileAttributesParserImpl::FFileAttributesParserImpl(IPlatformFile& InPlatformFile)
+		: PlatformFile(InPlatformFile)
 	{
 		AttributeSetters.Add(TEXT("readonly"),   [](FFileAttributes& Attributes, FString Value){ Attributes.bReadOnly = true; });
 		AttributeSetters.Add(TEXT("compressed"), [](FFileAttributes& Attributes, FString Value){ Attributes.bCompressed = true; });
@@ -41,15 +41,12 @@ namespace BuildPatchServices
 
 	bool FFileAttributesParserImpl::ParseFileAttributes(const FString& MetaFilename, TMap<FString, FFileAttributes>& FileAttributes)
 	{
-		FArchive* Reader = FileManager->CreateFileReader(*MetaFilename);
-		if (Reader != nullptr)
+		TAutoPtr<IFileHandle> Handle(PlatformFile.OpenRead(*MetaFilename));
+		if (Handle.IsValid())
 		{
 			TArray<uint8> FileData;
-			FileData.AddUninitialized(Reader->TotalSize());
-			Reader->Serialize(FileData.GetData(), FileData.Num());
-			bool Successful = Reader->Close();
-			delete Reader;
-			if (Successful)
+			FileData.AddUninitialized(Handle->Size());
+			if (Handle->Read(FileData.GetData(), FileData.Num()))
 			{
 				FString FileDataString;
 				FFileHelper::BufferToString(FileDataString, FileData.GetData(), FileData.Num());
@@ -147,9 +144,9 @@ namespace BuildPatchServices
 		return Successful;
 	}
 
-	FFileAttributesParserRef FFileAttributesParserFactory::Create(IFileManager* FileManager)
+	FFileAttributesParserRef FFileAttributesParserFactory::Create(IPlatformFile& PlatformFile)
 	{
-		return MakeShareable(new FFileAttributesParserImpl(FileManager));
+		return MakeShareable(new FFileAttributesParserImpl(PlatformFile));
 	}
 }
 
diff --git a/Engine/Source/Runtime/Online/BuildPatchServices/Private/Generation/FileAttributesParser.h b/Engine/Source/Runtime/Online/BuildPatchServices/Private/Generation/FileAttributesParser.h
--- a/Engine/Source/Runtime/Online/BuildPatchServices/Private/Generation/FileAttributesParser.h
+++ b/Engine/Source/Runtime/Online/BuildPatchServices/Private/Generation/FileAttributesParser.h
@@ -23,7 +23,7 @@ namespace BuildPatchServices
 	class FFileAttributesParserFactory
 	{
 	public:
-		static FFileAttributesParserRef Create(IFileManager* FileManager = &IFileManager::Get());
+		static FFileAttributesParserRef Create(IPlatformFile& PlatformFile = FPlatformFileManager::Get().GetPlatformFile());
 	};
 }
 
