fixes to build against IOS 9 beta 4 SDK
#ios9

diff --git a/Engine/Source/Runtime/Core/Public/IOS/IOSAppDelegate.h b/Engine/Source/Runtime/Core/Public/IOS/IOSAppDelegate.h
--- a/Engine/Source/Runtime/Core/Public/IOS/IOSAppDelegate.h
+++ b/Engine/Source/Runtime/Core/Public/IOS/IOSAppDelegate.h
@@ -4,7 +4,11 @@
 
 #import <UIKit/UIKit.h>
 #import <AVFoundation/AVAudioSession.h>
+#ifdef __IPHONE_9_0
+#import <GameCenter/GKGameCenterViewController.h>
+#else
 #import <GameKit/GKGameCenterViewController.h>
+#endif
 
 class CORE_API FIOSCoreDelegates
 {
diff --git a/Engine/Source/Runtime/IOS/IOSAudio/Private/IOSAudioDevice.cpp b/Engine/Source/Runtime/IOS/IOSAudio/Private/IOSAudioDevice.cpp
--- a/Engine/Source/Runtime/IOS/IOSAudio/Private/IOSAudioDevice.cpp
+++ b/Engine/Source/Runtime/IOS/IOSAudio/Private/IOSAudioDevice.cpp
@@ -122,7 +122,11 @@ bool FIOSAudioDevice::InitializeHardware()
 
 	// Setup audo mixer unit
 	UnitDescription.componentType         = kAudioUnitType_Mixer;
+#ifdef __IPHONE_9_0
+    UnitDescription.componentSubType      = kAudioUnitSubType_SpatialMixer;
+#else
 	UnitDescription.componentSubType      = kAudioUnitSubType_AU3DMixerEmbedded;
+#endif
 	UnitDescription.componentManufacturer = kAudioUnitManufacturer_Apple;
 	UnitDescription.componentFlags        = 0;
 	UnitDescription.componentFlagsMask    = 0;
diff --git a/Engine/Source/Runtime/Online/IOS/OnlineSubsystemIOS/Private/OnlineAchievementsInterfaceIOS.cpp b/Engine/Source/Runtime/Online/IOS/OnlineSubsystemIOS/Private/OnlineAchievementsInterfaceIOS.cpp
--- a/Engine/Source/Runtime/Online/IOS/OnlineSubsystemIOS/Private/OnlineAchievementsInterfaceIOS.cpp
+++ b/Engine/Source/Runtime/Online/IOS/OnlineSubsystemIOS/Private/OnlineAchievementsInterfaceIOS.cpp
@@ -4,9 +4,13 @@
 #include "OnlineSubsystemIOSPrivatePCH.h"
 
 // Game Center includes
+#ifdef __IPHONE_9_0
+#include <GameCenter/GKAchievement.h>
+#include <GameCenter/GKAchievementDescription.h>
+#else
 #include <GameKit/GKAchievement.h>
 #include <GameKit/GKAchievementDescription.h>
-
+#endif
 
 FOnlineAchievementsIOS::FOnlineAchievementsIOS(class FOnlineSubsystemIOS* InSubsystem)
 {
diff --git a/Engine/Source/Runtime/Online/IOS/OnlineSubsystemIOS/Private/OnlineLeaderboardsInterfaceIOS.cpp b/Engine/Source/Runtime/Online/IOS/OnlineSubsystemIOS/Private/OnlineLeaderboardsInterfaceIOS.cpp
--- a/Engine/Source/Runtime/Online/IOS/OnlineSubsystemIOS/Private/OnlineLeaderboardsInterfaceIOS.cpp
+++ b/Engine/Source/Runtime/Online/IOS/OnlineSubsystemIOS/Private/OnlineLeaderboardsInterfaceIOS.cpp
@@ -4,9 +4,13 @@
 #include "OnlineSubsystemIOSPrivatePCH.h"
 
 // GameCenter includes
+#ifdef __IPHONE_9_0
+#include <GameCenter/GKLeaderboard.h>
+#include <GameCenter/GKScore.h>
+#else
 #include <GameKit/GKLeaderboard.h>
 #include <GameKit/GKScore.h>
-
+#endif
 
 FOnlineLeaderboardsIOS::FOnlineLeaderboardsIOS(FOnlineSubsystemIOS* InSubsystem)
 {
diff --git a/Engine/Source/Runtime/Online/IOS/OnlineSubsystemIOS/Private/OnlineSessionInterfaceIOS.h b/Engine/Source/Runtime/Online/IOS/OnlineSubsystemIOS/Private/OnlineSessionInterfaceIOS.h
--- a/Engine/Source/Runtime/Online/IOS/OnlineSubsystemIOS/Private/OnlineSessionInterfaceIOS.h
+++ b/Engine/Source/Runtime/Online/IOS/OnlineSubsystemIOS/Private/OnlineSessionInterfaceIOS.h
@@ -1,13 +1,21 @@
 // Copyright 1998-2015 Epic Games, Inc. All Rights Reserved.
 
 #pragma once
+#ifdef __IPHONE_9_0
+#include <GameCenter/GKLocalPlayer.h>
+#else
 #include <GameKit/GKLocalPlayer.h>
+#endif
 
 #include "OnlineSessionInterface.h"
 #include "OnlineSubsystemIOSTypes.h"
 
 #if __IPHONE_OS_VERSION_MIN_REQUIRED < __IPHONE_7_0
+#ifdef __IPHONE_9_0
+#include <GameCenter/GKSession.h>
+#else
 #include <GameKit/GKSession.h>
+#endif
 
 @interface FGameCenterSessionDelegateGK : UIViewController<GKSessionDelegate>
 {
