[merge]

#xmpp
- removed unused shadow var

diff --git a/Engine/Source/Runtime/Online/XMPP/Private/XmppJingle/XmppMultiUserChatJingle.cpp b/Engine/Source/Runtime/Online/XMPP/Private/XmppJingle/XmppMultiUserChatJingle.cpp
--- a/Engine/Source/Runtime/Online/XMPP/Private/XmppJingle/XmppMultiUserChatJingle.cpp
+++ b/Engine/Source/Runtime/Online/XMPP/Private/XmppJingle/XmppMultiUserChatJingle.cpp
@@ -161,7 +161,6 @@ private:
 		XForm->AddElement(FormTypeField);
 
 		// Setup other feature fields
-		FRoomFeatureValuePairs::const_iterator FeaturePairIt = RoomFeatureValuePairs.begin();
 		for (FRoomFeatureValuePairs::const_iterator FeaturePairIt = RoomFeatureValuePairs.begin()
 			; FeaturePairIt != RoomFeatureValuePairs.end()
 			; ++FeaturePairIt)
