Fix Google Play leaderboard read failures.

diff --git a/Engine/Source/Runtime/Online/Android/OnlineSubsystemGooglePlay/Private/OnlineAsyncTaskGooglePlayShowLoginUI.cpp b/Engine/Source/Runtime/Online/Android/OnlineSubsystemGooglePlay/Private/OnlineAsyncTaskGooglePlayShowLoginUI.cpp
--- a/Engine/Source/Runtime/Online/Android/OnlineSubsystemGooglePlay/Private/OnlineAsyncTaskGooglePlayShowLoginUI.cpp
+++ b/Engine/Source/Runtime/Online/Android/OnlineSubsystemGooglePlay/Private/OnlineAsyncTaskGooglePlayShowLoginUI.cpp
@@ -42,11 +42,16 @@ void FOnlineAsyncTaskGooglePlayShowLoginUI::Finalize()
 
 void FOnlineAsyncTaskGooglePlayShowLoginUI::TriggerDelegates()
 {
-	TSharedPtr<FUniqueNetId> UserId;
+	TSharedPtr<const FUniqueNetIdString> UserId = Subsystem->GetIdentityGooglePlay()->GetCurrentUserId();
 
-	if (bWasSuccessful)
+	if (bWasSuccessful && !UserId.IsValid())
 	{
 		UserId = MakeShareable(new FUniqueNetIdString());
+		Subsystem->GetIdentityGooglePlay()->SetCurrentUserId(UserId);
+	}
+	else if (!bWasSuccessful)
+	{
+		Subsystem->GetIdentityGooglePlay()->SetCurrentUserId(nullptr);
 	}
 
 	Delegate.ExecuteIfBound(UserId, PlayerId);
diff --git a/Engine/Source/Runtime/Online/Android/OnlineSubsystemGooglePlay/Private/OnlineIdentityInterfaceGooglePlay.h b/Engine/Source/Runtime/Online/Android/OnlineSubsystemGooglePlay/Private/OnlineIdentityInterfaceGooglePlay.h
--- a/Engine/Source/Runtime/Online/Android/OnlineSubsystemGooglePlay/Private/OnlineIdentityInterfaceGooglePlay.h
+++ b/Engine/Source/Runtime/Online/Android/OnlineSubsystemGooglePlay/Private/OnlineIdentityInterfaceGooglePlay.h
@@ -22,7 +22,7 @@ private:
 	bool bRegisteringUser;
 
 	/** UID for this identity */
-	TSharedPtr< FUniqueNetIdString > UniqueNetId;
+	TSharedPtr< const FUniqueNetIdString > UniqueNetId;
 
 	struct FPendingConnection
 	{
@@ -41,7 +41,9 @@ PACKAGE_SCOPE:
 	FOnlineIdentityGooglePlay(FOnlineSubsystemGooglePlay* InSubsystem);
 
 	/** Allow individual interfaces to access the currently signed-in user's id */
-	TSharedPtr<const FUniqueNetId> GetCurrentUserId() const { return UniqueNetId; }
+	TSharedPtr<const FUniqueNetIdString> GetCurrentUserId() const { return UniqueNetId; }
+
+	void SetCurrentUserId(TSharedPtr<const FUniqueNetIdString> InUniqueNetId) { UniqueNetId = InUniqueNetId; }
 
 public:
 
