Adding cancel all functionality to BuildPatchServices so that an application can receive completion delegates on shutdown if it calls the new function.

Merging using LauncherDev -> UE4 CL# 2660816

diff --git a/Engine/Source/Runtime/Online/BuildPatchServices/Private/BuildPatchServicesModule.cpp b/Engine/Source/Runtime/Online/BuildPatchServices/Private/BuildPatchServicesModule.cpp
--- a/Engine/Source/Runtime/Online/BuildPatchServices/Private/BuildPatchServicesModule.cpp
+++ b/Engine/Source/Runtime/Online/BuildPatchServices/Private/BuildPatchServicesModule.cpp
@@ -180,24 +180,17 @@ bool FBuildPatchServicesModule::Tick( float Delta )
 	check( bIsCalledFromMainThread );
 
 	// Call complete delegate on each finished installer
-	for(auto InstallerIt = BuildPatchInstallers.CreateIterator(); InstallerIt; ++InstallerIt)
+	for (auto& Installer : BuildPatchInstallers)
 	{
-		if( (*InstallerIt).IsValid() && (*InstallerIt)->IsComplete() )
+		if (Installer.IsValid() && Installer->IsComplete())
 		{
-			(*InstallerIt)->ExecuteCompleteDelegate();
-			(*InstallerIt).Reset();
+			Installer->ExecuteCompleteDelegate();
+			Installer.Reset();
 		}
 	}
 
 	// Remove completed (invalids) from the list
-	for(int32 BuildPatchInstallersIdx = 0; BuildPatchInstallersIdx < BuildPatchInstallers.Num(); ++BuildPatchInstallersIdx )
-	{
-		const FBuildPatchInstallerPtr* Installer = &BuildPatchInstallers[ BuildPatchInstallersIdx ];
-		if( !Installer->IsValid() )
-		{
-			BuildPatchInstallers.RemoveAt( BuildPatchInstallersIdx-- );
-		}
-	}
+	BuildPatchInstallers.RemoveAll([](const FBuildPatchInstallerPtr& Installer){ return Installer.IsValid() == false; });
 
 	// More ticks
 	return true;
@@ -270,25 +263,43 @@ void FBuildPatchServicesModule::RegisterAppInstallation(IBuildManifestRef AppMan
 	InstallationInfo.RegisterAppInstallation(AppManifest, AppInstallDirectory);
 }
 
+void FBuildPatchServicesModule::CancelAllInstallers(bool WaitForThreads)
+{
+	// Using a local bool for this check will improve the assert message that gets displayed
+	const bool bIsCalledFromMainThread = IsInGameThread();
+	check(bIsCalledFromMainThread);
+
+	// Loop each installer, cancel it, and optionally wait to make completion delegate call
+	for (auto& Installer : BuildPatchInstallers)
+	{
+		if (Installer.IsValid())
+		{
+			Installer->CancelInstall();
+			if (WaitForThreads)
+			{
+				Installer->WaitForThread();
+				Installer->ExecuteCompleteDelegate();
+				Installer.Reset();
+			}
+		}
+	}
+
+	// Remove completed (invalids) from the list
+	BuildPatchInstallers.RemoveAll([](const FBuildPatchInstallerPtr& Installer){ return Installer.IsValid() == false; });
+}
+
 void FBuildPatchServicesModule::PreExit()
 {
 	// Set shutdown error so any running threads know to exit.
 	FBuildPatchInstallError::SetFatalError(EBuildPatchInstallError::ApplicationClosing);
 
+	// Cleanup installers
+	CancelAllInstallers(true);
+	BuildPatchInstallers.Empty();
+
 	// Release our ptr to analytics
 	FBuildPatchAnalytics::SetAnalyticsProvider(NULL);
 	FBuildPatchAnalytics::SetHttpTracker(nullptr);
-
-	// Cleanup installers
-	for (auto& BuildPatchInstaller : BuildPatchInstallers)
-	{
-		// Make sure it is not paused, this function only un-pauses when in error state
-		BuildPatchInstaller->TogglePauseInstall();
-		// We still have to manually wait for the thread as another system could hold a shared ptr
-		// thus we would not be calling the destructor here
-		BuildPatchInstaller->WaitForThread();
-	}
-	BuildPatchInstallers.Empty();
 }
 
 const FString& FBuildPatchServicesModule::GetStagingDirectory()
diff --git a/Engine/Source/Runtime/Online/BuildPatchServices/Private/BuildPatchServicesModule.h b/Engine/Source/Runtime/Online/BuildPatchServices/Private/BuildPatchServicesModule.h
--- a/Engine/Source/Runtime/Online/BuildPatchServices/Private/BuildPatchServicesModule.h
+++ b/Engine/Source/Runtime/Online/BuildPatchServices/Private/BuildPatchServicesModule.h
@@ -82,7 +82,8 @@ public:
 	virtual void SetBackupDirectory( const FString& BackupDir ) override;
 	virtual void SetAnalyticsProvider( TSharedPtr< IAnalyticsProvider > AnalyticsProvider ) override;
 	virtual void SetHttpTracker( TSharedPtr< FHttpServiceTracker > HttpTracker ) override;
-	virtual void RegisterAppInstallation( IBuildManifestRef AppManifest, const FString AppInstallDirectory ) override;
+	virtual void RegisterAppInstallation(IBuildManifestRef AppManifest, const FString AppInstallDirectory) override;
+	virtual void CancelAllInstallers(bool WaitForThreads) override;
 #if WITH_BUILDPATCHGENERATION
 	virtual bool GenerateChunksManifestFromDirectory( const FBuildPatchSettings& Settings ) override;
 	virtual bool GenerateFilesManifestFromDirectory( const FBuildPatchSettings& Settings ) override;
diff --git a/Engine/Source/Runtime/Online/BuildPatchServices/Public/Interfaces/IBuildPatchServicesModule.h b/Engine/Source/Runtime/Online/BuildPatchServices/Public/Interfaces/IBuildPatchServicesModule.h
--- a/Engine/Source/Runtime/Online/BuildPatchServices/Public/Interfaces/IBuildPatchServicesModule.h
+++ b/Engine/Source/Runtime/Online/BuildPatchServices/Public/Interfaces/IBuildPatchServicesModule.h
@@ -119,6 +119,12 @@ public:
 	 */
 	virtual void RegisterAppInstallation(IBuildManifestRef AppManifest, const FString AppInstallDirectory) = 0;
 
+	/**
+	 * Call to force the exit out of all current installers, optionally blocks until threads have exited and complete delegates are called.
+	 * @param WaitForThreads		If true, will block on threads exit and completion delegates
+	 */
+	virtual void CancelAllInstallers(bool WaitForThreads) = 0;
+
 #if WITH_BUILDPATCHGENERATION
 	/**
 	 * Processes a Build Image to determine new chunks and produce a chunk based manifest, all saved to the cloud.
