Making it so that the OSS interface tests use the automation system.

I'm leaving it in teh current execcmds as well at the moment.

--------
Integrated using branch UE4-To-UE4-LauncherDev (reversed) of change#2659457 by [REDACTED_PERSON] on 2015/08/18 12:08:03.

diff --git a/Engine/Source/Runtime/Online/OnlineSubsystemUtils/Private/Tests/OnlineSubsystemUtilitiesTests.cpp b/Engine/Source/Runtime/Online/OnlineSubsystemUtils/Private/Tests/OnlineSubsystemUtilitiesTests.cpp
new file mode 100644
--- /dev/null
+++ b/Engine/Source/Runtime/Online/OnlineSubsystemUtils/Private/Tests/OnlineSubsystemUtilitiesTests.cpp
@@ -0,0 +1,197 @@
+// Copyright 1998-2015 Epic Games, Inc. All Rights Reserved.
+
+#include "OnlineSubsystemUtilsPrivatePCH.h"
+#include "TestIdentityInterface.h"
+
+DEFINE_LOG_CATEGORY_STATIC(OSSUtilsTestLog, All, All);
+
+namespace OSSUtilsTestHelper
+{
+	// Enums used to hold the state of the current test.
+	enum type
+	{
+		NotStarted,
+		BeginTest,
+		PerformingTest,
+	};
+
+	/**
+	* Get the username and password from the command line.
+	* Command line arguments are -OSSIDTESTUSER="YourUser" and -OSSIDTESTPSSWD="YouPassword"
+	* @param OutUsername is the referenced variable that will hold the found username.
+	* @param OutPassword is the referenced variable that will hold the found password.
+	*/
+	bool GetUserInfo( FString& OutUsername, FString& OutPassword )
+	{
+		OutUsername = TEXT( "" );
+		OutPassword = TEXT( "" );
+		// Get the username and password from the command line
+		const FString CommandLine( FCommandLine::Get() );
+		FParse::Value( *CommandLine, TEXT( "OSSIDTESTUSER=" ), OutUsername );
+		FParse::Value( *CommandLine, TEXT( "OSSIDTESTPSSWD=" ), OutPassword );
+
+		if ( OutUsername.IsEmpty() || OutPassword.IsEmpty() )
+		{
+			UE_LOG( OSSUtilsTestLog, Error, TEXT( "Missing test username and/or password." ) );
+			return false;
+		}
+		return true;
+	}
+
+	void GetOSSNamesForTesting( TArray<FString>& OutBeautifiedNames, TArray<FString>& OutTestCommands, const FString& PostTestName = TEXT( "" ) )
+	{
+		TArray< FName > OSSList;
+		OSSList.AddUnique(TEXT( "MCP" ) );
+		OSSList.AddUnique(TEXT( "Null" ) );
+		OSSList.AddUnique(TEXT( "Steam" ) );
+		OSSList.AddUnique(TEXT( "Facebook" ) );
+		OSSList.AddUnique(TEXT( "IOS" ) );
+		OSSList.AddUnique(TEXT( "WeChat" ) );
+		OSSList.AddUnique(TEXT( "GooglePlay" ) );
+		OSSList.AddUnique(TEXT( "Live" ) );
+		OSSList.AddUnique(TEXT( "Thunderhead" ) );
+		OSSList.AddUnique(TEXT( "Amazon" ) );
+
+		for ( auto& OSSName : OSSList )
+		{
+			if (IOnlineSubsystem::Get(OSSName))
+			{
+				FString PostName = FString::Printf( TEXT(".%s"), *PostTestName );
+				FString PrettyName = FString::Printf(TEXT( "%s%s" ), 
+					*OSSName.ToString(), 
+					PostTestName.IsEmpty() ? TEXT( "" ) : *PostName );
+				OutBeautifiedNames.Add( PrettyName );
+				OutTestCommands.Add( OSSName.ToString() );
+			}
+		}
+	}
+}
+
+/**
+* Struct used to hold the user account credentials and the identity interface.
+*/
+struct FTestIdentityInterfaceStruct
+{
+	UWorld* CurrentWorld;
+	FOnlineAccountCredentials* AccountCredentials;
+	FTestIdentityInterface* IdentityInterface;
+	bool bIsLogoutTest;
+	OSSUtilsTestHelper::type TestStatus;
+
+	/**
+	* Constructor for the struct.
+	* @param InUsername is the valid username that can be used with the OSS service.
+	* @param InPassword is the valid password for the given username and OSS service.
+	* @param InSubsystem is the OSS service to be used.
+	*/
+	FTestIdentityInterfaceStruct( const FString& InUsername, const FString& InPassword, const FString& InAccountType, const FString& InSubsystem )
+		: CurrentWorld( GEngine->GetWorld() )
+		, AccountCredentials( new FOnlineAccountCredentials( InAccountType, InUsername, InPassword ) )
+		, IdentityInterface( new FTestIdentityInterface( InSubsystem ) )
+		, bIsLogoutTest( false )
+		, TestStatus( OSSUtilsTestHelper::NotStarted )
+	{
+	}
+};
+
+/**
+* Latent command used to run the login and logout tests.
+* @param TestStruct is the name of the struct that holds the test identity information.
+*/
+DEFINE_LATENT_AUTOMATION_COMMAND_ONE_PARAMETER( FTestIdentityInterfaceLoginOut, FTestIdentityInterfaceStruct, TestStruct );
+
+bool FTestIdentityInterfaceLoginOut::Update()
+{
+	switch ( TestStruct.TestStatus )
+	{
+	case OSSUtilsTestHelper::NotStarted:
+		UE_LOG( OSSUtilsTestLog, Log, TEXT( "Starting the 'Login' or 'Logout' functional test." ) );
+		TestStruct.TestStatus = OSSUtilsTestHelper::BeginTest;
+		break;
+	case OSSUtilsTestHelper::BeginTest:
+		UE_LOG( OSSUtilsTestLog, Log, TEXT( "Attempting to %s." ), 
+			TestStruct.bIsLogoutTest ? TEXT( "'Logout'" ) : TEXT( "'Login'" ) );
+
+		TestStruct.IdentityInterface->Test( 
+			TestStruct.CurrentWorld, 
+			*TestStruct.AccountCredentials, 
+			TestStruct.bIsLogoutTest );
+		TestStruct.TestStatus = OSSUtilsTestHelper::PerformingTest;
+		break;
+	case OSSUtilsTestHelper::PerformingTest:
+		if ( TestStruct.IdentityInterface->GetTestStatus() ) return true;
+	default:
+		break;
+	}
+
+	return false;
+}
+
+/**
+* This Login and Logout functional test is intended to test that a user can
+* login and logout from the service.
+* Use -OSSIDTESTUSER='TestUser' and -OSSIDTESTPSSWD='YouPassword'
+*/
+IMPLEMENT_COMPLEX_AUTOMATION_TEST( FIdenitiyInterfacetFunctionalLoginLogoutTest, "System.OSS", EAutomationTestFlags::ATF_Game )
+
+void FIdenitiyInterfacetFunctionalLoginLogoutTest::GetTests( TArray<FString>& OutBeautifiedNames, TArray<FString>& OutTestCommands ) const
+{
+	OSSUtilsTestHelper::GetOSSNamesForTesting(
+		OutBeautifiedNames,
+		OutTestCommands,
+		TEXT("Functional Tests.Log in and then log out"));
+}
+
+bool FIdenitiyInterfacetFunctionalLoginLogoutTest::RunTest( const FString& Parameters )
+{
+	// Get the user info that is used for logging in and out of the subsystem.
+	FString Username = TEXT( "" );
+	FString Password = TEXT( "" );
+	if ( !OSSUtilsTestHelper::GetUserInfo( Username, Password ) ) return false;
+
+	// Two identity structs are created because the test function will delete the object after each run.
+	FTestIdentityInterfaceStruct LoginTestStruct( Username, Password, TEXT( "epic" ), Parameters );
+	FTestIdentityInterfaceStruct LogoutTestStruct( Username, Password, TEXT("epic"), Parameters );
+	LogoutTestStruct.bIsLogoutTest = true;
+
+	/* The first latent command will run the login test.*/
+	ADD_LATENT_AUTOMATION_COMMAND( FTestIdentityInterfaceLoginOut( LoginTestStruct ) );
+	/* This second latent command will run the logout test.*/
+	ADD_LATENT_AUTOMATION_COMMAND( FTestIdentityInterfaceLoginOut( LogoutTestStruct ) );
+
+	return true;
+}
+
+
+IMPLEMENT_COMPLEX_AUTOMATION_TEST( FTestIdentityInterfaceLogoutWhileNotLoggedIn, "System.OSS", EAutomationTestFlags::ATF_Game )
+
+void FTestIdentityInterfaceLogoutWhileNotLoggedIn::GetTests( TArray<FString>& OutBeautifiedNames, TArray<FString>& OutTestCommands ) const
+{
+	OSSUtilsTestHelper::GetOSSNamesForTesting( 
+		OutBeautifiedNames, 
+		OutTestCommands, 
+		TEXT( "Functional Tests.Logging out while not logged in" ) );
+}
+
+bool FTestIdentityInterfaceLogoutWhileNotLoggedIn::RunTest(const FString& Parameters)
+{
+	// Get the user info that is used for logging in and out of the subsystem.
+	FString Username = TEXT( "" );
+	FString Password = TEXT( " ");
+	if ( !OSSUtilsTestHelper::GetUserInfo( Username, Password ) ) return false;
+
+	FTestIdentityInterfaceStruct LogoutTestStruct( Username, Password, TEXT( "epic" ), Parameters );
+	LogoutTestStruct.bIsLogoutTest = true;
+
+	// Stop the test here if the user is currently logged in.
+	if ( LogoutTestStruct.IdentityInterface->IsTheUserLoggedIn() )
+	{
+		AddError( TEXT( "Unable to test as the user is currently logged in." ) );
+		return false;
+	}
+
+	/* This latent command will run the logout test.*/
+	ADD_LATENT_AUTOMATION_COMMAND( FTestIdentityInterfaceLoginOut( LogoutTestStruct ) );
+
+	return true;
+}
\ No newline at end of file
diff --git a/Engine/Source/Runtime/Online/OnlineSubsystemUtils/Private/Tests/TestIdentityInterface.cpp b/Engine/Source/Runtime/Online/OnlineSubsystemUtils/Private/Tests/TestIdentityInterface.cpp
--- a/Engine/Source/Runtime/Online/OnlineSubsystemUtils/Private/Tests/TestIdentityInterface.cpp
+++ b/Engine/Source/Runtime/Online/OnlineSubsystemUtils/Private/Tests/TestIdentityInterface.cpp
@@ -14,12 +14,11 @@ void FTestIdentityInterface::Test(UWorld* InWorld, const FOnlineAccountCredentia
 	}
 
 	AccountCredentials = InAccountCredentials;
-
 	OnlineIdentity = Online::GetIdentityInterface(InWorld, SubsystemName.Len() ? FName(*SubsystemName, FNAME_Find) : NAME_None);
 	if (OnlineIdentity.IsValid())
 	{
 		// Add delegates for the various async calls
-		OnLoginCompleteDelegateHandle  = OnlineIdentity->AddOnLoginCompleteDelegate_Handle (LocalUserIdx, OnLoginCompleteDelegate);
+		OnLoginCompleteDelegateHandle  = OnlineIdentity->AddOnLoginCompleteDelegate_Handle(LocalUserIdx, OnLoginCompleteDelegate);
 		OnLogoutCompleteDelegateHandle = OnlineIdentity->AddOnLogoutCompleteDelegate_Handle(LocalUserIdx, OnLogoutCompleteDelegate);
 		
 		// kick off next test
@@ -44,6 +43,7 @@ void FTestIdentityInterface::StartNextTest()
 	}
 	else if (bRunLogoutTest)
 	{
+		IsTheUserLoggedIn();
 		OnlineIdentity->Logout(LocalUserIdx);
 	}
 	else
@@ -58,9 +58,10 @@ void FTestIdentityInterface::FinishTest()
 	if (OnlineIdentity.IsValid())
 	{
 		// Clear delegates for the various async calls
-		OnlineIdentity->ClearOnLoginCompleteDelegate_Handle (LocalUserIdx, OnLoginCompleteDelegateHandle);
+		OnlineIdentity->ClearOnLoginCompleteDelegate_Handle(LocalUserIdx, OnLoginCompleteDelegateHandle);
 		OnlineIdentity->ClearOnLogoutCompleteDelegate_Handle(LocalUserIdx, OnLogoutCompleteDelegateHandle);
 	}
+	SetTestStatus(true);
 	delete this;
 }
 
@@ -68,7 +69,7 @@ void FTestIdentityInterface::OnLoginComplete(int32 LocalUserNum, bool bWasSucces
 {
 	if (bWasSuccessful)
 	{
-		UE_LOG(LogOnline, Log, TEXT("Successful authenticated user. UserId=[%s] "), 
+		UE_LOG(LogOnline, Display, TEXT("Successful logged in user. UserId=[%s] "), 
 			*UserId.ToDebugString());
 
 		// update user info for newly registered user
@@ -76,8 +77,9 @@ void FTestIdentityInterface::OnLoginComplete(int32 LocalUserNum, bool bWasSucces
 	}
 	else
 	{
-		UE_LOG(LogOnline, Warning, TEXT("Failed to authenticate new user. Error=[%s]"), 
+		UE_LOG(LogOnline, Error, TEXT("Failed to log in new user. Error=[%s]"), 
 			*Error);
+
 	}
 	// Mark test as done
 	bRunLoginTest = false;
@@ -89,13 +91,19 @@ void FTestIdentityInterface::OnLogoutComplete(int32 LocalUserNum, bool bWasSucce
 {
 	if (bWasSuccessful)
 	{
-		UE_LOG(LogOnline, Log, TEXT("Successful logged out user. LocalUserNum=[%d] "),
+		UE_LOG(LogOnline, Display, TEXT("Successful logged out user. LocalUserNum=[%d] "),
 			LocalUserNum);
 	}
+	else if (!bIsUserLoggedIn)
+	{
+		UE_LOG(LogOnline, Display, TEXT("User is not logged in to be able to be logged out."));
+	}
+	// If the user was logged out at the start of the test then there will be nothing to log out.
 	else
 	{
-		UE_LOG(LogOnline, Warning, TEXT("Failed to log out user."));
+		UE_LOG(LogOnline, Error, TEXT("Failed to log out user."));
 	}
+
 	UserInfo.Reset();
 	// Mark test as done
 	bRunLogoutTest = false;
@@ -103,3 +111,30 @@ void FTestIdentityInterface::OnLogoutComplete(int32 LocalUserNum, bool bWasSucce
 	StartNextTest();
 }
 
+bool FTestIdentityInterface::GetTestStatus() const
+{
+	return bIsTestFinished;
+}
+
+void FTestIdentityInterface::SetTestStatus(const bool& NewStatus)
+{
+	bIsTestFinished = NewStatus;
+}
+
+bool FTestIdentityInterface::IsTheUserLoggedIn()
+{
+	// Create a new Identity pointer since this function can be potentially called before original pointer can be initialized.
+	IOnlineIdentityPtr TempInterfacePointer = Online::GetIdentityInterface( 
+		GEngine->GetWorld(), 
+		SubsystemName.Len() ? FName( *SubsystemName, FNAME_Find ) : NAME_None );
+	
+	ELoginStatus::Type CurrentLoginStatus = TempInterfacePointer->GetLoginStatus(LocalUserIdx);
+	bIsUserLoggedIn = false;
+
+	if ( CurrentLoginStatus == ELoginStatus::LoggedIn )
+	{
+		bIsUserLoggedIn = true;
+	}
+
+	return bIsUserLoggedIn;
+}
\ No newline at end of file
diff --git a/Engine/Source/Runtime/Online/OnlineSubsystemUtils/Private/Tests/TestIdentityInterface.h b/Engine/Source/Runtime/Online/OnlineSubsystemUtils/Private/Tests/TestIdentityInterface.h
--- a/Engine/Source/Runtime/Online/OnlineSubsystemUtils/Private/Tests/TestIdentityInterface.h
+++ b/Engine/Source/Runtime/Online/OnlineSubsystemUtils/Private/Tests/TestIdentityInterface.h
@@ -29,6 +29,10 @@
 	TSharedPtr<FUserOnlineAccount> UserInfo;
 	/** local user to run tests for */
 	int32 LocalUserIdx;
+	/** true if the test was successful. */
+	bool bIsTestFinished;
+	/** true if the user is logged in*/
+	bool bIsUserLoggedIn;
 
 	FOnlineAccountCredentials AccountCredentials;
 
@@ -69,12 +73,20 @@
 		, bRunLoginTest(true)
 		, bRunLogoutTest(false)
 		, LocalUserIdx(0)
+		, bIsTestFinished(false)
 	{
-		
 	}
 
 	/**
 	 * Kicks off all of the testing process
 	 */
 	void Test(UWorld* InWorld, const FOnlineAccountCredentials& InAccountCredentials, bool bOnlyRunLogoutTest=false);
+
+	/**
+	* Used to get and set the private variables.
+	*/
+	bool GetTestStatus() const;
+	void SetTestStatus( const bool& NewStatus );
+
+	bool IsTheUserLoggedIn();
  };
