Merging CL 2661713 using Main-Branch-Map
#UE4 #HTTP #Curl
- Fixed a bug that would prevent HTTP retries from working. TimeSinceLastResponse was not getting reset with ElapsedTime so once a request timed out it would always time out N more times until the retry sys gives up

diff --git a/Engine/Source/Runtime/Online/HTTP/Private/Curl/CurlHttp.cpp b/Engine/Source/Runtime/Online/HTTP/Private/Curl/CurlHttp.cpp
--- a/Engine/Source/Runtime/Online/HTTP/Private/Curl/CurlHttp.cpp
+++ b/Engine/Source/Runtime/Online/HTTP/Private/Curl/CurlHttp.cpp
@@ -608,6 +608,7 @@ bool FCurlHttpRequest::ProcessRequest()
 	FHttpModule::Get().GetHttpManager().AddRequest(SharedThis(this));
 	// reset timeout
 	ElapsedTime = 0.0f;
+	TimeSinceLastResponse = 0.0f;
 	
 	UE_LOG(LogHttp, Verbose, TEXT("%p: request (easy handle:%p) has been added to multi handle for processing"), this, EasyHandle );
 
