UE-19627 - fix for game restarting when coming back from Home or Lock Screen
#ios

diff --git a/Engine/Source/Runtime/Online/IOS/OnlineSubsystemIOS/Private/OnlineExternalUIInterfaceIOS.cpp b/Engine/Source/Runtime/Online/IOS/OnlineSubsystemIOS/Private/OnlineExternalUIInterfaceIOS.cpp
--- a/Engine/Source/Runtime/Online/IOS/OnlineSubsystemIOS/Private/OnlineExternalUIInterfaceIOS.cpp
+++ b/Engine/Source/Runtime/Online/IOS/OnlineSubsystemIOS/Private/OnlineExternalUIInterfaceIOS.cpp
@@ -32,57 +32,11 @@ bool FOnlineExternalUIIOS::ShowLoginUI(const int ControllerIndex, bool bShowOnli
 	// Not authenticated, set a handler
 	
 	// Copy the delegate so that the block can still access it when it runs.
-	FOnLoginUIClosedDelegate CopiedDelegate = Delegate;
+	CopiedDelegate = Delegate;
 	
-	// Trigger the login event on the main thread.
-	dispatch_async(dispatch_get_main_queue(), ^
-	{
-		[IdentityInterface->GetLocalGameCenterUser() setAuthenticateHandler:(^(UIViewController* viewcontroller, NSError *error)
-		{
-			UE_LOG(LogOnline, Log, TEXT("Game Center authenticateHandler. viewController = %p"), viewcontroller);
-			
-			// The login process has completed.
-			if (viewcontroller == nil)
-			{
-				FString ErrorMessage;
-				TSharedPtr<FUniqueNetIdString> UniqueNetId;
-
-				if (IdentityInterface->GetLocalGameCenterUser().isAuthenticated == YES)
-				{
-					// Trigger the completion delegate on the game thread.
-					const FString PlayerId(IdentityInterface->GetLocalGameCenterUser().playerID);
-					UniqueNetId = MakeShareable(new FUniqueNetIdString(PlayerId));
-					UE_LOG(LogOnline, Log, TEXT("The user %s has logged into Game Center"), *PlayerId);
-				}
-				else
-				{
-					ErrorMessage = TEXT("The user could not be authenticated by Game Center");
-					UE_LOG(LogOnline, Log, TEXT("%s"), *ErrorMessage);
-				}
-
-				if (error)
-				{
-					NSString *errstr = [error localizedDescription];
-					UE_LOG(LogOnline, Warning, TEXT("Game Center login has failed. %s]"), *FString(errstr));
-				}
-
-				// Report back to the game thread whether this succeeded.
-				[FIOSAsyncTask CreateTaskWithBlock : ^ bool(void)
-				{
-					IdentityInterface->SetLocalPlayerUniqueId(UniqueNetId);
-					CopiedDelegate.ExecuteIfBound(UniqueNetId, ControllerIndex);
-
-					return true;
-				}];
-			}
-			else
-			{
-				// Game Center has provided a view controller for us to login, we present it.
-				[[IOSAppDelegate GetDelegate].IOSController 
-					presentViewController:viewcontroller animated:YES completion:nil];
-			}
-		})];
-	});
+	// add a Login Complete delegat to the Identity Interface and attempt to Login
+	CompleteDelegate = IdentityInterface->AddOnLoginCompleteDelegate_Handle(ControllerIndex, FOnLoginCompleteDelegate::CreateRaw(this, &FOnlineExternalUIIOS::OnLoginComplete));
+	IdentityInterface->Login(ControllerIndex, FOnlineAccountCredentials());
 	
 	return true;
 }
@@ -126,3 +80,18 @@ bool FOnlineExternalUIIOS::ShowAccountUpgradeUI(const FUniqueNetId& UniqueId)
 {
 	return false;
 }
+
+void FOnlineExternalUIIOS::OnLoginComplete(int ControllerIndex, bool bWasSuccessful, const FUniqueNetId& UserId, const FString& ErrorString)
+{
+    FOnlineIdentityIOS* IdentityInterface = static_cast<FOnlineIdentityIOS*>(Subsystem->GetIdentityInterface().Get());
+    TSharedPtr<FUniqueNetIdString> UniqueNetId;
+    if (bWasSuccessful)
+    {
+        const FString PlayerId(IdentityInterface->GetLocalGameCenterUser().playerID);
+        UniqueNetId = MakeShareable(new FUniqueNetIdString(PlayerId));
+    }
+    CopiedDelegate.ExecuteIfBound(UniqueNetId, ControllerIndex);
+
+	check(IdentityInterface != nullptr);
+	IdentityInterface->ClearOnLoginCompleteDelegate_Handle(ControllerIndex, CompleteDelegate);
+}
\ No newline at end of file
diff --git a/Engine/Source/Runtime/Online/IOS/OnlineSubsystemIOS/Private/OnlineExternalUIInterfaceIOS.h b/Engine/Source/Runtime/Online/IOS/OnlineSubsystemIOS/Private/OnlineExternalUIInterfaceIOS.h
--- a/Engine/Source/Runtime/Online/IOS/OnlineSubsystemIOS/Private/OnlineExternalUIInterfaceIOS.h
+++ b/Engine/Source/Runtime/Online/IOS/OnlineSubsystemIOS/Private/OnlineExternalUIInterfaceIOS.h
@@ -28,6 +28,10 @@ public:
 	
 private:
 	FOnlineSubsystemIOS* Subsystem;
+	FDelegateHandle CompleteDelegate;
+	FOnLoginUIClosedDelegate CopiedDelegate;
+
+    void OnLoginComplete(int ControllerIndex, bool bWasSuccessful, const FUniqueNetId& UserId, const FString& ErrorString);
 };
 
 typedef TSharedPtr<FOnlineExternalUIIOS, ESPMode::ThreadSafe> FOnlineExternalUIIOSPtr;
\ No newline at end of file
