GitHub 1245 : Allow for servers to create multiple voice data channels.

diff --git a/Engine/Source/Runtime/Engine/Private/World.cpp b/Engine/Source/Runtime/Engine/Private/World.cpp
--- a/Engine/Source/Runtime/Engine/Private/World.cpp
+++ b/Engine/Source/Runtime/Engine/Private/World.cpp
@@ -3701,6 +3701,13 @@ bool UWorld::NotifyAcceptingChannel( UChannel* Channel )
 			//UE_LOG(LogWorld, Log,  "Client accepting actor channel" );
 			return 1;
 		}
+		else if (Channel->ChType == CHTYPE_Voice)
+		{
+			// Accept server requests to open a voice channel, allowing for custom voip implementations
+			// which utilize multiple server controlled voice channels.
+			//UE_LOG(LogNet, Log,  "Client accepting voice channel" );
+			return 1;
+		}
 		else
 		{
 			// Unwanted channel type.
diff --git a/Engine/Source/Runtime/Online/OnlineSubsystemUtils/Private/OnlineBeacon.cpp b/Engine/Source/Runtime/Online/OnlineSubsystemUtils/Private/OnlineBeacon.cpp
--- a/Engine/Source/Runtime/Online/OnlineSubsystemUtils/Private/OnlineBeacon.cpp
+++ b/Engine/Source/Runtime/Online/OnlineSubsystemUtils/Private/OnlineBeacon.cpp
@@ -144,6 +144,13 @@ bool AOnlineBeacon::NotifyAcceptingChannel(UChannel* Channel)
 			UE_LOG(LogNet, Log,  TEXT("Client accepting actor channel"));
 			return 1;
 		}
+		else if (Channel->ChType == CHTYPE_Voice)
+		{
+			// Accept server requests to open a voice channel, allowing for custom voip implementations
+			// which utilize multiple server controlled voice channels.
+			UE_LOG(LogNet, Log,  TEXT("Client accepting voice channel"));
+			return 1;
+		}
 		else
 		{
 			// Unwanted channel type.
