Fix a potential crash in OnlineSubsystemNull.

diff --git a/Engine/Source/Runtime/Online/OnlineSubsystemNull/Private/OnlineSessionInterfaceNull.cpp b/Engine/Source/Runtime/Online/OnlineSubsystemNull/Private/OnlineSessionInterfaceNull.cpp
--- a/Engine/Source/Runtime/Online/OnlineSubsystemNull/Private/OnlineSessionInterfaceNull.cpp
+++ b/Engine/Source/Runtime/Online/OnlineSubsystemNull/Private/OnlineSessionInterfaceNull.cpp
@@ -543,7 +543,7 @@ uint32 FOnlineSessionNull::FindLANSession()
 bool FOnlineSessionNull::CancelFindSessions()
 {
 	uint32 Return = E_FAIL;
-	if (CurrentSessionSearch->SearchState == EOnlineAsyncTaskState::InProgress)
+	if (CurrentSessionSearch.IsValid() && CurrentSessionSearch->SearchState == EOnlineAsyncTaskState::InProgress)
 	{
 		// Make sure it's the right type
 		Return = ERROR_SUCCESS;
