Integrate 2676006 from UE4-Orion

- Add option for host-migration for PS4 sessions using a custom session setting

diff --git a/Engine/Source/Runtime/Online/OnlineSubsystem/Public/OnlineSubsystemSessionSettings.h b/Engine/Source/Runtime/Online/OnlineSubsystem/Public/OnlineSubsystemSessionSettings.h
--- a/Engine/Source/Runtime/Online/OnlineSubsystem/Public/OnlineSubsystemSessionSettings.h
+++ b/Engine/Source/Runtime/Online/OnlineSubsystem/Public/OnlineSubsystemSessionSettings.h
@@ -44,3 +44,7 @@
 /** Session member group identifier (value is FString) - This is to be used with FString::Printf, populated with user xuid - Field required with Team Based matchmaking*/
 #define SETTING_GROUP_NAME FName(TEXT("USERGROUPNAME%s"))
 
+// These are PS4 specific settings to be used in FOnlineSessionSettings as keys
+
+/** Enables host migration for PS4 sessions, which is handled on Sony's servers. This means that all clients can update the session, we do not get callbacks when we become the host or a chance to set the host, it is picked automatically. */
+#define SETTING_HOST_MIGRATION FName(TEXT("HOSTMIGRATION"))
\ No newline at end of file
