Prevent BuildPatchTool updating timestamps of chunk files in Compactify mode, as they're no longer honored by builders.

Also, remove -nopatchdelete mode, as it no longer makes sense now that we don't update timestamps.

diff --git a/Engine/Source/Programs/AutomationTool/AutomationUtils/MCPPublic.cs b/Engine/Source/Programs/AutomationTool/AutomationUtils/MCPPublic.cs
--- a/Engine/Source/Programs/AutomationTool/AutomationUtils/MCPPublic.cs
+++ b/Engine/Source/Programs/AutomationTool/AutomationUtils/MCPPublic.cs
@@ -393,10 +393,6 @@ namespace EpicGames.MCP.Automation
 			/// </summary>
 			public bool bPreviewCompactify;
 			/// <summary>
-			/// Corresponds to the -nopatchdelete parameter
-			/// </summary>
-			public bool bNoPatchDeleteCompactify;
-			/// <summary>
 			/// The full list of manifest files in the compactify directory that we wish to keep; all others will be deleted.
 			/// </summary>
 			public string[] ManifestsToKeep;
diff --git a/Engine/Source/Programs/BuildPatchTool/Private/BuildPatchToolMain.cpp b/Engine/Source/Programs/BuildPatchTool/Private/BuildPatchToolMain.cpp
--- a/Engine/Source/Programs/BuildPatchTool/Private/BuildPatchToolMain.cpp
+++ b/Engine/Source/Programs/BuildPatchTool/Private/BuildPatchToolMain.cpp
@@ -57,7 +57,6 @@
 	Optional arguments:
 	-stdout				Adds stdout logging to the app.
 	-preview			Log all the actions it will take to update internal structures, but don't actually execute them.
-	-nopatchdelete		When specified, will only set the modified date of referenced files to the current time, but will NOT delete any patch data.
 	-ManifestsList=""			Specifies in quotes, the list of manifest filenames to keep following the operation. If omitted, all manifests are kept.
 	-ManifestsFile=""			Specifies in quotes, the name of the file (relative to CloudDir) which contains a list of manifests to keep, one manifest filename per line
 	-DataAgeThreshold=14.25		The maximum age in days of chunk files that will be retained. All older chunks will be deleted.
@@ -179,7 +178,6 @@ int32 BuildPatchToolMain( const TCHAR* CommandLine )
 
 	EBuildPatchToolMode::Type ToolMode = EBuildPatchToolMode::Unknown;
 	bool bPreview = false;
-	bool bNoPatchDelete = false;
 	bool bPatchWithReuseAgeThreshold = true;
 	bool bIncludeSizes = false;
 
@@ -220,7 +218,6 @@ int32 BuildPatchToolMain( const TCHAR* CommandLine )
 			ToolMode = EBuildPatchToolMode::PatchGeneration;
 		}
 		bPreview = Switches.IndexOfByPredicate(FCommandLineMatcher(TEXT("preview"))) != INDEX_NONE;
-		bNoPatchDelete = Switches.IndexOfByPredicate(FCommandLineMatcher(TEXT("nopatchdelete"))) != INDEX_NONE;
 		bIncludeSizes = Switches.IndexOfByPredicate(FCommandLineMatcher(TEXT("includesizes"))) != INDEX_NONE;
 		BuildRootIdx = Switches.IndexOfByPredicate(FCommandLineMatcher(TEXT("BuildRoot")));
 		CloudDirIdx = Switches.IndexOfByPredicate(FCommandLineMatcher(TEXT("CloudDir")));
@@ -415,12 +412,6 @@ int32 BuildPatchToolMain( const TCHAR* CommandLine )
 		return 1;
 	}
 
-	if (ToolMode == EBuildPatchToolMode::Compactify && bPreview && bNoPatchDelete)
-	{
-		GLog->Log(ELogVerbosity::Error, TEXT("Only one of -preview and -nopatchdelete can be specified"));
-		return 5;
-	}
-
 	if (!IgnoreListFile.IsEmpty() && !FPaths::FileExists(IgnoreListFile))
 	{
 		GLog->Logf(ELogVerbosity::Error, TEXT("Provided file ignore list was not found %s"), *IgnoreListFile);
@@ -479,10 +470,6 @@ int32 BuildPatchToolMain( const TCHAR* CommandLine )
 		{
 			CompactifyMode = ECompactifyMode::Preview;
 		}
-		else if (bNoPatchDelete)
-		{
-			CompactifyMode = ECompactifyMode::NoPatchDelete;
-		}
 
 		// Run the compactify routine
 		bSuccess = BuildPatchServicesModule->CompactifyCloudDirectory(ManifestsArr, DataAgeThreshold, CompactifyMode);
diff --git a/Engine/Source/Runtime/Online/BuildPatchServices/Private/BuildPatchCompactifier.cpp b/Engine/Source/Runtime/Online/BuildPatchServices/Private/BuildPatchCompactifier.cpp
--- a/Engine/Source/Runtime/Online/BuildPatchServices/Private/BuildPatchCompactifier.cpp
+++ b/Engine/Source/Runtime/Online/BuildPatchServices/Private/BuildPatchCompactifier.cpp
@@ -9,35 +9,27 @@ data that are no longer referenced by the manifests in a given cloud directory.
 
 #if WITH_BUILDPATCHGENERATION
 
-namespace BuildDataCompactifierDefs
-{
-	const static uint32 TouchFailureThresholdPercentage = 10;
-	const static int64 TimeStampTolerance = FTimespan::FromMinutes(1).GetTicks();
-}
-
 /* Constructors
 *****************************************************************************/
 
-FBuildDataCompactifier::FBuildDataCompactifier(const FString& InCloudDir, const bool bInPreview, const bool bInNoPatchDelete)
+FBuildDataCompactifier::FBuildDataCompactifier(const FString& InCloudDir, const bool bInPreview)
 	: CloudDir(InCloudDir)
 	, bPreview(bInPreview)
-	, bNoPatchDelete(bInNoPatchDelete)
 {
-
 }
 
 /* Public static methods
 *****************************************************************************/
 
-bool FBuildDataCompactifier::CompactifyCloudDirectory(const FString& CloudDir, const TArray<FString>& ManifestsToKeep, const float DataAgeThreshold, const bool bPreview, const bool bNoPatchDelete)
+bool FBuildDataCompactifier::CompactifyCloudDirectory(const FString& CloudDir, const TArray<FString>& ManifestsToKeep, const float DataAgeThreshold, const bool bPreview)
 {
-	FBuildDataCompactifier Compactifier(CloudDir, bPreview, bNoPatchDelete);
+	FBuildDataCompactifier Compactifier(CloudDir, bPreview);
 	return Compactifier.Compactify(ManifestsToKeep, DataAgeThreshold);
 }
 
-bool FBuildDataCompactifier::CompactifyCloudDirectory(const TArray<FString>& ManifestsToKeep, const float DataAgeThreshold, const bool bPreview, const bool bNoPatchDelete)
+bool FBuildDataCompactifier::CompactifyCloudDirectory(const TArray<FString>& ManifestsToKeep, const float DataAgeThreshold, const bool bPreview)
 {
-	return CompactifyCloudDirectory(FBuildPatchServicesModule::GetCloudDirectory(), ManifestsToKeep, DataAgeThreshold, bPreview, bNoPatchDelete);
+	return CompactifyCloudDirectory(FBuildPatchServicesModule::GetCloudDirectory(), ManifestsToKeep, DataAgeThreshold, bPreview);
 }
 
 /* Private methods
@@ -45,7 +37,7 @@ bool FBuildDataCompactifier::CompactifyCloudDirectory(const TArray<FString>& Man
 
 bool FBuildDataCompactifier::Compactify(const TArray<FString>& ManifestsToKeep, const float DataAgeThreshold) const
 {
-	GLog->Logf(TEXT("Running Compactify on %s%s"), *CloudDir, bPreview ? TEXT(". Preview mode. NO action will be taken.") : bNoPatchDelete ? TEXT(". NoPatchDelete mode. NO patch data will be deleted.") : TEXT(""));
+	GLog->Logf(TEXT("Running Compactify on %s%s"), *CloudDir, bPreview ? TEXT(". Preview mode. NO action will be taken.") : TEXT(""));
 	if (ManifestsToKeep.Num() > 0)
 	{
 		GLog->Logf(TEXT("Preserving manifest files: %s"), *FString::Join(ManifestsToKeep, TEXT(", ")));
@@ -124,24 +116,18 @@ bool FBuildDataCompactifier::Compactify(const TArray<FString>& ManifestsToKeep,
 		}
 	}
 
-	GLog->Logf(TEXT("Compactify walking %s to touch referenced chunks, remove all aged unreferenced chunks and compute statistics."), *CloudDir);
+	GLog->Logf(TEXT("Compactify walking %s to remove all mature unreferenced chunks and compute statistics."), *CloudDir);
 
 	uint32 FilesProcessed = 0;
-	uint32 FilesTouched = 0;
 	uint32 FilesSkipped = 0;
 	uint32 NonPatchFilesProcessed = 0;
 	uint32 FilesDeleted = DeletedManifestFilenames.Num();
 	uint64 BytesProcessed = 0;
-	uint64 BytesTouched = 0;
 	uint64 BytesSkipped = 0;
 	uint64 NonPatchBytesProcessed = 0;
 	uint64 BytesDeleted = ManifestBytesDeleted;
 	uint64 CurrentFileSize;
-	uint32 TouchFailureCount = 0;
-	uint64 BytesFailedTouch = 0;
 	FGuid FileGuid;
-	FDateTime Now = FDateTime::UtcNow();
-	const uint32 TouchFailureErrorThreshold = FMath::DivideAndRoundUp((uint32)AllFiles.Num(), BuildDataCompactifierDefs::TouchFailureThresholdPercentage);
 
 	for (const auto& File : AllFiles)
 	{
@@ -162,38 +148,21 @@ bool FBuildDataCompactifier::Compactify(const TArray<FString>& ManifestsToKeep,
 				continue;
 			}
 
-			if (ReferencedGuids.Contains(FileGuid))
+			if (!ReferencedGuids.Contains(FileGuid))
 			{
-				// This is a valid, referenced file, so we need to touch it with the current date
-				++FilesTouched;
-				BytesTouched += CurrentFileSize;
-				if (!bPreview)
+				if (IFileManager::Get().GetTimeStamp(*File) < Cutoff)
 				{
-					if (!SafeSetTimeStamp(*File, Now))
-					{
-						++TouchFailureCount;
-						BytesFailedTouch += CurrentFileSize;
-						GLog->Logf(ELogVerbosity::Warning, TEXT("Failed to set timestamp on file %s"), *File);
-						if (TouchFailureCount >= TouchFailureErrorThreshold)
-						{
-							GLog->Logf(ELogVerbosity::Error, TEXT("Failed to set timestamp on %u%% of files"), BuildDataCompactifierDefs::TouchFailureThresholdPercentage);
-							return false;
-						}
-					}
+					// This file is not referenced by any manifest, is a data file, and is older than we need to keep ...
+					// Let's get rid of it!
+					DeleteFile(File);
+					++FilesDeleted;
+					BytesDeleted += CurrentFileSize;
+				}
+				else
+				{
+					++FilesSkipped;
+					BytesSkipped += CurrentFileSize;
 				}
-			}
-			else if (IFileManager::Get().GetTimeStamp(*File) < Cutoff)
-			{
-				// This file is not referenced by any manifest, is a data file, and is older than we want to keep ...
-				// Let's get rid of it!
-				DeleteFile(File);
-				++FilesDeleted;
-				BytesDeleted += CurrentFileSize;
-			}
-			else
-			{
-				++FilesSkipped;
-				BytesSkipped += CurrentFileSize;
 			}
 		}
 		else
@@ -202,28 +171,13 @@ bool FBuildDataCompactifier::Compactify(const TArray<FString>& ManifestsToKeep,
 		}
 	}
 
-	FilesTouched -= TouchFailureCount;
-	BytesTouched -= BytesFailedTouch;
-
 	GLog->Logf(TEXT("Compactify of %s complete!"), *CloudDir);
 	GLog->Logf(TEXT("Compactify found %u files totalling %s."), FilesProcessed, *HumanReadableSize(BytesProcessed));
 	if (NonPatchFilesProcessed > 0)
 	{
 		GLog->Logf(TEXT("Of these, %u (totalling %s) were not chunk/manifest files."), NonPatchFilesProcessed, *HumanReadableSize(NonPatchBytesProcessed));
 	}
-	if (bNoPatchDelete)
-	{
-		GLog->Logf(TEXT("Compactify skipped deleting %u unreferenced chunk/manifest files (totalling %s) due to -nopatchdelete being specified."), FilesDeleted, *HumanReadableSize(BytesDeleted));
-	}
-	else
-	{
-		GLog->Logf(TEXT("Compactify deleted %u chunk/manifest files totalling %s."), FilesDeleted, *HumanReadableSize(BytesDeleted));
-	}
-	GLog->Logf(TEXT("Compactify touched %u chunk files totalling %s."), FilesTouched, *HumanReadableSize(BytesTouched));
-	if (TouchFailureCount > 0)
-	{
-		GLog->Logf(TEXT("Compactify failed to touch %u files toatlling %s."), TouchFailureCount, *HumanReadableSize(BytesFailedTouch));
-	}
+	GLog->Logf(TEXT("Compactify deleted %u chunk/manifest files totalling %s."), FilesDeleted, *HumanReadableSize(BytesDeleted));
 	GLog->Logf(TEXT("Compactify skipped %u unreferenced chunk files (totalling %s) which have not yet aged out."), FilesSkipped, *HumanReadableSize(BytesSkipped));
 	return true;
 }
@@ -231,7 +185,7 @@ bool FBuildDataCompactifier::Compactify(const TArray<FString>& ManifestsToKeep,
 void FBuildDataCompactifier::DeleteFile(const FString& FilePath) const
 {
 	FString LogMsg = FString::Printf(TEXT("Deprecated data %s"), *FilePath);
-	if (!bNoPatchDelete && !bPreview)
+	if (!bPreview)
 	{
 		LogMsg = LogMsg.Append(TEXT(" ... deleted"));
 		IFileManager::Get().Delete(*FilePath);
@@ -277,7 +231,7 @@ bool FBuildDataCompactifier::DeleteNonReferencedManifests(TArray<FString>& AllMa
 			DeletedManifests.Add(Manifest);
 			BytesDeleted += IFileManager::Get().FileSize(*ManifestPath);
 
-			if (!bPreview && !bNoPatchDelete)
+			if (!bPreview)
 			{
 				IFileManager::Get().Delete(*ManifestPath);
 				if (FPlatformFileManager::Get().GetPlatformFile().FileExists(*ManifestPath))
@@ -341,18 +295,4 @@ FString FBuildDataCompactifier::HumanReadableSize(uint64 NumBytes, uint8 Decimal
 
 	return FString::Printf(TEXT("%.*f %s"), DecimalPlaces, DataSize / (FMath::Pow(Base, Index)), Suffixes[bUseBase10 ? 0 : 1][Index]);
 }
-
-bool FBuildDataCompactifier::SafeSetTimeStamp(const FString& FilePath, FDateTime TimeStamp) const
-{
-	bool bSuccess = false;
-	IFileManager& FileMan = IFileManager::Get();
-	bSuccess = FileMan.SetTimeStamp(*FilePath, TimeStamp);
-	if (bSuccess)
-	{
-		FDateTime ActualTimeStamp = FileMan.GetTimeStamp(*FilePath);
-		int64 TickVariance = FMath::Abs(ActualTimeStamp.GetTicks() - TimeStamp.GetTicks());
-		bSuccess = TickVariance < BuildDataCompactifierDefs::TimeStampTolerance;
-	}
-	return bSuccess;
-}
 #endif
diff --git a/Engine/Source/Runtime/Online/BuildPatchServices/Private/BuildPatchCompactifier.h b/Engine/Source/Runtime/Online/BuildPatchServices/Private/BuildPatchCompactifier.h
--- a/Engine/Source/Runtime/Online/BuildPatchServices/Private/BuildPatchCompactifier.h
+++ b/Engine/Source/Runtime/Online/BuildPatchServices/Private/BuildPatchCompactifier.h
@@ -23,42 +23,38 @@ public:
 	 * @param ManifestsToKeep   If specified, these manifests will be retained, and all others will be deleted
 	 * @param DataAgeThreshold  Chunks which are not referenced by a valid manifest, and which are older than this age (in days), will be deleted
 	 * @param bPreview          If true, then no actual work will be done, but all operations which would be carried out will be logged.
-	 * @param bNoPatchDelete    If true, then no patch data will be deleted, but referenced data will have its timestamp updated to the current time.
 	 * @return                  true if no file errors occurred
 	 */
-	static bool CompactifyCloudDirectory(const TArray<FString>& ManifestsToKeep, const float DataAgeThreshold, const bool bPreview, const bool bNoPatchDelete);
+	static bool CompactifyCloudDirectory(const TArray<FString>& ManifestsToKeep, const float DataAgeThreshold, const bool bPreview);
 
 	/**
 	 * Processes the specified Cloud Directory to identify and delete any orphaned chunks or files.
 	 * NOTE: THIS function is blocking and will not return until finished. Don't run on main thread.
-	 * @param CloudDir      The path to the Cloud Dir to compactify
+	 * @param CloudDir          The path to the Cloud Dir to compactify
 	 * @param ManifestsToKeep   If specified, these manifests will be retained, and all others will be deleted
 	 * @param DataAgeThreshold  Chunks which are not referenced by a valid manifest, and which are older than this age (in days), will be deleted
 	 * @param bPreview          If true, then no actual work will be done, but all operations which would be carried out will be logged.
-	 * @param bNoPatchDelete    If true, then no patch data will be deleted, but referenced data will have its timestamp updated to the current time.
 	 * @return                  true if no file errors occurred
 	 */
-	static bool CompactifyCloudDirectory(const FString& CloudDir, const TArray<FString>& ManifestsToKeep, const float DataAgeThreshold, const bool bPreview, const bool bNoPatchDelete);
+	static bool CompactifyCloudDirectory(const FString& CloudDir, const TArray<FString>& ManifestsToKeep, const float DataAgeThreshold, const bool bPreview);
 
 private:
 	const FString CloudDir;
 	const bool bPreview;
-	const bool bNoPatchDelete;
 
 	/**
 	 * Constructor.
 	 * @param CloudDir         The directory that this compactifier should operate on.
 	 * @param bPreview         If true, then no actual work will be done, but all operations which would be carried out will be logged.
-	 * @param bNoPatchDelete   If true, then no patch data will be deleted, but referenced data will have its timestamp updated to the current time.
 	 */
-	FBuildDataCompactifier(const FString& CloudDir, const bool bPreview, const bool bNoPatchDelete);
+	FBuildDataCompactifier(const FString& CloudDir, const bool bPreview);
 
 	/**
 	 * Processes the Cloud Directory to identify and delete any orphaned chunks or files.
 	 * NOTE: THIS function is blocking and will not return until finished. Don't run on main thread.
 	 * @param ManfestsToKeep      The list of manifest files that we want to keep. If specified, all other manifest files will be deleted.
 	 * @param DataAgeThreshold    The minimum age (in days) of data before it is eligible to be deleted.
-	 * @return              true if no file errors occurred
+	 * @return                    true if no file errors occurred
 	 */
 	bool Compactify(const TArray<FString>& ManifestsToKeep, const float DataAgeThreshold) const;
 
@@ -108,14 +104,6 @@ private:
 	 * @return               A string containing the human readable size
 	 */
 	FString HumanReadableSize(uint64 NumBytes, uint8 DecimalPlaces = 2, bool bUseBase10 = false) const;
-
-	/**
-	 * Sets the modified time on a file, and checks that setting the time was successful
-	 * @param FilePath    The full path to the filename on which to set the modified date
-	 * @param TimeStamp   The timestamp to apply to the file
-	 * @return            True if the operation succeeded, false otherwise
-	 */
-	bool SafeSetTimeStamp(const FString& FilePath, FDateTime TimeStamp) const;
 };
 
 #endif //WITH_BUILDPATCHGENERATION
diff --git a/Engine/Source/Runtime/Online/BuildPatchServices/Private/BuildPatchServicesModule.cpp b/Engine/Source/Runtime/Online/BuildPatchServices/Private/BuildPatchServicesModule.cpp
--- a/Engine/Source/Runtime/Online/BuildPatchServices/Private/BuildPatchServicesModule.cpp
+++ b/Engine/Source/Runtime/Online/BuildPatchServices/Private/BuildPatchServicesModule.cpp
@@ -210,8 +210,7 @@ bool FBuildPatchServicesModule::GenerateFilesManifestFromDirectory( const FBuild
 bool FBuildPatchServicesModule::CompactifyCloudDirectory(const TArray<FString>& ManifestsToKeep, const float DataAgeThreshold, const ECompactifyMode::Type Mode)
 {
 	const bool bPreview = Mode == ECompactifyMode::Preview;
-	const bool bNoPatchDelete = Mode == ECompactifyMode::NoPatchDelete;
-	return FBuildDataCompactifier::CompactifyCloudDirectory(ManifestsToKeep, DataAgeThreshold, bPreview, bNoPatchDelete);
+	return FBuildDataCompactifier::CompactifyCloudDirectory(ManifestsToKeep, DataAgeThreshold, bPreview);
 }
 
 bool FBuildPatchServicesModule::EnumerateManifestData(FString ManifestFilePath, FString OutputFile, const bool bIncludeSizes)
diff --git a/Engine/Source/Runtime/Online/BuildPatchServices/Public/Interfaces/IBuildPatchServicesModule.h b/Engine/Source/Runtime/Online/BuildPatchServices/Public/Interfaces/IBuildPatchServicesModule.h
--- a/Engine/Source/Runtime/Online/BuildPatchServices/Public/Interfaces/IBuildPatchServicesModule.h
+++ b/Engine/Source/Runtime/Online/BuildPatchServices/Public/Interfaces/IBuildPatchServicesModule.h
@@ -20,7 +20,6 @@ namespace ECompactifyMode
 	enum Type
 	{
 		Preview,
-		NoPatchDelete,
 		Full
 	};
 }
