Support for nested serializable objects during conversion to JSON.

https://udn.unrealengine.com/questions/258375/json-serializer-bug-when-using-online-json-seriali.html

diff --git a/Engine/Source/Runtime/Json/Public/Serialization/JsonWriter.h b/Engine/Source/Runtime/Json/Public/Serialization/JsonWriter.h
--- a/Engine/Source/Runtime/Json/Public/Serialization/JsonWriter.h
+++ b/Engine/Source/Runtime/Json/Public/Serialization/JsonWriter.h
@@ -27,7 +27,7 @@ public:
 
 	void WriteObjectStart()
 	{
-		check(CanWriteValueWithoutIdentifier());
+		check(CanWriteObjectWithoutIdentifier());
 		if (PreviousTokenWritten != EJsonToken::None )
 		{
 			WriteCommaIfNeeded();
@@ -351,6 +351,11 @@ protected:
 		return Stack.Num() <= 0 || Stack.Top() == EJson::Array || PreviousTokenWritten == EJsonToken::Identifier;
 	}
 
+	FORCEINLINE bool CanWriteObjectWithoutIdentifier() const
+	{
+		return Stack.Num() <= 0 || Stack.Top() == EJson::Array || PreviousTokenWritten == EJsonToken::Identifier || PreviousTokenWritten == EJsonToken::Colon;
+	}
+
 	FORCEINLINE void WriteCommaIfNeeded()
 	{
 		if ( PreviousTokenWritten != EJsonToken::CurlyOpen && PreviousTokenWritten != EJsonToken::SquareOpen && PreviousTokenWritten != EJsonToken::Identifier)
diff --git a/Engine/Source/Runtime/Online/OnlineSubsystem/Public/OnlineJsonSerializer.h b/Engine/Source/Runtime/Online/OnlineSubsystem/Public/OnlineJsonSerializer.h
--- a/Engine/Source/Runtime/Online/OnlineSubsystem/Public/OnlineJsonSerializer.h
+++ b/Engine/Source/Runtime/Online/OnlineSubsystem/Public/OnlineJsonSerializer.h
@@ -84,26 +84,7 @@
 		}
 
 #define ONLINE_JSON_SERIALIZE_OBJECT_SERIALIZABLE(JsonName, JsonSerializableObject) \
-		/* Process the JsonName field differently because it is an object */ \
-		if (Serializer.IsLoading()) \
-		{ \
-			/* Read in the value from the JsonName field */ \
-			if (Serializer.GetObject()->HasTypedField<EJson::Object>(JsonName)) \
-			{ \
-				TSharedPtr<FJsonObject> JsonObj = Serializer.GetObject()->GetObjectField(JsonName); \
-				if (JsonObj.IsValid()) \
-				{ \
-					JsonSerializableObject.FromJson(JsonObj); \
-				} \
-			} \
-		} \
-		else \
-		{ \
-			/* Write the value to the JsonName field */ \
-			Serializer.StartObject(JsonName); \
-			JsonSerializableObject.Serialize(Serializer); \
-			Serializer.EndObject(); \
-		}
+		Serializer.SerializeObject(TEXT(JsonName), JsonSerializableObject);
 
 /** Array of string data */
 typedef TArray<FString> FJsonSerializableArray;
@@ -113,6 +94,8 @@ typedef FOnlineKeyValuePairs<FString, FString> FJsonSerializableKeyValueMap;
 typedef FOnlineKeyValuePairs<FString, int32> FJsonSerializableKeyValueMapInt;
 typedef FOnlineKeyValuePairs<FString, FVariantData> FJsonSerializeableKeyValueMapVariant;
 
+struct FOnlineJsonSerializable;
+
 /**
  * Base interface used to serialize to/from JSON. Hides the fact there are separate read/write classes
  */
@@ -134,6 +117,7 @@ struct FOnlineJsonSerializerBase
 	virtual void Serialize(const TCHAR* Name, float& Value) = 0;
 	virtual void Serialize(const TCHAR* Name, double& Value) = 0;
 	virtual void Serialize(const TCHAR* Name, FDateTime& Value) = 0;
+	virtual void SerializeObject(const TCHAR* Name, FOnlineJsonSerializable& Value) = 0;
 	virtual void SerializeArray(FJsonSerializableArray& Array) = 0;
 	virtual void SerializeArray(const TCHAR* Name, FJsonSerializableArray& Value) = 0;
 	virtual void SerializeMap(const TCHAR* Name, FJsonSerializableKeyValueMap& Map) = 0;
@@ -290,6 +274,13 @@ public:
 			JsonWriter->WriteValue(Name, Value.ToIso8601());
 		}
 	}
+	/**
+	 * Writes the field name and corresponding object value to the JSON data
+	 *
+	 * @param Name the field name to write out
+	 * @param Object the object to write out
+	 */
+	virtual void SerializeObject(const TCHAR* Name, FOnlineJsonSerializable& Object) override;
 	/**
 	 * Serializes an array of values
 	 *
@@ -517,6 +508,13 @@ public:
 			FDateTime::ParseIso8601(*JsonObject->GetStringField(Name), Value);
 		}
 	}
+	/**
+	 * Writes the field name and the corresponding object value to the JSON data
+	 *
+	 * @param Name the field name to write out
+	 * @Object Value the object to write out
+	 */
+	virtual void SerializeObject(const TCHAR* Name, FOnlineJsonSerializable& Object) override;
 	/**
 	 * Serializes an array of values
 	 *
@@ -666,3 +664,23 @@ struct FOnlineJsonSerializable
 	virtual void Serialize(FOnlineJsonSerializerBase& Serializer) = 0;
 };
 
+inline void FOnlineJsonSerializerReader::SerializeObject(const TCHAR* Name, FOnlineJsonSerializable& Value)
+{
+	/* Read in the value from the Name field */
+	if (GetObject()->HasTypedField<EJson::Object>(Name))
+	{
+		TSharedPtr<FJsonObject> JsonObj = GetObject()->GetObjectField(Name);
+		if (JsonObj.IsValid())
+		{
+			Value.FromJson(JsonObj);
+		}
+	}
+}
+
+template <class CharType, class PrintPolicy>
+inline void FOnlineJsonSerializerWriter<CharType, PrintPolicy>::SerializeObject(const TCHAR* Name, FOnlineJsonSerializable& Object)
+{
+	/* Write the value to the Name field */
+	JsonWriter->WriteIdentifierPrefix(Name);
+	Object.Serialize(*this);
+}
