BuildPatchServices: Tag based installation support
#OPP-3824: BuildPatchServices module interface supports install tags

diff --git a/Engine/Source/Runtime/Online/BuildPatchServices/Private/BuildPatchInstaller.cpp b/Engine/Source/Runtime/Online/BuildPatchServices/Private/BuildPatchInstaller.cpp
--- a/Engine/Source/Runtime/Online/BuildPatchServices/Private/BuildPatchInstaller.cpp
+++ b/Engine/Source/Runtime/Online/BuildPatchServices/Private/BuildPatchInstaller.cpp
@@ -49,21 +49,60 @@ FBuildPatchInstaller::FBuildPatchInstaller(FBuildPatchBoolManifestDelegate InOnC
 	, TimePausedAt( 0.0 )
 	, InstallationInfo( InstallationInfoRef )
 {
-	// Start thread!
-	const TCHAR* ThreadName = TEXT( "BuildPatchInstallerThread" );
-	Thread = FRunnableThread::Create(this, ThreadName);
 }
 
 FBuildPatchInstaller::~FBuildPatchInstaller()
 {
-	if( Thread != NULL )
+	if (Thread != nullptr)
 	{
 		Thread->WaitForCompletion();
 		delete Thread;
-		Thread = NULL;
+		Thread = nullptr;
 	}
 }
 
+bool FBuildPatchInstaller::SetRequiredInstallTags(const TSet<FString>& Tags)
+{
+	// We do not yet support changing tags after installing was started
+	if (Thread != nullptr)
+	{
+		return false;
+	}
+
+	// Check provided tags are all valid
+	TSet<FString> ValidTags;
+	NewBuildManifest->GetFileTagList(ValidTags);
+	if (Tags.Difference(ValidTags).Num() > 0)
+	{
+		return false;
+	}
+
+	// Store for use later
+	InstallTags = Tags;
+	return true;
+}
+
+bool FBuildPatchInstaller::StartInstallation()
+{
+	if (Thread == nullptr)
+	{
+		// Pre-process install tags. Doing this logic here means it doesn't need repeating around lower level code
+		// No tags means full installation
+		if (InstallTags.Num() == 0)
+		{
+			NewBuildManifest->GetFileTagList(InstallTags);
+		}
+
+		// Always require the empty tag
+		InstallTags.Add(TEXT(""));
+
+		// Start thread!
+		const TCHAR* ThreadName = TEXT("BuildPatchInstallerThread");
+		Thread = FRunnableThread::Create(this, ThreadName);
+	}
+	return Thread != nullptr;
+}
+
 bool FBuildPatchInstaller::Init()
 {
 	// Make sure the installation directory exists
@@ -127,7 +166,7 @@ uint32 FBuildPatchInstaller::Run()
 		// Backup local changes then move generated files
 		bInstallSuccess = bInstallSuccess && RunBackupAndMove();
 
-		// There is no more potential for initialising
+		// There is no more potential for initializing
 		BuildProgress.SetStateProgress(EBuildPatchProgress::Initializing, 1.0f);
 
 		// Setup file attributes
@@ -285,6 +324,10 @@ bool FBuildPatchInstaller::RunInstallation(TArray<FString>& CorruptFiles)
 		BuildStats.NumFilesInBuild = NumFilesInBuild;
 	}
 
+	// Get the list of required files, by the tags
+	TaggedFiles.Empty();
+	NewBuildManifest->GetTaggedFileList(InstallTags, TaggedFiles);
+
 	// Check if we should skip out of this process due to existing installation,
 	// that will mean we start with the verification stage
 	bool bFirstTimeRun = CorruptFiles.Num() == 0;
@@ -306,15 +349,17 @@ bool FBuildPatchInstaller::RunInstallation(TArray<FString>& CorruptFiles)
 		return true;
 	}
 
-	// Get the list of files needing construction
-	TArray< FString > FilesToConstruct;
+	// Get the list of files actually needing construction
+	TArray<FString> FilesToConstruct;
 	if (CorruptFiles.Num() > 0)
 	{
 		FilesToConstruct.Append(CorruptFiles);
 	}
 	else
 	{
-		FBuildPatchAppManifest::GetOutdatedFiles(CurrentBuildManifest, NewBuildManifest, InstallDirectory, FilesToConstruct);
+		TSet<FString> OutdatedFiles;
+		FBuildPatchAppManifest::GetOutdatedFiles(CurrentBuildManifest, NewBuildManifest, InstallDirectory, OutdatedFiles);
+		FilesToConstruct = OutdatedFiles.Intersect(TaggedFiles).Array();
 	}
 	GLog->Logf(TEXT("BuildPatchServices: Requiring %d files"), FilesToConstruct.Num());
 
@@ -547,6 +592,12 @@ bool FBuildPatchInstaller::RunBackupAndMove()
 			{
 				FBuildPatchAppManifest::GetRemovableFiles(CurrentBuildManifest.ToSharedRef(), NewBuildManifest, FilesToRemove);
 			}
+			// And also files that may no longer be required (removal of tags)
+			TArray<FString> NewBuildFiles;
+			NewBuildManifest->GetFileList(NewBuildFiles);
+			TSet<FString> NewBuildFilesSet(NewBuildFiles);
+			TSet<FString> RemovableBuildFiles = NewBuildFilesSet.Difference(TaggedFiles);
+			FilesToRemove.Append(RemovableBuildFiles.Array());
 			// Add to build stats
 			ThreadLock.Lock();
 			BuildStats.NumFilesToRemove = FilesToRemove.Num();
@@ -688,6 +739,7 @@ bool FBuildPatchInstaller::RunVerification(TArray< FString >& CorruptFiles)
 	auto Verifier = FBuildPatchVerificationFactory::Create(NewBuildManifest, ProgressDelegate, IsPausedDelegate, InstallDirectory, OptionalStageDirectory);
 
 	// Verify the build
+	Verifier->SetRequiredFiles(TaggedFiles.Array());
 	bool bVerifySuccess = Verifier->VerifyAgainstDirectory(CorruptFiles, VerifyPauseTime);
 	VerifyTime = FPlatformTime::Seconds() - VerifyTime - VerifyPauseTime;
 	if (!bVerifySuccess)
diff --git a/Engine/Source/Runtime/Online/BuildPatchServices/Private/BuildPatchInstaller.h b/Engine/Source/Runtime/Online/BuildPatchServices/Private/BuildPatchInstaller.h
--- a/Engine/Source/Runtime/Online/BuildPatchServices/Private/BuildPatchInstaller.h
+++ b/Engine/Source/Runtime/Online/BuildPatchServices/Private/BuildPatchInstaller.h
@@ -94,7 +94,13 @@ private:
 	double TimePausedAt;
 
 	// Holds a list of files that have been placed into the install directory
-	TArray< FString > FilesInstalled;
+	TArray<FString> FilesInstalled;
+
+	// Holds the install tags for describing the files we want to install
+	TSet<FString> InstallTags;
+
+	// Holds the files which are all required
+	TSet<FString> TaggedFiles;
 
 	// Referecne to the module's installation info
 	FBuildPatchInstallationInfo& InstallationInfo;
@@ -143,6 +149,21 @@ public:
 	virtual bool TogglePauseInstall() override;
 	// End IBuildInstaller interface
 
+	/**
+	 * Set the list of install tags to be installed. Only files which are tagged with one of these will be installed. Untagged files are always installed.
+	 * Special cases:
+	 *      Providing an empty set will result in every file being installed.
+	 *      Providing a set with a single empty string entry, is minimal install, defined by only untagged files.
+	 * @return true if the new set was accepted. false can be returned if the installation is already running, or the set contained a tag not present on the install manifest.
+	 */
+	bool SetRequiredInstallTags(const TSet<FString>& Tags);
+
+	/**
+	 * Begin the installation process
+	 * @return true if the installation started successfully, or is already running
+	 */
+	bool StartInstallation();
+
 	/**
 	 * Executes the on complete delegate. This should only be called when completed, and is separated out
 	 * to allow control to make this call on the main thread.
diff --git a/Engine/Source/Runtime/Online/BuildPatchServices/Private/BuildPatchManifest.cpp b/Engine/Source/Runtime/Online/BuildPatchServices/Private/BuildPatchManifest.cpp
--- a/Engine/Source/Runtime/Online/BuildPatchServices/Private/BuildPatchManifest.cpp
+++ b/Engine/Source/Runtime/Online/BuildPatchServices/Private/BuildPatchManifest.cpp
@@ -715,6 +715,7 @@ void FBuildPatchAppManifest::InitLookups()
 	// Setup internals
 	TotalBuildSize = 0;
 	FileManifestLookup.Empty(Data->FileManifestList.Num());
+	TaggedFilesLookup.Empty();
 	FileNameLookup.Empty(Data->bIsFileData ? Data->FileManifestList.Num() : 0);
 	for (auto& File : Data->FileManifestList)
 	{
@@ -726,6 +727,19 @@ void FBuildPatchAppManifest::InitLookups()
 			// File data chunk parts should have been checked already
 			FileNameLookup.Add(File.FileChunkParts[0].Guid, &File.Filename);
 		}
+		if (File.InstallTags.Num() == 0)
+		{
+			// Untagged files are required
+			TaggedFilesLookup.FindOrAdd(TEXT("")).Add(&File);
+		}
+		else
+		{
+			// Fill out lookup for optional files
+			for (auto& FileTag : File.InstallTags)
+			{
+				TaggedFilesLookup.FindOrAdd(FileTag).Add(&File);
+			}
+		}
 	}
 	TotalDownloadSize = 0;
 	ChunkInfoLookup.Empty(Data->ChunkList.Num());
@@ -1171,6 +1185,9 @@ bool FBuildPatchAppManifest::DeserializeFromJSON( const FString& JSONInput )
 	// Mark as should be re-saved, client that stores manifests should start using binary
 	bNeedsResaving = true;
 
+	// Setup internal lookups
+	InitLookups();
+
 	// Make sure we don't have any half loaded data
 	if( !bSuccess )
 	{
@@ -1284,6 +1301,28 @@ void FBuildPatchAppManifest::GetFileList(TArray<FString>& Filenames) const
 	FileManifestLookup.GetKeys(Filenames);
 }
 
+void FBuildPatchAppManifest::GetFileTagList(TSet<FString>& Tags) const
+{
+	TArray<FString> TagsArray;
+	TaggedFilesLookup.GetKeys(TagsArray);
+	Tags.Append(MoveTemp(TagsArray));
+}
+
+void FBuildPatchAppManifest::GetTaggedFileList(const TSet<FString>& Tags, TSet<FString>& TaggedFiles) const
+{
+	for (auto& Tag : Tags)
+	{
+		auto* Files = TaggedFilesLookup.Find(Tag);
+		if (Files != nullptr)
+		{
+			for (auto& File : *Files)
+			{
+				TaggedFiles.Add(File->Filename);
+			}
+		}
+	}
+}
+
 void FBuildPatchAppManifest::GetDataList(TArray<FGuid>& DataGuids) const
 {
 	ChunkInfoLookup.GetKeys(DataGuids);
@@ -1640,12 +1679,14 @@ bool FBuildPatchAppManifest::NeedsResaving() const
 	return bNeedsResaving;
 }
 
-void FBuildPatchAppManifest::GetOutdatedFiles(FBuildPatchAppManifestPtr OldManifest, FBuildPatchAppManifestRef NewManifest, const FString& InstallDirectory, TArray< FString >& OutDatedFiles)
+void FBuildPatchAppManifest::GetOutdatedFiles(FBuildPatchAppManifestPtr OldManifest, FBuildPatchAppManifestRef NewManifest, const FString& InstallDirectory, TSet< FString >& OutDatedFiles)
 {
 	if (!OldManifest.IsValid())
 	{
 		// All files are outdated if no OldManifest
-		NewManifest->FileManifestLookup.GetKeys(OutDatedFiles);
+		TArray<FString> Filenames;
+		NewManifest->FileManifestLookup.GetKeys(Filenames);
+		OutDatedFiles.Append(MoveTemp(Filenames));
 	}
 	else
 	{
diff --git a/Engine/Source/Runtime/Online/BuildPatchServices/Private/BuildPatchManifest.h b/Engine/Source/Runtime/Online/BuildPatchServices/Private/BuildPatchManifest.h
--- a/Engine/Source/Runtime/Online/BuildPatchServices/Private/BuildPatchManifest.h
+++ b/Engine/Source/Runtime/Online/BuildPatchServices/Private/BuildPatchManifest.h
@@ -491,7 +491,20 @@ public:
 	 * Get the list of files described by this manifest
 	 * @param Filenames		OUT		Receives the array of files.
 	 */
-	void GetFileList(TArray< FString >& Filenames) const;
+	void GetFileList(TArray<FString>& Filenames) const;
+
+	/**
+	 * Get the list of install tags in this manifest
+	 * @param Tags			OUT		Receives the tags referenced.
+	 */
+	void GetFileTagList(TSet<FString>& Tags) const;
+
+	/**
+	 * Get the list of files that are tagged with the provided tags
+	 * @param Tags					The tags for the required file groups.
+	 * @param TaggedFiles	OUT		Receives the tagged files.
+	 */
+	void GetTaggedFileList(const TSet<FString>& Tags, TSet<FString>& TaggedFiles) const;
 
 	/**
 	* Get the list of Guids for all files described by this manifest
@@ -573,9 +586,9 @@ public:
 	 * @param OldManifest		IN		The Build Manifest that is currently installed. Shared Ptr - Can be invalid.
 	 * @param NewManifest		IN		The Build Manifest that is being patched to. Shared Ref - Implicitly valid.
 	 * @param InstallDirectory	IN		The Build installation directory, so that it can be checked for missing files.
-	 * @param OutDatedFiles		OUT		The array of files that do not match or are new.
+	 * @param OutDatedFiles		OUT		The files that changed hash, are new, are wrong size, or missing on disk.
 	 */
-	static void GetOutdatedFiles(FBuildPatchAppManifestPtr OldManifest, FBuildPatchAppManifestRef NewManifest, const FString& InstallDirectory, TArray< FString >& OutDatedFiles);
+	static void GetOutdatedFiles(FBuildPatchAppManifestPtr OldManifest, FBuildPatchAppManifestRef NewManifest, const FString& InstallDirectory, TSet<FString>& OutDatedFiles);
 
 	/**
 	 * Check a single file to see if it will be effected by patching
@@ -624,6 +637,7 @@ private:
 	/** Some lookups to optimize data access */
 	TMap<FGuid, FString*> FileNameLookup;
 	TMap<FString, FFileManifestData*> FileManifestLookup;
+	TMap<FString, TArray<FFileManifestData*>> TaggedFilesLookup;
 	TMap<FGuid, FChunkInfoData*> ChunkInfoLookup;
 	TMap<FString, FCustomFieldData*> CustomFieldLookup;
 
diff --git a/Engine/Source/Runtime/Online/BuildPatchServices/Private/BuildPatchServicesModule.cpp b/Engine/Source/Runtime/Online/BuildPatchServices/Private/BuildPatchServicesModule.cpp
--- a/Engine/Source/Runtime/Online/BuildPatchServices/Private/BuildPatchServicesModule.cpp
+++ b/Engine/Source/Runtime/Online/BuildPatchServices/Private/BuildPatchServicesModule.cpp
@@ -126,7 +126,7 @@ bool FBuildPatchServicesModule::SaveManifestToFile(const FString& Filename, IBui
 	return StaticCastSharedRef< FBuildPatchAppManifest >(Manifest)->SaveToFile(Filename, bUseBinary);
 }
 
-IBuildInstallerPtr FBuildPatchServicesModule::StartBuildInstall( IBuildManifestPtr CurrentManifest, IBuildManifestPtr InstallManifest, const FString& InstallDirectory, FBuildPatchBoolManifestDelegate OnCompleteDelegate )
+IBuildInstallerPtr FBuildPatchServicesModule::StartBuildInstall(IBuildManifestPtr CurrentManifest, IBuildManifestPtr InstallManifest, const FString& InstallDirectory, FBuildPatchBoolManifestDelegate OnCompleteDelegate, TSet<FString> InstallTags)
 {
 	// Using a local bool for this check will improve the assert message that gets displayed
 	const bool bIsCalledFromMainThread = IsInGameThread();
@@ -148,11 +148,14 @@ IBuildInstallerPtr FBuildPatchServicesModule::StartBuildInstall( IBuildManifestP
 	// Make sure the http wrapper is already created
 	FBuildPatchHTTP::Initialize();
 	// Run the install thread
-	BuildPatchInstallers.Add( MakeShareable( new FBuildPatchInstaller( OnCompleteDelegate, CurrentManifestInternal, InstallManifestInternal.ToSharedRef(), InstallDirectory, GetStagingDirectory(), InstallationInfo, false ) ) );
-	return BuildPatchInstallers.Top();
+	FBuildPatchInstallerRef Installer = MakeShareable(new FBuildPatchInstaller(OnCompleteDelegate, CurrentManifestInternal, InstallManifestInternal.ToSharedRef(), InstallDirectory, GetStagingDirectory(), InstallationInfo, false));
+	Installer->SetRequiredInstallTags(InstallTags);
+	Installer->StartInstallation();
+	BuildPatchInstallers.Add(Installer);
+	return Installer;
 }
 
-IBuildInstallerPtr FBuildPatchServicesModule::StartBuildInstallStageOnly(IBuildManifestPtr CurrentManifest, IBuildManifestPtr InstallManifest, const FString& InstallDirectory, FBuildPatchBoolManifestDelegate OnCompleteDelegate)
+IBuildInstallerPtr FBuildPatchServicesModule::StartBuildInstallStageOnly(IBuildManifestPtr CurrentManifest, IBuildManifestPtr InstallManifest, const FString& InstallDirectory, FBuildPatchBoolManifestDelegate OnCompleteDelegate, TSet<FString> InstallTags)
 {
 	// Using a local bool for this check will improve the assert message that gets displayed
 	const bool bIsCalledFromMainThread = IsInGameThread();
@@ -168,8 +171,11 @@ IBuildInstallerPtr FBuildPatchServicesModule::StartBuildInstallStageOnly(IBuildM
 	// Make sure the http wrapper is already created
 	FBuildPatchHTTP::Initialize();
 	// Run the install thread
-	BuildPatchInstallers.Add( MakeShareable( new FBuildPatchInstaller( OnCompleteDelegate, CurrentManifestInternal, InstallManifestInternal.ToSharedRef(), InstallDirectory, GetStagingDirectory(), InstallationInfo, true ) ) );
-	return BuildPatchInstallers.Top();
+	FBuildPatchInstallerRef Installer = MakeShareable(new FBuildPatchInstaller(OnCompleteDelegate, CurrentManifestInternal, InstallManifestInternal.ToSharedRef(), InstallDirectory, GetStagingDirectory(), InstallationInfo, true));
+	Installer->SetRequiredInstallTags(InstallTags);
+	Installer->StartInstallation();
+	BuildPatchInstallers.Add(Installer);
+	return Installer;
 }
 
 bool FBuildPatchServicesModule::Tick( float Delta )
diff --git a/Engine/Source/Runtime/Online/BuildPatchServices/Private/BuildPatchServicesModule.h b/Engine/Source/Runtime/Online/BuildPatchServices/Private/BuildPatchServicesModule.h
--- a/Engine/Source/Runtime/Online/BuildPatchServices/Private/BuildPatchServicesModule.h
+++ b/Engine/Source/Runtime/Online/BuildPatchServices/Private/BuildPatchServicesModule.h
@@ -72,25 +72,25 @@ public:
 	// End IModuleInterface interface
 
 	// Begin IBuildPatchServicesModule interface
-	virtual IBuildManifestPtr LoadManifestFromFile( const FString& Filename ) override;
-	virtual IBuildManifestPtr MakeManifestFromData( const TArray<uint8>& ManifestData ) override;
-	virtual bool SaveManifestToFile( const FString& Filename, IBuildManifestRef Manifest, bool bUseBinary = true ) override;
-	virtual IBuildInstallerPtr StartBuildInstall(IBuildManifestPtr CurrentManifest, IBuildManifestPtr InstallManifest, const FString& InstallDirectory, FBuildPatchBoolManifestDelegate OnCompleteDelegate) override;
-	virtual IBuildInstallerPtr StartBuildInstallStageOnly(IBuildManifestPtr CurrentManifest, IBuildManifestPtr InstallManifest, const FString& InstallDirectory, FBuildPatchBoolManifestDelegate OnCompleteDelegate) override;
-	virtual void SetStagingDirectory( const FString& StagingDir ) override;
-	virtual void SetCloudDirectory( const FString& CloudDir ) override;
-	virtual void SetBackupDirectory( const FString& BackupDir ) override;
-	virtual void SetAnalyticsProvider( TSharedPtr< IAnalyticsProvider > AnalyticsProvider ) override;
-	virtual void SetHttpTracker( TSharedPtr< FHttpServiceTracker > HttpTracker ) override;
+	virtual IBuildManifestPtr LoadManifestFromFile(const FString& Filename) override;
+	virtual IBuildManifestPtr MakeManifestFromData(const TArray<uint8>& ManifestData) override;
+	virtual bool SaveManifestToFile(const FString& Filename, IBuildManifestRef Manifest, bool bUseBinary = true) override;
+	virtual IBuildInstallerPtr StartBuildInstall(IBuildManifestPtr CurrentManifest, IBuildManifestPtr InstallManifest, const FString& InstallDirectory, FBuildPatchBoolManifestDelegate OnCompleteDelegate, TSet<FString> InstallTags = TSet<FString>()) override;
+	virtual IBuildInstallerPtr StartBuildInstallStageOnly(IBuildManifestPtr CurrentManifest, IBuildManifestPtr InstallManifest, const FString& InstallDirectory, FBuildPatchBoolManifestDelegate OnCompleteDelegate, TSet<FString> InstallTags = TSet<FString>()) override;
+	virtual void SetStagingDirectory(const FString& StagingDir) override;
+	virtual void SetCloudDirectory(const FString& CloudDir) override;
+	virtual void SetBackupDirectory(const FString& BackupDir) override;
+	virtual void SetAnalyticsProvider(TSharedPtr< IAnalyticsProvider > AnalyticsProvider) override;
+	virtual void SetHttpTracker(TSharedPtr< FHttpServiceTracker > HttpTracker) override;
 	virtual void RegisterAppInstallation(IBuildManifestRef AppManifest, const FString AppInstallDirectory) override;
 	virtual void CancelAllInstallers(bool WaitForThreads) override;
 #if WITH_BUILDPATCHGENERATION
-	virtual bool GenerateChunksManifestFromDirectory( const FBuildPatchSettings& Settings ) override;
-	virtual bool GenerateFilesManifestFromDirectory( const FBuildPatchSettings& Settings ) override;
-	virtual bool CompactifyCloudDirectory( const TArray<FString>& ManifestsToKeep, const float DataAgeThreshold, const ECompactifyMode::Type Mode ) override;
+	virtual bool GenerateChunksManifestFromDirectory(const FBuildPatchSettings& Settings) override;
+	virtual bool GenerateFilesManifestFromDirectory(const FBuildPatchSettings& Settings) override;
+	virtual bool CompactifyCloudDirectory(const TArray<FString>& ManifestsToKeep, const float DataAgeThreshold, const ECompactifyMode::Type Mode) override;
 	virtual bool EnumerateManifestData(FString ManifestFilePath, FString OutputFile, const bool bIncludeSizes) override;
 #endif
-	virtual IBuildManifestPtr MakeManifestFromJSON( const FString& ManifestJSON ) override;
+	virtual IBuildManifestPtr MakeManifestFromJSON(const FString& ManifestJSON) override;
 	// End IBuildPatchServicesModule interface
 
 	/**
@@ -117,7 +117,7 @@ private:
 	 * @param		Delta	Time since last tick
 	 * @return	Whether to continue ticking
 	 */
-	bool Tick( float Delta );
+	bool Tick(float Delta);
 
 	/**
 	 * This will get called when core PreExits. Make sure any running installers are canceled out.
diff --git a/Engine/Source/Runtime/Online/BuildPatchServices/Private/BuildPatchVerification.cpp b/Engine/Source/Runtime/Online/BuildPatchServices/Private/BuildPatchVerification.cpp
--- a/Engine/Source/Runtime/Online/BuildPatchServices/Private/BuildPatchVerification.cpp
+++ b/Engine/Source/Runtime/Online/BuildPatchServices/Private/BuildPatchVerification.cpp
@@ -9,10 +9,13 @@ public:
 
 	virtual ~FBuildPatchVerificationImpl(){}
 
+	virtual void SetRequiredFiles(TArray<FString> RequiredFiles) override;
+
 	virtual bool VerifyAgainstDirectory(TArray<FString>& OutDatedFiles, double& TimeSpentPaused) override;
 
 private:
 	FBuildPatchAppManifestRef Manifest;
+	TArray<FString> RequiredFiles;
 	FBuildPatchFloatDelegate ProgressDelegate;
 	FBuildPatchBoolRetDelegate ShouldPauseDelegate;
 	const FString VerifyDirectory;
@@ -36,21 +39,28 @@ FBuildPatchVerificationImpl::FBuildPatchVerificationImpl(const FBuildPatchAppMan
 	, CurrentFileWeight(0)
 {}
 
+void FBuildPatchVerificationImpl::SetRequiredFiles(TArray<FString> InRequiredFiles)
+{
+	RequiredFiles = MoveTemp(InRequiredFiles);
+}
+
 bool FBuildPatchVerificationImpl::VerifyAgainstDirectory(TArray<FString>& OutDatedFiles, double& TimeSpentPaused)
 {
 	bool bAllCorrect = true;
 	OutDatedFiles.Empty();
 	TimeSpentPaused = 0;
+	if (RequiredFiles.Num() == 0)
+	{
+		Manifest->GetFileList(RequiredFiles);
+	}
 
 	// Setup progress tracking
-	double TotalBuildSizeDouble = Manifest->GetBuildSize();
+	double TotalBuildSizeDouble = Manifest->GetFileSize(RequiredFiles);
 	double ProcessedBytes = 0;
 	CurrentBuildPercentage = 0;
 
 	// For all files in the manifest, check that they produce the correct SHA1 hash, adding any that don't to the list
-	TArray<FString> BuildFiles;
-	Manifest->GetFileList(BuildFiles);
-	for (const FString& BuildFile : BuildFiles)
+	for (const FString& BuildFile : RequiredFiles)
 	{
 		// Break if quitting
 		if (FBuildPatchInstallError::HasFatalError())
diff --git a/Engine/Source/Runtime/Online/BuildPatchServices/Private/BuildPatchVerification.h b/Engine/Source/Runtime/Online/BuildPatchServices/Private/BuildPatchVerification.h
--- a/Engine/Source/Runtime/Online/BuildPatchServices/Private/BuildPatchVerification.h
+++ b/Engine/Source/Runtime/Online/BuildPatchServices/Private/BuildPatchVerification.h
@@ -8,6 +8,12 @@ DECLARE_DELEGATE_RetVal(bool, FBuildPatchBoolRetDelegate);
 class FBuildPatchVerification
 {
 public:
+	/**
+	 * Override which files should be verified. The default is all build files
+	 * @param RequiredFiles    The array of files to be verified
+	 */
+	virtual void SetRequiredFiles(TArray<FString> RequiredFiles) = 0;
+
 	/**
 	 * Verifies a local directory structure against a given manifest.
 	 * NOTE: This function is blocking and will not return until finished. Don't run on main thread.
diff --git a/Engine/Source/Runtime/Online/BuildPatchServices/Public/Interfaces/IBuildPatchServicesModule.h b/Engine/Source/Runtime/Online/BuildPatchServices/Public/Interfaces/IBuildPatchServicesModule.h
--- a/Engine/Source/Runtime/Online/BuildPatchServices/Public/Interfaces/IBuildPatchServicesModule.h
+++ b/Engine/Source/Runtime/Online/BuildPatchServices/Public/Interfaces/IBuildPatchServicesModule.h
@@ -65,9 +65,10 @@ public:
 	 * @param	InstallManifest			The manifest to be installed
 	 * @param	InstallDirectory		The directory to install the App to
 	 * @param	OnCompleteDelegate		The delegate to call on completion
+	 * @param	InstallTags				The set of tags that describe what to be installed, default empty set means full installation
 	 * @return		An interface to the created installer. Will be an invalid ptr if error.
 	 */
-	virtual IBuildInstallerPtr StartBuildInstall( IBuildManifestPtr CurrentManifest, IBuildManifestPtr InstallManifest, const FString& InstallDirectory, FBuildPatchBoolManifestDelegate OnCompleteDelegate ) = 0;
+	virtual IBuildInstallerPtr StartBuildInstall(IBuildManifestPtr CurrentManifest, IBuildManifestPtr InstallManifest, const FString& InstallDirectory, FBuildPatchBoolManifestDelegate OnCompleteDelegate, TSet<FString> InstallTags = TSet<FString>()) = 0;
 
 	/**
 	 * Starts an installer thread for the provided manifests, only producing the necessary stage. Useful for handling specific install directory write access requirements yourself.
@@ -76,9 +77,10 @@ public:
 	 * @param	InstallManifest			The manifest to be installed
 	 * @param	InstallDirectory		The directory to install the App to - this should still be the real install directory. It may be read from for patching.
 	 * @param	OnCompleteDelegate		The delegate to call on completion
+	 * @param	InstallTags				The set of tags that describe what to be installed, default empty set means full installation
 	 * @return		An interface to the created installer. Will be an invalid ptr if error.
 	 */
-	virtual IBuildInstallerPtr StartBuildInstallStageOnly( IBuildManifestPtr CurrentManifest, IBuildManifestPtr InstallManifest, const FString& InstallDirectory, FBuildPatchBoolManifestDelegate OnCompleteDelegate ) = 0;
+	virtual IBuildInstallerPtr StartBuildInstallStageOnly(IBuildManifestPtr CurrentManifest, IBuildManifestPtr InstallManifest, const FString& InstallDirectory, FBuildPatchBoolManifestDelegate OnCompleteDelegate, TSet<FString> InstallTags = TSet<FString>()) = 0;
 
 	/**
 	 * Sets the directory used for staging intermediate files.
