BuildPatchServices: Ability to remove a custom field from a manifest
Merging CL# 2682800 using UE4-LauncherDev to UE4

diff --git a/Engine/Source/Runtime/Online/BuildPatchServices/Private/BuildPatchManifest.cpp b/Engine/Source/Runtime/Online/BuildPatchServices/Private/BuildPatchManifest.cpp
--- a/Engine/Source/Runtime/Online/BuildPatchServices/Private/BuildPatchManifest.cpp
+++ b/Engine/Source/Runtime/Online/BuildPatchServices/Private/BuildPatchManifest.cpp
@@ -1509,6 +1509,16 @@ const IManifestFieldPtr FBuildPatchAppManifest::SetCustomField(const FString& Fi
 	return SetCustomField(FieldName, ToStringBlob(Value));
 }
 
+void FBuildPatchAppManifest::RemoveCustomField(const FString& FieldName)
+{
+	Data->CustomFields.RemoveAll([&](const FCustomFieldData& Entry){ return Entry.Key == FieldName; });
+	CustomFieldLookup.Empty(Data->CustomFields.Num());
+	for (auto& CustomField : Data->CustomFields)
+	{
+		CustomFieldLookup.Add(CustomField.Key, &CustomField);
+	}
+}
+
 void FBuildPatchAppManifest::EnumerateProducibleChunks( const FString& InstallDirectory, const TArray< FGuid >& ChunksRequired, TArray< FGuid >& ChunksAvailable ) const
 {
 	// A struct that will store byte ranges
diff --git a/Engine/Source/Runtime/Online/BuildPatchServices/Private/BuildPatchManifest.h b/Engine/Source/Runtime/Online/BuildPatchServices/Private/BuildPatchManifest.h
--- a/Engine/Source/Runtime/Online/BuildPatchServices/Private/BuildPatchManifest.h
+++ b/Engine/Source/Runtime/Online/BuildPatchServices/Private/BuildPatchManifest.h
@@ -393,6 +393,7 @@ public:
 	virtual const IManifestFieldPtr SetCustomField(const FString& FieldName, const FString& Value) override;
 	virtual const IManifestFieldPtr SetCustomField(const FString& FieldName, const double& Value) override;
 	virtual const IManifestFieldPtr SetCustomField(const FString& FieldName, const int64& Value) override;
+	virtual void RemoveCustomField(const FString& FieldName) override;
 	virtual IBuildManifestRef Duplicate() const override;
 	// END IBuildManifest Interface
 
diff --git a/Engine/Source/Runtime/Online/BuildPatchServices/Public/Interfaces/IBuildManifest.h b/Engine/Source/Runtime/Online/BuildPatchServices/Public/Interfaces/IBuildManifest.h
--- a/Engine/Source/Runtime/Online/BuildPatchServices/Public/Interfaces/IBuildManifest.h
+++ b/Engine/Source/Runtime/Online/BuildPatchServices/Public/Interfaces/IBuildManifest.h
@@ -163,6 +163,12 @@ public:
 	virtual const IManifestFieldPtr SetCustomField(const FString& FieldName, const FString& Value) = 0;
 	virtual const IManifestFieldPtr SetCustomField(const FString& FieldName, const double& Value) = 0;
 	virtual const IManifestFieldPtr SetCustomField(const FString& FieldName, const int64& Value) = 0;
+
+	/**
+	 * Remove a custom field from the manifest
+	 * @param	FieldName	The name of the custom field
+	 */
+	virtual void RemoveCustomField(const FString& FieldName) = 0;
 	
 	/**
 	 * Duplicated this manifest to create a copy. Should be used if storing a received manifest as an installed
