Integrate 2683662 from UE4-Orion - Made some UParty functions virtual.

diff --git a/Engine/Source/Runtime/Online/OnlineFramework/Party/Public/Party.h b/Engine/Source/Runtime/Online/OnlineFramework/Party/Public/Party.h
--- a/Engine/Source/Runtime/Online/OnlineFramework/Party/Public/Party.h
+++ b/Engine/Source/Runtime/Online/OnlineFramework/Party/Public/Party.h
@@ -114,7 +114,7 @@ public:
 	/**
 	 * Initialization of the party management code, sets up listening for all party activity
 	 */
-	void Init();
+	virtual void Init();
 
 	/**
 	 * Initialization specific to play in editor
@@ -159,7 +159,7 @@ public:
 	 * @param InUserId user creating the party (should be primary player)
 	 * @param InCompletionDelegate delegate called upon completion
 	 */
-	void CreatePersistentParty(const FUniqueNetId& InUserId, const FOnCreatePartyComplete& InCompletionDelegate);
+	virtual void CreatePersistentParty(const FUniqueNetId& InUserId, const FOnCreatePartyComplete& InCompletionDelegate);
 
 	/**
 	 * Join a persistent party (party that will hold users for the duration of app)
@@ -176,7 +176,7 @@ public:
 	 * @param InUserId user creating the party (should be primary player)
 	 * @param InCompletionDelegate delegate called upon completion
 	 */
-	void LeavePersistentParty(const FUniqueNetId& InUserId, const FOnLeavePartyComplete& InCompletionDelegate);
+	virtual void LeavePersistentParty(const FUniqueNetId& InUserId, const FOnLeavePartyComplete& InCompletionDelegate);
 
 	/**
 	 * Reestablish all party state and information upon returning to the main menu
