[INTEGRATE]  Change 2684351 by [REDACTED_PERSON][REDACTED_HANDLE] on 2015/09/09 06:31

        Change patch generation process to only reuse files referenced by manifests newer than a certain age.

diff --git a/Engine/Source/Programs/BuildPatchTool/Private/BuildPatchToolMain.cpp b/Engine/Source/Programs/BuildPatchTool/Private/BuildPatchToolMain.cpp
--- a/Engine/Source/Programs/BuildPatchTool/Private/BuildPatchToolMain.cpp
+++ b/Engine/Source/Programs/BuildPatchTool/Private/BuildPatchToolMain.cpp
@@ -35,7 +35,7 @@
 	-PrereqName=""			Specifies in quotes, the display name for the prerequisites installer
 	-PrereqPath=""			Specifies in quotes, the prerequisites installer to launch on successful product install
 	-PrereqArgs=""			Specifies in quotes, the commandline to send to prerequisites installer on launch
-	-DataAgeThreshold=12.5	Specified the maximum age (in days) of existing patch data which can be reused in the generated manifest
+	-DataAgeThreshold=12.5	Specified the maximum age (in days) of existing manifest files whose referenced patch data can be reused in the generated manifest
 
 	NB: -DataAgeThreshold can also be present in the [PatchGeneration] section of BuildPatchTool.ini in the cloud directory
 	NB: If -DataAgeThreshold is not supplied, either on the command-line or in BuildPatchTool.ini, then all existing data is eligible for reuse in the generated manifest
@@ -59,7 +59,7 @@
 	-preview			Log all the actions it will take to update internal structures, but don't actually execute them.
 	-ManifestsList=""			Specifies in quotes, the list of manifest filenames to keep following the operation. If omitted, all manifests are kept.
 	-ManifestsFile=""			Specifies in quotes, the name of the file (relative to CloudDir) which contains a list of manifests to keep, one manifest filename per line
-	-DataAgeThreshold=14.25		The maximum age in days of chunk files that will be retained. All older chunks will be deleted.
+	-DataAgeThreshold=2			The minimum age in days of chunk files that will be deleted. Any unreferenced chunks older than this will be deleted.
 
 	NB: If -ManifestsList is specified, then -ManifestsFile is ignored.
 	NB: -DataAgeThreshold can also be present in the [Compactify] section of BuildPatchTool.ini in the cloud directory
diff --git a/Engine/Source/Runtime/Online/BuildPatchServices/Private/BuildPatchGeneration.cpp b/Engine/Source/Runtime/Online/BuildPatchServices/Private/BuildPatchGeneration.cpp
--- a/Engine/Source/Runtime/Online/BuildPatchServices/Private/BuildPatchGeneration.cpp
+++ b/Engine/Source/Runtime/Online/BuildPatchServices/Private/BuildPatchGeneration.cpp
@@ -168,13 +168,12 @@ void FBuildDataFileProcessor::GetFileStats( uint32& OutNewFiles, uint32& OutKnow
 
 /* FBuildSimpleChunkCache::FChunkReader implementation
 *****************************************************************************/
-FBuildGenerationChunkCache::FChunkReader::FChunkReader( const FString& InChunkFilePath, TSharedRef< FChunkFile > InChunkFile, uint32* InBytesRead, const FDateTime& InDataAgeThreshold )
+FBuildGenerationChunkCache::FChunkReader::FChunkReader(const FString& InChunkFilePath, TSharedRef< FChunkFile > InChunkFile, uint32* InBytesRead)
 	: ChunkFilePath( InChunkFilePath )
 	, ChunkFileReader( NULL )
 	, ChunkFile( InChunkFile )
 	, FileBytesRead( InBytesRead )
 	, MemoryBytesRead( 0 )
-	, DataAgeThreshold( InDataAgeThreshold )
 {
 	ChunkFile->GetDataLock( &ChunkData, &ChunkHeader );
 }
@@ -210,13 +209,6 @@ FArchive* FBuildGenerationChunkCache::FChunkReader::GetArchive()
 		{
 			*ChunkFileReader << *ChunkHeader;
 		}
-		// Check the file is not too old to reuse
-		if (IFileManager::Get().GetTimeStamp(*ChunkFilePath) < DataAgeThreshold)
-		{
-			// Break the magic to mark as invalid (this chunk is too old to reuse)
-			ChunkHeader->Magic = 0;
-			return ChunkFileReader;
-		}
 		// Check we can seek otherwise bad chunk
 		const int64 ExpectedFileSize = ChunkHeader->DataSize + ChunkHeader->HeaderSize;
 		const int64 ChunkFileSize = ChunkFileReader->TotalSize();
@@ -330,9 +322,8 @@ const uint32 FBuildGenerationChunkCache::FChunkReader::BytesLeft()
 
 /* FBuildGenerationChunkCache implementation
 *****************************************************************************/
-FBuildGenerationChunkCache::FBuildGenerationChunkCache(const FDateTime& InDataAgeThreshold)
-	: DataAgeThreshold(InDataAgeThreshold)
-	, StatChunksInDataCache(nullptr)
+FBuildGenerationChunkCache::FBuildGenerationChunkCache()
+	: StatChunksInDataCache(nullptr)
 	, StatNumCacheLoads(nullptr)
 	, StatNumCacheBoots(nullptr)
 {
@@ -389,7 +380,7 @@ TSharedRef< FBuildGenerationChunkCache::FChunkReader > FBuildGenerationChunkCach
 			FStatsCollector::Accumulate(StatNumCacheLoads, 1);
 		}
 	}
-	return MakeShareable( new FChunkReader( ChunkFilePath, ChunkCache[ ChunkFilePath ], BytesReadPerChunk[ ChunkFilePath ], DataAgeThreshold ) );
+	return MakeShareable( new FChunkReader(ChunkFilePath, ChunkCache[ ChunkFilePath ], BytesReadPerChunk[ ChunkFilePath ]) );
 }
 
 void FBuildGenerationChunkCache::Cleanup()
@@ -410,11 +401,11 @@ void FBuildGenerationChunkCache::Cleanup()
 *****************************************************************************/
 TSharedPtr< FBuildGenerationChunkCache > FBuildGenerationChunkCache::SingletonInstance = NULL;
 
-void FBuildGenerationChunkCache::Init(const FDateTime& DataAgeThreshold)
+void FBuildGenerationChunkCache::Init()
 {
 	// We won't allow misuse of these functions
 	check( !SingletonInstance.IsValid() );
-	SingletonInstance = MakeShareable(new FBuildGenerationChunkCache(DataAgeThreshold));
+	SingletonInstance = MakeShareable(new FBuildGenerationChunkCache());
 }
 
 FBuildGenerationChunkCache& FBuildGenerationChunkCache::Get()
@@ -499,7 +490,8 @@ bool FBuildDataGenerator::GenerateChunksManifestFromDirectory( const FBuildPatch
 	FStatsCollectorRef StatsCollector = FStatsCollectorFactory::Create();
 
 	// Enumerate Chunks
-	FCloudEnumerationRef CloudEnumeration = FCloudEnumerationFactory::Create(FBuildPatchServicesModule::GetCloudDirectory());
+	const FDateTime Cutoff = Settings.bShouldHonorReuseThreshold ? FDateTime::UtcNow() - FTimespan::FromDays(Settings.DataAgeThreshold) : FDateTime::MinValue();
+	FCloudEnumerationRef CloudEnumeration = FCloudEnumerationFactory::Create(FBuildPatchServicesModule::GetCloudDirectory(), Cutoff);
 
 	// Force waiting on cloud enumeration for more accurate stats, this line can be removed when stats are not required.
 	CloudEnumeration->GetChunkInventory();
@@ -514,8 +506,7 @@ bool FBuildDataGenerator::GenerateChunksManifestFromDirectory( const FBuildPatch
 	GLog->Logf(TEXT("Running Chunks Patch Generation for: %u:%s %s"), Settings.AppID, *Settings.AppName, *Settings.BuildVersion);
 
 	// Create our chunk cache
-	const FDateTime Cutoff = Settings.bShouldHonorReuseThreshold ? FDateTime::UtcNow() - FTimespan::FromDays(Settings.DataAgeThreshold) : FDateTime::MinValue();
-	FBuildGenerationChunkCache::Init(Cutoff);
+	FBuildGenerationChunkCache::Init();
 	FBuildGenerationChunkCache::Get().SetStatsCollector(StatsCollector);
 
 	// Create the manifest builder
diff --git a/Engine/Source/Runtime/Online/BuildPatchServices/Private/BuildPatchGeneration.h b/Engine/Source/Runtime/Online/BuildPatchServices/Private/BuildPatchGeneration.h
--- a/Engine/Source/Runtime/Online/BuildPatchServices/Private/BuildPatchGeneration.h
+++ b/Engine/Source/Runtime/Online/BuildPatchServices/Private/BuildPatchGeneration.h
@@ -120,18 +120,14 @@ public:
 		// How many bytes we have read out of this chunk, used to calculate when we need more data from file
 		uint32 MemoryBytesRead;
 
-		// The validity period of chunk data. Anything older than this will be considered to be invalid
-		FDateTime DataAgeThreshold;
-
 	public:
 		/**
 		 * Constructs the reader from required data
 		 * @param	InChunkFilePath		The file path of the source file for the chunk
 		 * @param	InChunkFile			The data structure for this chunk
 		 * @param	InBytesRead			Pointer to an unsigned int to keep track of how much data is read from file
-		 * @param	InDataAgeThreshold	The validity period of chunk data. Anything older than this will be considered to be invalid
 		 */
-		FChunkReader( const FString& InChunkFilePath, TSharedRef< FChunkFile > InChunkFile, uint32* InBytesRead, const FDateTime& InDataAgeThreshold );
+		FChunkReader(const FString& InChunkFilePath, TSharedRef< FChunkFile > InChunkFile, uint32* InBytesRead);
 
 		/**
 		 * Destructor. Cleans up and releases data locks.
@@ -175,16 +171,13 @@ public:
 private:
 	/**
 	* Constructor
-	* @param	DataAgeThreshold		The cutoff point before which data chunks should be regarded as isvalid
 	*/
-	FBuildGenerationChunkCache(const FDateTime& DataAgeThreshold);
+	FBuildGenerationChunkCache();
 
 private:
 	// The number of chunks to cache in memory
 	const static int32 NumChunksToCache = 1024;
 
-	FDateTime DataAgeThreshold;
-
 	// The map of chunks we have cached
 	TMap< FString, TSharedRef< FChunkFile > > ChunkCache;
 
@@ -216,9 +209,8 @@ public:
 
 	/**
 	 * Call to initialise static singleton so that the cache system can be used
-	 * @param DataAgeThreshold		Any chunks older than this date will be classed as invalid
 	 */
-	static void Init(const FDateTime& DataAgeThreshold);
+	static void Init();
 
 	/**
 	 * Get a reference to the chunk system. Only call between Init and Shutdown calls.
diff --git a/Engine/Source/Runtime/Online/BuildPatchServices/Private/Generation/CloudEnumeration.cpp b/Engine/Source/Runtime/Online/BuildPatchServices/Private/Generation/CloudEnumeration.cpp
--- a/Engine/Source/Runtime/Online/BuildPatchServices/Private/Generation/CloudEnumeration.cpp
+++ b/Engine/Source/Runtime/Online/BuildPatchServices/Private/Generation/CloudEnumeration.cpp
@@ -12,7 +12,7 @@ namespace BuildPatchServices
 		: public FCloudEnumeration
 	{
 	public:
-		FCloudEnumerationImpl(const FString& CloudDirectory);
+		FCloudEnumerationImpl(const FString& CloudDirectory, const FDateTime& ManifestAgeThreshold);
 		virtual ~FCloudEnumerationImpl();
 
 		virtual TSet<FGuid> GetChunkSet(uint64 ChunkHash) const override;
@@ -25,6 +25,7 @@ namespace BuildPatchServices
 
 	private:
 		const FString CloudDirectory;
+		const FDateTime ManifestAgeThreshold;
 		mutable FCriticalSection InventoryCS;
 		TMap<uint64, TSet<FGuid>> ChunkInventory;
 		TMap<FGuid, int64> ChunkFileSizes;
@@ -35,8 +36,9 @@ namespace BuildPatchServices
 		TFuture<void> Future;
 	};
 
-	FCloudEnumerationImpl::FCloudEnumerationImpl(const FString& InCloudDirectory)
+	FCloudEnumerationImpl::FCloudEnumerationImpl(const FString& InCloudDirectory, const FDateTime& InManifestAgeThreshold)
 		: CloudDirectory(InCloudDirectory)
+		, ManifestAgeThreshold(InManifestAgeThreshold)
 		, NumChunksFound(0)
 		, NumFilesFound(0)
 	{
@@ -103,6 +105,7 @@ namespace BuildPatchServices
 			FileManager.FindFiles(AllManifests, *(CloudDirectory / TEXT("*.manifest")), true, false);
 			const double EnumerateTime = FPlatformTime::Seconds() - StartEnumerate;
 			GLog->Logf(TEXT("FCloudEnumeration: Found %d manifests in %.1f seconds"), AllManifests.Num(), EnumerateTime);
+			int NumSkipped = 0;
 
 			// Load all manifest files
 			const double StartLoadAllManifest = FPlatformTime::Seconds();
@@ -110,6 +113,12 @@ namespace BuildPatchServices
 			{
 				// Determine chunks from manifest file
 				const FString ManifestFilename = CloudDirectory / ManifestFile;
+				if (IFileManager::Get().GetTimeStamp(*ManifestFilename) < ManifestAgeThreshold)
+				{
+					++NumSkipped;
+					GLog->Logf(TEXT("FCloudEnumeration: Skipping %s as it is too old to reuse"), *ManifestFile);
+					continue;
+				}
 				FBuildPatchAppManifestRef BuildManifest = MakeShareable(new FBuildPatchAppManifest());
 				const double StartLoadManifest = FPlatformTime::Seconds();
 				if (BuildManifest->LoadFromFile(ManifestFilename))
@@ -124,7 +133,7 @@ namespace BuildPatchServices
 				}
 			}
 			const double LoadAllManifestTime = FPlatformTime::Seconds() - StartLoadAllManifest;
-			GLog->Logf(TEXT("FCloudEnumeration: Used %d manifests to enumerate %llu chunks in %.1f seconds"), AllManifests.Num(), NumChunksFound, LoadAllManifestTime);
+			GLog->Logf(TEXT("FCloudEnumeration: Used %d manifests to enumerate %llu chunks in %.1f seconds"), AllManifests.Num() - NumSkipped, NumChunksFound, LoadAllManifestTime);
 		}
 		else
 		{
@@ -176,9 +185,9 @@ namespace BuildPatchServices
 		}
 	}
 
-	FCloudEnumerationRef FCloudEnumerationFactory::Create(const FString& CloudDirectory)
+	FCloudEnumerationRef FCloudEnumerationFactory::Create(const FString& CloudDirectory, const FDateTime& ManifestAgeThreshold)
 	{
-		return MakeShareable(new FCloudEnumerationImpl(CloudDirectory));
+		return MakeShareable(new FCloudEnumerationImpl(CloudDirectory, ManifestAgeThreshold));
 	}
 }
 
diff --git a/Engine/Source/Runtime/Online/BuildPatchServices/Private/Generation/CloudEnumeration.h b/Engine/Source/Runtime/Online/BuildPatchServices/Private/Generation/CloudEnumeration.h
--- a/Engine/Source/Runtime/Online/BuildPatchServices/Private/Generation/CloudEnumeration.h
+++ b/Engine/Source/Runtime/Online/BuildPatchServices/Private/Generation/CloudEnumeration.h
@@ -20,7 +20,7 @@ namespace BuildPatchServices
 	class FCloudEnumerationFactory
 	{
 	public:
-		static FCloudEnumerationRef Create(const FString& CloudDirectory);
+		static FCloudEnumerationRef Create(const FString& CloudDirectory, const FDateTime& ManifestAgeThreshold);
 	};
 }
 
