#UE4 - added completion delegate to destroy session

--------
Integrated using branch UE4-To-UE4-Orion (reversed) of change#2687022 by [REDACTED_PERSON] on 2015/09/10 16:41:29.

diff --git a/Engine/Source/Runtime/Online/IOS/OnlineSubsystemIOS/Private/OnlineSessionInterfaceIOS.cpp b/Engine/Source/Runtime/Online/IOS/OnlineSubsystemIOS/Private/OnlineSessionInterfaceIOS.cpp
--- a/Engine/Source/Runtime/Online/IOS/OnlineSubsystemIOS/Private/OnlineSessionInterfaceIOS.cpp
+++ b/Engine/Source/Runtime/Online/IOS/OnlineSubsystemIOS/Private/OnlineSessionInterfaceIOS.cpp
@@ -457,7 +457,7 @@ bool FOnlineSessionIOS::EndSession(FName SessionName)
 }
 
 
-bool FOnlineSessionIOS::DestroySession(FName SessionName)
+bool FOnlineSessionIOS::DestroySession(FName SessionName, const FOnDestroySessionCompleteDelegate& CompletionDelegate)
 {
 	bool bSuccessfullyDestroyedSession = false;
 	
@@ -475,6 +475,7 @@ bool FOnlineSessionIOS::DestroySession(FName SessionName)
 		bSuccessfullyDestroyedSession = true;
 	}
 
+	CompletionDelegate.ExecuteIfBound(SessionName, bSuccessfullyDestroyedSession);
 	TriggerOnDestroySessionCompleteDelegates(SessionName, bSuccessfullyDestroyedSession);
 
 	return bSuccessfullyDestroyedSession;
diff --git a/Engine/Source/Runtime/Online/IOS/OnlineSubsystemIOS/Private/OnlineSessionInterfaceIOS.h b/Engine/Source/Runtime/Online/IOS/OnlineSubsystemIOS/Private/OnlineSessionInterfaceIOS.h
--- a/Engine/Source/Runtime/Online/IOS/OnlineSubsystemIOS/Private/OnlineSessionInterfaceIOS.h
+++ b/Engine/Source/Runtime/Online/IOS/OnlineSubsystemIOS/Private/OnlineSessionInterfaceIOS.h
@@ -127,7 +127,7 @@ public:
 
 	virtual bool EndSession(FName SessionName) override;
 
-	virtual bool DestroySession(FName SessionName) override;
+	virtual bool DestroySession(FName SessionName, const FOnDestroySessionCompleteDelegate& CompletionDelegate = FOnDestroySessionCompleteDelegate()) override;
 
 	virtual bool IsPlayerInSession(FName SessionName, const FUniqueNetId& UniqueId) override;
 
diff --git a/Engine/Source/Runtime/Online/OnlineSubsystem/Public/Interfaces/OnlineSessionInterface.h b/Engine/Source/Runtime/Online/OnlineSubsystem/Public/Interfaces/OnlineSessionInterface.h
--- a/Engine/Source/Runtime/Online/OnlineSubsystem/Public/Interfaces/OnlineSessionInterface.h
+++ b/Engine/Source/Runtime/Online/OnlineSubsystem/Public/Interfaces/OnlineSessionInterface.h
@@ -378,7 +378,7 @@ public:
 	 *
 	 * @return true if successful destroying the session, false otherwise
 	 */
-	virtual bool DestroySession(FName SessionName) = 0;
+	virtual bool DestroySession(FName SessionName, const FOnDestroySessionCompleteDelegate& CompletionDelegate = FOnDestroySessionCompleteDelegate()) = 0;
 
 	/**
 	 * Delegate fired when a destroying an online session has completed
diff --git a/Engine/Source/Runtime/Online/OnlineSubsystemNull/Private/OnlineSessionInterfaceNull.cpp b/Engine/Source/Runtime/Online/OnlineSubsystemNull/Private/OnlineSessionInterfaceNull.cpp
--- a/Engine/Source/Runtime/Online/OnlineSubsystemNull/Private/OnlineSessionInterfaceNull.cpp
+++ b/Engine/Source/Runtime/Online/OnlineSubsystemNull/Private/OnlineSessionInterfaceNull.cpp
@@ -417,7 +417,7 @@ bool FOnlineSessionNull::EndSession(FName SessionName)
 	return Result == ERROR_SUCCESS || Result == ERROR_IO_PENDING;
 }
 
-bool FOnlineSessionNull::DestroySession(FName SessionName)
+bool FOnlineSessionNull::DestroySession(FName SessionName, const FOnDestroySessionCompleteDelegate& CompletionDelegate)
 {
 	uint32 Result = E_FAIL;
 	// Find the session in question
@@ -436,6 +436,7 @@ bool FOnlineSessionNull::DestroySession(FName SessionName)
 
 	if (Result != ERROR_IO_PENDING)
 	{
+		CompletionDelegate.ExecuteIfBound(SessionName, (Result == ERROR_SUCCESS) ? true : false);
 		TriggerOnDestroySessionCompleteDelegates(SessionName, (Result == ERROR_SUCCESS) ? true : false);
 	}
 
diff --git a/Engine/Source/Runtime/Online/OnlineSubsystemNull/Private/OnlineSessionInterfaceNull.h b/Engine/Source/Runtime/Online/OnlineSubsystemNull/Private/OnlineSessionInterfaceNull.h
--- a/Engine/Source/Runtime/Online/OnlineSubsystemNull/Private/OnlineSessionInterfaceNull.h
+++ b/Engine/Source/Runtime/Online/OnlineSubsystemNull/Private/OnlineSessionInterfaceNull.h
@@ -277,7 +277,7 @@ public:
 	virtual bool StartSession(FName SessionName) override;
 	virtual bool UpdateSession(FName SessionName, FOnlineSessionSettings& UpdatedSessionSettings, bool bShouldRefreshOnlineData = true) override;
 	virtual bool EndSession(FName SessionName) override;
-	virtual bool DestroySession(FName SessionName) override;
+	virtual bool DestroySession(FName SessionName, const FOnDestroySessionCompleteDelegate& CompletionDelegate = FOnDestroySessionCompleteDelegate()) override;
 	virtual bool IsPlayerInSession(FName SessionName, const FUniqueNetId& UniqueId) override;
 	virtual bool StartMatchmaking(const TArray< TSharedRef<const FUniqueNetId> >& LocalPlayers, FName SessionName, const FOnlineSessionSettings& NewSessionSettings, TSharedRef<FOnlineSessionSearch>& SearchSettings) override;
 	virtual bool CancelMatchmaking(int32 SearchingPlayerNum, FName SessionName) override;
diff --git a/Engine/Source/Runtime/Online/OnlineSubsystemSteam/Private/OnlineSessionInterfaceSteam.cpp b/Engine/Source/Runtime/Online/OnlineSubsystemSteam/Private/OnlineSessionInterfaceSteam.cpp
--- a/Engine/Source/Runtime/Online/OnlineSubsystemSteam/Private/OnlineSessionInterfaceSteam.cpp
+++ b/Engine/Source/Runtime/Online/OnlineSubsystemSteam/Private/OnlineSessionInterfaceSteam.cpp
@@ -127,11 +127,14 @@ class FOnlineAsyncTaskSteamDestroySession : public FOnlineAsyncTaskSteam
 private:
 	/** Name of session ending */
 	FName SessionName;
+	/** */
+	FOnDestroySessionCompleteDelegate CompletionDelegate;
 
 public:
-	FOnlineAsyncTaskSteamDestroySession(FOnlineSubsystemSteam* InSubsystem, FName InSessionName) :
+	FOnlineAsyncTaskSteamDestroySession(FOnlineSubsystemSteam* InSubsystem, FName InSessionName, const FOnDestroySessionCompleteDelegate& InCompletionDelegate) :
 		FOnlineAsyncTaskSteam(InSubsystem, k_uAPICallInvalid),
-		SessionName(InSessionName)
+		SessionName(InSessionName),
+		CompletionDelegate(InCompletionDelegate)
 	{
 	}
 
@@ -197,6 +200,7 @@ public:
 		IOnlineSessionPtr SessionInt = Subsystem->GetSessionInterface();
 		if (SessionInt.IsValid())
 		{
+			CompletionDelegate.ExecuteIfBound(SessionName, bWasSuccessful);
 			SessionInt->TriggerOnDestroySessionCompleteDelegates(SessionName, bWasSuccessful);
 		}
 	}
@@ -531,7 +535,7 @@ uint32 FOnlineSessionSteam::EndInternetSession(FNamedOnlineSession* Session)
 	return ERROR_IO_PENDING;
 }
 
-bool FOnlineSessionSteam::DestroySession(FName SessionName)
+bool FOnlineSessionSteam::DestroySession(FName SessionName, const FOnDestroySessionCompleteDelegate& CompletionDelegate)
 {
 	uint32 Result = E_FAIL;
 	// Find the session in question
@@ -550,11 +554,11 @@ bool FOnlineSessionSteam::DestroySession(FName SessionName)
 
 				if (Session->SessionSettings.bUsesPresence)
 				{
-					Result = DestroyLobbySession(Session);
+					Result = DestroyLobbySession(Session, CompletionDelegate);
 				}
 				else
 				{
-					Result = DestroyInternetSession(Session);
+					Result = DestroyInternetSession(Session, CompletionDelegate);
 				}
 			}
 			else
@@ -574,6 +578,7 @@ bool FOnlineSessionSteam::DestroySession(FName SessionName)
 			{
 				// The session info is no longer needed
 				RemoveNamedSession(Session->SessionName);
+				CompletionDelegate.ExecuteIfBound(SessionName, (Result == ERROR_SUCCESS) ? true : false);
 				TriggerOnDestroySessionCompleteDelegates(SessionName, (Result == ERROR_SUCCESS) ? true : false);
 			}
 		}
@@ -586,13 +591,14 @@ bool FOnlineSessionSteam::DestroySession(FName SessionName)
 	else
 	{
 		UE_LOG_ONLINE(Warning, TEXT("Can't destroy a null online session (%s)"), *SessionName.ToString());
+		CompletionDelegate.ExecuteIfBound(SessionName, false);
 		TriggerOnDestroySessionCompleteDelegates(SessionName, false);
 	}
 
 	return Result == ERROR_SUCCESS || Result == ERROR_IO_PENDING;
 }
 
-uint32 FOnlineSessionSteam::DestroyLobbySession(FNamedOnlineSession* Session)
+uint32 FOnlineSessionSteam::DestroyLobbySession(FNamedOnlineSession* Session, const FOnDestroySessionCompleteDelegate& CompletionDelegate)
 {
 	Session->SessionState = EOnlineSessionState::Destroying;
 
@@ -605,13 +611,13 @@ uint32 FOnlineSessionSteam::DestroyLobbySession(FNamedOnlineSession* Session)
 		SteamSubsystem->QueueAsyncTask(NewTask);
 	}
 
-	FOnlineAsyncTaskSteamDestroySession* NewTask = new FOnlineAsyncTaskSteamDestroySession(SteamSubsystem, Session->SessionName);
+	FOnlineAsyncTaskSteamDestroySession* NewTask = new FOnlineAsyncTaskSteamDestroySession(SteamSubsystem, Session->SessionName, CompletionDelegate);
 	SteamSubsystem->QueueAsyncTask(NewTask);
 
 	return ERROR_IO_PENDING;
 }
 
-uint32 FOnlineSessionSteam::DestroyInternetSession(FNamedOnlineSession* Session)
+uint32 FOnlineSessionSteam::DestroyInternetSession(FNamedOnlineSession* Session, const FOnDestroySessionCompleteDelegate& CompletionDelegate)
 {
 	Session->SessionState = EOnlineSessionState::Destroying;
 
@@ -629,7 +635,7 @@ uint32 FOnlineSessionSteam::DestroyInternetSession(FNamedOnlineSession* Session)
 	}
 
 	// Destroy the session
-	FOnlineAsyncTaskSteamDestroySession* DestroyTask = new FOnlineAsyncTaskSteamDestroySession(SteamSubsystem, Session->SessionName);
+	FOnlineAsyncTaskSteamDestroySession* DestroyTask = new FOnlineAsyncTaskSteamDestroySession(SteamSubsystem, Session->SessionName, CompletionDelegate);
 	SteamSubsystem->QueueAsyncTask(DestroyTask);
 
 	return ERROR_IO_PENDING;
diff --git a/Engine/Source/Runtime/Online/OnlineSubsystemSteam/Private/OnlineSessionInterfaceSteam.h b/Engine/Source/Runtime/Online/OnlineSubsystemSteam/Private/OnlineSessionInterfaceSteam.h
--- a/Engine/Source/Runtime/Online/OnlineSubsystemSteam/Private/OnlineSessionInterfaceSteam.h
+++ b/Engine/Source/Runtime/Online/OnlineSubsystemSteam/Private/OnlineSessionInterfaceSteam.h
@@ -110,7 +110,7 @@ private:
 	 * 
 	 * @return ERROR_SUCCESS if successful, an error code otherwise
 	 */
-	uint32 DestroyLobbySession(class FNamedOnlineSession* Session);
+	uint32 DestroyLobbySession(class FNamedOnlineSession* Session, const FOnDestroySessionCompleteDelegate& CompletionDelegate);
 
 	/**
 	 *	Destroy an internet session, advertised on the Steam backend
@@ -119,7 +119,7 @@ private:
 	 * 
 	 * @return ERROR_SUCCESS if successful, an error code otherwise
 	 */
-	uint32 DestroyInternetSession(class FNamedOnlineSession* Session);
+	uint32 DestroyInternetSession(class FNamedOnlineSession* Session, const FOnDestroySessionCompleteDelegate& CompletionDelegate);
 
 	/**
 	 *	Create a local LAN session, managed by a beacon on the host
@@ -494,7 +494,7 @@ public:
 	virtual bool StartSession(FName SessionName) override;
 	virtual bool UpdateSession(FName SessionName, FOnlineSessionSettings& UpdatedSessionSettings, bool bShouldRefreshOnlineData = true) override;
 	virtual bool EndSession(FName SessionName) override;
-	virtual bool DestroySession(FName SessionName) override;
+	virtual bool DestroySession(FName SessionName, const FOnDestroySessionCompleteDelegate& CompletionDelegate = FOnDestroySessionCompleteDelegate()) override;
 	virtual bool IsPlayerInSession(FName SessionName, const FUniqueNetId& UniqueId) override;
 	virtual bool StartMatchmaking(const TArray< TSharedRef<const FUniqueNetId> >& LocalPlayers, FName SessionName, const FOnlineSessionSettings& NewSessionSettings, TSharedRef<FOnlineSessionSearch>& SearchSettings) override;
 	virtual bool CancelMatchmaking(int32 SearchingPlayerNum, FName SessionName) override;
