BuildPatchServices: Download code for file protocols now use Thread Pool to load the files rather than one at a time.
Merging CL# 2687647 from UE4-LauncherDev to UE4

diff --git a/Engine/Source/Runtime/Online/BuildPatchServices/Private/BuildPatchDownloader.cpp b/Engine/Source/Runtime/Online/BuildPatchServices/Private/BuildPatchDownloader.cpp
--- a/Engine/Source/Runtime/Online/BuildPatchServices/Private/BuildPatchDownloader.cpp
+++ b/Engine/Source/Runtime/Online/BuildPatchServices/Private/BuildPatchDownloader.cpp
@@ -6,6 +6,7 @@
 =============================================================================*/
 
 #include "BuildPatchServicesPrivatePCH.h"
+#include "Async.h"
 
 #define EXTRA_DOWNLOAD_LOGGING 0
 
@@ -293,27 +294,31 @@ uint32 FBuildPatchDownloader::Run()
 		else
 		{
 			// Load file from drive/network
-			TArray< uint8 > FileDataArray;
-			FArchive* Reader = IFileManager::Get().CreateFileReader( *DownloadUrl );
-			bool bSuccess = Reader != NULL;
-			if( Reader )
-			{
-				const int64 BytesPerCall = 256*1024;
-				const int64 FileSize = Reader->TotalSize();
-				FileDataArray.Reset();
-				FileDataArray.AddUninitialized( FileSize );
-				int64 BytesRead = 0;
-				while ( BytesRead < FileSize && !FBuildPatchInstallError::HasFatalError() )
+			TFunction<void()> Task = [this, DownloadUrl, NextGuid]() {
+				TArray< uint8 > FileDataArray;
+				FArchive* Reader = IFileManager::Get().CreateFileReader( *DownloadUrl );
+				bool bSuccess = Reader != NULL;
+				if( Reader )
 				{
-					const int64 ReadLen = FMath::Min<int64>( BytesPerCall, FileSize - BytesRead );
-					Reader->Serialize( FileDataArray.GetData() + BytesRead, ReadLen );
-					BytesRead += ReadLen;
-					OnDownloadProgress( NextGuid, BytesRead );
+					const int64 BytesPerCall = 256*1024;
+					const int64 FileSize = Reader->TotalSize();
+					FileDataArray.Reset();
+					FileDataArray.AddUninitialized( FileSize );
+					int64 BytesRead = 0;
+					while ( BytesRead < FileSize && !FBuildPatchInstallError::HasFatalError() )
+					{
+						const int64 ReadLen = FMath::Min<int64>( BytesPerCall, FileSize - BytesRead );
+						Reader->Serialize( FileDataArray.GetData() + BytesRead, ReadLen );
+						BytesRead += ReadLen;
+						OnDownloadProgress( NextGuid, BytesRead );
+					}
+					bSuccess = Reader->Close() && !FBuildPatchInstallError::HasFatalError();
+					delete Reader;
 				}
-				bSuccess = Reader->Close() && !FBuildPatchInstallError::HasFatalError();
-				delete Reader;
-			}
-			OnDownloadComplete( NextGuid, DownloadUrl, FileDataArray, bSuccess, INDEX_NONE );
+				OnDownloadComplete(NextGuid, DownloadUrl, FileDataArray, bSuccess, INDEX_NONE);
+				FileDataArray.Empty();
+			};
+			Async(EAsyncExecution::ThreadPool, Task);
 		}
 	}
 	SetIdle( true );
