Integrate 2699413 from UE4-Orion

- When setting custom PS4 session data, convert to UTF-8 to reduce space
- Added more support for getting PS4 session info
- Removed OnConnectionStatusChanged callback in PS4 session subsystem since it wasn't doing anything

diff --git a/Engine/Source/Runtime/Online/OnlineFramework/Party/Public/Party.h b/Engine/Source/Runtime/Online/OnlineFramework/Party/Public/Party.h
--- a/Engine/Source/Runtime/Online/OnlineFramework/Party/Public/Party.h
+++ b/Engine/Source/Runtime/Online/OnlineFramework/Party/Public/Party.h
@@ -168,7 +168,7 @@ public:
 	 * @param InPartyDetails credentials for the party to join
 	 * @param InCompletionDelegate delegate called upon completion
 	 */
-	void JoinPersistentParty(const FUniqueNetId& InUserId, const FPartyDetails& InPartyDetails, const FOnJoinPartyComplete& InCompletionDelegate);
+	virtual void JoinPersistentParty(const FUniqueNetId& InUserId, const FPartyDetails& InPartyDetails, const FOnJoinPartyComplete& InCompletionDelegate);
 
 	/**
 	 * Leave a persistent party, a new one will need to be created for the given user afterward
