[INTEGRATE] merge ue4 <-ue4-launcherDev cl 2695116

diff --git a/Engine/Source/Runtime/Online/BuildPatchServices/Private/BuildPatchManifest.cpp b/Engine/Source/Runtime/Online/BuildPatchServices/Private/BuildPatchManifest.cpp
--- a/Engine/Source/Runtime/Online/BuildPatchServices/Private/BuildPatchManifest.cpp
+++ b/Engine/Source/Runtime/Online/BuildPatchServices/Private/BuildPatchManifest.cpp
@@ -1229,11 +1229,27 @@ int64 FBuildPatchAppManifest::GetDownloadSize() const
 	return TotalDownloadSize;
 }
 
+int64 FBuildPatchAppManifest::GetDownloadSize(const TSet<FString>& Tags) const
+{
+	TSet<FString> TaggedFiles;
+	GetTaggedFileList(Tags, TaggedFiles);
+	TArray<FGuid> RequiredChunks;
+	GetChunksRequiredForFiles(TaggedFiles.Array(), RequiredChunks, true);
+	return GetDataSize(RequiredChunks);
+}
+
 int64 FBuildPatchAppManifest::GetBuildSize() const
 {
 	return TotalBuildSize;
 }
 
+int64 FBuildPatchAppManifest::GetBuildSize(const TSet<FString>& Tags) const
+{
+	TSet<FString> TaggedFiles;
+	GetTaggedFileList(Tags, TaggedFiles);
+	return GetFileSize(TaggedFiles.Array());
+}
+
 TArray<FString> FBuildPatchAppManifest::GetBuildFileList() const
 {
 	TArray<FString> Filenames;
diff --git a/Engine/Source/Runtime/Online/BuildPatchServices/Private/BuildPatchManifest.h b/Engine/Source/Runtime/Online/BuildPatchServices/Private/BuildPatchManifest.h
--- a/Engine/Source/Runtime/Online/BuildPatchServices/Private/BuildPatchManifest.h
+++ b/Engine/Source/Runtime/Online/BuildPatchServices/Private/BuildPatchManifest.h
@@ -383,7 +383,9 @@ public:
 	virtual const FString& GetPrereqPath() const override;
 	virtual const FString& GetPrereqArgs() const override;
 	virtual int64 GetDownloadSize() const override;
+	virtual int64 GetDownloadSize(const TSet<FString>& Tags) const override;
 	virtual int64 GetBuildSize() const override;
+	virtual int64 GetBuildSize(const TSet<FString>& Tags) const override;
 	virtual TArray<FString> GetBuildFileList() const override;
 	virtual void GetRemovableFiles(IBuildManifestRef OldManifest, TArray< FString >& RemovableFiles) const override;
 	virtual void GetRemovableFiles(const TCHAR* InstallPath, TArray< FString >& RemovableFiles) const override;
diff --git a/Engine/Source/Runtime/Online/BuildPatchServices/Public/Interfaces/IBuildManifest.h b/Engine/Source/Runtime/Online/BuildPatchServices/Public/Interfaces/IBuildManifest.h
--- a/Engine/Source/Runtime/Online/BuildPatchServices/Public/Interfaces/IBuildManifest.h
+++ b/Engine/Source/Runtime/Online/BuildPatchServices/Public/Interfaces/IBuildManifest.h
@@ -108,12 +108,26 @@ public:
 	 */
 	virtual int64 GetDownloadSize() const = 0;
 
+	/**
+	 * Get the size of the download of this set of tags
+	 * @param Tags	IN	A list of the tags we want to know the size of
+	 * @return		the download size of the tags in bytes
+	 */
+	virtual int64 GetDownloadSize(const TSet<FString>& Tags) const = 0;
+
 	/**
 	 * Get the size of this build
 	 * @return		the build size in bytes
 	 */
 	virtual int64 GetBuildSize() const = 0;
 
+	/**
+	 * Get the size of the build of this set of tags
+	 * @param Tags	IN	A list of the tags we want to know the size of
+	 * @return		the build size of the tags in bytes
+	 */
+	virtual int64 GetBuildSize(const TSet<FString>& Tags) const = 0;
+
 	/**
 	 * Get the list of files in this build. Filenames are all relative to an install directory
 	 * @return		an array containing build files
