[INTEGRATE] merge ue4 <-ue4-launcherDev cl 2700782

diff --git a/Engine/Source/Runtime/Online/BuildPatchServices/Private/BuildPatchManifest.h b/Engine/Source/Runtime/Online/BuildPatchServices/Private/BuildPatchManifest.h
--- a/Engine/Source/Runtime/Online/BuildPatchServices/Private/BuildPatchManifest.h
+++ b/Engine/Source/Runtime/Online/BuildPatchServices/Private/BuildPatchManifest.h
@@ -387,6 +387,7 @@ public:
 	virtual int64 GetBuildSize() const override;
 	virtual int64 GetBuildSize(const TSet<FString>& Tags) const override;
 	virtual TArray<FString> GetBuildFileList() const override;
+	virtual void GetFileTagList(TSet<FString>& Tags) const override;
 	virtual void GetRemovableFiles(IBuildManifestRef OldManifest, TArray< FString >& RemovableFiles) const override;
 	virtual void GetRemovableFiles(const TCHAR* InstallPath, TArray< FString >& RemovableFiles) const override;
 	virtual bool NeedsResaving() const override;
@@ -496,12 +497,6 @@ public:
 	 */
 	void GetFileList(TArray<FString>& Filenames) const;
 
-	/**
-	 * Get the list of install tags in this manifest
-	 * @param Tags			OUT		Receives the tags referenced.
-	 */
-	void GetFileTagList(TSet<FString>& Tags) const;
-
 	/**
 	 * Get the list of files that are tagged with the provided tags
 	 * @param Tags					The tags for the required file groups.
diff --git a/Engine/Source/Runtime/Online/BuildPatchServices/Public/Interfaces/IBuildManifest.h b/Engine/Source/Runtime/Online/BuildPatchServices/Public/Interfaces/IBuildManifest.h
--- a/Engine/Source/Runtime/Online/BuildPatchServices/Public/Interfaces/IBuildManifest.h
+++ b/Engine/Source/Runtime/Online/BuildPatchServices/Public/Interfaces/IBuildManifest.h
@@ -134,6 +134,12 @@ public:
 	 */
 	virtual TArray<FString> GetBuildFileList() const = 0;
 
+	/**
+	 * Get the list of install tags in this manifest
+	 * @param Tags			OUT		Receives the tags referenced.
+	 */
+	virtual void GetFileTagList(TSet<FString>& Tags) const = 0;
+
 	/**
 	 * Gets a list of files that were installed with the Old Manifest, but no longer required by this Manifest.
 	 * @param OldManifest	IN		The Build Manifest that is currently installed.
