#xmpp
- fix potential threading issue when chat room member exits

--------
Integrated using branch UE4-Orion-To-UE4 of change#2698959 by [REDACTED_PERSON] on 2015/09/21 10:39:42.

diff --git a/Engine/Source/Runtime/Online/XMPP/Private/XmppJingle/XmppMultiUserChatJingle.cpp b/Engine/Source/Runtime/Online/XMPP/Private/XmppJingle/XmppMultiUserChatJingle.cpp
--- a/Engine/Source/Runtime/Online/XMPP/Private/XmppJingle/XmppMultiUserChatJingle.cpp
+++ b/Engine/Source/Runtime/Online/XMPP/Private/XmppJingle/XmppMultiUserChatJingle.cpp
@@ -1762,6 +1762,8 @@ void FXmppMultiUserChatJingle::MemberExited(
 	FXmppUserJid MemberJid;
 	FXmppJingle::ConvertToJid(MemberJid, XmppMember->member_jid());
 
+	FScopeLock Lock(&ChatroomsLock);
+
 	FXmppRoomId RoomId(UTF8_TO_TCHAR(RoomModule->chatroom_jid().node().c_str()));
 
 	FScopeLock Lock(&ChatroomsLock);
