Back out revision 8 from //depot/UE4/Engine/Source/Runtime/Online/XMPP/Private/XmppJingle/XmppMultiUserChatJingle.cpp

diff --git a/Engine/Source/Runtime/Online/XMPP/Private/XmppJingle/XmppMultiUserChatJingle.cpp b/Engine/Source/Runtime/Online/XMPP/Private/XmppJingle/XmppMultiUserChatJingle.cpp
--- a/Engine/Source/Runtime/Online/XMPP/Private/XmppJingle/XmppMultiUserChatJingle.cpp
+++ b/Engine/Source/Runtime/Online/XMPP/Private/XmppJingle/XmppMultiUserChatJingle.cpp
@@ -1762,8 +1762,6 @@ void FXmppMultiUserChatJingle::MemberExited(
 	FXmppUserJid MemberJid;
 	FXmppJingle::ConvertToJid(MemberJid, XmppMember->member_jid());
 
-	FScopeLock Lock(&ChatroomsLock);
-
 	FXmppRoomId RoomId(UTF8_TO_TCHAR(RoomModule->chatroom_jid().node().c_str()));
 
 	FScopeLock Lock(&ChatroomsLock);
