Store "Content-Length" header value when recieved so GetContentLength() can be used before a response is complete (original CL 2594678)

diff --git a/Engine/Source/Runtime/Online/HTTP/Private/Curl/CurlHttp.cpp b/Engine/Source/Runtime/Online/HTTP/Private/Curl/CurlHttp.cpp
--- a/Engine/Source/Runtime/Online/HTTP/Private/Curl/CurlHttp.cpp
+++ b/Engine/Source/Runtime/Online/HTTP/Private/Curl/CurlHttp.cpp
@@ -273,6 +273,12 @@ size_t FCurlHttpRequest::ReceiveResponseHeaderCallback(void* Ptr, size_t SizeInB
 			if (Header.Split(TEXT(": "), &HeaderName, &Param))
 			{
 				Response->Headers.Add(HeaderName, Param);
+
+				//Store the content length so OnRequestProgress() delegates have something to work with
+				if (HeaderName == TEXT("Content-Length"))
+				{
+					Response->ContentLength = FCString::Atoi(*Param);
+				}
 			}
 			return HeaderSize;
 		}
