GetBeaconType should be virtual so it can be overridden by child OnlineBeaconClients

diff --git a/Engine/Source/Runtime/Online/OnlineSubsystemUtils/Public/OnlineBeaconClient.h b/Engine/Source/Runtime/Online/OnlineSubsystemUtils/Public/OnlineBeaconClient.h
--- a/Engine/Source/Runtime/Online/OnlineSubsystemUtils/Public/OnlineBeaconClient.h
+++ b/Engine/Source/Runtime/Online/OnlineSubsystemUtils/Public/OnlineBeaconClient.h
@@ -71,7 +71,7 @@ class ONLINESUBSYSTEMUTILS_API AOnlineBeaconClient : public AOnlineBeacon
 	 *
 	 * @return string representing the type of beacon 
 	 */
-	FString GetBeaconType() const;
+	virtual FString GetBeaconType() const;
 	
 	/**
 	 * A connection has been made and RPC/replication can begin
