#merge: More logs to diagnose timeouts.

diff --git a/Engine/Source/Runtime/Online/OnlineSubsystemUtils/Private/IpNetDriver.cpp b/Engine/Source/Runtime/Online/OnlineSubsystemUtils/Private/IpNetDriver.cpp
--- a/Engine/Source/Runtime/Online/OnlineSubsystemUtils/Private/IpNetDriver.cpp
+++ b/Engine/Source/Runtime/Online/OnlineSubsystemUtils/Private/IpNetDriver.cpp
@@ -168,6 +168,8 @@ void UIpNetDriver::TickDispatch( float DeltaTime )
 
 	ISocketSubsystem* SocketSubsystem = GetSocketSubsystem();
 
+	const double StartReceiveTime = FPlatformTime::Seconds();
+
 	// Process all incoming packets.
 	uint8 Data[MAX_PACKET_SIZE];
 	TSharedRef<FInternetAddr> FromAddr = SocketSubsystem->CreateInternetAddr();
@@ -303,6 +305,15 @@ void UIpNetDriver::TickDispatch( float DeltaTime )
 			}
 		}
 	}
+
+	const double EndReceiveTime		= FPlatformTime::Seconds();
+	const float DeltaReceiveTime	= EndReceiveTime - StartReceiveTime;
+	const float Threshold			= 10.0f;
+
+	if ( DeltaReceiveTime > Threshold )
+	{
+		UE_LOG( LogNet, Warning, TEXT( "UIpNetDriver::TickDispatch: Took too long to receive packets. Time: %2.2f %s" ), DeltaReceiveTime, *GetName() );
+	}
 }
 
 void UIpNetDriver::ProcessRemoteFunction(class AActor* Actor, UFunction* Function, void* Parameters, FOutParmRec* OutParms, FFrame* Stack, class UObject* SubObject )
