Work around bug caused by compiler bug in VS2015 shipping builds

diff --git a/Engine/Source/Runtime/Online/BuildPatchServices/Private/Generation/DataScanner.h b/Engine/Source/Runtime/Online/BuildPatchServices/Private/Generation/DataScanner.h
--- a/Engine/Source/Runtime/Online/BuildPatchServices/Private/Generation/DataScanner.h
+++ b/Engine/Source/Runtime/Online/BuildPatchServices/Private/Generation/DataScanner.h
@@ -23,6 +23,9 @@ namespace BuildPatchServices
 
 	struct FChunkInfo
 	{
+		// Work-around for a compiler bug in VS2015 affecting shipping builds.
+		FChunkInfo() { }
+
 		// The hash for this chunk
 		uint64 Hash;
 		// The SHA hash for this chunk
