[INTEGRATE] merge UE4-LauncherDev >UE4 CL 2712469
OSS changes to support blocked players

diff --git a/Engine/Source/Runtime/Online/OnlineSubsystem/Public/Interfaces/OnlineFriendsInterface.h b/Engine/Source/Runtime/Online/OnlineSubsystem/Public/Interfaces/OnlineFriendsInterface.h
--- a/Engine/Source/Runtime/Online/OnlineSubsystem/Public/Interfaces/OnlineFriendsInterface.h
+++ b/Engine/Source/Runtime/Online/OnlineSubsystem/Public/Interfaces/OnlineFriendsInterface.h
@@ -105,6 +105,30 @@ typedef FOnRejectInviteComplete::FDelegate FOnRejectInviteCompleteDelegate;
 DECLARE_MULTICAST_DELEGATE_FiveParams(FOnDeleteFriendComplete, int32, bool, const FUniqueNetId&, const FString&, const FString&);
 typedef FOnDeleteFriendComplete::FDelegate FOnDeleteFriendCompleteDelegate;
 
+/**
+ * Delegate used when a block request has completed
+ *
+ * @param LocalUserNum the controller number of the associated user that made the request
+ * @param bWasSuccessful true if the async action completed without error, false if there was an error
+ * @param UniqueID Player blocked
+ * @param ListName name of the friends list that was operated on
+ * @param ErrorStr string representing the error condition
+ */
+DECLARE_MULTICAST_DELEGATE_FiveParams(FOnBlockedPlayerComplete, int32, bool, const FUniqueNetId&, const FString&, const FString&);
+typedef FOnBlockedPlayerComplete::FDelegate FOnBlockedPlayerCompleteDelegate;
+
+/**
+ * Delegate used when an unblock request has completed
+ *
+ * @param LocalUserNum the controller number of the associated user that made the request
+ * @param bWasSuccessful true if the async action completed without error, false if there was an error
+ * @param UniqueID Player unblocked
+ * @param ListName name of the friends list that was operated on
+ * @param ErrorStr string representing the error condition
+ */
+DECLARE_MULTICAST_DELEGATE_FiveParams(FOnUnblockedPlayerComplete, int32, bool, const FUniqueNetId&, const FString&, const FString&);
+typedef FOnUnblockedPlayerComplete::FDelegate FOnUnblockedPlayerCompleteDelegate;
+
 /**
  * Delegate used when the query for recent players has completed
  *
@@ -116,6 +140,16 @@ typedef FOnDeleteFriendComplete::FDelegate FOnDeleteFriendCompleteDelegate;
 DECLARE_MULTICAST_DELEGATE_FourParams(FOnQueryRecentPlayersComplete, const FUniqueNetId& /*UserId*/, const FString& /*Namespace*/, bool /*bWasSuccessful*/, const FString& /*Error*/);
 typedef FOnQueryRecentPlayersComplete::FDelegate FOnQueryRecentPlayersCompleteDelegate;
 
+/**
+ * Delegate used when the query for blocked players has completed
+ *
+ * @param UserId the id of the user that made the request
+ * @param bWasSuccessful true if the async action completed without error, false if there was an error
+ * @param ErrorStr string representing the error condition
+ */
+DECLARE_MULTICAST_DELEGATE_ThreeParams(FOnQueryBlockedPlayersComplete, const FUniqueNetId& /*UserId*/, bool /*bWasSuccessful*/, const FString& /*Error*/);
+typedef FOnQueryBlockedPlayersComplete::FDelegate FOnQueryBlockedPlayersCompleteDelegate;
+
 /**
  * Delegate called when remote friend sends an invite
  *
@@ -287,6 +321,28 @@ public:
 	 */
 	DEFINE_ONLINE_PLAYER_DELEGATE_FOUR_PARAM(MAX_LOCAL_PLAYERS, OnDeleteFriendComplete, bool, const FUniqueNetId&, const FString&, const FString&);
 
+	/**
+	 * Delegate used when a block player request has completed
+	 *
+	 * @param LocalUserNum the controller number of the associated user that made the request
+	 * @param bWasSuccessful true if the async action completed without error, false if there was an error
+	 * @param UniqueIF player that was blocked
+	 * @param ListName name of the friends list that was operated on
+	 * @param ErrorStr string representing the error condition
+	 */
+	DEFINE_ONLINE_PLAYER_DELEGATE_FOUR_PARAM(MAX_LOCAL_PLAYERS, OnBlockedPlayerComplete, bool, const FUniqueNetId&, const FString&, const FString&);
+
+	/**
+	 * Delegate used when an ublock player request has completed
+	 *
+	 * @param LocalUserNum the controller number of the associated user that made the request
+	 * @param bWasSuccessful true if the async action completed without error, false if there was an error
+	 * @param UniqueIF player that was unblocked
+	 * @param ListName name of the friends list that was operated on
+	 * @param ErrorStr string representing the error condition
+	 */
+	DEFINE_ONLINE_PLAYER_DELEGATE_FOUR_PARAM(MAX_LOCAL_PLAYERS, OnUnblockedPlayerComplete, bool, const FUniqueNetId&, const FString&, const FString&);
+
 	/**
 	 * Copies the list of friends for the player previously retrieved from the online service
 	 *
@@ -340,6 +396,15 @@ public:
 	 */
 	DEFINE_ONLINE_DELEGATE_FOUR_PARAM(OnQueryRecentPlayersComplete, const FUniqueNetId& /*UserId*/, const FString& /*Namespace*/, bool /*bWasSuccessful*/, const FString& /*Error*/);
 
+	/**
+	 * Delegate used when the query for blocked players has completed
+	 *
+	 * @param UserId the id of the user that made the request
+	 * @param bWasSuccessful true if the async action completed without error, false if there was an error
+	 * @param Error string representing the error condition
+	 */
+	DEFINE_ONLINE_DELEGATE_THREE_PARAM(OnQueryBlockedPlayersComplete, const FUniqueNetId& /*UserId*/, bool /*bWasSuccessful*/, const FString& /*Error*/);
+
 	/**
 	 * Copies the cached list of recent players for a given user
 	 *
@@ -350,6 +415,46 @@ public:
 	 * @return true if recent players list was found for the given user
 	 */
 	virtual bool GetRecentPlayers(const FUniqueNetId& UserId, const FString& Namespace, TArray< TSharedRef<FOnlineRecentPlayer> >& OutRecentPlayers) = 0;
+
+	/**
+	 * Block a player
+	 *
+	 * @param LocalUserNum The user to check for
+	 * @param PlayerId The player to block
+	 *
+	 * @return true if query was started
+	 */
+	virtual bool BlockPlayer(int32 LocalUserNum, const FUniqueNetId& PlayerId) = 0;
+
+	/**
+	 * Unblock a player
+	 *
+	 * @param LocalUserNum The user to check for
+	 * @param PlayerId The player to unblock
+	 *
+	 * @return true if query was started
+	 */
+	virtual bool UnblockPlayer(int32 LocalUserNum, const FUniqueNetId& PlayerId) = 0;
+
+	/**
+	 * Query for blocked players
+	 *
+	 * @param UserId user to query blocked players for
+	 *
+	 * @return true if query was started
+	 */
+	virtual bool QueryBlockedPlayers(const FUniqueNetId& UserId) = 0;
+
+	/**
+	 * Get the list of blocked players
+	 *
+	 * @param UserId user to retrieve blocked players for
+	 * @param OuBlockedPlayers [out] array that receives the copied data
+	 *
+	 * @return true if blocked players list was found for the given user
+	 */
+	virtual bool GetBlockedPlayers(const FUniqueNetId& UserId, TArray< TSharedRef<FOnlineBlockedPlayer> >& OutBlockedPlayers) = 0;
+
 };
 
 typedef TSharedPtr<IOnlineFriends, ESPMode::ThreadSafe> IOnlineFriendsPtr;
diff --git a/Engine/Source/Runtime/Online/OnlineSubsystem/Public/OnlineSubsystemTypes.h b/Engine/Source/Runtime/Online/OnlineSubsystem/Public/OnlineSubsystemTypes.h
--- a/Engine/Source/Runtime/Online/OnlineSubsystem/Public/OnlineSubsystemTypes.h
+++ b/Engine/Source/Runtime/Online/OnlineSubsystem/Public/OnlineSubsystemTypes.h
@@ -1094,7 +1094,9 @@ namespace EInviteStatus
 		/** Friend has sent player an invite, but it has not been accepted/rejected */
 		PendingInbound,
 		/** Player has sent friend an invite, but it has not been accepted/rejected */
-		PendingOutbound
+		PendingOutbound,
+		/** Player has been blocked */
+		Blocked,
 	};
 
 	/** 
@@ -1120,6 +1122,10 @@ namespace EInviteStatus
 			{
 				return TEXT("PendingOutbound");
 			}
+			case Blocked:
+			{
+				return TEXT("Blocked");
+			}
 		}
 		return TEXT("");
 	}
@@ -1156,6 +1162,13 @@ public:
 	virtual FDateTime GetLastSeen() const = 0;
 };
 
+/**
+ * Blocked user info returned via IOnlineFriends interface
+ */
+class FOnlineBlockedPlayer : public FOnlineUser
+{
+};
+
 /** The possible permission categories we can choose from to read from the server */
 namespace EOnlineSharingReadCategory
 {
