UE4 - Added a cache to avoid parsing online subsystem module names in a singleton.

diff --git a/Engine/Source/Runtime/Online/OnlineSubsystem/Private/OnlineSubsystemModule.cpp b/Engine/Source/Runtime/Online/OnlineSubsystem/Private/OnlineSubsystemModule.cpp
--- a/Engine/Source/Runtime/Online/OnlineSubsystem/Private/OnlineSubsystemModule.cpp
+++ b/Engine/Source/Runtime/Online/OnlineSubsystem/Private/OnlineSubsystemModule.cpp
@@ -138,6 +138,7 @@ void FOnlineSubsystemModule::UnregisterPlatformService(const FName FactoryName)
 	}
 }
 
+
 void FOnlineSubsystemModule::ParseOnlineSubsystemName(const FName& FullName, FName& SubsystemName, FName& InstanceName) const
 {
 	SubsystemName = DefaultPlatformService;
@@ -168,7 +169,7 @@ void FOnlineSubsystemModule::ParseOnlineSubsystemName(const FName& FullName, FNa
 	}
 }
 
-IOnlineSubsystem* FOnlineSubsystemModule::GetOnlineSubsystem(const FName InSubsystemName)
+FName FOnlineSubsystemModule::ParseOnlineSubsystemName(FName InSubsystemName) const
 {
 	FName SubsystemName, InstanceName;
 	ParseOnlineSubsystemName(InSubsystemName, SubsystemName, InstanceName);
@@ -176,11 +177,44 @@ IOnlineSubsystem* FOnlineSubsystemModule::GetOnlineSubsystem(const FName InSubsy
 	IOnlineSubsystemPtr* OnlineSubsystem = NULL;
 	if (SubsystemName != NAME_None)
 	{
-		FName KeyName = FName(*FString::Printf(TEXT("%s:%s"), *SubsystemName.ToString(), *InstanceName.ToString()));
+		return FName(*FString::Printf(TEXT("%s:%s"), *SubsystemName.ToString(), *InstanceName.ToString()));
+	}
+	return NAME_None;
+}
 
+FName FOnlineSubsystemModule::ParseOnlineSubsystemNameCached(FName FullName) const
+{
+#if !WITH_EDITOR
+	if (IsInGameThread()) // yuk!
+	{
+		static TMap<FName, FName> Cache;
+		FName* Existing = Cache.Find(FullName);
+		if (Existing)
+		{
+			return *Existing;
+		}
+		FName Result = ParseOnlineSubsystemName(FullName);
+		Cache.Add(FullName, Result);
+		return Result;
+	}
+#endif
+	return ParseOnlineSubsystemName(FullName);
+}
+
+IOnlineSubsystem* FOnlineSubsystemModule::GetOnlineSubsystem(const FName InSubsystemName)
+{
+	QUICK_SCOPE_CYCLE_COUNTER(STAT_FOnlineSubsystemModule_GetOnlineSubsystem);
+	IOnlineSubsystemPtr* OnlineSubsystem = NULL;
+
+	FName KeyName = ParseOnlineSubsystemNameCached(InSubsystemName);
+
+	if (KeyName != NAME_None)
+	{
 		OnlineSubsystem = OnlineSubsystems.Find(KeyName);
 		if (OnlineSubsystem == NULL)
 		{
+			FName SubsystemName, InstanceName;
+			ParseOnlineSubsystemName(InSubsystemName, SubsystemName, InstanceName);
 			IOnlineFactory** OSSFactory = OnlineFactories.Find(SubsystemName);
 			if (OSSFactory == NULL)
 			{
diff --git a/Engine/Source/Runtime/Online/OnlineSubsystem/Public/OnlineSubsystemModule.h b/Engine/Source/Runtime/Online/OnlineSubsystem/Public/OnlineSubsystemModule.h
--- a/Engine/Source/Runtime/Online/OnlineSubsystem/Public/OnlineSubsystemModule.h
+++ b/Engine/Source/Runtime/Online/OnlineSubsystem/Public/OnlineSubsystemModule.h
@@ -49,6 +49,12 @@ private:
 	 */
 	void ParseOnlineSubsystemName(const FName& FullName, FName& SubsystemName, FName& InstanceName) const;
 
+	// As above but recombines the SubsystemName and InstanceName into a key
+	FName ParseOnlineSubsystemName(FName FullName) const;
+
+	// As above but uses a cache to avoid all of the string operations
+	FName ParseOnlineSubsystemNameCached(FName FullName) const;
+
 	/**
 	 *	Shuts down all registered online subsystem platforms and unloads their modules
 	 */
