#online,mcp,json
- json serializer better handling of inline objects

--------
Integrated using branch UE4-Orion-To-UE4 of change#2704481 by [REDACTED_PERSON] on 2015/09/24 11:05:31.

diff --git a/Engine/Source/Runtime/Online/OnlineSubsystem/Public/OnlineJsonSerializer.h b/Engine/Source/Runtime/Online/OnlineSubsystem/Public/OnlineJsonSerializer.h
--- a/Engine/Source/Runtime/Online/OnlineSubsystem/Public/OnlineJsonSerializer.h
+++ b/Engine/Source/Runtime/Online/OnlineSubsystem/Public/OnlineJsonSerializer.h
@@ -8,19 +8,12 @@
  * Macros used to generate a serialization function for a class derived from FJsonSerializable
  */
 #define BEGIN_ONLINE_JSON_SERIALIZER \
-	virtual void Serialize(FOnlineJsonSerializerBase& Serializer) override \
+	virtual void Serialize(FOnlineJsonSerializerBase& Serializer, bool bFlatObject) override \
 	{ \
-		Serializer.StartObject();
+		if (!bFlatObject) Serializer.StartObject();
 
 #define END_ONLINE_JSON_SERIALIZER \
-		Serializer.EndObject(); \
-	}
-
-#define BEGIN_ONLINE_JSON_SERIALIZER_FLAT \
-	virtual void Serialize(FOnlineJsonSerializerBase& Serializer) override \
-	{
-
-#define END_ONLINE_JSON_SERIALIZER_FLAT \
+		if (!bFlatObject) Serializer.EndObject(); \
 	}
 
 #define ONLINE_JSON_SERIALIZE(JsonName, JsonValue) \
@@ -33,7 +26,7 @@
 		Serializer.SerializeMap(TEXT(JsonName), JsonMap)
 
 #define ONLINE_JSON_SERIALIZE_SERIALIZABLE(JsonName, JsonValue) \
-		JsonValue.Serialize(Serializer)
+		JsonValue.Serialize(Serializer, false)
 
 #define ONLINE_JSON_SERIALIZE_ARRAY_SERIALIZABLE(JsonName, JsonArray, ElementType) \
 		if (Serializer.IsLoading()) \
@@ -52,7 +45,7 @@
 			Serializer.StartArray(JsonName); \
 			for (auto It = JsonArray.CreateIterator(); It; ++It) \
 			{ \
-				It->Serialize(Serializer); \
+				It->Serialize(Serializer, false); \
 			} \
 			Serializer.EndArray(); \
 		}
@@ -77,7 +70,7 @@
 			for (auto It = JsonMap.CreateIterator(); It; ++It) \
 			{ \
 				Serializer.StartObject(It.Key()); \
-				It.Value().Serialize(Serializer); \
+				It.Value().Serialize(Serializer, true); \
 				Serializer.EndObject(); \
 			} \
 			Serializer.EndObject(); \
@@ -614,14 +607,14 @@ struct FOnlineJsonSerializable
 		{
 			TSharedRef<TJsonWriter<> > JsonWriter = TJsonWriterFactory<>::Create(&JsonStr);
 			FOnlineJsonSerializerWriter<> Serializer(JsonWriter);
-			Serialize(Serializer);
+			Serialize(Serializer, false);
 			JsonWriter->Close();
 		}
 		else
 		{
 			TSharedRef< TJsonWriter< TCHAR, TCondensedJsonPrintPolicy< TCHAR > > > JsonWriter = TJsonWriterFactory< TCHAR, TCondensedJsonPrintPolicy< TCHAR > >::Create( &JsonStr );
 			FOnlineJsonSerializerWriter<TCHAR, TCondensedJsonPrintPolicy< TCHAR >> Serializer(JsonWriter);
-			Serialize(Serializer);
+			Serialize(Serializer, false);
 			JsonWriter->Close();
 		}
 		return JsonStr;
@@ -640,7 +633,7 @@ struct FOnlineJsonSerializable
 			JsonObject.IsValid())
 		{
 			FOnlineJsonSerializerReader Serializer(JsonObject);
-			Serialize(Serializer);
+			Serialize(Serializer, false);
 			return true;
 		}
 		return false;
@@ -650,7 +643,7 @@ struct FOnlineJsonSerializable
 		if (JsonObject.IsValid())
 		{
 			FOnlineJsonSerializerReader Serializer(JsonObject);
-			Serialize(Serializer);
+			Serialize(Serializer, false);
 			return true;
 		}
 		return false;
@@ -660,8 +653,9 @@ struct FOnlineJsonSerializable
 	 * Abstract method that needs to be supplied using the macros
 	 *
 	 * @param Serializer the object that will perform serialization in/out of JSON
+	 * @param bFlatObject if true then no object wrapper is used
 	 */
-	virtual void Serialize(FOnlineJsonSerializerBase& Serializer) = 0;
+	virtual void Serialize(FOnlineJsonSerializerBase& Serializer, bool bFlatObject) = 0;
 };
 
 inline void FOnlineJsonSerializerReader::SerializeObject(const TCHAR* Name, FOnlineJsonSerializable& Value)
@@ -682,5 +676,5 @@ inline void FOnlineJsonSerializerWriter<CharType, PrintPolicy>::SerializeObject(
 {
 	/* Write the value to the Name field */
 	JsonWriter->WriteIdentifierPrefix(Name);
-	Object.Serialize(*this);
+	Object.Serialize(*this, true);
 }
