#online,identity,mcp
- Added account creation APIs
- primarily used for creating "headless" account
- can provide external auth credentials to associate with an external account at creation time

--------
Integrated using branch UE4-Orion-To-UE4 of change#2709388 by [REDACTED_PERSON] on 2015/09/29 10:23:30.

diff --git a/Engine/Source/Runtime/Online/OnlineSubsystem/Public/OnlineJsonSerializer.h b/Engine/Source/Runtime/Online/OnlineSubsystem/Public/OnlineJsonSerializer.h
--- a/Engine/Source/Runtime/Online/OnlineSubsystem/Public/OnlineJsonSerializer.h
+++ b/Engine/Source/Runtime/Online/OnlineSubsystem/Public/OnlineJsonSerializer.h
@@ -618,7 +618,16 @@ struct FOnlineJsonSerializable
 			JsonWriter->Close();
 		}
 		return JsonStr;
-
+	}
+	virtual void ToJson(TSharedRef<TJsonWriter<> >& JsonWriter, bool bFlatObject) const
+	{
+		FOnlineJsonSerializerWriter<> Serializer(JsonWriter);
+		((FOnlineJsonSerializable*)this)->Serialize(Serializer, bFlatObject);
+	}
+	virtual void ToJson(TSharedRef< TJsonWriter< TCHAR, TCondensedJsonPrintPolicy< TCHAR > > >& JsonWriter, bool bFlatObject) const
+	{
+		FOnlineJsonSerializerWriter<TCHAR, TCondensedJsonPrintPolicy< TCHAR >> Serializer(JsonWriter);
+		((FOnlineJsonSerializable*)this)->Serialize(Serializer, bFlatObject);
 	}
 	/**
 	 * Serializes the contents of a JSON string into this object
