Merging CL#2728976

Quick fix to possible mismatch of packet MTU size and receive buffer size + warnings when this might be happening in the wild

diff --git a/Engine/Source/Runtime/Online/OnlineSubsystemUtils/Private/IpConnection.cpp b/Engine/Source/Runtime/Online/OnlineSubsystemUtils/Private/IpConnection.cpp
--- a/Engine/Source/Runtime/Online/OnlineSubsystemUtils/Private/IpConnection.cpp
+++ b/Engine/Source/Runtime/Online/OnlineSubsystemUtils/Private/IpConnection.cpp
@@ -142,6 +142,12 @@ void UIpConnection::LowLevelSend( void* Data, int32 Count )
 	// Send to remote.
 	int32 BytesSent = 0;
 	CLOCK_CYCLES(Driver->SendCycles);
+
+	if ( Count > MaxPacket )
+	{
+		UE_LOG( LogNet, Warning, TEXT( "UIpConnection::LowLevelSend: Count > MaxPacketSize! Count: %i, MaxPacket: %i %s" ), Count, MaxPacket, *Describe() );
+	}
+
 	Socket->SendTo(DataToSend, Count, BytesSent, *RemoteAddr);
 	UNCLOCK_CYCLES(Driver->SendCycles);
 	NETWORK_PROFILER(GNetworkProfiler.FlushOutgoingBunches(this));
diff --git a/Engine/Source/Runtime/Online/OnlineSubsystemUtils/Private/IpNetDriver.cpp b/Engine/Source/Runtime/Online/OnlineSubsystemUtils/Private/IpNetDriver.cpp
--- a/Engine/Source/Runtime/Online/OnlineSubsystemUtils/Private/IpNetDriver.cpp
+++ b/Engine/Source/Runtime/Online/OnlineSubsystemUtils/Private/IpNetDriver.cpp
@@ -171,7 +171,7 @@ void UIpNetDriver::TickDispatch( float DeltaTime )
 	const double StartReceiveTime = FPlatformTime::Seconds();
 
 	// Process all incoming packets.
-	uint8 Data[MAX_PACKET_SIZE];
+	uint8 Data[MAX_PACKET_SIZE * 2];		// x 2 for now until we resolve some issues with the sending side incorrectly calculating fragmentation size
 	TSharedRef<FInternetAddr> FromAddr = SocketSubsystem->CreateInternetAddr();
 	for( ; Socket != NULL; )
 	{
