Build fix for OnlineUserFacebook - implement pure virtual functions.

diff --git a/Engine/Source/Runtime/Online/OnlineSubsystemFacebook/Private/IOS/OnlineUserFacebook.cpp b/Engine/Source/Runtime/Online/OnlineSubsystemFacebook/Private/IOS/OnlineUserFacebook.cpp
--- a/Engine/Source/Runtime/Online/OnlineSubsystemFacebook/Private/IOS/OnlineUserFacebook.cpp
+++ b/Engine/Source/Runtime/Online/OnlineSubsystemFacebook/Private/IOS/OnlineUserFacebook.cpp
@@ -208,3 +208,19 @@ bool FOnlineUserFacebook::QueryUserIdMapping(const FUniqueNetId& UserId, const F
 	return false;
 }
 
+bool FOnlineUserFacebook::QueryExternalIdMappings(const FUniqueNetId& LocalUserId, const FString& AuthType, const TArray<FString>& ExternalIds, const FOnQueryExternalIdMappingsComplete& Delegate)
+{
+	Delegate.ExecuteIfBound(false, LocalUserId, AuthType, ExternalIds, TEXT("not implemented"));
+	return false;
+}
+
+void FOnlineUserFacebook::GetExternalIdMappings(const FString& AuthType, const TArray<FString>& ExternalIds, TArray<TSharedPtr<const FUniqueNetId>>& OutIds)
+{
+	// Not implemented for PS4 - return an array full of empty values
+	OutIds.SetNum(ExternalIds.Num());
+}
+
+TSharedPtr<const FUniqueNetId> FOnlineUserFacebook::GetExternalIdMapping(const FString& AuthType, const FString& ExternalId)
+{
+	return TSharedPtr<FUniqueNetId>();
+}
diff --git a/Engine/Source/Runtime/Online/OnlineSubsystemFacebook/Private/IOS/OnlineUserFacebook.h b/Engine/Source/Runtime/Online/OnlineSubsystemFacebook/Private/IOS/OnlineUserFacebook.h
--- a/Engine/Source/Runtime/Online/OnlineSubsystemFacebook/Private/IOS/OnlineUserFacebook.h
+++ b/Engine/Source/Runtime/Online/OnlineSubsystemFacebook/Private/IOS/OnlineUserFacebook.h
@@ -79,6 +79,9 @@ public:
 	virtual bool GetAllUserInfo(int32 LocalUserNum, TArray< TSharedRef<class FOnlineUser> >& OutUsers) override;
 	virtual TSharedPtr<FOnlineUser> GetUserInfo(int32 LocalUserNum, const class FUniqueNetId& UserId) override;	
 	virtual bool QueryUserIdMapping(const FUniqueNetId& UserId, const FString& DisplayNameOrEmail, const FOnQueryUserMappingComplete& Delegate = FOnQueryUserMappingComplete()) override;
+	virtual bool QueryExternalIdMappings(const FUniqueNetId& LocalUserId, const FString& AuthType, const TArray<FString>& ExternalIds, const FOnQueryExternalIdMappingsComplete& Delegate = FOnQueryExternalIdMappingsComplete()) override;
+	virtual void GetExternalIdMappings(const FString& AuthType, const TArray<FString>& ExternalIds, TArray<TSharedPtr<const FUniqueNetId>>& OutIds) override;
+	virtual TSharedPtr<const FUniqueNetId> GetExternalIdMapping(const FString& AuthType, const FString& ExternalId) override;
 
 	// FOnlineUserFacebook
 
