Merging using UE4-To-UE4-Main

Change: 2746803

Merging using UE4-To-UE4-LauncherDev

merging up to cl 276442
from //depot/UE4-LauncherDev/Engine/Source/Runtime/Online/...
to //depot/UE4/Engine/Source/Runtime/Online/...

diff --git a/Engine/Source/Runtime/Online/BuildPatchServices/Private/BuildPatchAnalytics.cpp b/Engine/Source/Runtime/Online/BuildPatchServices/Private/BuildPatchAnalytics.cpp
--- a/Engine/Source/Runtime/Online/BuildPatchServices/Private/BuildPatchAnalytics.cpp
+++ b/Engine/Source/Runtime/Online/BuildPatchServices/Private/BuildPatchAnalytics.cpp
@@ -81,12 +81,14 @@ void FBuildPatchAnalytics::RecordConstructionError( const FString& Filename, con
 	}
 }
 
-void FBuildPatchAnalytics::RecordPrereqInstallnError( const FString& Filename, const FString& CommandLine, const int32 ErrorCode, const FString& ErrorString )
+void FBuildPatchAnalytics::RecordPrereqInstallationError(const FString& AppName, const FString& AppVersion, const FString& Filename, const FString& CommandLine, const int32 ErrorCode, const FString& ErrorString)
 {
 	TArray<FAnalyticsEventAttribute> Attributes;
+	Attributes.Add( FAnalyticsEventAttribute( TEXT( "AppName" ), AppName ) );
+	Attributes.Add( FAnalyticsEventAttribute( TEXT( "AppVersion" ), AppVersion ) );
 	Attributes.Add( FAnalyticsEventAttribute( TEXT( "Filename" ), Filename ) );
 	Attributes.Add( FAnalyticsEventAttribute( TEXT( "CommandLine" ), CommandLine ) );
-	Attributes.Add( FAnalyticsEventAttribute( TEXT( "ErrorCode" ), ErrorCode ) );
+	Attributes.Add( FAnalyticsEventAttribute( TEXT( "ReturnCode" ), ErrorCode ) );
 	Attributes.Add( FAnalyticsEventAttribute( TEXT( "ErrorString" ), ErrorString ) );
 	FBuildPatchHTTP::QueueAnalyticsEvent( TEXT( "Patcher.Error.Prerequisites" ), Attributes );
 
diff --git a/Engine/Source/Runtime/Online/BuildPatchServices/Private/BuildPatchAnalytics.h b/Engine/Source/Runtime/Online/BuildPatchServices/Private/BuildPatchAnalytics.h
--- a/Engine/Source/Runtime/Online/BuildPatchServices/Private/BuildPatchAnalytics.h
+++ b/Engine/Source/Runtime/Online/BuildPatchServices/Private/BuildPatchAnalytics.h
@@ -87,13 +87,15 @@ public:
 	static void RecordConstructionError( const FString& Filename, const int32 LastError, const FString& ErrorString );
 
 	/**
-	 * Record a file construction error event
+	 * Record a prerequisite installation error
+	 * @param AppName		The app name for the installation
+	 * @param AppVersion	The version string for the installation
 	 * @param Filename		The installer file
 	 * @param CommandLine	The command line passed to the installer
-	 * @param ErrorCode		The OS error code if appropriate
+	 * @param ErrorCode		The return code from the run
 	 * @param ErrorString	The type of error that has occurred
 	 */
-	static void RecordPrereqInstallnError( const FString& Filename, const FString& CommandLine, const int32 ErrorCode, const FString& ErrorString );
+	static void RecordPrereqInstallationError(const FString& AppName, const FString& AppVersion, const FString& Filename, const FString& CommandLine, const int32 ErrorCode, const FString& ErrorString);
 
 	/**
 	 * Set the Http Service Tracker to be used for tracking Http Service responsiveness.
diff --git a/Engine/Source/Runtime/Online/BuildPatchServices/Private/BuildPatchDownloader.cpp b/Engine/Source/Runtime/Online/BuildPatchServices/Private/BuildPatchDownloader.cpp
--- a/Engine/Source/Runtime/Online/BuildPatchServices/Private/BuildPatchDownloader.cpp
+++ b/Engine/Source/Runtime/Online/BuildPatchServices/Private/BuildPatchDownloader.cpp
@@ -71,8 +71,9 @@ private:
 
 /* FBuildPatchDownloader implementation
 *****************************************************************************/
-FBuildPatchDownloader::FBuildPatchDownloader( const FString& InSaveDirectory, const FBuildPatchAppManifestRef& InInstallManifest, FBuildPatchProgress* InBuildProgress )
+FBuildPatchDownloader::FBuildPatchDownloader(const FString& InSaveDirectory, const FString& InCloudDirectory, const FBuildPatchAppManifestRef& InInstallManifest, FBuildPatchProgress* InBuildProgress)
 	: SaveDirectory( InSaveDirectory )
+	, CloudDirectory( InCloudDirectory )
 	, InstallManifest( InInstallManifest )
 	, Thread( NULL )
 	, bIsRunning( false )
@@ -122,6 +123,7 @@ uint32 FBuildPatchDownloader::Run()
 
 	// The Average chunk download time
 	FMeanChunkTime MeanChunkTime;
+	TArray< FGuid > InFlightKeys;
 
 	// While there is the possibility of more chunks to download
 	while ( ShouldBeRunning() && !FBuildPatchInstallError::HasFatalError() )
@@ -130,7 +132,7 @@ uint32 FBuildPatchDownloader::Run()
 		FPlatformProcess::Sleep( 0.0f );
 		// Check inflight downloads for completion and process them
 		InFlightDownloadsLock.Lock();
-		TArray< FGuid > InFlightKeys;
+		InFlightKeys.Empty();
 		InFlightDownloads.GetKeys( InFlightKeys );
 		for( auto InFlightKeysIt = InFlightKeys.CreateConstIterator(); InFlightKeysIt && !FBuildPatchInstallError::HasFatalError(); ++InFlightKeysIt )
 		{
@@ -273,7 +275,7 @@ uint32 FBuildPatchDownloader::Run()
 		SetIdle( false );
 
 		// Make the filename and download url
-		FString DownloadUrl = FBuildPatchUtils::GetDataFilename(InstallManifest, FBuildPatchServicesModule::GetCloudDirectory(), NextGuid);
+		FString DownloadUrl = FBuildPatchUtils::GetDataFilename(InstallManifest, CloudDirectory, NextGuid);
 		const bool bIsHTTPRequest = DownloadUrl.Contains( TEXT( "http" ), ESearchCase::IgnoreCase );
 
 		// Start the download
@@ -321,6 +323,27 @@ uint32 FBuildPatchDownloader::Run()
 			Async(EAsyncExecution::ThreadPool, MoveTemp(Task));
 		}
 	}
+
+	// Wait for all canceled downloads to complete
+	InFlightDownloadsLock.Lock();
+	while (InFlightDownloads.Num() > 0)
+	{
+		InFlightKeys.Empty();
+		InFlightDownloads.GetKeys(InFlightKeys);
+		for (const FGuid& InFlightKey : InFlightKeys)
+		{
+			// If the download is not running
+			if (InFlightDownloads[InFlightKey].StateFlag.GetValue() != EDownloadState::DownloadRunning)
+			{
+				InFlightDownloads.Remove(InFlightKey);
+			}
+		}
+		InFlightDownloadsLock.Unlock();
+		FPlatformProcess::Sleep(0.0f);
+		InFlightDownloadsLock.Lock();
+	}
+	InFlightDownloadsLock.Unlock();
+
 	SetIdle( true );
 	SetRunning( false );
 	return 0;
@@ -471,7 +494,7 @@ bool FBuildPatchDownloader::ShouldBeRunning()
 void FBuildPatchDownloader::HttpRequestProgress( FHttpRequestPtr Request, int32 BytesSent, int32 BytesReceived )
 {
 #if EXTRA_DOWNLOAD_LOGGING
-	GWarn->Logf(TEXT("BuildPatchDownloader: Request %p Bytes: %d"), Request.Get(), BytesSoFar);
+	GWarn->Logf(TEXT("BuildPatchDownloader: Request %p BytesSent: %d BytesReceived: %d"), Request.Get(), BytesSent, BytesReceived);
 #endif
 	if ( !FBuildPatchInstallError::HasFatalError() )
 	{
@@ -485,18 +508,15 @@ void FBuildPatchDownloader::HttpRequestComplete( FHttpRequestPtr Request, FHttpR
 {
 	// track the request if we want CDN analytics.
 	FBuildPatchAnalytics::TrackRequest(Request);
-	if ( !FBuildPatchInstallError::HasFatalError() )
+	FGuid DataGUID;
+	FBuildPatchUtils::GetGUIDFromFilename( Request->GetURL(), DataGUID );
+	if( bSucceeded && Response.IsValid() && EHttpResponseCodes::IsOk( Response->GetResponseCode() ) )
 	{
-		FGuid DataGUID;
-		FBuildPatchUtils::GetGUIDFromFilename( Request->GetURL(), DataGUID );
-		if( bSucceeded && Response.IsValid() && EHttpResponseCodes::IsOk( Response->GetResponseCode() ) )
-		{
-			OnDownloadComplete( DataGUID, Request->GetURL(), Response->GetContent(), true, Response->GetResponseCode() );
-		}
-		else
-		{
-			OnDownloadComplete( DataGUID, Request->GetURL(), TArray< uint8 >(), false, Response.IsValid() ? Response->GetResponseCode() : INDEX_NONE );
-		}
+		OnDownloadComplete( DataGUID, Request->GetURL(), Response->GetContent(), true, Response->GetResponseCode() );
+	}
+	else
+	{
+		OnDownloadComplete( DataGUID, Request->GetURL(), TArray< uint8 >(), false, Response.IsValid() ? Response->GetResponseCode() : INDEX_NONE );
 	}
 }
 
@@ -527,7 +547,6 @@ void FBuildPatchDownloader::OnDownloadComplete( const FGuid& Guid, const FString
 #if EXTRA_DOWNLOAD_LOGGING
 	GWarn->Logf(TEXT("BuildPatchDownloader: OnDownloadComplete %s %s %d"), *Guid.ToString(), bSucceeded?TEXT("SUCCESS"):TEXT("FAIL"), ResponseCode);
 #endif
-	if (!FBuildPatchInstallError::HasFatalError())
 	{
 		FScopeLock ScopeLock(&InFlightDownloadsLock);
 		FDownloadJob& CurrentJob = InFlightDownloads[Guid];
@@ -629,11 +648,11 @@ bool FBuildPatchDownloader::GetNextDownload( FGuid& Job )
 *****************************************************************************/
 TSharedPtr< FBuildPatchDownloader, ESPMode::ThreadSafe > FBuildPatchDownloader::SingletonInstance = NULL;
 
-void FBuildPatchDownloader::Create( const FString& InSaveDirectory, const FBuildPatchAppManifestRef& InInstallManifest, FBuildPatchProgress* InBuildProgress )
+void FBuildPatchDownloader::Create(const FString& SaveDirectory, const FString& CloudDirectory, const FBuildPatchAppManifestRef& InstallManifest, FBuildPatchProgress* BuildProgress)
 {
 	// We won't allow misuse of these functions
 	check( !SingletonInstance.IsValid() );
-	SingletonInstance = MakeShareable( new FBuildPatchDownloader( InSaveDirectory, InInstallManifest, InBuildProgress ) );
+	SingletonInstance = MakeShareable(new FBuildPatchDownloader(SaveDirectory, CloudDirectory, InstallManifest, BuildProgress));
 }
 
 FBuildPatchDownloader& FBuildPatchDownloader::Get()
diff --git a/Engine/Source/Runtime/Online/BuildPatchServices/Private/BuildPatchDownloader.h b/Engine/Source/Runtime/Online/BuildPatchServices/Private/BuildPatchDownloader.h
--- a/Engine/Source/Runtime/Online/BuildPatchServices/Private/BuildPatchDownloader.h
+++ b/Engine/Source/Runtime/Online/BuildPatchServices/Private/BuildPatchDownloader.h
@@ -97,6 +97,9 @@ private:
 	// The directory to save data to
 	const FString SaveDirectory;
 
+	// The base url for downloads
+	const FString CloudDirectory;
+
 	// The Manifest that is being installed
 	const FBuildPatchAppManifestRef InstallManifest;
 
@@ -152,11 +155,12 @@ public:
 
 	/**
 	 * Constructor
-	 * @param	InSaveDirectory		The chunk save to directory
-	 * @param	InInstallManifest	The manifest being installed so we can pull info from it
-	 * @param	InBuildProgress		Pointer to the progress tracker, for setting progress and getting pause state
+	 * @param	SaveDirectory		The chunk save to directory
+	 * @param	CloudDirectory		The base path or url for the downloads
+	 * @param	InstallManifest		The manifest being installed so we can pull info from it
+	 * @param	BuildProgress		Pointer to the progress tracker, for setting progress and getting pause state
 	 */
-	FBuildPatchDownloader( const FString& InSaveDirectory, const FBuildPatchAppManifestRef& InInstallManifest, FBuildPatchProgress* InBuildProgress );
+	FBuildPatchDownloader(const FString& SaveDirectory, const FString& CloudDirectory, const FBuildPatchAppManifestRef& InstallManifest, FBuildPatchProgress* BuildProgress);
 
 	/**
 	 * Default Destructor, deletes allocated Thread
@@ -312,11 +316,12 @@ public:
 
 	/**
 	 * Creates the singleton downloader system. Takes the arguments required for construction.
-	 * @param	InSaveDirectory		The chunk save to directory
-	 * @param	InInstallManifest	The manifest being installed so we can pull info from it
-	 * @param	InBuildProgress		Pointer to the progress tracker, for setting progress and getting pause state
+	 * @param	SaveDirectory		The chunk save to directory
+	 * @param	CloudDirectory		The base path or url for the downloads
+	 * @param	InstallManifest		The manifest being installed so we can pull info from it
+	 * @param	BuildProgress		Pointer to the progress tracker, for setting progress and getting pause state
 	 */
-	static void Create( const FString& InSaveDirectory, const FBuildPatchAppManifestRef& InInstallManifest, FBuildPatchProgress* InBuildProgress );
+	static void Create(const FString& SaveDirectory, const FString& CloudDirectory, const FBuildPatchAppManifestRef& InstallManifest, FBuildPatchProgress* BuildProgress);
 
 	/**
 	 * Get the singleton class object. Must only be called between Create and Shutdown calls otherwise the code will assert.
diff --git a/Engine/Source/Runtime/Online/BuildPatchServices/Private/BuildPatchError.cpp b/Engine/Source/Runtime/Online/BuildPatchServices/Private/BuildPatchError.cpp
--- a/Engine/Source/Runtime/Online/BuildPatchServices/Private/BuildPatchError.cpp
+++ b/Engine/Source/Runtime/Online/BuildPatchServices/Private/BuildPatchError.cpp
@@ -22,7 +22,7 @@ bool FBuildPatchInstallError::HasFatalError()
 	return ErrorState != EBuildPatchInstallError::NoError;
 }
 
-EBuildPatchInstallError::Type FBuildPatchInstallError::GetErrorState()
+EBuildPatchInstallError FBuildPatchInstallError::GetErrorState()
 {
 	FScopeLock ScopeLock( &ThreadLock );
 	return ErrorState;
@@ -63,6 +63,7 @@ const FText FBuildPatchInstallError::GetErrorText()
 	static const FText MoveFileToInstall( LOCTEXT( "BuildPatchInstallError_MoveFileToInstall", "Please check your running processes." ) );
 	static const FText PathLengthExceeded(LOCTEXT("BuildPatchInstallError_PathLengthExceeded", "Please specify a shorter install location."));
 	static const FText OutOfDiskSpace(LOCTEXT("BuildPatchInstallError_OutOfDiskSpace", "Please free up some disk space and try again."));
+	static const FText ErrorCode(LOCTEXT("BuildPatchInstallError_ErrorCode", "Error Code {0}."));
 
 	FScopeLock ScopeLock( &ThreadLock );
 	FText ErrorText = FText::GetEmpty();
@@ -73,6 +74,13 @@ const FText FBuildPatchInstallError::GetErrorText()
 		case EBuildPatchInstallError::MoveFileToInstall: ErrorText = MoveFileToInstall; break;
 		case EBuildPatchInstallError::PathLengthExceeded: ErrorText = PathLengthExceeded; break;
 		case EBuildPatchInstallError::OutOfDiskSpace: ErrorText = OutOfDiskSpace; break;
+		case EBuildPatchInstallError::PrerequisiteError:
+			if (ErrorString.IsEmpty() == false)
+			{
+				ErrorText = FText::Format(ErrorCode, FText::FromString(ErrorString));
+			}
+			break;
+
 	}
 	return FText::Format(LOCTEXT("BuildPatchInstallLongError", "{0} {1}"), GetShortErrorText(), ErrorText);
 }
@@ -87,7 +95,7 @@ const FText& FBuildPatchInstallError::GetShortErrorText()
 	static const FText ApplicationClosing(LOCTEXT("BuildPatchInstallShortError_ApplicationClosing", "The application is closing."));
 	static const FText ApplicationError(LOCTEXT("BuildPatchInstallShortError_ApplicationError", "Patching service could not start."));
 	static const FText UserCanceled(LOCTEXT("BuildPatchInstallShortError_UserCanceled", "User cancelled."));
-	static const FText PrerequisiteError(LOCTEXT("BuildPatchInstallShortError_PrerequisiteError", "Prerequisites install failed."));
+	static const FText PrerequisiteError(LOCTEXT("BuildPatchInstallShortError_PrerequisiteError", "The necessary prerequisites have failed to install."));
 	static const FText InitializationError(LOCTEXT("BuildPatchInstallShortError_InitializationError", "The installer failed to initialize."));
 	static const FText PathLengthExceeded(LOCTEXT("BuildPatchInstallShortError_PathLengthExceeded", "Maximum path length exceeded."));
 	static const FText OutOfDiskSpace(LOCTEXT("BuildPatchInstallShortError_OutOfDiskSpace", "Not enough disk space available."));
@@ -112,7 +120,7 @@ const FText& FBuildPatchInstallError::GetShortErrorText()
 	}
 }
 
-void FBuildPatchInstallError::SetFatalError( const EBuildPatchInstallError::Type& ErrorType, const FString& ErrorLog )
+void FBuildPatchInstallError::SetFatalError( const EBuildPatchInstallError& ErrorType, const FString& ErrorLog )
 {
 	FScopeLock ScopeLock( &ThreadLock );
 	// Only accept the first error
@@ -134,7 +142,7 @@ void FBuildPatchInstallError::SetFatalError( const EBuildPatchInstallError::Type
 	ErrorString = ErrorLog;
 }
 
-const FString& FBuildPatchInstallError::ToString( const EBuildPatchInstallError::Type& ErrorType )
+const FString& FBuildPatchInstallError::ToString( const EBuildPatchInstallError& ErrorType )
 {
 	// Const enum strings, special case no error.
 	static const FString NoError( "SUCCESS" );
@@ -179,7 +187,7 @@ const FString& FBuildPatchInstallError::ToString()
  * Static FBuildPatchInstallError variables
  */
 FCriticalSection FBuildPatchInstallError::ThreadLock;
-EBuildPatchInstallError::Type FBuildPatchInstallError::ErrorState = EBuildPatchInstallError::NoError;
+EBuildPatchInstallError FBuildPatchInstallError::ErrorState = EBuildPatchInstallError::NoError;
 FString FBuildPatchInstallError::ErrorString;
 
 #undef LOCTEXT_NAMESPACE
diff --git a/Engine/Source/Runtime/Online/BuildPatchServices/Private/BuildPatchError.h b/Engine/Source/Runtime/Online/BuildPatchServices/Private/BuildPatchError.h
--- a/Engine/Source/Runtime/Online/BuildPatchServices/Private/BuildPatchError.h
+++ b/Engine/Source/Runtime/Online/BuildPatchServices/Private/BuildPatchError.h
@@ -9,51 +9,6 @@
 
 #pragma once
 
-/**
- * Declares the error type enum for use with the fatal error system
- */
-namespace EBuildPatchInstallError
-{
-	enum Type
-	{
-		// There has been no registered error
-		NoError = 0,
-
-		// A download registered a fatal error
-		DownloadError = 1,
-
-		// A file failed to construct properly
-		FileConstructionFail = 2,
-
-		// An error occurred trying to move the file to the install location
-		MoveFileToInstall = 3,
-
-		// The installed build failed to verify
-		BuildVerifyFail = 4,
-
-		// The user or some process has closed the application
-		ApplicationClosing = 5,
-
-		// An application error, such as module fail to load.
-		ApplicationError = 6,
-
-		// User canceled download
-		UserCanceled = 7,
-
-		// A prerequisites installer failed
-		PrerequisiteError = 8,
-
-		// An initialization error
-		InitializationError = 9,
-
-		// An error occurred creating a file due to excessive path length
-		PathLengthExceeded = 10,
-
-		// An error occurred creating a file due to their not being enough space left on the disk
-		OutOfDiskSpace = 11
-	};
-};
-
 /**
  * A struct to hold static fatal error flagging
  */
@@ -65,7 +20,7 @@ private:
 	static FCriticalSection ThreadLock;
 
 	// The current error state
-	static EBuildPatchInstallError::Type ErrorState;
+	static EBuildPatchInstallError ErrorState;
 
 	// The current error string
 	static FString ErrorString;
@@ -99,7 +54,7 @@ public:
 	 * Static function to get fatal error
 	 * @return The fatal error type
 	 */
-	static EBuildPatchInstallError::Type GetErrorState();
+	static EBuildPatchInstallError GetErrorState();
 
 	/**
 	 * Static function to get the error string for logging
@@ -124,7 +79,7 @@ public:
 	 * @param ErrorType		The error type value.
 	 * @param ErrorLog		An optional string with extra information.
 	 */
-	static void SetFatalError( const EBuildPatchInstallError::Type& ErrorType, const FString& ErrorLog = TEXT( "" ) );
+	static void SetFatalError( const EBuildPatchInstallError& ErrorType, const FString& ErrorLog = TEXT( "" ) );
 
 private:
 	/**
@@ -132,7 +87,7 @@ private:
 	 * @param ErrorType - The error type value.
 	 * @return A string value.
 	 */
-	static const FString& ToString( const EBuildPatchInstallError::Type& ErrorType );
+	static const FString& ToString( const EBuildPatchInstallError& ErrorType );
 
 	/**
 	 * Returns the string representation of the current FBuildPatchInstallError value. Used for logging only.
diff --git a/Engine/Source/Runtime/Online/BuildPatchServices/Private/BuildPatchInstaller.cpp b/Engine/Source/Runtime/Online/BuildPatchServices/Private/BuildPatchInstaller.cpp
--- a/Engine/Source/Runtime/Online/BuildPatchServices/Private/BuildPatchInstaller.cpp
+++ b/Engine/Source/Runtime/Online/BuildPatchServices/Private/BuildPatchInstaller.cpp
@@ -19,6 +19,33 @@
 
 #define NUM_FILE_MOVE_RETRIES	5
 
+namespace
+{
+	// A simple scope timer class
+	class FScopeTimer
+	{
+	public:
+		FScopeTimer(float& InValueToSet, FCriticalSection* InSynchObject)
+			: Time(InValueToSet)
+			, SynchObject(InSynchObject)
+		{
+			StartSecs = FPlatformTime::Seconds();
+		}
+		~FScopeTimer()
+		{
+			StartSecs = FPlatformTime::Seconds() - StartSecs;
+			SynchObject->Lock();
+			Time = StartSecs;
+			SynchObject->Unlock();
+		}
+
+	private:
+		double StartSecs;
+		float& Time;
+		FCriticalSection* SynchObject;
+	};
+}
+
 /* FBuildPatchInstaller implementation
 *****************************************************************************/
 FBuildPatchInstaller::FBuildPatchInstaller(FBuildPatchBoolManifestDelegate InOnCompleteDelegate, FBuildPatchAppManifestPtr CurrentManifest, FBuildPatchAppManifestRef InstallManifest, const FString& InInstallDirectory, const FString& InStagingDirectory, FBuildPatchInstallationInfo& InstallationInfoRef, bool ShouldStageOnly)
@@ -30,6 +57,8 @@ FBuildPatchInstaller::FBuildPatchInstaller(FBuildPatchBoolManifestDelegate InOnC
 	, StagingDirectory( InStagingDirectory )
 	, DataStagingDir( InStagingDirectory / TEXT( "PatchData" ) )
 	, InstallStagingDir( InStagingDirectory / TEXT( "Install" ) )
+	, CloudDirectory( FBuildPatchServicesModule::GetCloudDirectory() )
+	, BackupDirectory( FBuildPatchServicesModule::GetBackupDirectory() )
 	, PreviousMoveMarker( InstallDirectory / TEXT( "$movedMarker" ) )
 	, ThreadLock()
 	, bIsFileData( InstallManifest->IsFileDataManifest() )
@@ -39,7 +68,7 @@ FBuildPatchInstaller::FBuildPatchInstaller(FBuildPatchBoolManifestDelegate InOnC
 	, bSuccess( true )
 	, bIsRunning( false )
 	, bIsInited( false )
-	, DownloadSpeedValue( 0.0 )
+	, DownloadSpeedValue( -1.0 )
 	, DownloadBytesLeft( 0 )
 	, BuildStats()
 	, BuildProgress()
@@ -142,7 +171,7 @@ uint32 FBuildPatchInstaller::Run()
 	TArray<FString> CorruptFiles;
 
 	// Init prereqs progress value
-	const bool bInstallPrereqs = !CurrentBuildManifest.IsValid() && !NewBuildManifest->GetPrereqPath().IsEmpty();
+	const bool bInstallPrereqs = !NewBuildManifest->GetPrereqPath().IsEmpty();
 
 	// Get the start time
 	double StartTime = FPlatformTime::Seconds();
@@ -154,6 +183,9 @@ uint32 FBuildPatchInstaller::Run()
 	int32 InstallRetries = 5;
 	while (!bProcessSuccess && bCanRetry)
 	{
+		// No longer queued
+		BuildProgress.SetStateProgress(EBuildPatchProgress::Queued, 1.0f);
+
 		// Run the install
 		bool bInstallSuccess = RunInstallation(CorruptFiles);
 		BuildProgress.SetStateProgress(EBuildPatchProgress::PrerequisitesInstall, bInstallPrereqs ? 0.0f : 1.0f);
@@ -227,6 +259,7 @@ uint32 FBuildPatchInstaller::Run()
 		BuildStats.FailureReason = FBuildPatchInstallError::GetErrorString();
 		BuildStats.FailureReasonText = FBuildPatchInstallError::GetErrorText();
 		BuildStats.CleanUpTime = CleanUpTime;
+		BuildStats.FailureType = FBuildPatchInstallError::GetErrorState();
 
 		// Log stats
 		GLog->Logf(TEXT("BuildPatchServices: Build Stat: AppName: %s"), *BuildStats.AppName);
@@ -250,6 +283,7 @@ uint32 FBuildPatchInstaller::Run()
 		GLog->Logf(TEXT("BuildPatchServices: Build Stat: TheoreticalDownloadTime: %s"), *FPlatformTime::PrettyTime(BuildStats.TheoreticalDownloadTime));
 		GLog->Logf(TEXT("BuildPatchServices: Build Stat: VerifyTime: %s"), *FPlatformTime::PrettyTime(BuildStats.VerifyTime));
 		GLog->Logf(TEXT("BuildPatchServices: Build Stat: CleanUpTime: %s"), *FPlatformTime::PrettyTime(BuildStats.CleanUpTime));
+		GLog->Logf(TEXT("BuildPatchServices: Build Stat: PrereqTime: %s"), *FPlatformTime::PrettyTime(BuildStats.PrereqTime));
 		GLog->Logf(TEXT("BuildPatchServices: Build Stat: ProcessExecuteTime: %s"), *FPlatformTime::PrettyTime(BuildStats.ProcessExecuteTime));
 		GLog->Logf(TEXT("BuildPatchServices: Build Stat: ProcessPausedTime: %.1f sec"), BuildStats.ProcessPausedTime);
 		GLog->Logf(TEXT("BuildPatchServices: Build Stat: ProcessSuccess: %s"), BuildStats.ProcessSuccess ? TEXT("TRUE") : TEXT("FALSE"));
@@ -305,6 +339,7 @@ bool FBuildPatchInstaller::RunInstallation(TArray<FString>& CorruptFiles)
 	FBuildPatchInstallError::Reset();
 	FBuildPatchAnalytics::ResetCounters();
 	BuildProgress.Reset();
+	BuildProgress.SetStateProgress(EBuildPatchProgress::Queued, 1.0f);
 	BuildProgress.SetStateProgress(EBuildPatchProgress::Initializing, 0.01f);
 	BuildProgress.SetStateProgress(EBuildPatchProgress::CleanUp, 0.0f);
 
@@ -320,7 +355,7 @@ bool FBuildPatchInstaller::RunInstallation(TArray<FString>& CorruptFiles)
 		BuildStats.AppName = NewBuildManifest->GetAppName();
 		BuildStats.AppPatchVersion = NewBuildManifest->GetVersionString();
 		BuildStats.AppInstalledVersion = CurrentBuildManifest.IsValid() ? CurrentBuildManifest->GetVersionString() : TEXT("NONE");
-		BuildStats.CloudDirectory = FBuildPatchServicesModule::GetCloudDirectory();
+		BuildStats.CloudDirectory = CloudDirectory;
 		BuildStats.NumFilesInBuild = NumFilesInBuild;
 	}
 
@@ -390,7 +425,7 @@ bool FBuildPatchInstaller::RunInstallation(TArray<FString>& CorruptFiles)
 	}
 
 	// Create the downloader
-	FBuildPatchDownloader::Create(DataStagingDir, NewBuildManifest, &BuildProgress);
+	FBuildPatchDownloader::Create(DataStagingDir, CloudDirectory, NewBuildManifest, &BuildProgress);
 
 	// Create chunk cache
 	if (bIsChunkData)
@@ -777,7 +812,6 @@ bool FBuildPatchInstaller::RunVerification(TArray< FString >& CorruptFiles)
 
 bool FBuildPatchInstaller::BackupFileIfNecessary(const FString& Filename, bool bDiscoveredByVerification /*= false */)
 {
-	const FString BackupDirectory = FBuildPatchServicesModule::GetBackupDirectory();
 	const FString InstalledFilename = InstallDirectory / Filename;
 	const FString BackupFilename = BackupDirectory / Filename;
 	const bool bBackupOriginals = !BackupDirectory.IsEmpty();
@@ -838,28 +872,34 @@ bool FBuildPatchInstaller::BackupFileIfNecessary(const FString& Filename, bool b
 
 bool FBuildPatchInstaller::RunPrereqInstaller()
 {
-	FString PrereqPath = InstallDirectory / NewBuildManifest->GetPrereqPath();
-	PrereqPath = FPaths::ConvertRelativePathToFull(PrereqPath);
+	FScopeTimer ScopeTimer(BuildStats.PrereqTime, &ThreadLock);
+
+	FString PrereqPath = FPaths::ConvertRelativePathToFull(InstallDirectory / NewBuildManifest->GetPrereqPath());
+	if (bShouldStageOnly)
+	{
+		FString StagedPrereqPath = FPaths::ConvertRelativePathToFull(InstallStagingDir / NewBuildManifest->GetPrereqPath());
+		if (IFileManager::Get().FileSize(*StagedPrereqPath) > 0)
+		{
+			PrereqPath = MoveTemp(StagedPrereqPath);
+		}
+	}
 	FString PrereqCommandline = NewBuildManifest->GetPrereqArgs();
+	int32 ReturnCode = INDEX_NONE;
 
 	GLog->Logf(TEXT("BuildPatchServices: Running prerequisites installer %s %s"), *PrereqPath, *PrereqCommandline);
 
 	BuildProgress.SetStateProgress(EBuildPatchProgress::PrerequisitesInstall, 0.0f);
 
-	// @TODO: Tell our installer to run with no UI since we will have BuildPatchProgress
-	// @TODO: Pass in params to the installer that tell it to fire up the portal/launcher after install if it has to perform a restart.
 	FProcHandle ProcessHandle = FPlatformProcess::CreateProc(*PrereqPath, *PrereqCommandline, true, false, false, NULL, 0, *FPaths::GetPath(PrereqPath), NULL);
 	bool bPrereqInstallSuccessful = true;
 
 	if (!ProcessHandle.IsValid())
 	{
 		GLog->Logf(TEXT("BuildPatchServices: ERROR: Failed to start the prerequisites install process."));
-		FBuildPatchAnalytics::RecordPrereqInstallnError(PrereqPath, PrereqCommandline, -1, TEXT("Failed to start installer"));
-		// @TODO: Do we need to do anything special to fail?
+		FBuildPatchAnalytics::RecordPrereqInstallationError(NewBuildManifest->GetAppName(), NewBuildManifest->GetVersionString(), PrereqPath, PrereqCommandline, ReturnCode, TEXT("Failed to start installer"));
 		bPrereqInstallSuccessful = false;
 	}
 
-	int32 ReturnCode;
 	if (bPrereqInstallSuccessful)
 	{
 		FPlatformProcess::WaitForProc(ProcessHandle);
@@ -869,17 +909,9 @@ bool FBuildPatchInstaller::RunPrereqInstaller()
 
 		if (ReturnCode != 0)
 		{
-			if (ReturnCode == 3010)
-			{
-				GLog->Logf(TEXT("BuildPatchServices: Prerequisites executable returned restart required code %d"), ReturnCode);
-				// @TODO: Inform app that system restart is required
-			}
-			else
-			{
-				GLog->Logf(TEXT("BuildPatchServices: ERROR: Prerequisites executable failed with code %d"), ReturnCode);
-				FBuildPatchAnalytics::RecordPrereqInstallnError(PrereqPath, PrereqCommandline, ReturnCode, TEXT("Failed to install"));
-				bPrereqInstallSuccessful = false;
-			}
+			GLog->Logf(TEXT("BuildPatchServices: ERROR: Prerequisites executable failed with code %d"), ReturnCode);
+			FBuildPatchAnalytics::RecordPrereqInstallationError(NewBuildManifest->GetAppName(), NewBuildManifest->GetVersionString(), PrereqPath, PrereqCommandline, ReturnCode, TEXT("Failed to install"));
+			bPrereqInstallSuccessful = false;
 		}
 	}
 
@@ -889,7 +921,8 @@ bool FBuildPatchInstaller::RunPrereqInstaller()
 	}
 	else
 	{
-		FBuildPatchInstallError::SetFatalError(EBuildPatchInstallError::PrerequisiteError);
+		FString ExtraPrereqCode = ReturnCode != INDEX_NONE ? FString::Printf(TEXT("%u"), ReturnCode) : TEXT("");
+		FBuildPatchInstallError::SetFatalError(EBuildPatchInstallError::PrerequisiteError, *ExtraPrereqCode);
 	}
 
 	return bPrereqInstallSuccessful;
@@ -1047,7 +1080,7 @@ bool FBuildPatchInstaller::IsComplete()
 bool FBuildPatchInstaller::IsCanceled()
 {
 	FScopeLock Lock( &ThreadLock );
-	return FBuildPatchInstallError::GetErrorState() == EBuildPatchInstallError::UserCanceled;
+	return BuildStats.FailureType == EBuildPatchInstallError::UserCanceled;
 }
 
 bool FBuildPatchInstaller::IsPaused()
@@ -1058,7 +1091,7 @@ bool FBuildPatchInstaller::IsPaused()
 bool FBuildPatchInstaller::IsResumable()
 {
 	FScopeLock Lock( &ThreadLock );
-	if( FBuildPatchInstallError::GetErrorState() == EBuildPatchInstallError::PathLengthExceeded )
+	if (BuildStats.FailureType == EBuildPatchInstallError::PathLengthExceeded)
 	{
 		return false;
 	}
@@ -1068,13 +1101,19 @@ bool FBuildPatchInstaller::IsResumable()
 bool FBuildPatchInstaller::HasError()
 {
 	FScopeLock Lock( &ThreadLock );
-	if( FBuildPatchInstallError::GetErrorState() == EBuildPatchInstallError::UserCanceled )
+	if (BuildStats.FailureType == EBuildPatchInstallError::UserCanceled)
 	{
 		return false;
 	}
 	return !BuildStats.ProcessSuccess;
 }
 
+EBuildPatchInstallError FBuildPatchInstaller::GetErrorType()
+{
+	FScopeLock Lock(&ThreadLock);
+	return BuildStats.FailureType;
+}
+
 //@todo this is deprecated and shouldn't be used anymore [6/4/2014 justin.sargent]
 FText FBuildPatchInstaller::GetPercentageText()
 {
@@ -1161,7 +1200,7 @@ void FBuildPatchInstaller::ExecuteCompleteDelegate()
 	check( IsInGameThread() );
 	check( IsComplete() );
 	// Call the complete delegate
-	OnCompleteDelegate.Execute( bSuccess, NewBuildManifest );
+	OnCompleteDelegate.ExecuteIfBound(bSuccess, NewBuildManifest);
 }
 
 void FBuildPatchInstaller::WaitForThread() const
diff --git a/Engine/Source/Runtime/Online/BuildPatchServices/Private/BuildPatchInstaller.h b/Engine/Source/Runtime/Online/BuildPatchServices/Private/BuildPatchInstaller.h
--- a/Engine/Source/Runtime/Online/BuildPatchServices/Private/BuildPatchInstaller.h
+++ b/Engine/Source/Runtime/Online/BuildPatchServices/Private/BuildPatchInstaller.h
@@ -42,6 +42,12 @@ private:
 	// The directory created in staging to construct install files to
 	FString InstallStagingDir;
 
+	// The base path for the data cloud
+	FString CloudDirectory;
+
+	// The backup directory
+	FString BackupDirectory;
+
 	// The filename used to mark a previous install that did not complete but moved staged files into the install directory
 	FString PreviousMoveMarker;
 
@@ -134,6 +140,7 @@ public:
 	virtual bool IsPaused() override;
 	virtual bool IsResumable() override;
 	virtual bool HasError() override;
+	virtual EBuildPatchInstallError GetErrorType() override;
 	//@todo this is deprecated and shouldn't be used anymore [6/4/2014 justin.sargent]
 	virtual FText GetPercentageText() override;
 	//@todo this is deprecated and shouldn't be used anymore [6/4/2014 justin.sargent]
diff --git a/Engine/Source/Runtime/Online/BuildPatchServices/Private/BuildPatchProgress.cpp b/Engine/Source/Runtime/Online/BuildPatchServices/Private/BuildPatchProgress.cpp
--- a/Engine/Source/Runtime/Online/BuildPatchServices/Private/BuildPatchProgress.cpp
+++ b/Engine/Source/Runtime/Online/BuildPatchServices/Private/BuildPatchProgress.cpp
@@ -14,6 +14,7 @@
 const FString& EBuildPatchProgress::ToString(const EBuildPatchProgress::Type& ProgressValue)
 {
 	// Static const fixed FString values so that they are not constantly constructed
+	static const FString Queued(TEXT("Queued"));
 	static const FString Initializing(TEXT("Initialising"));
 	static const FString Resuming(TEXT("Resuming"));
 	static const FString Downloading(TEXT("Downloading"));
@@ -31,6 +32,8 @@ const FString& EBuildPatchProgress::ToString(const EBuildPatchProgress::Type& Pr
 
 	switch (ProgressValue)
 	{
+		case EBuildPatchProgress::Queued:
+			return Queued;
 		case EBuildPatchProgress::Initializing:
 			return Initializing;
 		case EBuildPatchProgress::Resuming:
@@ -63,13 +66,14 @@ const FString& EBuildPatchProgress::ToString(const EBuildPatchProgress::Type& Pr
 const FText& EBuildPatchProgress::ToText(const EBuildPatchProgress::Type& ProgressType)
 {
 	// Static const fixed FText values so that they are not constantly constructed
+	static const FText Queued = LOCTEXT("EBuildPatchProgress_Queued", "Queued");
 	static const FText Initializing = LOCTEXT("EBuildPatchProgress_Initialising", "Initializing");
 	static const FText Resuming = LOCTEXT("EBuildPatchProgress_Resuming", "Resuming");
 	static const FText Downloading = LOCTEXT("EBuildPatchProgress_Downloading", "Downloading");
 	static const FText Installing = LOCTEXT("EBuildPatchProgress_Installing", "Installing");
 	static const FText BuildVerification = LOCTEXT("EBuildPatchProgress_BuildVerification", "Verifying");
 	static const FText CleanUp = LOCTEXT("EBuildPatchProgress_CleanUp", "Cleaning up");
-	static const FText PrerequisitesInstall = LOCTEXT("EBuildPatchProgress_PrerequisitesInstall", "Installing Prerequisites");
+	static const FText PrerequisitesInstall = LOCTEXT("EBuildPatchProgress_PrerequisitesInstall", "Prerequisites");
 	static const FText Completed = LOCTEXT("EBuildPatchProgress_Complete", "Complete");
 	static const FText Error = LOCTEXT("EBuildPatchProgress_Error", "Error");
 	static const FText Paused = LOCTEXT("EBuildPatchProgress_Paused", "Paused");
@@ -77,6 +81,8 @@ const FText& EBuildPatchProgress::ToText(const EBuildPatchProgress::Type& Progre
 
 	switch (ProgressType)
 	{
+		case EBuildPatchProgress::Queued:
+			return Queued;
 		case EBuildPatchProgress::Initializing:
 			return Initializing;
 		case EBuildPatchProgress::Resuming:
@@ -114,7 +120,7 @@ FBuildPatchProgress::FBuildPatchProgress()
 void FBuildPatchProgress::Reset()
 {
 	TotalWeight = 0.0f;
-	CurrentState = EBuildPatchProgress::Initializing;
+	CurrentState = EBuildPatchProgress::Queued;
 	CurrentProgress = 0.0f;
 	ErrorText = FText::GetEmpty();
 	ShortErrorText = FText::GetEmpty();
@@ -288,6 +294,7 @@ void FBuildPatchProgress::UpdateCachedValues()
 *****************************************************************************/
 const bool FBuildPatchProgress::bHasProgressValue[EBuildPatchProgress::NUM_PROGRESS_STATES] =
 {
+	false, // Queued
 	false, // Initializing
 	true,  // Resuming
 	true,  // Downloading
@@ -304,6 +311,7 @@ const bool FBuildPatchProgress::bHasProgressValue[EBuildPatchProgress::NUM_PROGR
 
 const bool FBuildPatchProgress::bCountsTowardsProgress[EBuildPatchProgress::NUM_PROGRESS_STATES] =
 {
+	false, // Queued
 	false, // Initializing
 	false, // Resuming
 	true,  // Downloading
diff --git a/Engine/Source/Runtime/Online/BuildPatchServices/Private/BuildPatchProgress.h b/Engine/Source/Runtime/Online/BuildPatchServices/Private/BuildPatchProgress.h
--- a/Engine/Source/Runtime/Online/BuildPatchServices/Private/BuildPatchProgress.h
+++ b/Engine/Source/Runtime/Online/BuildPatchServices/Private/BuildPatchProgress.h
@@ -14,8 +14,11 @@ namespace EBuildPatchProgress
 {
 	enum Type
 	{
+		// The patch process is waiting for other installs
+		Queued = 0,
+
 		// The patch process is initializing
-		Initializing = 0,
+		Initializing,
 
 		// The patch process is enumerating existing staged data
 		Resuming,
diff --git a/Engine/Source/Runtime/Online/BuildPatchServices/Public/Interfaces/IBuildInstaller.h b/Engine/Source/Runtime/Online/BuildPatchServices/Public/Interfaces/IBuildInstaller.h
--- a/Engine/Source/Runtime/Online/BuildPatchServices/Public/Interfaces/IBuildInstaller.h
+++ b/Engine/Source/Runtime/Online/BuildPatchServices/Public/Interfaces/IBuildInstaller.h
@@ -9,6 +9,48 @@
 typedef TSharedPtr< class IBuildInstaller, ESPMode::ThreadSafe > IBuildInstallerPtr;
 typedef TSharedRef< class IBuildInstaller, ESPMode::ThreadSafe > IBuildInstallerRef;
 
+/**
+ * Declares the error type enum for use with the error system
+ */
+enum class EBuildPatchInstallError
+{
+	// There has been no registered error
+	NoError = 0,
+
+	// A download registered a fatal error
+	DownloadError = 1,
+
+	// A file failed to construct properly
+	FileConstructionFail = 2,
+
+	// An error occurred trying to move the file to the install location
+	MoveFileToInstall = 3,
+
+	// The installed build failed to verify
+	BuildVerifyFail = 4,
+
+	// The user or some process has closed the application
+	ApplicationClosing = 5,
+
+	// An application error, such as module fail to load.
+	ApplicationError = 6,
+
+	// User canceled download
+	UserCanceled = 7,
+
+	// A prerequisites installer failed
+	PrerequisiteError = 8,
+
+	// An initialization error
+	InitializationError = 9,
+
+	// An error occurred creating a file due to excessive path length
+	PathLengthExceeded = 10,
+
+	// An error occurred creating a file due to their not being enough space left on the disk
+	OutOfDiskSpace = 11
+};
+
 /**
  * A struct to hold stats for the build process.
  */
@@ -33,9 +75,11 @@ struct FBuildInstallStats
 		, TheoreticalDownloadTime(0.0f)
 		, VerifyTime(0.0f)
 		, CleanUpTime(0.0f)
+		, PrereqTime(0.0f)
 		, ProcessExecuteTime(0.0f)
 		, ProcessPausedTime(0.0f)
 		, ProcessSuccess(false)
+		, FailureType(EBuildPatchInstallError::InitializationError)
 	{}
 
 	// The name of the app being installed
@@ -80,6 +124,8 @@ struct FBuildInstallStats
 	float VerifyTime;
 	// The total time that the install cleanup took to complete
 	float CleanUpTime;
+	// The total time that the prereq install took to complete
+	float PrereqTime;
 	// The total time that the install process took to complete
 	float ProcessExecuteTime;
 	// The amount of time that was spent paused
@@ -90,6 +136,8 @@ struct FBuildInstallStats
 	FString FailureReason;
 	// The localised, more generic failure reason
 	FText FailureReasonText;
+	// The failure type for the install
+	EBuildPatchInstallError FailureType;
 };
 
 /**
@@ -133,6 +181,12 @@ public:
 	 */
 	virtual bool HasError() = 0;
 
+	/**
+	 * Get the type of error for a failure that has occurred.
+	 * @return	the enum representing the type of error
+	 */
+	virtual EBuildPatchInstallError GetErrorType() = 0;
+
 	/**
 	 * This is deprecated and shouldn't be used anymore [6/4/2014 justin.sargent]
 	 *
diff --git a/Engine/Source/Runtime/Online/OnlineSubsystem/Public/Interfaces/OnlinePresenceInterface.h b/Engine/Source/Runtime/Online/OnlineSubsystem/Public/Interfaces/OnlinePresenceInterface.h
--- a/Engine/Source/Runtime/Online/OnlineSubsystem/Public/Interfaces/OnlinePresenceInterface.h
+++ b/Engine/Source/Runtime/Online/OnlineSubsystem/Public/Interfaces/OnlinePresenceInterface.h
@@ -19,6 +19,9 @@ const FString CustomPresenceDataKey = "CustomData";
 /** Id of the client that sent the presence update */
 const FString DefaultClientIdKey = "ClientId";
 
+/** Override Id of the client to set the presence state to */
+const FString OverrideClientIdKey = "OverrideClientId";
+
 /** Id of the session for the presence update. @todo samz - SessionId on presence data should be FUniqueNetId not uint64 */
 const FString DefaultSessionIdKey = "SessionId";
 
