Mergine changes from //depot/UE4 to //UE4/Main

diff --git a/Engine/Source/Runtime/Online/Hotfix/Private/OnlineHotfixManager.cpp b/Engine/Source/Runtime/Online/Hotfix/Private/OnlineHotfixManager.cpp
--- a/Engine/Source/Runtime/Online/Hotfix/Private/OnlineHotfixManager.cpp
+++ b/Engine/Source/Runtime/Online/Hotfix/Private/OnlineHotfixManager.cpp
@@ -3,9 +3,31 @@
 #include "HotfixPrivatePCH.h"
 #include "OnlineHotfixManager.h"
 #include "Internationalization.h"
+#include "OnlineSubsystemUtils.h"
 
 DEFINE_LOG_CATEGORY(LogHotfixManager);
 
+FName NAME_HotfixManager(TEXT("HotfixManager"));
+
+struct FHotfixManagerExec :
+	public FSelfRegisteringExec
+{
+	virtual bool Exec(UWorld* InWorld, const TCHAR* Cmd, FOutputDevice& Ar) override
+	{
+		if (FParse::Command(&Cmd, TEXT("HOTFIX")))
+		{
+			UOnlineHotfixManager* HotfixManager = UOnlineHotfixManager::Get(InWorld);
+			if (HotfixManager != nullptr)
+			{
+				HotfixManager->StartHotfixProcess();
+			}
+			return true;
+		}
+		return false;
+	}
+};
+static FHotfixManagerExec HotfixManagerExec;
+
 UOnlineHotfixManager::UOnlineHotfixManager() :
 	Super()
 {
@@ -16,6 +38,31 @@ UOnlineHotfixManager::UOnlineHotfixManager() :
 	PlatformPrefix += TEXT("_");
 }
 
+UOnlineHotfixManager* UOnlineHotfixManager::Get(UWorld* World)
+{
+	UOnlineHotfixManager* DefaultObject = UOnlineHotfixManager::StaticClass()->GetDefaultObject<UOnlineHotfixManager>();
+	IOnlineSubsystem* OnlineSub = Online::GetSubsystem(World, DefaultObject->OSSName.Len() > 0 ? FName(*DefaultObject->OSSName) : NAME_None);
+	if (OnlineSub != nullptr)
+	{
+		UOnlineHotfixManager* HotfixManager = Cast<UOnlineHotfixManager>(OnlineSub->GetNamedInterface(NAME_HotfixManager));
+		if (HotfixManager == nullptr)
+		{
+			FString HotfixManagerClassName = DefaultObject->HotfixManagerClassName;
+			UClass* HotfixManagerClass = LoadClass<UOnlineHotfixManager>(nullptr, *HotfixManagerClassName, nullptr, LOAD_None, nullptr);
+			if (HotfixManagerClass == nullptr)
+			{
+				// Just use the default class if it couldn't load what was specified
+				HotfixManagerClass = UOnlineHotfixManager::StaticClass();
+			}
+			// Create it and store it
+			HotfixManager = NewObject<UOnlineHotfixManager>(GetTransientPackage(), HotfixManagerClass);
+			OnlineSub->SetNamedInterface(NAME_HotfixManager, HotfixManager);
+			return HotfixManager;
+		}
+	}
+	return nullptr;
+}
+
 void UOnlineHotfixManager::Init()
 {
 	// Build the name of the loc file that we'll care about
diff --git a/Engine/Source/Runtime/Online/OnlineFramework/Party/Private/Party.cpp b/Engine/Source/Runtime/Online/OnlineFramework/Party/Private/Party.cpp
--- a/Engine/Source/Runtime/Online/OnlineFramework/Party/Private/Party.cpp
+++ b/Engine/Source/Runtime/Online/OnlineFramework/Party/Private/Party.cpp
@@ -67,11 +67,10 @@ void UParty::OnLogoutComplete(int32 LocalUserNum, bool bWasSuccessful)
 		UE_LOG(LogParty, Log, TEXT("Party cleanup on logout"));
 		for (const auto& PartyKeyValue : JoinedParties)
 		{
-			const FString& PartyId = PartyKeyValue.Key;
 			UPartyGameState* Party = PartyKeyValue.Value;
 			if (Party)
 			{
-				UE_LOG(LogParty, Log, TEXT("[%s] Removed"), *PartyId);
+				UE_LOG(LogParty, Log, TEXT("[%s] Removed"), *(Party->GetPartyId()->ToDebugString()));
 				Party->HandleRemovedFromParty(EMemberExitedReason::Left);
 			}
 		}
@@ -89,11 +88,10 @@ void UParty::OnLoginStatusChanged(int32 LocalUserNum, ELoginStatus::Type OldStat
 			UE_LOG(LogParty, Log, TEXT("Party cleanup on logout"));
 			for (const auto& PartyKeyValue : JoinedParties)
 			{
-				const FString& PartyId = PartyKeyValue.Key;
 				UPartyGameState* Party = PartyKeyValue.Value;
 				if (Party)
 				{
-					UE_LOG(LogParty, Log, TEXT("[%s] Removed"), *PartyId);
+					UE_LOG(LogParty, Log, TEXT("[%s] Removed"), *(Party->GetPartyId()->ToDebugString()));
 					Party->HandleRemovedFromParty(EMemberExitedReason::Left);
 				}
 			}
@@ -101,6 +99,8 @@ void UParty::OnLoginStatusChanged(int32 LocalUserNum, ELoginStatus::Type OldStat
 			JoinedParties.Empty();
 		}
 	}
+
+	PendingPartyJoin.Reset();
 }
 
 void UParty::AddReferencedObjects(UObject* InThis, FReferenceCollector& Collector)
@@ -210,24 +210,39 @@ void UParty::NotifyPreClientTravel()
 
 UPartyGameState* UParty::GetParty(const FOnlinePartyId& InPartyId) const
 {
-	if (InPartyId.IsValid())
+	UPartyGameState* ResultPtr = nullptr;
+
+	for (auto& JoinedParty : JoinedParties)
 	{
-		UPartyGameState* const * PartyPtr = JoinedParties.Find(InPartyId.ToString());
-		return PartyPtr ? *PartyPtr : nullptr;
+		TSharedPtr<const FOnlinePartyId> PartyId = JoinedParty.Value->GetPartyId();
+		if (*PartyId == InPartyId)
+		{
+			ResultPtr = JoinedParty.Value;
+			break;
+		}
 	}
+
+	return ResultPtr;
+}
+
+UPartyGameState* UParty::GetParty(const FOnlinePartyTypeId InPartyTypeId) const
+{
+	UPartyGameState* ResultPtr = nullptr;
+
+	UPartyGameState* const * PartyPtr = JoinedParties.Find(InPartyTypeId);
+	ResultPtr = PartyPtr ? *PartyPtr : nullptr;
 	
-	return nullptr;
+	return ResultPtr;
 }
 
 UPartyGameState* UParty::GetPersistentParty() const
 {
-	if (PersistentPartyId.IsValid())
-	{
-		UPartyGameState* const * PartyPtr = JoinedParties.Find(PersistentPartyId->ToString());
-		return PartyPtr ? *PartyPtr : nullptr;
-	}
+	UPartyGameState* ResultPtr = nullptr;
 
-	return nullptr;
+	UPartyGameState* const * PartyPtr = JoinedParties.Find(IOnlinePartySystem::GetPrimaryPartyTypeId());
+	ResultPtr = PartyPtr ? *PartyPtr : nullptr;
+
+	return ResultPtr;
 }
 
 void UParty::PartyConfigChangedInternal(const FUniqueNetId& InLocalUserId, const FOnlinePartyId& InPartyId, const TSharedRef<FPartyConfiguration>& InPartyConfig)
@@ -368,6 +383,7 @@ void UParty::PartyMemberExitedInternal(const FUniqueNetId& InLocalUserId, const
 		else
 		{
 			UPartyGameState* PartyState = GetParty(InPartyId);
+			FOnlinePartyTypeId PartyTypeId = PartyState->GetPartyTypeId();
 			if (!PartyState)
 			{
 				UE_LOG(LogParty, Warning, TEXT("[%s]: Missing party state during local player exit."), *InPartyId.ToString());
@@ -382,7 +398,7 @@ void UParty::PartyMemberExitedInternal(const FUniqueNetId& InLocalUserId, const
 				PartyState->HandleRemovedFromParty(InReason);
 			}
 
-			JoinedParties.Remove(InPartyId.ToString());
+			JoinedParties.Remove(PartyTypeId);
 		}
 
 		// If the removal was the persistent party, make sure we are in a good state
@@ -426,23 +442,102 @@ void UParty::PartyPromotionLockoutStateChangedInternal(const FUniqueNetId& Local
 	}
 }
 
-void UParty::CreatePartyInternal(const FUniqueNetId& InUserId, const FPartyConfiguration& InPartyConfig, const FOnCreatePartyComplete& InCompletionDelegate)
+void UParty::CreatePartyInternal(const FUniqueNetId& InUserId, const FOnlinePartyTypeId InPartyTypeId, const FPartyConfiguration& InPartyConfig, const UPartyDelegates::FOnCreateUPartyComplete& InCompletionDelegate)
 {
+	ECreatePartyCompletionResult Result = ECreatePartyCompletionResult::UnknownClientFailure;
+	FString ErrorMsg;
+
 	UWorld* World = GetWorld();
 	IOnlinePartyPtr PartyInt = Online::GetPartyInterface(World);
 	if (PartyInt.IsValid())
 	{
-		PartyInt->CreateParty(InUserId, InPartyConfig, InCompletionDelegate);
+		FOnCreatePartyComplete CompletionDelegate;
+		CompletionDelegate.BindUObject(this, &ThisClass::OnCreatePartyInternalComplete, InPartyTypeId, InCompletionDelegate);
+		PartyInt->CreateParty(InUserId, InPartyTypeId, InPartyConfig, CompletionDelegate);
+		Result = ECreatePartyCompletionResult::Succeeded;
 	}
 	else
 	{
-		FString ErrorStr(TEXT("No party interface during CreateParty()"));
-		InCompletionDelegate.ExecuteIfBound(InUserId, InUserId, ECreatePartyCompletionResult::UnknownClientFailure);
+		Result = ECreatePartyCompletionResult::UnknownClientFailure;
+		ErrorMsg = FString::Printf(TEXT("No party interface during JoinParty()"));
+	}
+
+	if (Result != ECreatePartyCompletionResult::Succeeded)
+	{
+		UE_LOG(LogParty, Warning, TEXT("%s"), *ErrorMsg);
+		InCompletionDelegate.ExecuteIfBound(InUserId, Result);
 	}
 }
 
-void UParty::JoinPartyInternal(const FUniqueNetId& InUserId, const FPartyDetails& InPartyDetails, const FOnJoinPartyComplete& InCompletionDelegate)
+void UParty::OnCreatePartyInternalComplete(const FUniqueNetId& LocalUserId, const TSharedPtr<const FOnlinePartyId>& InPartyId, const ECreatePartyCompletionResult Result, const FOnlinePartyTypeId InPartyTypeId, UPartyDelegates::FOnCreateUPartyComplete InCompletionDelegate)
 {
+	bool bWasSuccessful = Result == ECreatePartyCompletionResult::Succeeded;
+
+	FString PartyIdDebugString = InPartyId.IsValid() ? InPartyId->ToDebugString() : TEXT("Invalid");
+
+	UE_LOG(LogParty, Display, TEXT("OnCreatePartyInternalComplete() %s %d %s"), *PartyIdDebugString, bWasSuccessful, ToString(Result));
+
+	ECreatePartyCompletionResult LocalResult = Result;
+	if (bWasSuccessful)
+	{
+		UWorld* World = GetWorld();
+		IOnlinePartyPtr PartyInt = Online::GetPartyInterface(World);
+		if (PartyInt.IsValid())
+		{
+			TSharedPtr<const FOnlineParty> Party = PartyInt->GetParty(LocalUserId, InPartyTypeId);
+			if (Party.IsValid())
+			{
+				TSubclassOf<UPartyGameState>* PartyGameStateClass = PartyClasses.Find(InPartyTypeId);
+				if (PartyGameStateClass != nullptr)
+				{
+					UPartyGameState* NewParty = NewObject<UPartyGameState>(this, *PartyGameStateClass);
+
+					// Add right away so future delegate broadcasts have this available
+					JoinedParties.Add(InPartyTypeId, NewParty);
+
+					// Initialize and trigger delegates
+					NewParty->InitFromCreate(LocalUserId, Party);
+
+					LocalResult = ECreatePartyCompletionResult::Succeeded;
+				}
+				else
+				{
+					LocalResult = ECreatePartyCompletionResult::UnknownClientFailure;
+				}
+			}
+			else
+			{
+				LocalResult = ECreatePartyCompletionResult::UnknownClientFailure;
+			}
+		}
+		else
+		{
+			LocalResult = ECreatePartyCompletionResult::UnknownClientFailure;
+		}
+
+		if (LocalResult != ECreatePartyCompletionResult::Succeeded)
+		{
+			bWasSuccessful = false;
+		}
+	}
+
+
+	if (!bWasSuccessful)
+	{
+		if (LocalResult != ECreatePartyCompletionResult::Succeeded)
+		{
+			UE_LOG(LogParty, Warning, TEXT("Error when creating party %s error=%s"), *PartyIdDebugString, ToString(LocalResult));
+		}
+	}
+
+	InCompletionDelegate.ExecuteIfBound(LocalUserId, LocalResult);
+}
+
+void UParty::JoinPartyInternal(const FUniqueNetId& InUserId, const FPartyDetails& InPartyDetails, const UPartyDelegates::FOnJoinUPartyComplete& InCompletionDelegate)
+{
+	EJoinPartyCompletionResult Result = EJoinPartyCompletionResult::UnknownClientFailure;
+	FString ErrorMsg;
+
 	IOnlinePartyPtr PartyInt = Online::GetPartyInterface(GetWorld());
 	if (PartyInt.IsValid())
 	{
@@ -452,227 +547,348 @@ void UParty::JoinPartyInternal(const FUniqueNetId& InUserId, const FPartyDetails
 			// High level party data check
 			UPartyGameState* PartyState = GetParty(PartyId);
 			// Interface level party data check should not be out of sync
-			const TSharedPtr<FOnlinePartyInfo>& PartyInfo = PartyInt->GetPartyInfo(InUserId, PartyId);
+			TSharedPtr<const FOnlineParty> Party = PartyInt->GetParty(InUserId, PartyId);
 			if (!PartyState)
 			{
-				if (!PartyInfo.IsValid())
+				if (!Party.IsValid())
 				{
-					PartyInt->JoinParty(InUserId, *(InPartyDetails.PartyJoinInfo), InCompletionDelegate);
+					FOnJoinPartyComplete CompletionDelegate;
+					CompletionDelegate.BindUObject(this, &ThisClass::OnJoinPartyInternalComplete, InPartyDetails.GetPartyTypeId(), InCompletionDelegate);
+					PartyInt->JoinParty(InUserId, *(InPartyDetails.PartyJoinInfo), CompletionDelegate);
+					Result = EJoinPartyCompletionResult::Succeeded;
 				}
 				else
 				{
-					UE_LOG(LogParty, Verbose, TEXT("Already joining party %s, not joining again."), *InPartyDetails.GetPartyId()->ToString());
-					InCompletionDelegate.ExecuteIfBound(InUserId, *InPartyDetails.GetPartyId(), EJoinPartyCompletionResult::AlreadyJoiningParty, 0);
+					Result = EJoinPartyCompletionResult::AlreadyJoiningParty;
+					ErrorMsg = FString::Printf(TEXT("Already joining party %s, not joining again."), *InPartyDetails.GetPartyId()->ToString());
 				}
 			}
 			else
 			{
-				ensure(PartyInfo.IsValid());
-				UE_LOG(LogParty, Verbose, TEXT("Already in party %s, not joining again."), *InPartyDetails.GetPartyId()->ToString());
-				InCompletionDelegate.ExecuteIfBound(InUserId, *InPartyDetails.GetPartyId(), EJoinPartyCompletionResult::AlreadyInParty, 0);
+				Result = EJoinPartyCompletionResult::AlreadyInParty;
+				ErrorMsg = FString::Printf(TEXT("Already in party %s, not joining again."), *InPartyDetails.GetPartyId()->ToString());
 			}
 		}
 		else
 		{
-			UE_LOG(LogParty, Verbose, TEXT("Invalid party details, cannot join. Details: %s"), *InPartyDetails.ToString());
-			InCompletionDelegate.ExecuteIfBound(InUserId, *InPartyDetails.GetPartyId(), EJoinPartyCompletionResult::JoinInfoInvalid, 0);
+			Result = EJoinPartyCompletionResult::JoinInfoInvalid;
+			ErrorMsg = FString::Printf(TEXT("Invalid party details, cannot join. Details: %s"), *InPartyDetails.ToString());
 		}
 	}
 	else
 	{
-		UE_LOG(LogParty, Verbose, TEXT("No party interface during JoinParty()"));
-		InCompletionDelegate.ExecuteIfBound(InUserId, *InPartyDetails.GetPartyId(), EJoinPartyCompletionResult::UnknownClientFailure, 0);
+		Result = EJoinPartyCompletionResult::UnknownClientFailure;
+		ErrorMsg = FString::Printf(TEXT("No party interface during JoinParty()"));
+	}
+
+	if (Result != EJoinPartyCompletionResult::Succeeded)
+	{
+		UE_LOG(LogParty, Warning, TEXT("%s"), *ErrorMsg);
+		InCompletionDelegate.ExecuteIfBound(InUserId, Result, 0);
 	}
 }
 
-void UParty::LeavePartyInternal(const FUniqueNetId& InUserId, const FOnlinePartyId& InPartyId, const FOnLeavePartyComplete& InCompletionDelegate)
+void UParty::OnJoinPartyInternalComplete(const FUniqueNetId& LocalUserId, const FOnlinePartyId& InPartyId, const EJoinPartyCompletionResult Result, int32 DeniedResultCode, const FOnlinePartyTypeId InPartyTypeId, UPartyDelegates::FOnJoinUPartyComplete CompletionDelegate)
 {
+	bool bWasSuccessful = Result == EJoinPartyCompletionResult::Succeeded;
+
+	FString PartyIdDebugString = InPartyId.ToDebugString();
+	UE_LOG(LogParty, Display, TEXT("OnJoinPartyInternalComplete() %s %d %s."), *PartyIdDebugString, bWasSuccessful, ToString(Result));
+
+	EJoinPartyCompletionResult LocalResult = Result;
+	if (bWasSuccessful)
+	{
+		UWorld* World = GetWorld();
+		IOnlinePartyPtr PartyInt = Online::GetPartyInterface(World);
+		if (PartyInt.IsValid())
+		{
+			TSharedPtr<const FOnlineParty> Party = PartyInt->GetParty(LocalUserId, InPartyTypeId);
+			if (Party.IsValid())
+			{
+				TSubclassOf<UPartyGameState>* PartyGameStateClass = PartyClasses.Find(InPartyTypeId);
+				if (PartyGameStateClass != nullptr)
+				{
+					UPartyGameState* NewParty = NewObject<UPartyGameState>(this, *PartyGameStateClass);
+
+					// Add right away so future delegate broadcasts have this available
+					JoinedParties.Add(InPartyTypeId, NewParty);
+
+					// Initialize and trigger delegates
+					NewParty->InitFromJoin(LocalUserId, Party);
+				}
+				else
+				{
+					LocalResult = EJoinPartyCompletionResult::AlreadyInPartyOfSpecifiedType;
+				}
+			}
+			else
+			{
+				LocalResult = EJoinPartyCompletionResult::UnknownClientFailure;
+			}
+		}
+		else
+		{
+			LocalResult = EJoinPartyCompletionResult::UnknownClientFailure;
+		}
+
+		if (LocalResult != EJoinPartyCompletionResult::Succeeded)
+		{
+			bWasSuccessful = false;
+		}
+	}
+
+	if (!bWasSuccessful)
+	{
+		if (LocalResult != EJoinPartyCompletionResult::Succeeded)
+		{
+			UE_LOG(LogParty, Warning, TEXT("Error when joining party %s error=%s"), *PartyIdDebugString, ToString(LocalResult));
+		}
+	}
+
+	int32 OutDeniedResultCode = (Result == EJoinPartyCompletionResult::NotApproved) ? DeniedResultCode : 0;
+	CompletionDelegate.ExecuteIfBound(LocalUserId, LocalResult, OutDeniedResultCode);
+}
+
+void UParty::LeavePartyInternal(const FUniqueNetId& InUserId, const FOnlinePartyTypeId InPartyTypeId, const UPartyDelegates::FOnLeaveUPartyComplete& InCompletionDelegate)
+{
+	ELeavePartyCompletionResult Result = ELeavePartyCompletionResult::UnknownClientFailure;
+	FString ErrorMsg;
+
 	IOnlinePartyPtr PartyInt = Online::GetPartyInterface(GetWorld());
-	if (PartyInt.IsValid() && InPartyId.IsValid())
+	if (PartyInt.IsValid())
 	{
-		UPartyGameState* PartyState = GetParty(InPartyId);
+		TSharedPtr<const FOnlineParty> Party = PartyInt->GetParty(InUserId, InPartyTypeId);
+		TSharedPtr<const FOnlinePartyId> PartyId = Party.IsValid() ? Party->PartyId : TSharedPtr<const FOnlinePartyId>();
+
+		UPartyGameState* PartyState = PartyId.IsValid() ? GetParty(*PartyId) : nullptr;
 		if (PartyState)
 		{
 			PartyState->HandleLeavingParty();
-			PartyInt->LeaveParty(InUserId, InPartyId, InCompletionDelegate);
+			FOnLeavePartyComplete CompletionDelegate;
+			CompletionDelegate.BindUObject(this, &ThisClass::OnLeavePartyInternalComplete, InPartyTypeId, InCompletionDelegate);
+			PartyInt->LeaveParty(InUserId, *PartyId, CompletionDelegate);
+			Result = ELeavePartyCompletionResult::Succeeded;
 		}
 		else
 		{
-			FString ErrorStr(TEXT("Party not found in LeaveParty()"));
-			InCompletionDelegate.ExecuteIfBound(InUserId, InUserId, ELeavePartyCompletionResult::UnknownClientFailure);
+			Result = ELeavePartyCompletionResult::UnknownClientFailure;
+			ErrorMsg = FString::Printf(TEXT("Party not found in LeaveParty()"));
 		}
 	}
 	else
 	{
-		FString ErrorStr(TEXT("No party interface during LeaveParty()"));
-		InCompletionDelegate.ExecuteIfBound(InUserId, InUserId, ELeavePartyCompletionResult::UnknownClientFailure);
+		Result = ELeavePartyCompletionResult::UnknownClientFailure;
+		ErrorMsg = FString::Printf(TEXT("No party interface during JoinParty()"));
+	}
+
+	if (Result != ELeavePartyCompletionResult::Succeeded)
+	{
+		UE_LOG(LogParty, Warning, TEXT("%s"), *ErrorMsg);
+		InCompletionDelegate.ExecuteIfBound(InUserId, Result);
 	}
 }
 
-void UParty::CreateParty(const FUniqueNetId& InUserId, const FPartyConfiguration& InPartyConfig, const FOnCreatePartyComplete& InCompletionDelegate)
+void UParty::OnLeavePartyInternalComplete(const FUniqueNetId& LocalUserId, const FOnlinePartyId& InPartyId, const ELeavePartyCompletionResult Result, const FOnlinePartyTypeId InPartyTypeId, UPartyDelegates::FOnLeaveUPartyComplete CompletionDelegate)
 {
-	FOnCreatePartyComplete CompletionDelegate;
-	CompletionDelegate.BindUObject(this, &ThisClass::OnCreatePartyComplete, InCompletionDelegate);
-	CreatePartyInternal(InUserId, InPartyConfig, CompletionDelegate);
+	FString PartyIdDebugString = InPartyId.ToDebugString();
+	UE_LOG(LogParty, Display, TEXT("OnLeavePartyComplete() %s %d %s."), *PartyIdDebugString, Result == ELeavePartyCompletionResult::Succeeded, ToString(Result));
+
+	UPartyGameState* PartyState = GetParty(InPartyTypeId);
+	if (PartyState)
+	{
+		PartyState->HandleRemovedFromParty(EMemberExitedReason::Left);
+		JoinedParties.Remove(InPartyTypeId);
+	}
+	else
+	{
+		UE_LOG(LogParty, Warning, TEXT("OnLeavePartyComplete: Missing party state %s"), *PartyIdDebugString);
+	}
+
+	CompletionDelegate.ExecuteIfBound(LocalUserId, Result);
 }
 
-void UParty::OnCreatePartyComplete(const FUniqueNetId& LocalUserId, const FOnlinePartyId& InPartyId, const ECreatePartyCompletionResult Result, FOnCreatePartyComplete InCompletionDelegate)
+void UParty::CreateParty(const FUniqueNetId& InUserId, const FOnlinePartyTypeId InPartyTypeId, const FPartyConfiguration& InPartyConfig, const UPartyDelegates::FOnCreateUPartyComplete& InCompletionDelegate)
 {
-	UE_LOG(LogParty, Display, TEXT("OnCreatePartyComplete() %s %s"), *InPartyId.ToString(), ToString(Result));
+	UPartyDelegates::FOnCreateUPartyComplete CompletionDelegate;
+	CompletionDelegate.BindUObject(this, &ThisClass::OnCreatePartyComplete, InPartyTypeId, InCompletionDelegate);
+	CreatePartyInternal(InUserId, InPartyTypeId, InPartyConfig, CompletionDelegate);
+}
 
-	ECreatePartyCompletionResult LocalResult = Result;
-	if (Result == ECreatePartyCompletionResult::Succeeded)
+void UParty::OnCreatePartyComplete(const FUniqueNetId& LocalUserId, const ECreatePartyCompletionResult Result, const FOnlinePartyTypeId InPartyTypeId, UPartyDelegates::FOnCreateUPartyComplete InCompletionDelegate)
+{
+	bool bWasSuccessful = Result == ECreatePartyCompletionResult::Succeeded;
+
+	UE_LOG(LogParty, Display, TEXT("OnCreatePartyComplete() type(0x%08x) %d %s"), InPartyTypeId.GetValue(), bWasSuccessful, ToString(Result));
+
+	InCompletionDelegate.ExecuteIfBound(LocalUserId, Result);
+}
+
+EOnCreatePartyCompleteInternalResult UParty::OnCreatePartyCompleteInternal(const FUniqueNetId& LocalUserId, const FOnlinePartyTypeId InPartyTypeId)
+{
+	EOnCreatePartyCompleteInternalResult ReturnResult = EOnCreatePartyCompleteInternalResult::UnknownFailure;
+
+	UWorld* World = GetWorld();
+	IOnlinePartyPtr PartyInt = Online::GetPartyInterface(World);
+	if (PartyInt.IsValid())
 	{
-		UWorld* World = GetWorld();
-		IOnlinePartyPtr PartyInt = Online::GetPartyInterface(World);
-		if (PartyInt.IsValid())
+		TSharedPtr<const FOnlineParty> Party = PartyInt->GetParty(LocalUserId, InPartyTypeId);
+		if (Party.IsValid())
 		{
-			TSharedPtr<FOnlinePartyInfo> PartyInfo = PartyInt->GetPartyInfo(LocalUserId, InPartyId);
-			if (PartyInfo.IsValid())
+			TSubclassOf<UPartyGameState>* PartyGameStateClass = PartyClasses.Find(InPartyTypeId);
+			if (PartyGameStateClass != nullptr)
 			{
-				FVariantData PartyConfigType;
-				PartyInfo->Config->ClientConfigData.GetAttribute(PARTY_CONFIG_TYPE, PartyConfigType);
-				if (PartyConfigType.GetType() == EOnlineKeyValuePairDataType::Int32)
-				{
-					int32 PartyType = 0;
-					PartyConfigType.GetValue(PartyType);
-
-					TSubclassOf<UPartyGameState>* PartyGameStateClass = PartyClasses.Find(PartyType);
-					if (PartyGameStateClass != nullptr)
-					{
-						UPartyGameState* NewParty = NewObject<UPartyGameState>(this, *PartyGameStateClass);
-
-						// Add right away so future delegate broadcasts have this available
-						JoinedParties.Add(InPartyId.ToString(), NewParty);
-
-						// Initialize and trigger delegates
-						NewParty->InitFromCreate(LocalUserId, PartyInfo);
-					}
-					else
-					{
-						UE_LOG(LogParty, Warning, TEXT("No valid party class when creating party %s"), *InPartyId.ToString());
-						LocalResult = ECreatePartyCompletionResult::UnknownClientFailure;
-					}
-				}
-				else
-				{
-					UE_LOG(LogParty, Warning, TEXT("Invalid party config type when creating party %s"), *InPartyId.ToString());
-					LocalResult = ECreatePartyCompletionResult::UnknownClientFailure;
-				}
+				UPartyGameState* NewParty = NewObject<UPartyGameState>(this, *PartyGameStateClass);
+
+				// Add right away so future delegate broadcasts have this available
+				JoinedParties.Add(InPartyTypeId, NewParty);
+
+				// Initialize and trigger delegates
+				NewParty->InitFromCreate(LocalUserId, Party);
+
+				ReturnResult = EOnCreatePartyCompleteInternalResult::Success;
 			}
 			else
 			{
-				UE_LOG(LogParty, Warning, TEXT("Invalid PartyInfo when creating party %s"), *InPartyId.ToString());
-				LocalResult = ECreatePartyCompletionResult::UnknownClientFailure;
+				ReturnResult = EOnCreatePartyCompleteInternalResult::InvalidPartyType;
 			}
 		}
+		else
+		{
+			ReturnResult = EOnCreatePartyCompleteInternalResult::InvalidParty;
+		}
+	}
+	else
+	{
+		ReturnResult = EOnCreatePartyCompleteInternalResult::InvalidSystemInterface;
 	}
 
-	InCompletionDelegate.ExecuteIfBound(LocalUserId, InPartyId, LocalResult);
+	return ReturnResult;
 }
 
-void UParty::JoinParty(const FUniqueNetId& InUserId, const FPartyDetails& InPartyDetails, const FOnJoinPartyComplete& InCompletionDelegate)
+void UParty::JoinParty(const FUniqueNetId& InUserId, const FPartyDetails& InPartyDetails, const UPartyDelegates::FOnJoinUPartyComplete& InCompletionDelegate)
 {
-	FOnJoinPartyComplete CompletionDelegate;
-	CompletionDelegate.BindUObject(this, &ThisClass::OnJoinPartyComplete, InCompletionDelegate);
+	UPartyDelegates::FOnJoinUPartyComplete CompletionDelegate;
+	CompletionDelegate.BindUObject(this, &ThisClass::OnJoinPartyComplete, InPartyDetails.GetPartyTypeId(), InCompletionDelegate);
 	JoinPartyInternal(InUserId, InPartyDetails, CompletionDelegate);
 }
 
-void UParty::OnJoinPartyComplete(const FUniqueNetId& LocalUserId, const FOnlinePartyId& InPartyId, EJoinPartyCompletionResult Result, int32 DeniedResultCode, FOnJoinPartyComplete CompletionDelegate)
+void UParty::OnJoinPartyComplete(const FUniqueNetId& LocalUserId, EJoinPartyCompletionResult Result, int32 DeniedResultCode, const FOnlinePartyTypeId InPartyTypeId, UPartyDelegates::FOnJoinUPartyComplete InCompletionDelegate)
 {
-	UE_LOG(LogParty, Display, TEXT("OnJoinPartyComplete() %s %s."), *InPartyId.ToString(), ToString(Result));
+	bool bWasSuccessful = Result == EJoinPartyCompletionResult::Succeeded;
 
-	EJoinPartyCompletionResult LocalResult = Result;
-	FString ErrorString;
-	if (Result == EJoinPartyCompletionResult::Succeeded)
+	UE_LOG(LogParty, Display, TEXT("OnJoinPartyComplete() type(0x%08x) %d %s"), InPartyTypeId.GetValue(), bWasSuccessful, ToString(Result));
+
+	InCompletionDelegate.ExecuteIfBound(LocalUserId, Result, DeniedResultCode);
+
+	//bool bWasSuccessful = Result == EJoinPartyCompletionResult::Succeeded;
+
+	//FString PartyIdDebugString = InPartyId.ToDebugString();
+	//UE_LOG(LogParty, Display, TEXT("OnJoinPartyComplete() %s %d %s."), *PartyIdDebugString, bWasSuccessful, ToString(Result));
+
+	//EJoinPartyCompletionResult LocalResult = Result;
+	//if (bWasSuccessful)
+	//{
+	//	EOnJoinPartyCompleteInternalResult InternalResult = OnJoinPartyCompleteInternal(LocalUserId, InPartyId);
+
+	//	if (InternalResult != EOnJoinPartyCompleteInternalResult::Success)
+	//	{
+	//		LocalResult = EJoinPartyCompletionResult::UnknownClientFailure;
+	//		UE_LOG(LogParty, Warning, TEXT("Error when joining party %s error=%s"), *PartyIdDebugString, ToString(InternalResult));
+	//	}
+	//}
+
+	//int32 OutDeniedResultCode = (Result == EJoinPartyCompletionResult::NotApproved) ? DeniedResultCode : 0;
+	//CompletionDelegate.ExecuteIfBound(LocalUserId, InPartyId, LocalResult, OutDeniedResultCode);
+}
+
+EOnJoinPartyCompleteInternalResult UParty::OnJoinPartyCompleteInternal(const FUniqueNetId& LocalUserId, const FOnlinePartyTypeId InPartyTypeId)
+{
+	EOnJoinPartyCompleteInternalResult ReturnResult = EOnJoinPartyCompleteInternalResult::UnknownFailure;
+
+	UWorld* World = GetWorld();
+	IOnlinePartyPtr PartyInt = Online::GetPartyInterface(World);
+	if (PartyInt.IsValid())
 	{
-		UWorld* World = GetWorld();
-		IOnlinePartyPtr PartyInt = Online::GetPartyInterface(World);
-		if (PartyInt.IsValid())
+		TSharedPtr<const FOnlineParty> Party = PartyInt->GetParty(LocalUserId, InPartyTypeId);
+		if (Party.IsValid())
 		{
-			TSharedPtr<FOnlinePartyInfo> PartyInfo = PartyInt->GetPartyInfo(LocalUserId, InPartyId);
-			if (PartyInfo.IsValid())
+			TSubclassOf<UPartyGameState>* PartyGameStateClass = PartyClasses.Find(InPartyTypeId);
+			if (PartyGameStateClass != nullptr)
 			{
-				FVariantData PartyConfigType;
-				PartyInfo->Config->ClientConfigData.GetAttribute(PARTY_CONFIG_TYPE, PartyConfigType);
-				if (PartyConfigType.GetType() == EOnlineKeyValuePairDataType::Int32)
-				{
-					int32 PartyType = 0;
-					PartyConfigType.GetValue(PartyType);
-
-					TSubclassOf<UPartyGameState>* PartyGameStateClass = PartyClasses.Find(PartyType);
-					if (PartyGameStateClass != nullptr)
-					{
-						UPartyGameState* NewParty = NewObject<UPartyGameState>(this, *PartyGameStateClass);
-
-						// Add right away so future delegate broadcasts have this available
-						JoinedParties.Add(InPartyId.ToString(), NewParty);
-
-						// Initialize and trigger delegates
-						NewParty->InitFromJoin(LocalUserId, PartyInfo);
-					}
-					else
-					{
-						LocalResult = EJoinPartyCompletionResult::UnknownClientFailure;
-						ErrorString = FString::Printf(TEXT("No valid party class when joining party %s"), *InPartyId.ToString());
-					}
-				}
-				else
-				{
-					LocalResult = EJoinPartyCompletionResult::UnknownClientFailure;
-					ErrorString = FString::Printf(TEXT("Invalid party config type when joining party %s"), *InPartyId.ToString());
-				}
+				UPartyGameState* NewParty = NewObject<UPartyGameState>(this, *PartyGameStateClass);
+
+				// Add right away so future delegate broadcasts have this available
+				JoinedParties.Add(InPartyTypeId, NewParty);
+
+				// Initialize and trigger delegates
+				NewParty->InitFromJoin(LocalUserId, Party);
+
+				ReturnResult = EOnJoinPartyCompleteInternalResult::Success;
 			}
 			else
 			{
-				LocalResult = EJoinPartyCompletionResult::UnknownClientFailure;
-				ErrorString = FString::Printf(TEXT("Invalid PartyInfo when joining party %s"), *InPartyId.ToString());
+				ReturnResult = EOnJoinPartyCompleteInternalResult::InvalidPartyType;
 			}
 		}
 		else
 		{
-			LocalResult = EJoinPartyCompletionResult::UnknownClientFailure;
-			ErrorString = FString::Printf(TEXT("No party interface when joining party %s"), *InPartyId.ToString());
+			ReturnResult = EOnJoinPartyCompleteInternalResult::InvalidParty;
 		}
 	}
 	else
 	{
-		ErrorString = ToString(LocalResult);
+		ReturnResult = EOnJoinPartyCompleteInternalResult::InvalidSystemInterface;
 	}
 
-	if (LocalResult != EJoinPartyCompletionResult::Succeeded)
-	{
-		UE_LOG(LogParty, Warning, TEXT("%s"), *ErrorString);
-	}
-
-	int32 OutDeniedResultCode = (LocalResult == EJoinPartyCompletionResult::NotApproved) ? DeniedResultCode : 0;
+	return ReturnResult;
+}
 
-	CompletionDelegate.ExecuteIfBound(LocalUserId, InPartyId, LocalResult, OutDeniedResultCode);
+void UParty::LeaveParty(const FUniqueNetId& InUserId, const FOnlinePartyTypeId InPartyTypeId, const UPartyDelegates::FOnLeaveUPartyComplete& InCompletionDelegate)
+{
+	UPartyDelegates::FOnLeaveUPartyComplete CompletionDelegate;
+	CompletionDelegate.BindUObject(this, &ThisClass::OnLeavePartyComplete, InPartyTypeId, InCompletionDelegate);
+	LeavePartyInternal(InUserId, InPartyTypeId, CompletionDelegate);
 }
 
-void UParty::LeaveParty(const FUniqueNetId& InUserId, const FOnlinePartyId& InPartyId, const FOnLeavePartyComplete& InCompletionDelegate)
+void UParty::OnLeavePartyComplete(const FUniqueNetId& LocalUserId, const ELeavePartyCompletionResult Result, const FOnlinePartyTypeId InPartyTypeId, UPartyDelegates::FOnLeaveUPartyComplete InCompletionDelegate)
 {
-	FOnLeavePartyComplete CompletionDelegate;
-	CompletionDelegate.BindUObject(this, &ThisClass::OnLeavePartyComplete, InCompletionDelegate);
-	LeavePartyInternal(InUserId, InPartyId, CompletionDelegate);
+	bool bWasSuccessful = Result == ELeavePartyCompletionResult::Succeeded;
+
+	UE_LOG(LogParty, Display, TEXT("OnLeavePartyComplete() type(0x%08x) %d %s"), InPartyTypeId.GetValue(), bWasSuccessful, ToString(Result));
+
+	InCompletionDelegate.ExecuteIfBound(LocalUserId, Result);
+	//FString PartyIdDebugString = InPartyId.ToDebugString();
+	//UE_LOG(LogParty, Display, TEXT("OnLeavePartyComplete() %s %d %s."), *PartyIdDebugString, Result == ELeavePartyCompletionResult::Succeeded, ToString(Result));
+
+	//ELeavePartyCompletionResult LocalResult = Result;
+	//EOnLeavePartyCompleteInternalResult InternalResult = OnLeavePartyCompleteInternal(LocalUserId, InPartyId);
+
+	//if (InternalResult != EOnLeavePartyCompleteInternalResult::Success)
+	//{
+	//	LocalResult = ELeavePartyCompletionResult::UnknownClientFailure;
+	//	UE_LOG(LogParty, Warning, TEXT("Error when leaving party %s error=%s"), *PartyIdDebugString, ToString(InternalResult));
+	//}
+
+	//CompletionDelegate.ExecuteIfBound(LocalUserId, InPartyId, LocalResult);
 }
 
-void UParty::OnLeavePartyComplete(const FUniqueNetId& LocalUserId, const FOnlinePartyId& InPartyId, const ELeavePartyCompletionResult Result, FOnLeavePartyComplete CompletionDelegate)
+EOnLeavePartyCompleteInternalResult UParty::OnLeavePartyCompleteInternal(const FUniqueNetId& LocalUserId, const FOnlinePartyTypeId InPartyTypeId)
 {
-	UE_LOG(LogParty, Display, TEXT("OnLeavePartyComplete() %s %d %s."), *InPartyId.ToString(), Result == ELeavePartyCompletionResult::Succeeded, ToString(Result));
+	EOnLeavePartyCompleteInternalResult ReturnResult = EOnLeavePartyCompleteInternalResult::UnknownFailure;
 
-	UPartyGameState* PartyState = GetParty(InPartyId);
+	UPartyGameState* PartyState = GetParty(InPartyTypeId);
 	if (PartyState)
 	{
 		PartyState->HandleRemovedFromParty(EMemberExitedReason::Left);
-		JoinedParties.Remove(InPartyId.ToString());
+		JoinedParties.Remove(InPartyTypeId);
+		ReturnResult = EOnLeavePartyCompleteInternalResult::Success;
 	}
 	else
 	{
-		UE_LOG(LogParty, Warning, TEXT("OnLeavePartyComplete: Missing party state %s"), *InPartyId.ToString());
+		ReturnResult = EOnLeavePartyCompleteInternalResult::InvalidPartyState;
 	}
-	
-	CompletionDelegate.ExecuteIfBound(LocalUserId, InPartyId, Result);
+
+	return ReturnResult;
 }
 
 void UParty::GetDefaultPersistentPartySettings(EPartyType& PartyType, bool& bLeaderFriendsOnly, bool& bLeaderInvitesOnly)
@@ -691,20 +907,40 @@ void UParty::GetPersistentPartyConfiguration(FPartyConfiguration& PartyConfig)
 
 	bool bIsPrivate = (PartyType == EPartyType::Private);
 
-	PartyConfig.Permissions = bIsPrivate ? EPartyPermissions::Private : EPartyPermissions::Public;
-	PartyConfig.bIsAcceptingMembers = bIsPrivate ? false : true;
-	PartyConfig.bPresenceParty = true;
-	PartyConfig.bShouldPublishToPresence = !bIsPrivate;
-	PartyConfig.bCanNonLeaderPublishToPresence = !bIsPrivate && !bLeaderFriendsOnly;
-	PartyConfig.bCanNonLeaderInviteToParty = !bLeaderInvitesOnly;
+	PartySystemPermissions::EPresencePermissions PresencePermissions;
+	if (bLeaderFriendsOnly)
+	{
+		if (bIsPrivate)
+		{
+			PresencePermissions = PartySystemPermissions::FriendsInviteOnly;
+		}
+		else
+		{
+			PresencePermissions = PartySystemPermissions::FriendsOnly;
+		}
+	}
+	else
+	{
+		if (bIsPrivate)
+		{
+			PresencePermissions = PartySystemPermissions::PublicInviteOnly;
+		}
+		else
+		{
+			PresencePermissions = PartySystemPermissions::Public;
+		}
+	}
 
-	PartyConfig.ClientConfigData.SetAttribute(PARTY_CONFIG_TYPE, (int32)EPartyConfigType::PersistentParty);
+	PartyConfig.JoinRequestAction = EJoinRequestAction::Manual;
+	PartyConfig.bIsAcceptingMembers = bIsPrivate ? false : true;
+	PartyConfig.bShouldRemoveOnDisconnection = true;
+	PartyConfig.PresencePermissions = PresencePermissions;
+	PartyConfig.InvitePermissions = bLeaderInvitesOnly ? PartySystemPermissions::EInvitePermissions::Leader : PartySystemPermissions::EInvitePermissions::Anyone;
 
 	PartyConfig.MaxMembers = DefaultMaxPartySize;
-	PartyConfig.bShouldRemoveOnDisconnection = true;
 }
 
-void UParty::CreatePersistentParty(const FUniqueNetId& InUserId, const FOnCreatePartyComplete& InCompletionDelegate)
+void UParty::CreatePersistentParty(const FUniqueNetId& InUserId, const UPartyDelegates::FOnCreateUPartyComplete& InCompletionDelegate)
 {
 	if (PersistentPartyId.IsValid())
 	{
@@ -716,19 +952,49 @@ void UParty::CreatePersistentParty(const FUniqueNetId& InUserId, const FOnCreate
 	FPartyConfiguration PartyConfig;
 	GetPersistentPartyConfiguration(PartyConfig);
 
-	FOnCreatePartyComplete CompletionDelegate;
+	UPartyDelegates::FOnCreateUPartyComplete CompletionDelegate;
 	CompletionDelegate.BindUObject(this, &ThisClass::OnCreatePersistentPartyComplete, InCompletionDelegate);
-	CreatePartyInternal(InUserId, PartyConfig, CompletionDelegate);
+	CreatePartyInternal(InUserId, IOnlinePartySystem::GetPrimaryPartyTypeId(), PartyConfig, CompletionDelegate);
 }
 
-void UParty::OnCreatePersistentPartyComplete(const FUniqueNetId& LocalUserId, const FOnlinePartyId& InPartyId, const ECreatePartyCompletionResult Result, FOnCreatePartyComplete CompletionDelegate)
+void UParty::OnCreatePersistentPartyComplete(const FUniqueNetId& LocalUserId, const ECreatePartyCompletionResult Result, UPartyDelegates::FOnCreateUPartyComplete CompletionDelegate)
 {
-	if (Result == ECreatePartyCompletionResult::Succeeded)
+	bool bWasSuccessful = Result == ECreatePartyCompletionResult::Succeeded;
+
+	TSharedPtr<const FOnlineParty> Party;
+	TSharedPtr<const FOnlinePartyId> PartyId;
+	
+	UWorld* World = GetWorld();
+	IOnlinePartyPtr PartyInt = Online::GetPartyInterface(World);
+	if (PartyInt.IsValid())
 	{
-		PersistentPartyId = InPartyId.AsShared();
+		if (bWasSuccessful)
+		{
+			Party = PartyInt->GetParty(LocalUserId, IOnlinePartySystem::GetPrimaryPartyTypeId());
+			if (Party.IsValid())
+			{
+				PartyId = Party->PartyId;
+				PersistentPartyId = PartyId;
+			}
+		}
 	}
 
-	OnCreatePartyComplete(LocalUserId, InPartyId, Result, CompletionDelegate);
+	FString PartyIdDebugString = PartyId.IsValid() ? PartyId->ToDebugString() : TEXT("Invalid");
+
+	UE_LOG(LogParty, Display, TEXT("OnCreatePersistentPartyComplete() %s %d %s"), *PartyIdDebugString, bWasSuccessful, ToString(Result));
+
+	//ECreatePartyCompletionResult LocalResult = Result;
+	//if (bWasSuccessful)
+	//{
+	//	EOnCreatePartyCompleteInternalResult InternalResult = OnCreatePartyCompleteInternal(LocalUserId, IOnlinePartySystem::GetPrimaryPartyTypeId());
+	//	if (InternalResult != EOnCreatePartyCompleteInternalResult::Success)
+	//	{
+	//		LocalResult = ECreatePartyCompletionResult::UnknownClientFailure;
+	//		UE_LOG(LogParty, Warning, TEXT("Error when creating persistent party %s error=%s"), *PartyIdDebugString, ToString(InternalResult));
+	//	}
+	//}
+
+	CompletionDelegate.ExecuteIfBound(LocalUserId, Result);
 
 	if (Result == ECreatePartyCompletionResult::Succeeded)
 	{
@@ -760,7 +1026,7 @@ void UParty::OnCreatePersistentPartyComplete(const FUniqueNetId& LocalUserId, co
 	}
 }
 
-void UParty::JoinPersistentParty(const FUniqueNetId& InUserId, const FPartyDetails& InPartyDetails, const FOnJoinPartyComplete& InCompletionDelegate)
+void UParty::JoinPersistentParty(const FUniqueNetId& InUserId, const FPartyDetails& InPartyDetails, const UPartyDelegates::FOnJoinUPartyComplete& InCompletionDelegate)
 {
 	if (PersistentPartyId.IsValid())
 	{
@@ -769,19 +1035,50 @@ void UParty::JoinPersistentParty(const FUniqueNetId& InUserId, const FPartyDetai
 
 	PersistentPartyId = nullptr;
 
-	FOnJoinPartyComplete CompletionDelegate;
+	UPartyDelegates::FOnJoinUPartyComplete CompletionDelegate;
 	CompletionDelegate.BindUObject(this, &ThisClass::OnJoinPersistentPartyComplete, InCompletionDelegate);
 	JoinPartyInternal(InUserId, InPartyDetails, CompletionDelegate);
 }
 
-void UParty::OnJoinPersistentPartyComplete(const FUniqueNetId& LocalUserId, const FOnlinePartyId& InPartyId, EJoinPartyCompletionResult Result, int32 DeniedResultCode, FOnJoinPartyComplete CompletionDelegate)
+void UParty::OnJoinPersistentPartyComplete(const FUniqueNetId& LocalUserId, EJoinPartyCompletionResult Result, int32 DeniedResultCode, UPartyDelegates::FOnJoinUPartyComplete CompletionDelegate)
 {
-	if (Result == EJoinPartyCompletionResult::Succeeded)
+	bool bWasSuccessful = (Result == EJoinPartyCompletionResult::Succeeded);
+
+	TSharedPtr<const FOnlineParty> Party;
+	TSharedPtr<const FOnlinePartyId> PartyId;
+
+	UWorld* World = GetWorld();
+	IOnlinePartyPtr PartyInt = Online::GetPartyInterface(World);
+	if (PartyInt.IsValid())
 	{
-		PersistentPartyId = InPartyId.AsShared();
+		if (bWasSuccessful)
+		{
+			Party = PartyInt->GetParty(LocalUserId, IOnlinePartySystem::GetPrimaryPartyTypeId());
+			if (Party.IsValid())
+			{
+				PartyId = Party->PartyId;
+				PersistentPartyId = PartyId;
+			}
+		}
 	}
 
-	OnJoinPartyComplete(LocalUserId, InPartyId, Result, DeniedResultCode, CompletionDelegate);
+	FString PartyIdDebugString = PartyId.IsValid() ? PartyId->ToDebugString() : TEXT("Invalid");
+
+	UE_LOG(LogParty, Display, TEXT("OnJoinPersistentPartyComplete() %s %d %s"), *PartyIdDebugString, bWasSuccessful, ToString(Result));
+
+	//EJoinPartyCompletionResult LocalResult = Result;
+	//if (bWasSuccessful)
+	//{
+	//	EOnJoinPartyCompleteInternalResult InternalResult = OnJoinPartyCompleteInternal(LocalUserId, IOnlinePartySystem::GetPrimaryPartyTypeId());
+	//	if (InternalResult != EOnJoinPartyCompleteInternalResult::Success)
+	//	{
+	//		LocalResult = EJoinPartyCompletionResult::UnknownInternalFailure;
+	//		UE_LOG(LogParty, Warning, TEXT("Error when creating persistent party %s error=%s"), *PartyIdDebugString, ToString(InternalResult));
+	//	}
+	//}
+
+	int32 OutDeniedResultCode = (Result == EJoinPartyCompletionResult::NotApproved) ? DeniedResultCode : 0;
+	CompletionDelegate.ExecuteIfBound(LocalUserId, Result, OutDeniedResultCode);
 
 	if (Result == EJoinPartyCompletionResult::Succeeded)
 	{
@@ -808,7 +1105,6 @@ void UParty::OnJoinPersistentPartyComplete(const FUniqueNetId& LocalUserId, cons
 	{
 		if (Result != EJoinPartyCompletionResult::AlreadyJoiningParty)
 		{
-			UWorld* World = GetWorld();
 			if (World)
 			{
 				// Try to get back to a good state
@@ -829,17 +1125,17 @@ void UParty::HandleJoinPersistentPartyFailure()
 	RestorePersistentPartyState();
 }
 
-void UParty::LeavePersistentParty(const FUniqueNetId& InUserId, const FOnLeavePartyComplete& InCompletionDelegate)
+void UParty::LeavePersistentParty(const FUniqueNetId& InUserId, const UPartyDelegates::FOnLeaveUPartyComplete& InCompletionDelegate)
 {
-	FOnLeavePartyComplete CompletionDelegate;
-	CompletionDelegate.BindUObject(this, &ThisClass::OnLeavePersistentPartyComplete, InCompletionDelegate);
-
 	if (PersistentPartyId.IsValid())
 	{
+		UPartyDelegates::FOnLeaveUPartyComplete CompletionDelegate;
+		CompletionDelegate.BindUObject(this, &ThisClass::OnLeavePersistentPartyComplete, InCompletionDelegate);
+
 		if (!bLeavingPersistentParty)
 		{
 			bLeavingPersistentParty = true;
-			LeavePartyInternal(InUserId, *PersistentPartyId, CompletionDelegate);
+			LeavePartyInternal(InUserId, IOnlinePartySystem::GetPrimaryPartyTypeId(), CompletionDelegate);
 		}
 		else
 		{
@@ -850,81 +1146,88 @@ void UParty::LeavePersistentParty(const FUniqueNetId& InUserId, const FOnLeavePa
 	{
 		UE_LOG(LogParty, Warning, TEXT("No party during LeavePersistentParty()"));
 		// mark this as true so the ensure works in the delegate
-		bLeavingPersistentParty = true;
-		CompletionDelegate.ExecuteIfBound(InUserId, InUserId, ELeavePartyCompletionResult::UnknownClientFailure);
+		InCompletionDelegate.ExecuteIfBound(InUserId, ELeavePartyCompletionResult::UnknownClientFailure);
 	}
 }
 
-void UParty::OnLeavePersistentPartyComplete(const FUniqueNetId& LocalUserId, const FOnlinePartyId& InPartyId, const ELeavePartyCompletionResult Result, FOnLeavePartyComplete CompletionDelegate)
+void UParty::OnLeavePersistentPartyComplete(const FUniqueNetId& LocalUserId, const ELeavePartyCompletionResult Result, UPartyDelegates::FOnLeaveUPartyComplete CompletionDelegate)
 {
 	ensure(bLeavingPersistentParty);
-	ensure(!PersistentPartyId.IsValid() || *PersistentPartyId == InPartyId);
+	//ensure(!PersistentPartyId.IsValid());
 	bLeavingPersistentParty = false;
 
+	bool bWasSuccessful = (Result == ELeavePartyCompletionResult::Succeeded);
+
+	UE_LOG(LogParty, Display, TEXT("OnLeavePersistentPartyComplete() %d %s"), bWasSuccessful, ToString(Result));
+
 	// This might be the last reference of the party id, transfer the shared pointer for this brief period of time
 	// PersistentPartyId is expected to be null now as far as external code is concerned
-	TSharedPtr<const FOnlinePartyId> TempPartyId = PersistentPartyId.IsValid() ? PersistentPartyId : nullptr;
+	TSharedPtr<const FOnlinePartyId> TempPartyId = PersistentPartyId;
 	PersistentPartyId = nullptr;
-	OnLeavePartyComplete(LocalUserId, InPartyId, Result, CompletionDelegate);
 
-	TArray<FOnLeavePartyComplete> DelegatesCopy = LeavePartyCompleteDelegates;
+	//EOnLeavePartyCompleteInternalResult InternalResult = OnLeavePartyCompleteInternal(LocalUserId, IOnlinePartySystem::GetPrimaryPartyTypeId());
+	//ELeavePartyCompletionResult LocalResult = Result;
+	//if (InternalResult != EOnLeavePartyCompleteInternalResult::Success)
+	//{
+	//	LocalResult = ELeavePartyCompletionResult::UnknownClientFailure;
+	//}
+
+	CompletionDelegate.ExecuteIfBound(LocalUserId, Result);
+
+	TArray<UPartyDelegates::FOnLeaveUPartyComplete> DelegatesCopy = LeavePartyCompleteDelegates;
 	LeavePartyCompleteDelegates.Empty();
 
 	// fire delegates for any/all LeavePersistentParty calls made while this one was in flight
-	for (const FOnLeavePartyComplete& ExtraDelegate : DelegatesCopy)
+	for (const UPartyDelegates::FOnLeaveUPartyComplete& ExtraDelegate : DelegatesCopy)
 	{
-		ExtraDelegate.ExecuteIfBound(LocalUserId, InPartyId, Result);
+		ExtraDelegate.ExecuteIfBound(LocalUserId, Result);
 	}
 }
 
 void UParty::RestorePersistentPartyState()
 {
-	if (!GIsRequestingExit)
+	if (!bLeavingPersistentParty)
 	{
-		if (!bLeavingPersistentParty)
+		UWorld* World = GetWorld();
+		IOnlinePartyPtr PartyInt = Online::GetPartyInterface(World);
+		if (PartyInt.IsValid())
 		{
-			UWorld* World = GetWorld();
-			IOnlinePartyPtr PartyInt = Online::GetPartyInterface(World);
-			if (PartyInt.IsValid())
-			{
-				UGameInstance* GameInstance = GetGameInstance();
-				check(GameInstance);
+			UGameInstance* GameInstance = GetGameInstance();
+			check(GameInstance);
 
-				TSharedPtr<const FUniqueNetId> LocalUserId = GameInstance->GetPrimaryPlayerUniqueId();
-				check(LocalUserId.IsValid() && LocalUserId->IsValid());
+			TSharedPtr<const FUniqueNetId> LocalUserId = GameInstance->GetPrimaryPlayerUniqueId();
+			check(LocalUserId.IsValid() && LocalUserId->IsValid());
 
-				UPartyGameState* PersistentParty = GetPersistentParty();
+			UPartyGameState* PersistentParty = GetPersistentParty();
 
-				// Check for existing party and create a new one if necessary
-				bool bFoundExistingPersistentParty = PersistentParty ? true : false;
-				if (bFoundExistingPersistentParty)
+			// Check for existing party and create a new one if necessary
+			bool bFoundExistingPersistentParty = PersistentParty ? true : false;
+			if (bFoundExistingPersistentParty)
+			{
+				// In a party already, make sure the UI is aware of its state
+				if (PersistentParty->ResetForFrontend())
 				{
-					// In a party already, make sure the UI is aware of its state
-					if (PersistentParty->ResetForFrontend())
-					{
-						OnPartyResetForFrontend().Broadcast(PersistentParty);
-					}
-					else
-					{
-						// There was an issue resetting the party, so leave 
-						LeaveAndRestorePersistentParty();
-					}
+					OnPartyResetForFrontend().Broadcast(PersistentParty);
 				}
 				else
 				{
-					PersistentPartyId = nullptr;
-
-					// Create a new party
-					FOnCreatePartyComplete CompletionDelegate;
-					CreatePersistentParty(*LocalUserId, CompletionDelegate);
+					// There was an issue resetting the party, so leave 
+					LeaveAndRestorePersistentParty();
 				}
 			}
-		}
-		else
-		{
-			UE_LOG(LogParty, Verbose, TEXT("Can't RestorePersistentPartyState while leaving party, ignoring"));
+			else
+			{
+				PersistentPartyId = nullptr;
+
+				// Create a new party
+				CreatePersistentParty(*LocalUserId);
+			}
 		}
 	}
+	else
+	{
+		UE_LOG(LogParty, Verbose, TEXT("Can't RestorePersistentPartyState while leaving party, ignoring"));
+	}
 }
 
 void UParty::UpdatePersistentPartyLeader(const FUniqueNetIdRepl& NewPartyLeader)
@@ -978,8 +1281,7 @@ void UParty::KickFromPersistentParty()
 			if (ensure(LocalUserId.IsValid() && LocalUserId->IsValid()))
 			{
 				// Leave the party (restored in frontend)
-				FOnLeavePartyComplete CompletionDelegate;
-				LeavePersistentParty(*LocalUserId, CompletionDelegate);
+				LeavePersistentParty(*LocalUserId);
 			}
 		}
 		else
@@ -1015,7 +1317,7 @@ void UParty::LeaveAndRestorePersistentPartyInternal()
 	TSharedPtr<const FUniqueNetId> PrimaryUserId = GameInstance->GetPrimaryPlayerUniqueId();
 	check(PrimaryUserId.IsValid() && PrimaryUserId->IsValid());
 
-	FOnLeavePartyComplete CompletionDelegate;
+	UPartyDelegates::FOnLeaveUPartyComplete CompletionDelegate;
 	CompletionDelegate.BindUObject(this, &ThisClass::OnLeavePersistentPartyAndRestore);
 
 	// Unset this here, wouldn't be here if we were unable to call LeaveAndRestorePersistentParty, it will be reset inside
@@ -1024,39 +1326,38 @@ void UParty::LeaveAndRestorePersistentPartyInternal()
 	LeavePersistentParty(*PrimaryUserId, CompletionDelegate);
 }
 
-void UParty::OnLeavePersistentPartyAndRestore(const FUniqueNetId& LocalUserId, const FOnlinePartyId& PartyId, const ELeavePartyCompletionResult Result)
+void UParty::OnLeavePersistentPartyAndRestore(const FUniqueNetId& LocalUserId, const ELeavePartyCompletionResult Result)
 {
-	UE_LOG(LogParty, Display, TEXT("OnLeavePersistentPartyAndRestore [%s] Result: %s"), *PartyId.ToString(), ToString(Result));
+	UE_LOG(LogParty, Display, TEXT("OnLeavePersistentPartyAndRestore Result: %s"), ToString(Result));
+
 	RestorePersistentPartyState();
 }
 
-bool UParty::ProcessPendingPartyInvites()
+void UParty::AddPendingPartyJoin(const FUniqueNetId& LocalUserId, const FPartyDetails& PartyDetails, const UPartyDelegates::FOnJoinUPartyComplete& JoinCompleteDelegate)
 {
-	if (PendingPartyInvite.IsValid() && PendingPartyInvite->IsValid())
+	if (LocalUserId.IsValid() && PartyDetails.IsValid())
 	{
-		// Copy the values out because PendingPartyInvite is going to be remade inside this function
-		TSharedRef<const FUniqueNetId> LocalUserId = PendingPartyInvite->LocalUserId;
-		FPartyDetails PartyDetails = PendingPartyInvite->PartyDetails;
-		FOnJoinPartyComplete JoinCompleteDelegate = PendingPartyInvite->JoinCompleteDelegate;
-		HandlePendingInvite(*LocalUserId, PartyDetails, JoinCompleteDelegate);
-		return true;
+		if (!PendingPartyJoin.IsValid())
+		{
+			PendingPartyJoin = MakeShareable(new FPendingPartyJoin(LocalUserId.AsShared(), PartyDetails, JoinCompleteDelegate));
+		}
 	}
-
-	return false;
 }
 
-void UParty::HandlePendingInvite(const FUniqueNetId& LocalUserId, const FPartyDetails& PartyDetails, const FOnJoinPartyComplete& JoinCompleteDelegate)
+bool UParty::ProcessPendingPartyJoin()
 {
-	if (LocalUserId.IsValid() && PartyDetails.IsValid())
+	if (PendingPartyJoin.IsValid() )
 	{
-		// Setup the invite for processing
-		PendingPartyInvite = MakeShareable(new FPendingPartyInvite(LocalUserId.AsShared(), PartyDetails, JoinCompleteDelegate));
+		// Copy the values out because PendingPartyInvite is going to be remade inside this function
+		TSharedRef<const FUniqueNetId> LocalUserId = PendingPartyJoin->LocalUserId;
+		FPartyDetails PartyDetails = PendingPartyJoin->PartyDetails;
 
-		UWorld* World = GetWorld();
-		check(World);
+		HandlePendingJoin();
 
-		JoinCompleteDelegate.ExecuteIfBound(LocalUserId, *PartyDetails.GetPartyId(), EJoinPartyCompletionResult::UnknownClientFailure, 0);
+		return true;
 	}
+
+	return false;
 }
 
 UWorld* UParty::GetWorld() const
diff --git a/Engine/Source/Runtime/Online/OnlineFramework/Party/Private/PartyGameState.cpp b/Engine/Source/Runtime/Online/OnlineFramework/Party/Private/PartyGameState.cpp
--- a/Engine/Source/Runtime/Online/OnlineFramework/Party/Private/PartyGameState.cpp
+++ b/Engine/Source/Runtime/Online/OnlineFramework/Party/Private/PartyGameState.cpp
@@ -26,7 +26,7 @@ UPartyGameState::UPartyGameState(const FObjectInitializer& ObjectInitializer) :
 	PartyStateRef(nullptr),
 	OwningUserId(nullptr),
 	bDebugAcceptingMembers(false),
-	PartyInfo(nullptr),
+	OssParty(nullptr),
 	bPromotionLockoutState(false),
 	bStayWithPartyOnDisconnect(false),
 	ReservationBeaconClientClass(nullptr),
@@ -53,7 +53,7 @@ void UPartyGameState::OnShutdown()
 		UnregisterFrontendDelegates();
 	}
 
-	PartyInfo.Reset();
+	OssParty.Reset();
 	PartyMembersState.Empty();
 
 	if (PartyStateRefDef)
@@ -128,20 +128,20 @@ void UPartyGameState::ResetLocalPlayerState()
 	}
 }
 
-void UPartyGameState::Init(const FUniqueNetId& LocalUserId, TSharedPtr<FOnlinePartyInfo>& InPartyInfo)
+void UPartyGameState::Init(const FUniqueNetId& LocalUserId, TSharedPtr<const FOnlineParty>& InParty)
 {
 	OwningUserId.SetUniqueNetId(LocalUserId.AsShared());
-	PartyInfo = InPartyInfo;
-	CurrentConfig = *PartyInfo->Config;
+	OssParty = InParty;
+	CurrentConfig = *OssParty->Config;
 	bDebugAcceptingMembers = CurrentConfig.bIsAcceptingMembers;
 
 	// Last since it needs the party info/id set first
 	RegisterFrontendDelegates();
 }
 
-void UPartyGameState::InitFromCreate(const FUniqueNetId& LocalUserId, TSharedPtr<FOnlinePartyInfo>& InPartyInfo)
+void UPartyGameState::InitFromCreate(const FUniqueNetId& LocalUserId, TSharedPtr<const FOnlineParty>& InParty)
 {
-	Init(LocalUserId, InPartyInfo);
+	Init(LocalUserId, InParty);
 
 	// Setup initial data for the party
 	ResetPartyState();
@@ -158,9 +158,9 @@ void UPartyGameState::InitFromCreate(const FUniqueNetId& LocalUserId, TSharedPtr
 	SendLocalPlayerPartyData();
 }
 
-void UPartyGameState::InitFromJoin(const FUniqueNetId& LocalUserId, TSharedPtr<FOnlinePartyInfo>& InPartyInfo)
+void UPartyGameState::InitFromJoin(const FUniqueNetId& LocalUserId, TSharedPtr<const FOnlineParty>& InParty)
 {
-	Init(LocalUserId, InPartyInfo);
+	Init(LocalUserId, InParty);
 
 	// Broadcast join
 	UParty* Party = GetPartyOuter();
@@ -191,18 +191,18 @@ bool UPartyGameState::ResetForFrontend()
 
 	CleanupReservationBeacon();
 
-	if (PartyInfo.IsValid())
+	if (OssParty.IsValid())
 	{
 		UWorld* World = GetWorld();
 		IOnlinePartyPtr PartyInt = Online::GetPartyInterface(World);
 		if (PartyInt.IsValid())
 		{
-			TSharedPtr<FOnlinePartyInfo> ExistingPartyInfo = PartyInt->GetPartyInfo(*OwningUserId, *PartyInfo->PartyId);
-			if (ExistingPartyInfo.IsValid())
+			TSharedPtr<const FOnlineParty> ExistingParty = PartyInt->GetParty(*OwningUserId, *OssParty->PartyId);
+			if (ExistingParty.IsValid())
 			{
 				bSuccess = true;
 
-				Init(*OwningUserId, ExistingPartyInfo);
+				Init(*OwningUserId, ExistingParty);
 				bStayWithPartyOnDisconnect = false;
 
 				ResetPartyState();
@@ -220,10 +220,10 @@ bool UPartyGameState::ResetForFrontend()
 					const FUniqueNetIdRepl& MemberId = PartyMember.Key;
 					if (MemberId.IsValid())
 					{
-						TSharedPtr<FOnlinePartyMember> CheckPartyMember = PartyInt->GetPartyMember(*OwningUserId, *PartyInfo->PartyId, *MemberId);
+						TSharedPtr<FOnlinePartyMember> CheckPartyMember = PartyInt->GetPartyMember(*OwningUserId, *OssParty->PartyId, *MemberId);
 						if (!CheckPartyMember.IsValid())
 						{
-							UE_LOG(LogParty, Verbose, TEXT("[%s] Player %s left during fixup"), *PartyInfo->PartyId->ToString(), *MemberId.ToString());
+							UE_LOG(LogParty, Verbose, TEXT("[%s] Player %s left during fixup"), *OssParty->PartyId->ToString(), *MemberId.ToString());
 							HandlePartyMemberLeft(*MemberId, EMemberExitedReason::Left);
 						}
 					}
@@ -231,7 +231,7 @@ bool UPartyGameState::ResetForFrontend()
 
 				// Add members we don't have, but lower level does
 				TArray<TSharedRef<FOnlinePartyMember>> PartyMembers;
-				PartyInt->GetPartyMembers(*OwningUserId, *PartyInfo->PartyId, PartyMembers);
+				PartyInt->GetPartyMembers(*OwningUserId, *OssParty->PartyId, PartyMembers);
 				for (const auto& PartyMember : PartyMembers)
 				{
 					TSharedRef<const FUniqueNetId> MemberId = PartyMember->GetUserId();
@@ -241,10 +241,10 @@ bool UPartyGameState::ResetForFrontend()
 					UPartyMemberState* CurrentPartyMemberData = CurrentPartyMemberDataPtr ? *CurrentPartyMemberDataPtr : nullptr;
 					if (!CurrentPartyMemberData)
 					{
-						TSharedPtr<FOnlinePartyData> PartyMemberData = PartyInt->GetPartyMemberData(*OwningUserId, *PartyInfo->PartyId, *MemberId);
+						TSharedPtr<FOnlinePartyData> PartyMemberData = PartyInt->GetPartyMemberData(*OwningUserId, *OssParty->PartyId, *MemberId);
 						if (PartyMemberData.IsValid())
 						{
-							UE_LOG(LogParty, Verbose, TEXT("[%s] Player %s data received during fixup"), *PartyInfo->PartyId->ToString(), *UniqueId.ToString());
+							UE_LOG(LogParty, Verbose, TEXT("[%s] Player %s data received during fixup"), *OssParty->PartyId->ToString(), *UniqueId.ToString());
 							HandlePartyMemberDataReceived(*MemberId, PartyMemberData.ToSharedRef());
 						}
 					}
@@ -274,7 +274,7 @@ bool UPartyGameState::ResetForFrontend()
 
 	if (!bSuccess)
 	{
-		PartyInfo = nullptr;
+		OssParty = nullptr;
 
 		ResetPartyState();
 		ResetLocalPlayerState();
@@ -294,7 +294,7 @@ UPartyMemberState* UPartyGameState::CreateNewPartyMember(const FUniqueNetId& InM
 		IOnlinePartyPtr PartyInt = Online::GetPartyInterface(World);
 		if (PartyInt.IsValid())
 		{
-			TSharedPtr<FOnlinePartyMember> PartyMember = PartyInt->GetPartyMember(*OwningUserId, *PartyInfo->PartyId, InMemberId);
+			TSharedPtr<FOnlinePartyMember> PartyMember = PartyInt->GetPartyMember(*OwningUserId, *OssParty->PartyId, InMemberId);
 			if (PartyMember.IsValid())
 			{
 				NewPartyMemberState = NewObject<UPartyMemberState>(this, PartyMemberStateClass);
@@ -331,17 +331,17 @@ void UPartyGameState::SetInitMemberFunctor( TFunction<void( TSharedPtr<FOnlinePa
 
 void UPartyGameState::HandlePartyConfigChanged( const TSharedRef<FPartyConfiguration>& InPartyConfig )
 {
-	UE_LOG(LogParty, VeryVerbose, TEXT("[%s] HandlePartyConfigChanged"), PartyInfo.IsValid() ? *PartyInfo->PartyId->ToString() : TEXT("INVALID"));
-	if (PartyInfo.IsValid())
+	UE_LOG(LogParty, VeryVerbose, TEXT("[%s] HandlePartyConfigChanged"), OssParty.IsValid() ? *OssParty->PartyId->ToString() : TEXT("INVALID"));
+	if (OssParty.IsValid())
 	{
-		CurrentConfig = *PartyInfo->Config;
+		CurrentConfig = *OssParty->Config;
 		bDebugAcceptingMembers = CurrentConfig.bIsAcceptingMembers;
 	}
 }
 
 void UPartyGameState::HandlePartyMemberJoined(const FUniqueNetId& InMemberId)
 {
-	UE_LOG(LogParty, VeryVerbose, TEXT("[%s] HandlePartyMemberJoined %s"), PartyInfo.IsValid() ? *PartyInfo->PartyId->ToString() : TEXT("INVALID"), *InMemberId.ToString());
+	UE_LOG(LogParty, VeryVerbose, TEXT("[%s] HandlePartyMemberJoined %s"), OssParty.IsValid() ? *OssParty->PartyId->ToString() : TEXT("INVALID"), *InMemberId.ToString());
 
  	TSharedRef<const FUniqueNetId> IdRef = InMemberId.AsShared();
  	FUniqueNetIdRepl MemberId(IdRef);
@@ -363,7 +363,7 @@ void UPartyGameState::HandlePartyMemberJoined(const FUniqueNetId& InMemberId)
 
 void UPartyGameState::HandlePartyMemberLeft(const FUniqueNetId& InMemberId, EMemberExitedReason Reason)
 {
-	UE_LOG(LogParty, VeryVerbose, TEXT("[%s] HandlePartyMemberLeft %s"), PartyInfo.IsValid() ? *PartyInfo->PartyId->ToString() : TEXT("INVALID"), *InMemberId.ToString());
+	UE_LOG(LogParty, VeryVerbose, TEXT("[%s] HandlePartyMemberLeft %s"), OssParty.IsValid() ? *OssParty->PartyId->ToString() : TEXT("INVALID"), *InMemberId.ToString());
 
 	if (InMemberId.IsValid())
 	{
@@ -385,7 +385,7 @@ void UPartyGameState::HandlePartyMemberLeft(const FUniqueNetId& InMemberId, EMem
 
 void UPartyGameState::HandlePartyMemberPromoted(const FUniqueNetId& InMemberId)
 {
-	UE_LOG(LogParty, VeryVerbose, TEXT("[%s] HandlePartyMemberPromoted %s"), PartyInfo.IsValid() ? *PartyInfo->PartyId->ToString() : TEXT("INVALID"), *InMemberId.ToString());
+	UE_LOG(LogParty, VeryVerbose, TEXT("[%s] HandlePartyMemberPromoted %s"), OssParty.IsValid() ? *OssParty->PartyId->ToString() : TEXT("INVALID"), *InMemberId.ToString());
 
 	if (InMemberId.IsValid())
 	{
@@ -410,24 +410,24 @@ void UPartyGameState::ComparePartyData(const FPartyState& OldPartyData, const FP
 	{
 		if (OldPartyData.PartyType != NewPartyData.PartyType)
 		{
-			OnPartyTypeChanged().Broadcast(NewPartyData.PartyType);
+			OnClientPartyTypeChanged().Broadcast(NewPartyData.PartyType);
 		}
 
 		if (OldPartyData.bLeaderFriendsOnly != NewPartyData.bLeaderFriendsOnly)
 		{
-			OnLeaderFriendsOnlyChanged().Broadcast(NewPartyData.bLeaderFriendsOnly);
+			OnClientLeaderFriendsOnlyChanged().Broadcast(NewPartyData.bLeaderFriendsOnly);
 		}
 
 		if (OldPartyData.bLeaderInvitesOnly != NewPartyData.bLeaderInvitesOnly)
 		{
-			OnLeaderInvitesOnlyChanged().Broadcast(NewPartyData.bLeaderInvitesOnly);
+			OnClientLeaderInvitesOnlyChanged().Broadcast(NewPartyData.bLeaderInvitesOnly);
 		}
 	}
 }
 
 void UPartyGameState::HandlePartyDataReceived(const TSharedRef<FOnlinePartyData>& InPartyData)
 {
-	UE_LOG(LogParty, VeryVerbose, TEXT("[%s] HandlePartyDataReceived"), PartyInfo.IsValid() ? *PartyInfo->PartyId->ToString() : TEXT("INVALID"));
+	UE_LOG(LogParty, VeryVerbose, TEXT("[%s] HandlePartyDataReceived"), OssParty.IsValid() ? *OssParty->PartyId->ToString() : TEXT("INVALID"));
 
 	UWorld* World = GetWorld();
 	IOnlinePartyPtr PartyInt = Online::GetPartyInterface(World);
@@ -450,7 +450,7 @@ void UPartyGameState::HandlePartyDataReceived(const TSharedRef<FOnlinePartyData>
 
 void UPartyGameState::HandlePartyMemberDataReceived(const FUniqueNetId& InMemberId, const TSharedRef<FOnlinePartyData>& InPartyMemberData)
 {
-	UE_LOG(LogParty, VeryVerbose, TEXT("[%s] HandlePartyMemberDataReceived %s"), PartyInfo.IsValid() ? *PartyInfo->PartyId->ToString() : TEXT("INVALID"), *InMemberId.ToString());
+	UE_LOG(LogParty, VeryVerbose, TEXT("[%s] HandlePartyMemberDataReceived %s"), OssParty.IsValid() ? *OssParty->PartyId->ToString() : TEXT("INVALID"), *InMemberId.ToString());
 
 	UWorld* World = GetWorld();
 	IOnlinePartyPtr PartyInt = Online::GetPartyInterface(World);
@@ -532,7 +532,7 @@ void UPartyGameState::HandlePartyJoinRequestReceived(const FUniqueNetId& Recipie
 	IOnlinePartyPtr PartyInt = Online::GetPartyInterface(World);
 	if (PartyInt.IsValid())
 	{
-		PartyInt->ApproveJoinRequest(RecipientId, *PartyInfo->PartyId, SenderId, true);
+		PartyInt->ApproveJoinRequest(RecipientId, *Party->PartyId, SenderId, true);
 	}
 	return;
 #endif
@@ -540,7 +540,7 @@ void UPartyGameState::HandlePartyJoinRequestReceived(const FUniqueNetId& Recipie
 	EApprovalAction ApprovalAction = EApprovalAction::Deny;
 	EJoinPartyDenialReason DenialReason = EJoinPartyDenialReason::Busy;
 
-	if (PartyInfo.IsValid())
+	if (OssParty.IsValid())
 	{
 		int32 NumPartyMembers = GetPartySize();
 		int32 MaxPartyMembers = CurrentConfig.MaxMembers;
@@ -579,7 +579,7 @@ void UPartyGameState::HandlePartyJoinRequestReceived(const FUniqueNetId& Recipie
 		ApprovalAction == EApprovalAction::EnqueueAndStartBeacon)
 	{
 		// Enqueue for a more opportune time
-		UE_LOG(LogParty, Verbose, TEXT("[%s] Enqueuing approval request for %s"), PartyInfo.IsValid() ? *PartyInfo->PartyId->ToString() : TEXT("INVALID"), *SenderId.ToString());
+		UE_LOG(LogParty, Verbose, TEXT("[%s] Enqueuing approval request for %s"), OssParty.IsValid() ? *OssParty->PartyId->ToString() : TEXT("INVALID"), *SenderId.ToString());
 		
 		FPendingMemberApproval PendingApproval;
 		PendingApproval.RecipientId.SetUniqueNetId(RecipientId.AsShared());
@@ -597,22 +597,34 @@ void UPartyGameState::HandlePartyJoinRequestReceived(const FUniqueNetId& Recipie
 		bool bApproveRequest = (ApprovalAction == EApprovalAction::Approve);
 
 		// Respond now
-		UE_LOG(LogParty, Verbose, TEXT("[%s] Responding to approval request for %s with %s"), PartyInfo.IsValid() ? *PartyInfo->PartyId->ToString() : TEXT("INVALID"), *SenderId.ToString(), bApproveRequest ? TEXT("approved") : TEXT("denied"));
+		UE_LOG(LogParty, Verbose, TEXT("[%s] Responding to approval request for %s with %s"), OssParty.IsValid() ? *OssParty->PartyId->ToString() : TEXT("INVALID"), *SenderId.ToString(), bApproveRequest ? TEXT("approved") : TEXT("denied"));
 		
 		UWorld* World = GetWorld();
 		IOnlinePartyPtr PartyInt = Online::GetPartyInterface(World);
 		if (PartyInt.IsValid())
 		{
-			PartyInt->ApproveJoinRequest(RecipientId, *PartyInfo->PartyId, SenderId, bApproveRequest, (int32)(!bApproveRequest ? DenialReason : EJoinPartyDenialReason::NoReason));
+			PartyInt->ApproveJoinRequest(RecipientId, *OssParty->PartyId, SenderId, bApproveRequest, (int32)(!bApproveRequest ? DenialReason : EJoinPartyDenialReason::NoReason));
 		}
 	}
 }
 
+FOnlinePartyTypeId UPartyGameState::GetPartyTypeId() const
+{
+	FOnlinePartyTypeId Result;
+
+	if (OssParty.IsValid())
+	{
+		Result = OssParty->PartyTypeId;
+	}
+
+	return Result;
+}
+
 TSharedPtr<const FOnlinePartyId> UPartyGameState::GetPartyId() const
 {
-	if (PartyInfo.IsValid())
+	if (OssParty.IsValid())
 	{
-		return PartyInfo->PartyId;
+		return OssParty->PartyId;
 	}
 
 	return nullptr;
@@ -629,27 +641,34 @@ void UPartyGameState::SetPartyType(EPartyType InPartyType, bool bLeaderFriendsOn
 		{
 			bool bIsPrivate = (InPartyType == EPartyType::Private);
 
-			CurrentConfig.Permissions = bIsPrivate ? EPartyPermissions::Private : EPartyPermissions::Public;
-			CurrentConfig.bIsAcceptingMembers = bIsPrivate ? false : CurrentConfig.bIsAcceptingMembers;
-			CurrentConfig.bShouldPublishToPresence = !bIsPrivate;
-			CurrentConfig.bCanNonLeaderPublishToPresence = !bIsPrivate && !bLeaderFriendsOnly;
-			CurrentConfig.bCanNonLeaderInviteToParty = !bLeaderInvitesOnly;
-			UpdatePartyConfig(bIsPrivate);
-
-			if (PartyStateRef->PartyType != InPartyType)
+			PartySystemPermissions::EPresencePermissions PresencePermissions;
+			if (bLeaderFriendsOnly)
 			{
-				OnPartyTypeChanged().Broadcast(InPartyType);
+				if (bIsPrivate)
+				{
+					PresencePermissions = PartySystemPermissions::FriendsInviteOnly;
+				}
+				else
+				{
+					PresencePermissions = PartySystemPermissions::FriendsOnly;
+				}
 			}
-
-			if (PartyStateRef->bLeaderFriendsOnly != bLeaderFriendsOnly)
+			else
 			{
-				OnLeaderFriendsOnlyChanged().Broadcast(bLeaderFriendsOnly);
+				if (bIsPrivate)
+				{
+					PresencePermissions = PartySystemPermissions::PublicInviteOnly;
+				}
+				else
+				{
+					PresencePermissions = PartySystemPermissions::Public;
+				}
 			}
 
-			if (PartyStateRef->bLeaderInvitesOnly != bLeaderInvitesOnly)
-			{
-				OnLeaderInvitesOnlyChanged().Broadcast(bLeaderInvitesOnly);
-			}
+			CurrentConfig.PresencePermissions = PresencePermissions;
+			CurrentConfig.InvitePermissions = bLeaderInvitesOnly ? PartySystemPermissions::EInvitePermissions::Leader : PartySystemPermissions::EInvitePermissions::Anyone;
+
+			UpdatePartyConfig(bIsPrivate);
 
 			PartyStateRef->PartyType = InPartyType;
 			PartyStateRef->bLeaderFriendsOnly = bLeaderFriendsOnly;
@@ -679,7 +698,7 @@ int32 UPartyGameState::GetPartySize() const
 
 void UPartyGameState::SetPartyMaxSize(int32 NewSize)
 {
-	if (PartyInfo.IsValid())
+	if (OssParty.IsValid())
 	{
 		if (CurrentConfig.MaxMembers != NewSize)
 		{
@@ -696,9 +715,9 @@ void UPartyGameState::SetPartyMaxSize(int32 NewSize)
 
 int32 UPartyGameState::GetPartyMaxSize() const
 {
-	if (PartyInfo.IsValid())
+	if (OssParty.IsValid())
 	{
-		return PartyInfo->Config->MaxMembers;
+		return OssParty->Config->MaxMembers;
 	}
 
 	// Invalid party info
@@ -753,7 +772,7 @@ void UPartyGameState::UpdateAcceptingMembers()
 
 void UPartyGameState::SetAcceptingMembers(bool bIsAcceptingMembers, EJoinPartyDenialReason DenialReason)
 {
-	if (PartyInfo.IsValid())
+	if (OssParty.IsValid())
 	{
 		if (*OwningUserId == *GetPartyLeader())
 		{
@@ -808,34 +827,14 @@ bool UPartyGameState::IsInJoinableGameState() const
 	return !bInGame || (bInGame && bGameJoinable);
 }
 
-bool UPartyGameState::GetJoinability(bool& bFriendJoinable, bool& bInviteOnly, bool& bAllowInvites) const
+bool UPartyGameState::CanInvite() const
 {
-	bool bValidSettings = false;
-	if (PartyInfo.IsValid())
-	{
-		bValidSettings = true;
-
-		// In joinable game or no game at all
-		bool bGameJoinable = IsInJoinableGameState();
-
-		bFriendJoinable = false;
-		bInviteOnly = false;
-		bAllowInvites = false;
-
-		// If the game is joinable (ie not full in general), then party rules apply
-		if (bGameJoinable)
-		{
-			bool bPublicJoinable = false;
-			bValidSettings = PartyInfo->GetJoinability(*OwningUserId, bPublicJoinable, bFriendJoinable, bInviteOnly, bAllowInvites);
-		}
-	}
-
-	return bValidSettings;
+	return OssParty->CanLocalUserInvite(*OwningUserId);
 }
 
 void UPartyGameState::UpdatePartyConfig(bool bResetAccessKey)
 {
-	if (PartyInfo.IsValid())
+	if (OssParty.IsValid())
 	{
 		if (*OwningUserId == *GetPartyLeader())
 		{
@@ -845,14 +844,14 @@ void UPartyGameState::UpdatePartyConfig(bool bResetAccessKey)
 			{
 				FOnUpdatePartyComplete CompletionDelegate;
 				CompletionDelegate.BindUObject(this, &ThisClass::OnUpdatePartyConfigComplete);
-				if (!PartyInt->UpdateParty(*OwningUserId, *PartyInfo->PartyId, CurrentConfig, bResetAccessKey, CompletionDelegate))
+				if (!PartyInt->UpdateParty(*OwningUserId, *OssParty->PartyId, CurrentConfig, bResetAccessKey, CompletionDelegate))
 				{
-					UE_LOG(LogParty, Warning, TEXT("[%s] Failed to update party"), PartyInfo.IsValid() ? *PartyInfo->PartyId->ToString() : TEXT("INVALID"));
+					UE_LOG(LogParty, Warning, TEXT("[%s] Failed to update party"), OssParty.IsValid() ? *OssParty->PartyId->ToString() : TEXT("INVALID"));
 				}
 			}
 			else
 			{
-				UE_LOG(LogParty, Warning, TEXT("[%s] Invalid party interface updating party size"), PartyInfo.IsValid() ? *PartyInfo->PartyId->ToString() : TEXT("INVALID"));
+				UE_LOG(LogParty, Warning, TEXT("[%s] Invalid party interface updating party size"), OssParty.IsValid() ? *OssParty->PartyId->ToString() : TEXT("INVALID"));
 			}
 		}
 	}
@@ -862,13 +861,15 @@ void UPartyGameState::UpdatePartyConfig(bool bResetAccessKey)
 	}
 }
 
-void UPartyGameState::OnUpdatePartyConfigComplete(const FUniqueNetId& LocalUserId, bool bWasSuccessful, const FOnlinePartyId& InPartyId, const FString& Error)
+void UPartyGameState::OnUpdatePartyConfigComplete(const FUniqueNetId& LocalUserId, const FOnlinePartyId& InPartyId, const EUpdateConfigCompletionResult Result)
 {
-	UE_LOG(LogParty, Verbose, TEXT("[%s] Party config updated %d %s"), *InPartyId.ToString(), bWasSuccessful, *Error);
+	FString PartyIdDebugString = InPartyId.ToDebugString();
 
-	if (PartyInfo.IsValid())
+	UE_LOG(LogParty, Verbose, TEXT("[%s] Party config updated %d %s"), *PartyIdDebugString, Result == EUpdateConfigCompletionResult::Succeeded, ToString(Result));
+
+	if (OssParty.IsValid())
 	{
-		CurrentConfig = *PartyInfo->Config;
+		CurrentConfig = *OssParty->Config;
 		bDebugAcceptingMembers = CurrentConfig.bIsAcceptingMembers;
 	}
 }
@@ -877,9 +878,9 @@ TSharedPtr<const FUniqueNetId> UPartyGameState::GetPartyLeader() const
 {
 	TSharedPtr<const FUniqueNetId> ReturnValue;
 
-	if (PartyInfo.IsValid())
+	if (OssParty.IsValid())
 	{
-		ReturnValue = PartyInfo->GetLeaderId();
+		ReturnValue = OssParty->LeaderId;
 	}
 
 	return ReturnValue;
@@ -901,9 +902,14 @@ void UPartyGameState::GetAllPartyMembers(TArray<UPartyMemberState*>& PartyMember
 	PartyMembersState.GenerateValueArray(PartyMembers);
 }
 
-void UPartyGameState::OnPartyMemberPromoted(const FUniqueNetId& LocalUserId, bool bWasSuccessful, const FOnlinePartyId& InPartyId, const FUniqueNetId& InMemberId, const FString& Error)
+void UPartyGameState::OnPartyMemberPromoted(const FUniqueNetId& LocalUserId, const FOnlinePartyId& InPartyId, const FUniqueNetId& InMemberId, const EPromoteMemberCompletionResult Result)
 {
-	UE_LOG(LogParty, Verbose, TEXT("[%s] Player %s promoted %s %d %s"), *InPartyId.ToString(), *LocalUserId.ToString(), *InMemberId.ToString(), bWasSuccessful, *Error);
+	FString PartyIdDebugString = InPartyId.ToDebugString();
+	FString MemberIdDebugString = InMemberId.ToDebugString();
+
+	bool bWasSuccessful = Result == EPromoteMemberCompletionResult::Succeeded;
+
+	UE_LOG(LogParty, Verbose, TEXT("[%s] Player %s promoted %s %d %s"), *PartyIdDebugString, *LocalUserId.ToString(), *MemberIdDebugString, bWasSuccessful, ToString(Result));
 }
 
 void UPartyGameState::PromoteMember(const FUniqueNetIdRepl& NewPartyLeader)
@@ -917,20 +923,24 @@ void UPartyGameState::PromoteMember(const FUniqueNetIdRepl& NewPartyLeader)
 			{
 				FOnPromotePartyMemberComplete CompletionDelegate;
 				CompletionDelegate.BindUObject(this, &ThisClass::OnPartyMemberPromoted);
-				PartyInt->PromoteMember(*OwningUserId, *PartyInfo->PartyId, *NewPartyLeader, CompletionDelegate);
+				PartyInt->PromoteMember(*OwningUserId, *OssParty->PartyId, *NewPartyLeader, CompletionDelegate);
 			}
 		}
 	}
 	else
 	{
-		UE_LOG(LogParty, Verbose, TEXT("[%s] Promote member feature locked out."), PartyInfo.IsValid() ? *PartyInfo->PartyId->ToString() : TEXT("INVALID"));
+		UE_LOG(LogParty, Verbose, TEXT("[%s] Promote member feature locked out."), OssParty.IsValid() ? *OssParty->PartyId->ToString() : TEXT("INVALID"));
 	}
 }
 
 void UPartyGameState::OnPartyMemberKicked(const FUniqueNetId& LocalUserId, const FOnlinePartyId& InPartyId, const FUniqueNetId& InMemberId, const EKickMemberCompletionResult Result)
 {
+	FString PartyIdDebugString = InPartyId.ToDebugString();
+	FString MemberIdDebugString = InMemberId.ToDebugString();
+
 	bool bWasSuccessful = Result == EKickMemberCompletionResult::Succeeded;
-	UE_LOG(LogParty, Verbose, TEXT("[%s] Player %s kicked %s %d %s"), *InPartyId.ToString(), *LocalUserId.ToString(), *InMemberId.ToString(), bWasSuccessful, ToString(Result));
+
+	UE_LOG(LogParty, Verbose, TEXT("[%s] Player %s kicked %s %d %s"), *PartyIdDebugString, *LocalUserId.ToString(), *MemberIdDebugString, bWasSuccessful, ToString(Result));
 }
 
 void UPartyGameState::KickMember(const FUniqueNetIdRepl& PartyMemberToKick)
@@ -942,7 +952,7 @@ void UPartyGameState::KickMember(const FUniqueNetIdRepl& PartyMemberToKick)
 		{
 			FOnKickPartyMemberComplete CompletionDelegate;
 			CompletionDelegate.BindUObject(this, &ThisClass::OnPartyMemberKicked);
-			PartyInt->KickMember(*OwningUserId, *PartyInfo->PartyId, *PartyMemberToKick, CompletionDelegate);
+			PartyInt->KickMember(*OwningUserId, *OssParty->PartyId, *PartyMemberToKick, CompletionDelegate);
 		}
 	}
 }
@@ -997,10 +1007,10 @@ void UPartyGameState::UpdatePartyData(const FUniqueNetIdRepl& InLocalUserId)
 	IOnlinePartyPtr PartyInt = Online::GetPartyInterface(World);
 	if (PartyInt.IsValid())
 	{
-		FOnlinePartyAttrs PartyDataAttrs;
-		if (FVariantDataConverter::UStructToVariantMap(PartyStateRefDef, PartyStateRef, PartyDataAttrs, 0, CPF_Transient | CPF_RepSkip))
+		FOnlinePartyData PartyData;
+		if (FVariantDataConverter::UStructToVariantMap(PartyStateRefDef, PartyStateRef, PartyData.KeyValAttrs, 0, CPF_Transient | CPF_RepSkip))
 		{
-			PartyInt->UpdatePartyData(*InLocalUserId, *PartyInfo->PartyId, PartyDataAttrs);
+			PartyInt->UpdatePartyData(*InLocalUserId, *OssParty->PartyId, PartyData);
 		}
 	}
 }
@@ -1011,10 +1021,10 @@ void UPartyGameState::UpdatePartyMemberState(const FUniqueNetIdRepl& InLocalUser
 	IOnlinePartyPtr PartyInt = Online::GetPartyInterface(World);
 	if (PartyInt.IsValid())
 	{
-		FOnlinePartyAttrs PartyMemberStateAttrs;
-		if (FVariantDataConverter::UStructToVariantMap(InPartyMemberState->MemberStateRefDef, InPartyMemberState->MemberStateRef, PartyMemberStateAttrs, 0, CPF_Transient | CPF_RepSkip))
+		FOnlinePartyData PartyMemberData;
+		if (FVariantDataConverter::UStructToVariantMap(InPartyMemberState->MemberStateRefDef, InPartyMemberState->MemberStateRef, PartyMemberData.KeyValAttrs, 0, CPF_Transient | CPF_RepSkip))
 		{
-			PartyInt->UpdatePartyMemberData(*InLocalUserId, *PartyInfo->PartyId, PartyMemberStateAttrs);
+			PartyInt->UpdatePartyMemberData(*InLocalUserId, *OssParty->PartyId, PartyMemberData);
 		}
 	}
 }
@@ -1023,7 +1033,7 @@ void UPartyGameState::ConnectToReservationBeacon()
 {
 	UWorld* World = GetWorld();
 
-	if (PartyInfo.IsValid())
+	if (OssParty.IsValid())
 	{
 		FUniqueNetIdRepl PartyLeader(GetPartyLeader());
 		if (PartyLeader.IsValid() && *OwningUserId == *PartyLeader)
@@ -1062,7 +1072,7 @@ void UPartyGameState::RejectAllPendingJoinRequests()
 {
 	UWorld* World = GetWorld();
 	IOnlinePartyPtr PartyInt = Online::GetPartyInterface(World);
-	bool bValidInterface = PartyInt.IsValid() && PartyInfo.IsValid();
+	bool bValidInterface = PartyInt.IsValid() && OssParty.IsValid();
 
 	bool bEmpty = false;
 	FPendingMemberApproval PendingApproval;
@@ -1072,8 +1082,8 @@ void UPartyGameState::RejectAllPendingJoinRequests()
 		bEmpty = !PendingApprovals.Dequeue(PendingApproval);
 		if (!bEmpty && bValidInterface)
 		{
-			UE_LOG(LogParty, Verbose, TEXT("[%s] Responding to approval request for %s with denied"), PartyInfo.IsValid() ? *PartyInfo->PartyId->ToString() : TEXT("INVALID"), *PendingApproval.SenderId.ToString());
-			PartyInt->ApproveJoinRequest(*PendingApproval.RecipientId, *PartyInfo->PartyId, *PendingApproval.SenderId, false, (int32)EJoinPartyDenialReason::Busy);
+			UE_LOG(LogParty, Verbose, TEXT("[%s] Responding to approval request for %s with denied"), OssParty.IsValid() ? *OssParty->PartyId->ToString() : TEXT("INVALID"), *PendingApproval.SenderId.ToString());
+			PartyInt->ApproveJoinRequest(*PendingApproval.RecipientId, *OssParty->PartyId, *PendingApproval.SenderId, false, (int32)EJoinPartyDenialReason::Busy);
 		}
 	} while (!bEmpty);
 }
@@ -1096,7 +1106,7 @@ void UPartyGameState::OnReservationBeaconUpdateResponseReceived(EPartyReservatio
 	{
 		UWorld* World = GetWorld();
 		IOnlinePartyPtr PartyInt = Online::GetPartyInterface(World);
-		bool bValidInterface = PartyInt.IsValid() && PartyInfo.IsValid();
+		bool bValidInterface = PartyInt.IsValid() && OssParty.IsValid();
 
 		// There should be at least the one
 		FPendingMemberApproval PendingApproval;
@@ -1104,8 +1114,8 @@ void UPartyGameState::OnReservationBeaconUpdateResponseReceived(EPartyReservatio
 		{
 			if (bValidInterface)
 			{
-				UE_LOG(LogParty, Verbose, TEXT("[%s] Responding to approval request for %s with approved"), PartyInfo.IsValid() ? *PartyInfo->PartyId->ToString() : TEXT("INVALID"), *PendingApproval.SenderId.ToString());
-				PartyInt->ApproveJoinRequest(*PendingApproval.RecipientId, *PartyInfo->PartyId, *PendingApproval.SenderId, true);
+				UE_LOG(LogParty, Verbose, TEXT("[%s] Responding to approval request for %s with approved"), OssParty.IsValid() ? *OssParty->PartyId->ToString() : TEXT("INVALID"), *PendingApproval.SenderId.ToString());
+				PartyInt->ApproveJoinRequest(*PendingApproval.RecipientId, *OssParty->PartyId, *PendingApproval.SenderId, true);
 			}
 		}
 
@@ -1173,4 +1183,4 @@ UWorld* UPartyGameState::GetWorld() const
 UParty* UPartyGameState::GetPartyOuter() const
 {
 	return GetTypedOuter<UParty>();
-}
\ No newline at end of file
+}
diff --git a/Engine/Source/Runtime/Online/OnlineFramework/Party/Public/Party.h b/Engine/Source/Runtime/Online/OnlineFramework/Party/Public/Party.h
--- a/Engine/Source/Runtime/Online/OnlineFramework/Party/Public/Party.h
+++ b/Engine/Source/Runtime/Online/OnlineFramework/Party/Public/Party.h
@@ -8,7 +8,6 @@
 
 class UWorld;
 class FUniqueNetId;
-typedef FUniqueNetId FOnlinePartyId;
 typedef FUniqueNetIdRepl FOnlinePartyIdRepl;
 class UPartyGameState;
 class UGameInstance;
@@ -16,15 +15,6 @@ class IOnlinePartyJoinInfo;
 
 enum class EPartyType : uint8;
 
-#define PARTY_CONFIG_TYPE TEXT("PARTY_CONFIG_TYPE")
-enum class EPartyConfigType : uint8
-{
-	// Social "persistent" party
-	PersistentParty = 0,
-	// Transient team association
-	TeamParty
-};
-
 /**
  * Holds the basic information needed to join a party
  */
@@ -36,20 +26,25 @@ struct FPartyDetails
 
 	bool IsValid() const
 	{
-		return !PartyJoinInfo->IsInvalidForJoin();
+		return PartyJoinInfo->IsValid();
 	}
 
-	inline const TSharedRef<const FOnlinePartyId>& GetPartyId() const
+	const TSharedRef<const FOnlinePartyId>& GetPartyId() const
 	{
 		return PartyJoinInfo->GetPartyId();
 	}
 
-	inline const TSharedRef<const FUniqueNetId>& GetLeaderId() const
+	const FOnlinePartyTypeId GetPartyTypeId() const
+	{
+		return PartyJoinInfo->GetPartyTypeId();
+	}
+
+	const TSharedRef<const FUniqueNetId>& GetLeaderId() const
 	{
 		return PartyJoinInfo->GetLeaderId();
 	}
 
-	inline const FString& GetAppId() const
+	const FString& GetAppId() const
 	{
 		return PartyJoinInfo->GetAppId();
 	}
@@ -67,6 +62,35 @@ struct FPartyDetails
 	bool bAcceptInvite;
 };
 
+///////////////////////////////////////////////////////////////////
+// Completion delegates
+///////////////////////////////////////////////////////////////////
+namespace UPartyDelegates
+{
+	/**
+	 * Party creation async task completed callback
+	 *
+	 * @param LocalUserId - id of user that initiated the request
+	 * @param Result - result of the operation
+	 */
+	DECLARE_DELEGATE_TwoParams(FOnCreateUPartyComplete, const FUniqueNetId& /*LocalUserId*/, const ECreatePartyCompletionResult /*Result*/);
+	/**
+	 * Party join async task completed callback
+	 *
+	 * @param LocalUserId - id of user that initiated the request
+	 * @param Result - result of the operation
+	 * @param NotApprovedReason - client defined value describing why you were not approved
+	 */
+	DECLARE_DELEGATE_ThreeParams(FOnJoinUPartyComplete, const FUniqueNetId& /*LocalUserId*/, const EJoinPartyCompletionResult /*Result*/, const int32 /*NotApprovedReason*/);
+	/**
+	 * Party leave async task completed callback
+	 *
+	 * @param LocalUserId - id of user that initiated the request
+	 * @param Result - result of the operation
+	 */
+	DECLARE_DELEGATE_TwoParams(FOnLeaveUPartyComplete, const FUniqueNetId& /*LocalUserId*/, const ELeavePartyCompletionResult /*Result*/);
+}
+
 /**
  * High level singleton for the management of parties, all parties are contained within
  */
@@ -139,10 +163,11 @@ public:
 	 * Create a generic party
 	 *
 	 * @param InUserId user creating the party (should be primary player)
+	 * @param InPartyTypeId desired party type
 	 * @param InPartyConfig desired party configuration
 	 * @param InCompletionDelegate delegate called upon completion
 	 */
-	void CreateParty(const FUniqueNetId& InUserId, const FPartyConfiguration& InPartyConfig, const FOnCreatePartyComplete& InCompletionDelegate);
+	void CreateParty(const FUniqueNetId& InUserId, const FOnlinePartyTypeId InPartyTypeId, const FPartyConfiguration& InPartyConfig, const UPartyDelegates::FOnCreateUPartyComplete& InCompletionDelegate);
 
 	/**
 	 * Join a generic party
@@ -151,7 +176,7 @@ public:
 	 * @param InPartyDetails credentials for the party to join
 	 * @param InCompletionDelegate delegate called upon completion
 	 */
-	void JoinParty(const FUniqueNetId& InUserId, const FPartyDetails& InPartyDetails, const FOnJoinPartyComplete& InCompletionDelegate);
+	void JoinParty(const FUniqueNetId& InUserId, const FPartyDetails& InPartyDetails, const UPartyDelegates::FOnJoinUPartyComplete& InCompletionDelegate);
 
 	/**
 	 * Leave a generic party
@@ -160,7 +185,7 @@ public:
 	 * @param InPartyId id of the party to leave
 	 * @param InCompletionDelegate delegate called upon completion
 	 */
-	void LeaveParty(const FUniqueNetId& InUserId, const FOnlinePartyId& InPartyId, const FOnLeavePartyComplete& InCompletionDelegate);
+	void LeaveParty(const FUniqueNetId& InUserId, const FOnlinePartyTypeId InPartyTypeId, const UPartyDelegates::FOnLeaveUPartyComplete& InCompletionDelegate);
 	
 	/**
 	 * Create a persistent party (party that will hold users for the duration of app)
@@ -168,7 +193,7 @@ public:
 	 * @param InUserId user creating the party (should be primary player)
 	 * @param InCompletionDelegate delegate called upon completion
 	 */
-	virtual void CreatePersistentParty(const FUniqueNetId& InUserId, const FOnCreatePartyComplete& InCompletionDelegate);
+	virtual void CreatePersistentParty(const FUniqueNetId& InUserId, const UPartyDelegates::FOnCreateUPartyComplete& InCompletionDelegate = UPartyDelegates::FOnCreateUPartyComplete());
 
 	/**
 	 * Join a persistent party (party that will hold users for the duration of app)
@@ -177,7 +202,7 @@ public:
 	 * @param InPartyDetails credentials for the party to join
 	 * @param InCompletionDelegate delegate called upon completion
 	 */
-	virtual void JoinPersistentParty(const FUniqueNetId& InUserId, const FPartyDetails& InPartyDetails, const FOnJoinPartyComplete& InCompletionDelegate);
+	virtual void JoinPersistentParty(const FUniqueNetId& InUserId, const FPartyDetails& InPartyDetails, const UPartyDelegates::FOnJoinUPartyComplete& InCompletionDelegate = UPartyDelegates::FOnJoinUPartyComplete());
 
 	/**
 	 * Leave a persistent party, a new one will need to be created for the given user afterward
@@ -185,7 +210,7 @@ public:
 	 * @param InUserId user creating the party (should be primary player)
 	 * @param InCompletionDelegate delegate called upon completion
 	 */
-	virtual void LeavePersistentParty(const FUniqueNetId& InUserId, const FOnLeavePartyComplete& InCompletionDelegate);
+	virtual void LeavePersistentParty(const FUniqueNetId& InUserId, const UPartyDelegates::FOnLeaveUPartyComplete& InCompletionDelegate = UPartyDelegates::FOnLeaveUPartyComplete());
 
 	/**
 	 * Reestablish all party state and information upon returning to the main menu
@@ -209,7 +234,7 @@ public:
 	 *
 	 * @return true if a pending invite is available and a join will be attempted, false otherwise
 	 */
-	bool ProcessPendingPartyInvites();
+	bool ProcessPendingPartyJoin();
 
 	/**
 	 * Take ownership of a pending invite and do what is necessary to get the game to a point where it can be used
@@ -218,7 +243,8 @@ public:
 	 * @param PartyDetails party details used for the join
 	 * @param JoinCompleteDelegate delegate to call once the pending invite has been joined or in all failure cases
 	 */
-	virtual void HandlePendingInvite(const FUniqueNetId& LocalUserId, const FPartyDetails& PartyDetails, const FOnJoinPartyComplete& JoinCompleteDelegate);
+	void AddPendingPartyJoin(const FUniqueNetId& LocalUserId, const FPartyDetails& PartyDetails, const UPartyDelegates::FOnJoinUPartyComplete& JoinCompleteDelegate);
+	//void AddPendingPartyJoin(const FUniqueNetId& LocalUserId, const FPartyDetails& PartyDetails, const FOnJoinPartyComplete& JoinCompleteDelegate);
 
 	/**
 	 * Is any local player in the given party
@@ -255,6 +281,9 @@ public:
 	/** @return the party state for a given party */
 	UPartyGameState* GetParty(const FOnlinePartyId& InPartyId) const;
 
+	/** @return the party state for a given party type */
+	UPartyGameState* GetParty(const FOnlinePartyTypeId InPartyTypeId) const;
+
 	/** @return the party state for the persistent party */
 	UPartyGameState* GetPersistentParty() const;
 
@@ -267,33 +296,37 @@ public:
 
 protected:
 
-	/** Pending party invite */
-	struct FPendingPartyInvite
+	/** Pending party join */
+	struct FPendingPartyJoin : public TSharedFromThis<FPendingPartyJoin>
 	{
-		FPendingPartyInvite(
-			TSharedRef<const FUniqueNetId> InLocalUserId,
-			const FPartyDetails& InPartyDetails,
-			const FOnJoinPartyComplete& InJoinCompleteDelegate)
+		FPendingPartyJoin(TSharedRef<const FUniqueNetId> InLocalUserId, const FPartyDetails& InPartyDetails, const UPartyDelegates::FOnJoinUPartyComplete& InDelegate)
 			: LocalUserId(InLocalUserId)
 			, PartyDetails(InPartyDetails)
-			, JoinCompleteDelegate(InJoinCompleteDelegate)
-		{}
+			, Delegate(InDelegate)
+		{
+			
+			ensure(LocalUserId->IsValid());
+		}
 
 		/** user that sent the invite */
 		TSharedRef<const FUniqueNetId> LocalUserId;
 		/** details about party to join */
 		FPartyDetails PartyDetails;
-		/** delegate to call when the join for this party invite has been completed */
-		FOnJoinPartyComplete JoinCompleteDelegate;
-		/** @return true if info is available to attempt a game join */
-		inline bool IsValid() const
-		{
-			return PartyDetails.IsValid();
-		}
+
+		UPartyDelegates::FOnJoinUPartyComplete Delegate;
 	};
 
 	/** game invite info that is available to join via an accepted game invite */
-	TSharedPtr<FPendingPartyInvite> PendingPartyInvite;
+	TSharedPtr<FPendingPartyJoin> PendingPartyJoin;
+
+	/**
+	* Take ownership of a pending invite and do what is necessary to get the game to a point where it can be used
+	*
+	* @param LocalUserId who the invite is to
+	* @param PartyDetails party details used for the join
+	* @param JoinCompleteDelegate delegate to call once the pending invite has been joined or in all failure cases
+	*/
+	virtual void HandlePendingJoin() PURE_VIRTUAL(UParty::HandlePendingJoin, )
 
 	/**
 	 * Party configuration
@@ -306,7 +339,7 @@ public:
 protected:
 
 	/** Mapping of party types to party classes, up to game to define relationship */
-	TMap<int32, TSubclassOf<UPartyGameState>> PartyClasses;
+	TMap<FOnlinePartyTypeId, TSubclassOf<UPartyGameState>> PartyClasses;
 
 	UPROPERTY(Config)
 	int32 DefaultMaxPartySize;
@@ -337,11 +370,11 @@ private:
 	UPROPERTY()
 	bool bLeavingPersistentParty;
 	/** Array of leave persistent party delegates gathered while already leaving a persistent party */
-	TArray<FOnLeavePartyComplete> LeavePartyCompleteDelegates;
+	TArray<UPartyDelegates::FOnLeaveUPartyComplete> LeavePartyCompleteDelegates;
 	/** Id of the current persistent party, only one and it should always be valid except while transferring between parties */
 	TSharedPtr<const FOnlinePartyId> PersistentPartyId;
 	/** Mapping of all known joined parties */
-	TMap<FString, UPartyGameState*> JoinedParties;
+	TMap<FOnlinePartyTypeId, UPartyGameState*> JoinedParties;
 
 	/** Delegate when parties are joined */
 	FOnPartyUpdate PartyJoinedDelegates;
@@ -394,24 +427,35 @@ private:
 	/**
 	 * Common calls for managing parties
 	 */
-	void CreatePartyInternal(const FUniqueNetId& InUserId, const FPartyConfiguration& InPartyConfig, const FOnCreatePartyComplete& InCompletionDelegate);
-	void JoinPartyInternal(const FUniqueNetId& InUserId, const FPartyDetails& InPartyDetails, const FOnJoinPartyComplete& InCompletionDelegate);
-	void LeavePartyInternal(const FUniqueNetId& InUserId, const FOnlinePartyId& InPartyId, const FOnLeavePartyComplete& InCompletionDelegate);
+	void CreatePartyInternal(const FUniqueNetId& InUserId, const FOnlinePartyTypeId InPartyTypeId, const FPartyConfiguration& InPartyConfig, const UPartyDelegates::FOnCreateUPartyComplete& InCompletionDelegate);
+	void JoinPartyInternal  (const FUniqueNetId& InUserId, const FPartyDetails& InPartyDetails, const UPartyDelegates::FOnJoinUPartyComplete& InCompletionDelegate);
+	void LeavePartyInternal(const FUniqueNetId& InUserId, const FOnlinePartyTypeId InPartyTypeId, const UPartyDelegates::FOnLeaveUPartyComplete& InCompletionDelegate);
 
 	/**
 	 * Delegates for managing parties
 	 */
-	void OnCreatePartyComplete(const FUniqueNetId& LocalUserId, const FOnlinePartyId& InPartyId, const ECreatePartyCompletionResult Result, FOnCreatePartyComplete CompletionDelegate);
-	void OnJoinPartyComplete(const FUniqueNetId& LocalUserId, const FOnlinePartyId& InPartyId, const EJoinPartyCompletionResult Result, int32 DeniedResultCode, FOnJoinPartyComplete CompletionDelegate);
-	void OnLeavePartyComplete(const FUniqueNetId& LocalUserId, const FOnlinePartyId& InPartyId, const ELeavePartyCompletionResult Result, FOnLeavePartyComplete CompletionDelegate);
+	void OnCreatePartyComplete(const FUniqueNetId& LocalUserId, const ECreatePartyCompletionResult Result, const FOnlinePartyTypeId InPartyTypeId, UPartyDelegates::FOnCreateUPartyComplete CompletionDelegate);
+	void OnJoinPartyComplete(const FUniqueNetId& LocalUserId, const EJoinPartyCompletionResult Result, int32 DeniedResultCode, const FOnlinePartyTypeId InPartyTypeId, UPartyDelegates::FOnJoinUPartyComplete CompletionDelegate);
+	void OnLeavePartyComplete(const FUniqueNetId& LocalUserId, const ELeavePartyCompletionResult Result, const FOnlinePartyTypeId InPartyTypeId, UPartyDelegates::FOnLeaveUPartyComplete CompletionDelegate);
 	void HandleJoinPersistentPartyFailure();
 
+	void OnCreatePartyInternalComplete(const FUniqueNetId& LocalUserId, const TSharedPtr<const FOnlinePartyId>& InPartyId, const ECreatePartyCompletionResult Result, const FOnlinePartyTypeId InPartyTypeId, UPartyDelegates::FOnCreateUPartyComplete CompletionDelegate);
+	void OnJoinPartyInternalComplete(const FUniqueNetId& LocalUserId, const FOnlinePartyId& InPartyId, const EJoinPartyCompletionResult Result, int32 DeniedResultCode, const FOnlinePartyTypeId InPartyTypeId, UPartyDelegates::FOnJoinUPartyComplete CompletionDelegate);
+	void OnLeavePartyInternalComplete(const FUniqueNetId& LocalUserId, const FOnlinePartyId& InPartyId, const ELeavePartyCompletionResult Result, const FOnlinePartyTypeId InPartyTypeId, UPartyDelegates::FOnLeaveUPartyComplete CompletionDelegate);
+
+	/**
+	 * Internal helpers for delegates
+	 */
+	enum class EOnCreatePartyCompleteInternalResult OnCreatePartyCompleteInternal(const FUniqueNetId& LocalUserId, const FOnlinePartyTypeId InPartyTypeId);
+	enum class EOnJoinPartyCompleteInternalResult   OnJoinPartyCompleteInternal(const FUniqueNetId& LocalUserId, const FOnlinePartyTypeId InPartyTypeId);
+	enum class EOnLeavePartyCompleteInternalResult  OnLeavePartyCompleteInternal(const FUniqueNetId& LocalUserId, const FOnlinePartyTypeId InPartyTypeId);
+
 	/**
 	 * Delegates for managing persistent parties
 	 */
-	void OnCreatePersistentPartyComplete(const FUniqueNetId& LocalUserId, const FOnlinePartyId& InPartyId, const ECreatePartyCompletionResult Result, FOnCreatePartyComplete CompletionDelegate);
-	void OnJoinPersistentPartyComplete(const FUniqueNetId& LocalUserId, const FOnlinePartyId& InPartyId, const EJoinPartyCompletionResult Result, int32 DeniedResultCode, FOnJoinPartyComplete CompletionDelegate);
-	void OnLeavePersistentPartyComplete(const FUniqueNetId& LocalUserId, const FOnlinePartyId& InPartyId, const ELeavePartyCompletionResult Result, FOnLeavePartyComplete CompletionDelegate);
+	void OnCreatePersistentPartyComplete(const FUniqueNetId& LocalUserId, const ECreatePartyCompletionResult Result, UPartyDelegates::FOnCreateUPartyComplete CompletionDelegate);
+	void OnJoinPersistentPartyComplete(const FUniqueNetId& LocalUserId, const EJoinPartyCompletionResult Result, int32 DeniedResultCode, UPartyDelegates::FOnJoinUPartyComplete CompletionDelegate);
+	void OnLeavePersistentPartyComplete(const FUniqueNetId& LocalUserId, const ELeavePartyCompletionResult Result, UPartyDelegates::FOnLeaveUPartyComplete CompletionDelegate);
 
 	/**
 	 * Party state management delegates 
@@ -435,7 +479,7 @@ private:
 	 * Error handling to restore persistent party after leaving a different one
 	 */
 	void LeaveAndRestorePersistentPartyInternal();
-	void OnLeavePersistentPartyAndRestore(const FUniqueNetId& LocalUserId, const FOnlinePartyId& PartyId, const ELeavePartyCompletionResult Result);
+	void OnLeavePersistentPartyAndRestore(const FUniqueNetId& LocalUserId, const ELeavePartyCompletionResult Result);
 
 	/**
 	 * Delegate handles
@@ -470,3 +514,100 @@ protected:
 
 };
 
+enum class EOnCreatePartyCompleteInternalResult
+{
+	UnknownFailure,
+	InvalidSystemInterface,
+	InvalidPartyId,
+	InvalidParty,
+	InvalidPartyType,
+	Success
+};
+
+enum class EOnJoinPartyCompleteInternalResult
+{
+	UnknownFailure,
+	InvalidSystemInterface,
+	InvalidParty,
+	InvalidPartyType,
+	Success
+};
+
+enum class EOnLeavePartyCompleteInternalResult
+{
+	UnknownFailure,
+	InvalidPartyState,
+	Success
+};
+
+inline const TCHAR* ToString(EOnCreatePartyCompleteInternalResult Value)
+{
+	switch (Value)
+	{
+	case EOnCreatePartyCompleteInternalResult::UnknownFailure:
+	{
+		return TEXT("UnknownFailure");
+	}
+	case EOnCreatePartyCompleteInternalResult::InvalidSystemInterface:
+	{
+		return TEXT("InvalidSystemInterface");
+	}
+	case EOnCreatePartyCompleteInternalResult::InvalidPartyId:
+	{
+		return TEXT("InvalidPartyId");
+	}
+	case EOnCreatePartyCompleteInternalResult::InvalidParty:
+	{
+		return TEXT("InvalidParty");
+	}
+	case EOnCreatePartyCompleteInternalResult::Success:
+	{
+		return TEXT("Success");
+	}
+	}
+	return TEXT("");
+}
+
+inline const TCHAR* ToString(EOnJoinPartyCompleteInternalResult Value)
+{
+	switch (Value)
+	{
+	case EOnJoinPartyCompleteInternalResult::UnknownFailure:
+	{
+		return TEXT("UnknownFailure");
+	}
+	case EOnJoinPartyCompleteInternalResult::InvalidSystemInterface:
+	{
+		return TEXT("InvalidSystemInterface");
+	}
+	case EOnJoinPartyCompleteInternalResult::InvalidParty:
+	{
+		return TEXT("InvalidParty");
+	}
+	case EOnJoinPartyCompleteInternalResult::Success:
+	{
+		return TEXT("Success");
+	}
+	}
+	return TEXT("");
+}
+
+inline const TCHAR* ToString(EOnLeavePartyCompleteInternalResult Value)
+{
+	switch (Value)
+	{
+	case EOnLeavePartyCompleteInternalResult::UnknownFailure:
+	{
+		return TEXT("UnknownFailure");
+	}
+	case EOnLeavePartyCompleteInternalResult::InvalidPartyState:
+	{
+		return TEXT("InvalidPartyState");
+	}
+	case EOnLeavePartyCompleteInternalResult::Success:
+	{
+		return TEXT("Success");
+	}
+	}
+	return TEXT("");
+}
