Mergine changes from //depot/UE4 to //UE4/Main

diff --git a/Engine/Source/Runtime/Online/OnlineFramework/Party/Public/PartyGameState.h b/Engine/Source/Runtime/Online/OnlineFramework/Party/Public/PartyGameState.h
--- a/Engine/Source/Runtime/Online/OnlineFramework/Party/Public/PartyGameState.h
+++ b/Engine/Source/Runtime/Online/OnlineFramework/Party/Public/PartyGameState.h
@@ -8,10 +8,9 @@
 
 class UWorld;
 class FUniqueNetId;
-class FOnlinePartyInfo;
+class FOnlineParty;
 class FOnlinePartyData;
 struct FPartyConfiguration;
-typedef FUniqueNetId FOnlinePartyId;
 typedef FUniqueNetIdRepl FOnlinePartyIdRepl;
 class FOnlineSessionSearchResult;
 class UParty;
@@ -107,9 +106,9 @@ class PARTY_API UPartyGameState : public UObject
 {
 	GENERATED_UCLASS_BODY()
 
-	//~ Begin UObject Interface
+	// Begin UObject Interface
 	virtual void BeginDestroy() override;
-	//~ End UObject Interface
+	// End UObject Interface
 
 private:
 
@@ -133,21 +132,21 @@ private:
 	 * 
 	 * @param NewPartyType new party type
 	 */
-	DECLARE_MULTICAST_DELEGATE_OneParam(FOnPartyTypeChanged, EPartyType /* NewPartyType */);
+	DECLARE_MULTICAST_DELEGATE_OneParam(FOnClientPartyTypeChanged, EPartyType /* NewPartyType */);
 
 	/**
 	 * Delegate fired when a join via presence permissions change
 	 * 
 	 * @param bLeaderFriendsOnly only leader friends can join
 	 */
-	DECLARE_MULTICAST_DELEGATE_OneParam(FOnLeaderFriendsOnlyChanged, bool /* bLeaderFriendsOnly */);
+	DECLARE_MULTICAST_DELEGATE_OneParam(FOnClientLeaderFriendsOnlyChanged, bool /* bLeaderFriendsOnly */);
 
 	/**
 	 * Delegate fired when a invite permissions change
 	 * 
 	 * @param bLeaderInviteOnly only leader invites are allowed
 	 */
-	DECLARE_MULTICAST_DELEGATE_OneParam(FOnLeaderInvitesOnlyChanged, bool /* bLeaderInviteOnly */);
+	DECLARE_MULTICAST_DELEGATE_OneParam(FOnClientLeaderInvitesOnlyChanged, bool /* bLeaderInviteOnly */);
 
 protected:
 
@@ -160,11 +159,17 @@ protected:
 
 public:
 
-	//~ Begin UObject Interface begin
+	// UObject interface begin
 	virtual UWorld* GetWorld() const override;
 	static void AddReferencedObjects(UObject* InThis, FReferenceCollector& Collector);
-	//~ Begin UObject Interface end
+	// UObject interface end
 
+	
+	/** 
+	 * Allows external systems an opportunity to initialize newly created PartyMemberStates before being entered into the party system
+	 * 
+	 * @param InFunctor a TFunction to call immediately after allocation. 
+	 */
 	void SetInitMemberFunctor( TFunction<void( TSharedPtr<FOnlinePartyMember>, UPartyMemberState* )>&& InFunctor );
 
 	/**
@@ -182,6 +187,9 @@ public:
 	 */
 	virtual void OnShutdown();
 
+	/** @return the party id of this party */
+	FOnlinePartyTypeId GetPartyTypeId() const;
+
 	/** @return the party id of this party */
 	TSharedPtr<const FOnlinePartyId> GetPartyId() const;
 
@@ -235,16 +243,11 @@ public:
 	void SetAcceptingMembers(bool bIsAcceptingMembers, EJoinPartyDenialReason DenialReason);
 
 	/**
-	 * Calculate the possible joinability state of this party
-	 * check the values from left to right in order of precedence
-	 *
-	 * @param bFriendJoinable [out] is the party joinable by friends via presence (doesn't require invite)
-	 * @param bInviteOnly [out] is the party joinable via explicit invites
-	 * @param bAllowInvites [out] are invites possible (use with bInviteOnly to determine if invites are available right now)
+	 * Check if you have permission to send invites to this party
 	 *
-	 * @return true if the out params are valid, false otherwise
+	 * @return true if you have permission
 	 */
-	bool GetJoinability(bool& bFriendJoinable, bool& bInviteOnly, bool& bAllowInvites) const;
+	bool CanInvite() const;
 
 	/** @return the party leader for this party */
 	TSharedPtr<const FUniqueNetId> GetPartyLeader() const;
@@ -280,12 +283,12 @@ public:
 	void GetAllPartyMembers(TArray<UPartyMemberState*>& PartyMembers) const;
 
 	template< typename TPartyMemberState >
-	void GetTypedPartyMembers( TArray< TPartyMemberState* >& PartyMembers ) const
+	void GetTypedPartyMembers(TArray< TPartyMemberState* >& PartyMembers) const
 	{
-		PartyMembers.Empty( PartyMembersState.Num() );
-		for ( auto Iter = PartyMembersState.CreateConstIterator(); Iter; ++Iter )
+		PartyMembers.Empty(PartyMembersState.Num());
+		for (auto Iter = PartyMembersState.CreateConstIterator(); Iter; ++Iter)
 		{
-			PartyMembers.Add( Cast< TPartyMemberState >( Iter.Value() ) );
+			PartyMembers.Add(Cast< TPartyMemberState >(Iter.Value()));
 		}
 	}
 
@@ -295,11 +298,11 @@ public:
 	FOnPartyMemberDataChanged& OnPartyMemberDataChanged() { return PartyMemberDataChanged; }
 
 	/** @return delegate fired when the party type has changed */
-	FOnPartyTypeChanged& OnPartyTypeChanged() { return PartyTypeChanged; }
+	FOnClientPartyTypeChanged& OnClientPartyTypeChanged() { return ClientPartyTypeChanged; }
 	/** @return delegate fired when the leader based join permissions have changed */
-	FOnLeaderFriendsOnlyChanged& OnLeaderFriendsOnlyChanged() { return LeaderFriendsOnlyChanged; }
+	FOnClientLeaderFriendsOnlyChanged& OnClientLeaderFriendsOnlyChanged() { return ClientLeaderFriendsOnlyChanged; }
 	/** @return delegate fired when the leader based invite permissions have changed */
-	FOnLeaderInvitesOnlyChanged& OnLeaderInvitesOnlyChanged() { return LeaderInvitesOnlyChanged; }
+	FOnClientLeaderInvitesOnlyChanged& OnClientLeaderInvitesOnlyChanged() { return ClientLeaderInvitesOnlyChanged; }
 	
 protected:
 
@@ -317,17 +320,17 @@ protected:
 	 * Initialize a party that is newly created and a local player owns it
 	 *
 	 * @param LocalUserID party owner/leader (may change later if user promotes another leader)
-	 * @param InPartyInfo reference to the basic party info from the interface
+	 * @param InParty reference to the basic party info from the interface
 	 */
-	void InitFromCreate(const FUniqueNetId& LocalUserId, TSharedPtr<FOnlinePartyInfo>& InPartyInfo);
+	void InitFromCreate(const FUniqueNetId& LocalUserId, TSharedPtr<const FOnlineParty>& InParty);
 
 	/**
 	 * Initialize a party that has been joined and a local player is simply a member
 	 *
 	 * @param LocalUserID local primary player that joined the party
-	 * @param InPartyInfo reference to the basic party info from the interface
+	 * @param InParty reference to the basic party info from the interface
 	 */
-	void InitFromJoin(const FUniqueNetId& LocalUserId, TSharedPtr<FOnlinePartyInfo>& InPartyInfo);
+	void InitFromJoin(const FUniqueNetId& LocalUserId, TSharedPtr<const FOnlineParty>& InParty);
 
 protected:
 
@@ -336,11 +339,11 @@ protected:
 	const UScriptStruct* PartyStateRefDef;
 
 	/**
-	* Pointer to child USTRUCT that holds the current state of party member (set via InitPartyMemberState)
-	*
-	* Cached data for the party, only modifiable by the party leader
-	* Reference to the data structure defined in a child class
-	*/
+	 * Pointer to child USTRUCT that holds the current state of party member (set via InitPartyMemberState)
+	 *
+	 * Cached data for the party, only modifiable by the party leader
+	 * Reference to the data structure defined in a child class
+	 */
 	FPartyState* PartyStateRef;
 
 	/** User who created or joined this room (not the party leader) */
@@ -355,7 +358,7 @@ protected:
 	bool bDebugAcceptingMembers;
 
 	/** Reference to party info within OSS */
-	TSharedPtr<FOnlinePartyInfo> PartyInfo;
+	TSharedPtr<const FOnlineParty> OssParty;
 
 	/** Is leader promotion available at the moment */
 	bool bPromotionLockoutState;
@@ -404,14 +407,14 @@ protected:
 	/**
 	 * Delegates for party visibility/presence/invite permission changes
 	 */
-	FOnPartyTypeChanged PartyTypeChanged;
-	FOnLeaderFriendsOnlyChanged LeaderFriendsOnlyChanged;
-	FOnLeaderInvitesOnlyChanged LeaderInvitesOnlyChanged;
+	FOnClientPartyTypeChanged ClientPartyTypeChanged;
+	FOnClientLeaderFriendsOnlyChanged ClientLeaderFriendsOnlyChanged;
+	FOnClientLeaderInvitesOnlyChanged ClientLeaderInvitesOnlyChanged;
 
 	/**
 	 * Common initialization for a newly instantiated party
 	 */
-	void Init(const FUniqueNetId& LocalUserId, TSharedPtr<FOnlinePartyInfo>& InPartyInfo);
+	void Init(const FUniqueNetId& LocalUserId, TSharedPtr<const FOnlineParty>& InParty);
 
 	/** Register for game related delegates that affect the party */
 	virtual void RegisterFrontendDelegates();
@@ -468,7 +471,7 @@ protected:
 	 * @param bResetAccessKey should the update reset the access key (invalidates existing invites/presence)
 	 */
 	void UpdatePartyConfig(bool bResetAccessKey = false);
-	void OnUpdatePartyConfigComplete(const FUniqueNetId& LocalUserId, bool bWasSuccessful, const FOnlinePartyId& InPartyId, const FString& Error);
+	void OnUpdatePartyConfigComplete(const FUniqueNetId& LocalUserId, const FOnlinePartyId& InPartyId, const EUpdateConfigCompletionResult);
 
 	/**
 	 * Called for all existing party members when a party configuration setting changes
@@ -535,7 +538,7 @@ protected:
 	void HandleRemovedFromParty(EMemberExitedReason Reason);
 
 	/** Party leader delegates */
-	void OnPartyMemberPromoted(const FUniqueNetId& LocalUserId, bool bWasSuccessful, const FOnlinePartyId& InPartyId, const FUniqueNetId& InMemberId, const FString& Error);
+	void OnPartyMemberPromoted(const FUniqueNetId& LocalUserId, const FOnlinePartyId& InPartyId, const FUniqueNetId& InMemberId, const EPromoteMemberCompletionResult Result);
 	void OnPartyMemberKicked(const FUniqueNetId& LocalUserId, const FOnlinePartyId& InPartyId, const FUniqueNetId& InMemberId, const EKickMemberCompletionResult Result);
 
 	/**
@@ -706,4 +709,4 @@ inline const TCHAR* ToString(EApprovalAction Type)
 		return TEXT("");
 	}
 	}
-}
\ No newline at end of file
+}
diff --git a/Engine/Source/Runtime/Online/OnlineFramework/Party/Public/PartyMemberState.h b/Engine/Source/Runtime/Online/OnlineFramework/Party/Public/PartyMemberState.h
--- a/Engine/Source/Runtime/Online/OnlineFramework/Party/Public/PartyMemberState.h
+++ b/Engine/Source/Runtime/Online/OnlineFramework/Party/Public/PartyMemberState.h
@@ -27,16 +27,16 @@ class PARTY_API UPartyMemberState : public UObject
 {
 	GENERATED_UCLASS_BODY()
 
-	//~ Begin UObject Interface
+	// Begin UObject Interface
 	virtual void BeginDestroy() override;
-	//~ End UObject Interface
+	// End UObject Interface
 
 	/** Unique id of this party member */
 	UPROPERTY(Transient)
 	FUniqueNetIdRepl UniqueId;
 
 	/** Display name of this party member */
-	UPROPERTY(BlueprintReadOnly, Transient, Category = PartyMemberState )
+	UPROPERTY(BlueprintReadOnly, Transient, Category = PartyMemberState)
 	FText DisplayName;
 
 	/** @return the party this member is associated with */
@@ -78,6 +78,7 @@ protected:
 
 private:
 	
+
 	/** Scratch copy of child USTRUCT for handling replication comparisons */
 	FPartyMemberRepState* MemberStateRefScratch;
 
@@ -85,6 +86,5 @@ private:
 	UPROPERTY(Transient)
 	bool bHasAnnouncedJoin;
 
-
 	friend UPartyGameState;
 };
diff --git a/Engine/Source/Runtime/Online/OnlineSubsystem/Public/Interfaces/OnlinePartyInterface.h b/Engine/Source/Runtime/Online/OnlineSubsystem/Public/Interfaces/OnlinePartyInterface.h
--- a/Engine/Source/Runtime/Online/OnlineSubsystem/Public/Interfaces/OnlinePartyInterface.h
+++ b/Engine/Source/Runtime/Online/OnlineSubsystem/Public/Interfaces/OnlinePartyInterface.h
@@ -8,10 +8,6 @@
 #define F_PREFIX(TypeToPrefix) F##TypeToPrefix
 #define PARTY_DECLARE_DELEGATETYPE(Type) typedef F##Type::FDelegate F##Type##Delegate
 
-/**
- * Key value data that can be sent to party members
- */
-typedef FOnlineKeyValuePairs<FString, FVariantData> FOnlinePartyAttrs;
 
 /**
  * Party member user info returned by IOnlineParty interface
@@ -28,7 +24,7 @@ class FOnlinePartyData : public TSharedFromThis<FOnlinePartyData>
 {
 public:
 	FOnlinePartyData() {}
-	virtual ~FOnlinePartyData() {}
+	~FOnlinePartyData() {}
 
 	/**
 	 * Get an attribute from the party data
@@ -66,24 +62,28 @@ public:
 	void ToJson(FString& JsonString) const;
 	void FromJson(const FString& JsonString);
 
+	/** map of key/val attributes that represents the data */
+	FOnlineKeyValuePairs<FString, FVariantData>  KeyValAttrs;
+};
+
+/**
+* Info needed to join a party
+*/
+class IOnlinePartyPendingJoinRequestInfo
+{
+public:
+	IOnlinePartyPendingJoinRequestInfo() {}
+	virtual ~IOnlinePartyPendingJoinRequestInfo() {}
+
 	/**
-	 * Dump state about the party data for debugging
-	 */
-	virtual FString DumpDebug() const
-	{
-		FString Result;
-		const FOnlinePartyAttrs& PartyAttrs = KeyValAttrs;
-		for (FOnlinePartyAttrs::TConstIterator It(PartyAttrs); It; ++It)
-		{
-			FString Key = It.Key();
-			const FVariantData& Value = It.Value();
-			Result += FString::Printf(TEXT("\n%s=%s"), *Key, *Value.ToString());
-		}
-		return Result;
-	}
+	* @return id of the sender of this join request
+	*/
+	virtual const TSharedRef<const FUniqueNetId>& GetSenderId() const = 0;
 
-	/** map of key/val attributes that represents the data */
-	FOnlinePartyAttrs  KeyValAttrs;
+	/**
+	* @return display name of the sender of this join request
+	*/
+	virtual const FString& GetSenderDisplayName() const = 0;
 };
 
 /**
@@ -95,155 +95,191 @@ public:
 	IOnlinePartyJoinInfo() {}
 	virtual ~IOnlinePartyJoinInfo() {}
 
-	/**
-	 * @return id of the user this invite is from
-	 */
-	virtual const TSharedRef<const FUniqueNetId>& GetLeaderId() const = 0;
+	virtual bool IsValid() const = 0;
 
 	/**
-	 * @return id of the party associated with this invite
+	 * @return party id of party associated with this join invite
 	 */
 	virtual const TSharedRef<const FOnlinePartyId>& GetPartyId() const = 0;
 
 	/**
-	 * @return id of the client app associated with the sender of the party invite 
+	 * @return party id of party associated with this join invite
 	 */
-	virtual const FString& GetAppId() const = 0;
+	virtual const FOnlinePartyTypeId GetPartyTypeId() const = 0;
 
 	/**
-	 * @return id of the client app associated with the sender of the party invite 
+	 * @return user id of the leader of the party associated with this join info
 	 */
-	virtual int32 GetBuildId() const = 0;
+	virtual const TSharedRef<const FUniqueNetId>& GetLeaderId() const = 0;
 
 	/**
-	 * @return whether or not the party is accepting members
+	 * @return user id of the leader of the party associated with this join info
 	 */
-	virtual bool GetIsAcceptingMembers() const = 0;
+	virtual const FString& GetLeaderDisplayName() const = 0;
 
 	/**
-	 * @return why the party is not accepting members
+	 * @return user id of where this join info came from
 	 */
-	virtual int32 GetNotAcceptingReason() const = 0;
+	virtual const TSharedRef<const FUniqueNetId>& GetSourceUserId() const = 0;
 
 	/**
-	 * @return true if the join info cannot be used to join a party
+	 * @return user id of where this join info came from
 	 */
-	virtual bool IsInvalidForJoin() const = 0;
+	virtual const FString& GetSourceDisplayName() const = 0;
 
 	/**
-	 * @return true if the join info cannot be used to request an invite
+	 * @return true if the join info has some form of key(does not guarantee the validity of that key)
 	 */
-	virtual bool IsInvalidForInviteRequest() const = 0;
+	virtual bool HasKey() const = 0;
 
 	/**
-	 * @return the game-specific custom data that was sent along with this invite
-	 */
-	virtual const FOnlinePartyData& GetCustomPayload() const = 0;
-};
+	 * @return true if a password can be used to bypass generated access key
+     */
+	virtual bool HasPassword() const = 0;
 
-/**
- * The binding of a join info object to a given sender
- */
-class IOnlinePartyInvite
-{
-public:
-	IOnlinePartyInvite() {}
-	virtual ~IOnlinePartyInvite() {}
+	/**
+	 * @return true if the party is known to be accepting members
+     */
+	virtual bool IsAcceptingMembers() const = 0;
 
-	virtual const TSharedRef<const FUniqueNetId>& GetSenderId() const = 0;
-	virtual const TSharedRef<IOnlinePartyJoinInfo>& GetJoinInfo() const = 0;
-};
+	/**
+	 * @return why the party is not accepting members
+	 */
+	virtual int32 GetNotAcceptingReason() const = 0;
 
-/**
- * Info needed to join a party
- */
-class IOnlinePartyPendingJoinRequestInfo
-{
-public:
-	IOnlinePartyPendingJoinRequestInfo() {}
-	virtual ~IOnlinePartyPendingJoinRequestInfo() {}
+	/**
+	 * @return id of the client app associated with the sender of the party invite
+	 */
+	virtual const FString& GetAppId() const = 0;
 
 	/**
-	* @return id of the party associated with this invite
+	* @return id of the build associated with the sender of the party invite
 	*/
-	virtual const TSharedRef<const FOnlinePartyId>& GetPartyId() const = 0;
+	virtual const FString& GetBuildId() const = 0;
 
 	/**
-	 * @return id of the sender of this join request
+	 * @return client specify key-value data
 	 */
-	virtual const TSharedRef<const FUniqueNetId>& GetSenderId() const = 0;
+	virtual const FOnlinePartyData& GetClientData() const = 0;
 
 	/**
-	 * @return display name of the sender of this join request
+	 * @return whether or not the join info can be used to join
 	 */
-	virtual const FString& GetSenderDisplayName() const = 0;
+	virtual bool CanJoin() const = 0;
+
+	/**
+	 * @return whether or not the join info can be used to join with a password
+	 */
+	virtual bool CanJoinWithPassword() const = 0;
+
+	/**
+	 * @return whether or not the join info has the info to request an invite
+	 */
+	virtual bool CanRequestAnInvite() const = 0;
 };
 
 /**
- * Permissions for joining a party
+ * Permissions for party features
  */
-namespace EPartyPermissions
+namespace PartySystemPermissions
 {
-	enum Type
+	/**
+	 * EPresencePermissions details who can publish what to presence.
+	 * DoNotPublish is a valid setting for primary parties.
+	 */
+	enum class EPermissionType: uint8
 	{
-		/** any user that discovers the party can join */
-		Public,
-		/** invite required to join the party */
-		Private
+		None = 0x0,
+		Leader = 0x1,
+		Friend = 0x2,
+		Anyone = 0x4
 	};
 
-	inline const TCHAR* ToString(EPartyPermissions::Type PartyPermissions)
+	/**
+	 * FPublishPermissionValidator uses meta programming to cause errors for invalid combinations of permission flags
+	 */
+	template<bool bIsValidPermissionCombo>
+		struct FPublishPermissionValidator;
+
+	template<> struct FPublishPermissionValidator<true> {};
+
+	template<EPermissionType PublishIdPermission, EPermissionType PublishKeyPermission>
+	struct FPublishPermissionBuilder : public FPublishPermissionValidator<(PublishKeyPermission <= PublishIdPermission)>
 	{
-		switch (PartyPermissions)
-		{
-		case Public:
-		{
-			return TEXT("Public");
-		}
-		case Private:
-		{
-			return TEXT("Private");
-		}
-		}
-		return TEXT("");
-	}
+		static const uint8 Value = static_cast<uint8>(PublishIdPermission) << 4 | static_cast<uint8>(PublishKeyPermission);
+	};
+
+	/**
+	 * Pre-defined presence permission objects
+	 */
+	enum class EPresencePermissions: uint8
+	{
+		DoNotPublish                    = FPublishPermissionBuilder<EPermissionType::None,   EPermissionType::None>::Value,
+		LeaderPublishIdNonePublishKey   = FPublishPermissionBuilder<EPermissionType::Leader, EPermissionType::None>::Value,
+		LeaderPublishIdLeaderPublishKey = FPublishPermissionBuilder<EPermissionType::Leader, EPermissionType::Leader>::Value,
+		FriendPublishIdNonePublishKey   = FPublishPermissionBuilder<EPermissionType::Friend, EPermissionType::None>::Value,
+		FriendPublishIdLeaderPublishKey = FPublishPermissionBuilder<EPermissionType::Friend, EPermissionType::Leader>::Value,
+		FriendPublishIdFriendPublishKey = FPublishPermissionBuilder<EPermissionType::Friend, EPermissionType::Friend>::Value,
+		AnyonePublishIdNonePublishKey   = FPublishPermissionBuilder<EPermissionType::Anyone, EPermissionType::None>::Value,
+		AnyonePublishIdLeaderPublishKey = FPublishPermissionBuilder<EPermissionType::Anyone, EPermissionType::Leader>::Value,
+		AnyonePublishIdFriendPublishKey = FPublishPermissionBuilder<EPermissionType::Anyone, EPermissionType::Friend>::Value,
+		AnyonePublishIdAnyonePublishKey = FPublishPermissionBuilder<EPermissionType::Anyone, EPermissionType::Anyone>::Value,
+	};
+
+	/**
+	 * Presence permission aliases
+	 */
+	static const EPresencePermissions FriendsInviteOnly = EPresencePermissions::LeaderPublishIdNonePublishKey;
+	static const EPresencePermissions FriendsOfFriendsInviteOnly = EPresencePermissions::FriendPublishIdNonePublishKey;
+	static const EPresencePermissions PublicInviteOnly = EPresencePermissions::AnyonePublishIdNonePublishKey;
+	static const EPresencePermissions FriendsOnly = EPresencePermissions::LeaderPublishIdLeaderPublishKey;
+	static const EPresencePermissions FriendsOfFriendsOnly = EPresencePermissions::FriendPublishIdFriendPublishKey;
+	static const EPresencePermissions Public = EPresencePermissions::AnyonePublishIdAnyonePublishKey;
+
+	enum class EInvitePermissions
+	{
+		/** Available to the leader only */
+		Leader,
+		/** Available to friends of the leader only */
+		Friends,
+		/** Available to anyone */
+		Anyone,
+	};
 }
 
+enum class EJoinRequestAction
+{
+	Manual,
+	AutoApprove,
+	AutoReject
+};
 /**
  * Options for configuring a new party or for updating an existing party
  */
 struct FPartyConfiguration : public TSharedFromThis<FPartyConfiguration>
 {
 	FPartyConfiguration()
-		: Permissions(EPartyPermissions::Public)
-		, bPresenceParty(false)
-		, bShouldPublishToPresence(false)
-		, bCanNonLeaderPublishToPresence(false)
-		, bCanNonLeaderInviteToParty(false)
+		: JoinRequestAction(EJoinRequestAction::Manual)
+		, PresencePermissions(PartySystemPermissions::EPresencePermissions::AnyonePublishIdAnyonePublishKey)
+		, InvitePermissions(PartySystemPermissions::EInvitePermissions::Leader)
 		, bShouldRemoveOnDisconnection(false)
 		, bIsAcceptingMembers(false)
 		, NotAcceptingMembersReason(0)
 		, MaxMembers(0)
 	{}
 
-	/** Permission for configuring party */
-	EPartyPermissions::Type Permissions;
-	/** is this party tied to presence */
-	bool bPresenceParty;
-	/** should publish info to presence */
-	bool bShouldPublishToPresence;
 	/** should publish info to presence */
-	bool bCanNonLeaderPublishToPresence;
-	/** can non-leaders send an invite to people */
-	bool bCanNonLeaderInviteToParty;
+	EJoinRequestAction JoinRequestAction;
+	/** Permission for how the party can be  */
+	PartySystemPermissions::EPresencePermissions PresencePermissions;
+	/** Permission who can send invites */
+	PartySystemPermissions::EInvitePermissions InvitePermissions;
 	/** should remove on disconnection */
 	bool bShouldRemoveOnDisconnection;
 	/** is accepting members */
 	bool bIsAcceptingMembers;
 	/** not accepting members reason */
 	int32 NotAcceptingMembersReason;
-	/** clients can add whatever data they want for configuration options */
-	FOnlinePartyData ClientConfigData;
 	/** Maximum active members allowed. 0 means no maximum. */
 	int32 MaxMembers;
 	/** Human readable nickname */
@@ -252,35 +288,16 @@ struct FPartyConfiguration : public TSharedFromThis<FPartyConfiguration>
 	FString Description;
 	/** Human readable password for party. */
 	FString Password;
-
-	FString DumpState() const
-	{
-		return FString::Printf(TEXT("Permissions: %s PresenceEnabled: %d Publish: L(%d) NL(%d) NL Invite(%d) RemoveOnDisconnect: %d Accepting: %d Not Accepting Reason(%d) MaxMembers: %d Nickname: %s Description: %s Password: %s"),
-			ToString(Permissions),
-			bPresenceParty,
-			bShouldPublishToPresence,
-			bCanNonLeaderPublishToPresence,
-			bCanNonLeaderInviteToParty,
-			bShouldRemoveOnDisconnection,
-			bIsAcceptingMembers,
-			NotAcceptingMembersReason,
-			MaxMembers,
-			*Nickname,
-			*Description,
-			Password.IsEmpty() ? TEXT("not set") : TEXT("set"));
-	}
+	/** clients can add whatever data they want for configuration options */
+	FOnlinePartyData ClientConfigData;
 };
 
 enum class EPartyState
 {
 	None,
 	CreatePending,
-	DestroyPending,
 	JoinPending,
-	JoinChatPending,         // to be removed
-	JoinFinalizationPending, // to be removed
 	LeavePending,
-	RemovalPending,
 	Active,
 	Disconnected,
 	CleanUp
@@ -289,47 +306,30 @@ enum class EPartyState
 /**
  * Current state associated with a party
  */
-class FOnlinePartyInfo : public TSharedFromThis<FOnlinePartyInfo>
+class FOnlineParty : public TSharedFromThis<FOnlineParty>
 {
-	FOnlinePartyInfo();
-public:
-	FOnlinePartyInfo(const TSharedRef<const FOnlinePartyId>& InPartyId)
+protected:
+	FOnlineParty();
+	explicit FOnlineParty(const TSharedRef<const FOnlinePartyId>& InPartyId, const FOnlinePartyTypeId InPartyTypeId)
 		: PartyId(InPartyId)
+		, PartyTypeId(InPartyTypeId)
 		, State(EPartyState::None)
 		, Config(MakeShareable(new FPartyConfiguration()))
 	{}
 
-	virtual ~FOnlinePartyInfo() 
+public:
+	virtual ~FOnlineParty()
 	{}
 
-	virtual TSharedPtr<const FUniqueNetId> GetLeaderId() const = 0;
-	
-	/**
-	 * Return a structure suitable for joining this party
-	 */
-	virtual const TSharedRef<IOnlinePartyJoinInfo> GetJoinInfo() const = 0;
-
-	/**
-	 * Calculate the possible joinability state of this party
-	 * check the values from left to right in order of precedence
-	 *
-	 * @param LocalUserId local user making this request
-	 * @param bPublicJoinable [out] is the party joinable by anyone at all
-	 * @param bFriendJoinable [out] is the party joinable by friends via presence (doesn't require invite)
-	 * @param bInviteOnly [out] is the party joinable via explicit invites
-	 * @param bAllowInvites [out] are invites possible (use with bInviteOnly to determine if invites are available right now)
-	 *
-	 * @return true if the out params are valid, false otherwise
-	 */
-	virtual bool GetJoinability(const FUniqueNetId& LocalUserId, bool& bPublicJoinable, bool& bFriendJoinable, bool& bInviteOnly, bool& bAllowInvites) const = 0;
-
-	/**
-	 * Dump state about the party info for debugging
-	 */
-	virtual FString DumpState() const = 0;
+	virtual bool CanLocalUserInvite(const FUniqueNetId& LocalUserId) const = 0;
+	virtual bool IsJoinable() const = 0;
 
 	/** unique id of the party */
 	TSharedRef<const FOnlinePartyId> PartyId;
+	/** unique id of the party */
+	const FOnlinePartyTypeId PartyTypeId;
+	/** unique id of the leader */
+	TSharedPtr<const FUniqueNetId> LeaderId;
 	/** The current state of the party */
 	EPartyState State;
 	/** Current state of configuration */
@@ -350,180 +350,119 @@ enum class EMemberExitedReason
 	Unknown,
 	Left,
 	Removed,
-	Disbanded,
 	Kicked
 };
 
-enum class ECreatePartyCompletionResult
-{
-	UnknownClientFailure = -100,
-	AlreadyCreatingParty,
-	AlreadyInParty,
-	FailedToCreateMucRoom,
-	NoResponse,
-	UnknownInternalFailure = 0,
-	Succeeded = 1
-};
-
-enum class EJoinPartyCompletionResult
-{
-	UnknownClientFailure = -100,
-	BadBuild,
-	JoinInfoInvalid,
-	AlreadyJoiningParty,
-	AlreadyInParty,
-	UnknownLocalUser,
-	FailedToJoinChatRoom,
-	MessagingFailure,
-	NoSpace,
-	NotApproved,
-	RequesteeNotMember,
-	RequesteeNotLeader,
-	NoResponse,
-	UnknownTransportFailure,
-	UnknownInternalFailure = 0,
-	Succeeded = 1
-};
-
-enum class ELeavePartyCompletionResult
-{
-	UnknownClientFailure = -100,
-	AlreadyLeavingParty,
-	UnknownLocalUser,
-	UnknownParty,
-	MessagingFailure,
-	NoResponse,
-	UnknownTransportFailure,
-	UnknownInternalFailure = 0,
-	Succeeded = 1
-};
-
-enum class ESendInviteCompleteResult
-{
-	NotLoggedIn = -100,
-	InvitePending,
-	AlreadyInParty,
-	PartyFull,
-	NoPermission,
-	UnknownInternalFailure = 0,
-	Succeeded = 1
-};
-
-enum class EKickMemberCompletionResult
-{
-	UnknownClientFailure = -100,
-	UnknownParty,
-	LocalMemberNotMember,
-	LocalMemberNotLeader,
-	RemoteMemberNotMember,
-	MessagingFailure,
-	NoResponse,
-	UnknownInternalFailure = 0,
-	Succeeded = 1
-};
 
+enum class ECreatePartyCompletionResult;
+enum class EJoinPartyCompletionResult;
+enum class ELeavePartyCompletionResult;
+enum class EUpdateConfigCompletionResult;
+enum class ERequestPartyInvitationCompletionResult;
+enum class ESendPartyInvitationCompletionResult;
+enum class EAcceptPartyInvitationCompletionResult;
+enum class ERejectPartyInvitationCompletionResult;
+enum class EKickMemberCompletionResult;
+enum class EPromoteMemberCompletionResult;
 enum class EInvitationResponse;
 
+///////////////////////////////////////////////////////////////////
 // Completion delegates
-
+///////////////////////////////////////////////////////////////////
 /**
  * Party creation async task completed callback
  *
  * @param LocalUserId - id of user that initiated the request
- * @param bWasSuccessful - true if successfully created new party
  * @param PartyId - id associated with the party
- * @param Error - string with error info if any
+ * @param Result - result of the operation
  */
-DECLARE_DELEGATE_ThreeParams(FOnCreatePartyComplete, const FUniqueNetId& /*LocalUserId*/, const FOnlinePartyId& /*PartyId*/, const ECreatePartyCompletionResult /*Result*/);
+DECLARE_DELEGATE_ThreeParams(FOnCreatePartyComplete, const FUniqueNetId& /*LocalUserId*/, const TSharedPtr<const FOnlinePartyId>& /*PartyId*/, const ECreatePartyCompletionResult /*Result*/);
 /**
- * Party update async task completed callback
+ * Party join async task completed callback
  *
  * @param LocalUserId - id of user that initiated the request
- * @param bWasSuccessful - true if successfully updated party
  * @param PartyId - id associated with the party
- * @param Error - string with error info if any
+ * @param Result - result of the operation
+ * @param NotApprovedReason - client defined value describing why you were not approved
  */
-DECLARE_DELEGATE_FourParams(FOnUpdatePartyComplete, const FUniqueNetId& /*LocalUserId*/, bool /*bWasSuccessful*/, const FOnlinePartyId& /*PartyId*/, const FString& /*Error*/);
+DECLARE_DELEGATE_FourParams(FOnJoinPartyComplete, const FUniqueNetId& /*LocalUserId*/, const FOnlinePartyId& /*PartyId*/, const EJoinPartyCompletionResult /*Result*/, const int32 /*NotApprovedReason*/);
 /**
- * Party destroy async task completed callback
+ * Party leave async task completed callback
  *
  * @param LocalUserId - id of user that initiated the request
- * @param bWasSuccessful - true if successfully destroyed party
  * @param PartyId - id associated with the party
- * @param Error - string with error info if any
- */
-DECLARE_DELEGATE_FourParams(FOnDestroyPartyComplete, const FUniqueNetId& /*LocalUserId*/, bool /*bWasSuccessful*/, const FOnlinePartyId& /*PartyId*/, const FString& /*Error*/);
-/**
- * Not Yet Implemented
+ * @param Result - result of the operation
  */
-DECLARE_DELEGATE_FourParams(FOnRequestInvitationComplete, const FUniqueNetId& /*LocalUserId*/, bool /*bWasSuccessful*/, const FOnlinePartyId& /*PartyId*/, const FString& /*Error*/);
+DECLARE_DELEGATE_ThreeParams(FOnLeavePartyComplete, const FUniqueNetId& /*LocalUserId*/, const FOnlinePartyId& /*PartyId*/, const ELeavePartyCompletionResult /*Result*/);
 /**
- * Party join async task completed callback
+ * Party update async task completed callback
  *
  * @param LocalUserId - id of user that initiated the request
- * @param bWasSuccessful - true if successfully joined party
  * @param PartyId - id associated with the party
- * @param Error - string with error info if any
+ * @param Result - result of the operation
  */
-DECLARE_DELEGATE_FourParams(FOnJoinPartyComplete, const FUniqueNetId& /*LocalUserId*/, const FOnlinePartyId& /*PartyId*/, const EJoinPartyCompletionResult /*Result*/, const int32 /*NotApprovedReason*/);
+DECLARE_DELEGATE_ThreeParams(FOnUpdatePartyComplete, const FUniqueNetId& /*LocalUserId*/, const FOnlinePartyId& /*PartyId*/, const EUpdateConfigCompletionResult /*Result*/);
 /**
- * Party leave async task completed callback
+ * Party update async task completed callback
  *
  * @param LocalUserId - id of user that initiated the request
- * @param bWasSuccessful - true if successfully left party
  * @param PartyId - id associated with the party
- * @param Error - string with error info if any
+ * @param Result - result of the operation
  */
-DECLARE_DELEGATE_ThreeParams(FOnLeavePartyComplete, const FUniqueNetId& /*LocalUserId*/, const FOnlinePartyId& /*PartyId*/, const ELeavePartyCompletionResult /*Result*/);
+DECLARE_DELEGATE_ThreeParams(FOnRequestPartyInvitationComplete, const FUniqueNetId& /*LocalUserId*/, const FOnlinePartyId& /*PartyId*/, const ERequestPartyInvitationCompletionResult /*Result*/);
 /**
  * Party invitation sent completed callback
  *
  * @param LocalUserId - id of user that initiated the request
  * @param PartyId - id associated with the party
- * @param Result - result of the send operation
  * @param RecipientId - user invite was sent to
+ * @param Result - result of the send operation
  */
-DECLARE_DELEGATE_FourParams(FOnSendInvitationComplete, const FUniqueNetId& /*LocalUserId*/, const FOnlinePartyId& /*PartyId*/, ESendInviteCompleteResult /*Result*/, const FUniqueNetId& /*RecipientId*/);
+DECLARE_DELEGATE_FourParams(FOnSendPartyInvitationComplete, const FUniqueNetId& /*LocalUserId*/, const FOnlinePartyId& /*PartyId*/, const FUniqueNetId& /*RecipientId*/, const ESendPartyInvitationCompletionResult /*Result*/);
 /**
  * Accepting an invite to a user to join party async task completed callback
  *
  * @param LocalUserId - id of user that initiated the request
- * @param bWasSuccessful - true if successfully sent invite
  * @param PartyId - id associated with the party
- * @param Error - string with error info if any
+ * @param Result - string with error info if any
  */
-DECLARE_DELEGATE_FourParams(FOnAcceptPartyInviteComplete, const FUniqueNetId& /*LocalUserId*/, bool /*bWasSuccessful*/, const FOnlinePartyId& /*PartyId*/, const FString& /*Error*/);
+DECLARE_DELEGATE_ThreeParams(FOnAcceptPartyInvitationComplete, const FUniqueNetId& /*LocalUserId*/, const FOnlinePartyId& /*PartyId*/, const EAcceptPartyInvitationCompletionResult /*Result*/);
 /**
  * Rejecting an invite to a user to join party async task completed callback
  *
  * @param LocalUserId - id of user that initiated the request
  * @param bWasSuccessful - true if successfully sent invite
  * @param PartyId - id associated with the party
- * @param Error - string with error info if any
+ * @param Result - string with error info if any
  */
-DECLARE_DELEGATE_FourParams(FOnRejectPartyInviteComplete, const FUniqueNetId& /*LocalUserId*/, bool /*bWasSuccessful*/, const FOnlinePartyId& /*PartyId*/, const FString& /*Error*/);
+DECLARE_DELEGATE_ThreeParams(FOnRejectPartyInvitationComplete, const FUniqueNetId& /*LocalUserId*/, const FOnlinePartyId& /*PartyId*/, const ERejectPartyInvitationCompletionResult /*Result*/);
 /**
  * Kicking a member of a party async task completed callback
  *
  * @param LocalUserId - id of user that initiated the request
- * @param bWasSuccessful - true if successfully kicked member
  * @param PartyId - id associated with the party
  * @param MemberId - id of member being kicked
- * @param Error - string with error info if any
+ * @param Result - string with error info if any
  */
-DECLARE_DELEGATE_FourParams(FOnKickPartyMemberComplete, const FUniqueNetId& /*LocalUserId*/, const FOnlinePartyId& /*PartyId*/, const FUniqueNetId& /*MemberId*/, const EKickMemberCompletionResult /*Result*/);
+DECLARE_DELEGATE_FourParams(FOnKickPartyMemberComplete, const FUniqueNetId& /*LocalUserId*/, const FOnlinePartyId& /*PartyId*/, const FUniqueNetId& /*LocalUserId*/, const EKickMemberCompletionResult /*Result*/);
 /**
  * Promoting a member of a party async task completed callback
  *
  * @param LocalUserId - id of user that initiated the request
- * @param bWasSuccessful - true if successfully promoted member to admin
  * @param PartyId - id associated with the party
- * @param MemberId - id of member being promoted to admin
- * @param Error - string with error info if any
+ * @param MemberId - id of member being promoted to leader
+ * @param Result - string with error info if any
  */
-DECLARE_DELEGATE_FiveParams(FOnPromotePartyMemberComplete, const FUniqueNetId& /*LocalUserId*/, bool /*bWasSuccessful*/, const FOnlinePartyId& /*PartyId*/, const FUniqueNetId& /*MemberId*/, const FString& /*Error*/);
+DECLARE_DELEGATE_FourParams(FOnPromotePartyMemberComplete, const FUniqueNetId& /*LocalUserId*/, const FOnlinePartyId& /*PartyId*/, const FUniqueNetId& /*LocalUserId*/, const EPromoteMemberCompletionResult /*Result*/);
+
+
+
 
+
+
+///////////////////////////////////////////////////////////////////
 // Notification delegates
+///////////////////////////////////////////////////////////////////
 
 /**
  * notification when a party is joined
@@ -533,6 +472,14 @@ DECLARE_DELEGATE_FiveParams(FOnPromotePartyMemberComplete, const FUniqueNetId& /
 DECLARE_MULTICAST_DELEGATE_TwoParams(F_PREFIX(OnPartyJoined), const FUniqueNetId& /*LocalUserId*/, const FOnlinePartyId& /*PartyId*/);
 PARTY_DECLARE_DELEGATETYPE(OnPartyJoined);
 
+/**
+ * notification when a party is joined
+ * @param LocalUserId - id associated with this notification
+ * @param PartyId - id associated with the party
+ */
+DECLARE_MULTICAST_DELEGATE_TwoParams(F_PREFIX(OnPartyExited), const FUniqueNetId& /*LocalUserId*/, const FOnlinePartyId& /*PartyId*/);
+PARTY_DECLARE_DELEGATETYPE(OnPartyExited);
+
 /**
  * Notification when player promotion is locked out.
  * @param LocalUserId - id associated with this notification
@@ -632,7 +579,7 @@ PARTY_DECLARE_DELEGATETYPE(OnPartyInviteReceived);
  * @param SenderId - id of member that sent the invite
  * @param bWasAccepted - whether or not the invite was accepted
  */
-DECLARE_MULTICAST_DELEGATE_FourParams(F_PREFIX(OnPartyInviteResponseReceived), const FUniqueNetId& /*LocalUserId*/, const FOnlinePartyId& /*PartyId*/, const FUniqueNetId& /*SenderId*/, const EInvitationResponse /*bWasAccepted*/);
+DECLARE_MULTICAST_DELEGATE_FourParams(F_PREFIX(OnPartyInviteResponseReceived), const FUniqueNetId& /*LocalUserId*/, const FOnlinePartyId& /*PartyId*/, const FUniqueNetId& /*SenderId*/, const EInvitationResponse /*Response*/);
 PARTY_DECLARE_DELEGATETYPE(OnPartyInviteResponseReceived);
 
 /**
@@ -658,13 +605,13 @@ PARTY_DECLARE_DELEGATETYPE(OnPartyJoinRequestResponseReceived);
  * Interface definition for the online party services 
  * Allows for forming a party and communicating with party members
  */
-class IOnlineParty
+class IOnlinePartySystem
 {
 protected:
-	IOnlineParty() {};
+	IOnlinePartySystem() {};
 
 public:
-	virtual ~IOnlineParty() {};
+	virtual ~IOnlinePartySystem() {};
 
 	/**
 	 * Create a new party
@@ -676,7 +623,7 @@ public:
 	 *
 	 * @return true if task was started
 	 */
-	virtual bool CreateParty(const FUniqueNetId& LocalUserId, const FPartyConfiguration& PartyConfig, const FOnCreatePartyComplete& Delegate = FOnCreatePartyComplete()) = 0;
+	virtual bool CreateParty(const FUniqueNetId& LocalUserId, const FOnlinePartyTypeId PartyTypeId, const FPartyConfiguration& PartyConfig, const FOnCreatePartyComplete& Delegate = FOnCreatePartyComplete()) = 0;
 
 	/**
 	 * Update an existing party with new configuration
@@ -689,18 +636,6 @@ public:
 	 */
 	virtual bool UpdateParty(const FUniqueNetId& LocalUserId, const FOnlinePartyId& PartyId, const FPartyConfiguration& PartyConfig, bool bShouldRegenerateReservationKey = false, const FOnUpdatePartyComplete& Delegate = FOnUpdatePartyComplete()) = 0;
 
-	/**
-	 * Destroy an existing party. Party members will be disbanded
-	 * Only Owner/Admin can destroy the party
-	 *
-	 * @param LocalUserId - user making the request
-	 * @param PartyId - id of an existing party
-	 * @param Delegate - called on completion
-	 *
-	 * @return true if task was started
-	 */
-	virtual bool DestroyParty(const FUniqueNetId& LocalUserId, const FOnlinePartyId& PartyId, const FOnDestroyPartyComplete& Delegate = FOnDestroyPartyComplete()) = 0;
-
 	/**
 	 * Join an existing party
 	 *
@@ -738,20 +673,6 @@ public:
 	*/
 	virtual bool ApproveJoinRequest(const FUniqueNetId& LocalUserId, const FOnlinePartyId& PartyId, const FUniqueNetId& RecipientId, bool bIsApproved, int32 DeniedResultCode = 0) = 0;
 
-	/**
-	* sends an request to a party leader for an invite to a party that could not otherwise be joined
-	* if the player accepts the invite they will be sent the data needed to trigger a call to RequestReservation
-	*
-	* @param LocalUserId - user making the request
-	* @param PartyId - id of an existing party
-	* @param LeaderId - id of the user to request an invitation from
-	* @param RequestForId - id of the user you are requesting an invite for
-	* @param Delegate - called on completion
-	*
-	* @return true if task was started
-	*/
-	virtual bool RequestInvitation(const FUniqueNetId& LocalUserId, const FOnlinePartyId& PartyId, const FUniqueNetId& LeaderId, const FUniqueNetId& RequestForId, const FOnRequestInvitationComplete& Delegate = FOnRequestInvitationComplete()) = 0;
-
 	/**
 	 * sends an invitation to a user that could not otherwise join a party
 	 * if the player accepts the invite they will be sent the data needed to trigger a call to RequestReservation
@@ -759,12 +680,11 @@ public:
 	 * @param LocalUserId - user making the request
 	 * @param PartyId - id of an existing party
 	 * @param RecipientId - id of the user being invited
-	 * @param CustomPayload - game defined key value pairs
 	 * @param Delegate - called on completion
 	 *
 	 * @return true if task was started
 	 */
-	virtual bool SendInvitation(const FUniqueNetId& LocalUserId, const FOnlinePartyId& PartyId, const FUniqueNetId& RecipientId, const FOnlinePartyData& CustomPayload = FOnlinePartyData(), const FOnSendInvitationComplete& Delegate = FOnSendInvitationComplete()) = 0;
+	virtual bool SendInvitation(const FUniqueNetId& LocalUserId, const FOnlinePartyId& PartyId, const FUniqueNetId& RecipientId, const FOnlinePartyData& ClientData = FOnlinePartyData(), const FOnSendPartyInvitationComplete& Delegate = FOnSendPartyInvitationComplete()) = 0;
 
 	/**
 	* Accept an invite to a party. NOTE this does not initiate a join.
@@ -798,7 +718,7 @@ public:
 	 *
 	 * @return true if task was started
 	 */
-	virtual bool KickMember(const FUniqueNetId& LocalUserId, const FOnlinePartyId& PartyId, const FUniqueNetId& MemberId, const FOnKickPartyMemberComplete& Delegate = FOnKickPartyMemberComplete()) = 0;
+	virtual bool KickMember(const FUniqueNetId& LocalUserId, const FOnlinePartyId& PartyId, const FUniqueNetId& TargetMemberId, const FOnKickPartyMemberComplete& Delegate = FOnKickPartyMemberComplete()) = 0;
 
 	/**
 	 * Promote a user from an existing party to be admin
@@ -811,7 +731,7 @@ public:
 	 *
 	 * @return true if task was started
 	 */
-	virtual bool PromoteMember(const FUniqueNetId& LocalUserId, const FOnlinePartyId& PartyId, const FUniqueNetId& MemberId, const FOnPromotePartyMemberComplete& Delegate = FOnPromotePartyMemberComplete()) = 0;
+	virtual bool PromoteMember(const FUniqueNetId& LocalUserId, const FOnlinePartyId& PartyId, const FUniqueNetId& TargetMemberId, const FOnPromotePartyMemberComplete& Delegate = FOnPromotePartyMemberComplete()) = 0;
 
 	/**
 	 * Set party data and broadcast to all members
@@ -825,7 +745,7 @@ public:
 	 *
 	 * @return true if task was started
 	 */
-	virtual bool UpdatePartyData(const FUniqueNetId& LocalUserId, const FOnlinePartyId& PartyId, const FOnlinePartyAttrs& PartyData) = 0;
+	virtual bool UpdatePartyData(const FUniqueNetId& LocalUserId, const FOnlinePartyId& PartyId, const FOnlinePartyData& PartyData) = 0;
 
 	/**
 	 * Set party data for a single party member and broadcast to all members
@@ -839,7 +759,7 @@ public:
 	 *
 	 * @return true if task was started
 	 */
-	virtual bool UpdatePartyMemberData(const FUniqueNetId& LocalUserId, const FOnlinePartyId& PartyId, const FOnlinePartyAttrs& PartyMemberData) = 0;
+	virtual bool UpdatePartyMemberData(const FUniqueNetId& LocalUserId, const FOnlinePartyId& PartyId, const FOnlinePartyData& PartyMemberData) = 0;
 
 	/**
 	 * returns true if the user specified by MemberId is the leader of the party specified by PartyId
@@ -868,7 +788,17 @@ public:
 	 *
 	 * @return party info or nullptr if not found
 	 */
-	virtual TSharedPtr<FOnlinePartyInfo> GetPartyInfo(const FUniqueNetId& LocalUserId, const FOnlinePartyId& PartyId) const = 0;
+	virtual TSharedPtr<const FOnlineParty> GetParty(const FUniqueNetId& LocalUserId, const FOnlinePartyId& PartyId) const = 0;
+
+	/**
+	 * Get info associated with a party
+	 *
+	 * @param LocalUserId - user making the request
+	 * @param PartyTypeId - type of an existing party
+	 *
+	 * @return party info or nullptr if not found
+	 */
+	virtual TSharedPtr<const FOnlineParty> GetParty(const FUniqueNetId& LocalUserId, const FOnlinePartyTypeId& PartyTypeId) const = 0;
 
 	/**
 	 * Get a party member by id
@@ -905,16 +835,14 @@ public:
 	virtual TSharedPtr<FOnlinePartyData> GetPartyMemberData(const FUniqueNetId& LocalUserId, const FOnlinePartyId& PartyId, const FUniqueNetId& MemberId) const = 0;
 
 	/**
-	 * returns true if the user is advertising a party on their presence
+	 * returns true if the user is advertising a party of that type on their presence
 	 *
 	 * @param LocalUserId       - user making the request
 	 * @param UserId            - user to check
-	 * @param OutPartyIdString  - string id of advertised party
-	 * @param OutLeaderIdString - string id of advertised party's leader
-	 * @param OutReservationKey - (optional) reservation key of the advertised party
+	 * @param PartyTypeId       - type of party to query
 	 *
 	 */
-	virtual TSharedPtr<IOnlinePartyJoinInfo> GetAdvertisedParty(const FUniqueNetId& LocalUserId, const FUniqueNetId& UserId) = 0;
+	virtual TSharedPtr<IOnlinePartyJoinInfo> GetAdvertisedParty(const FUniqueNetId& LocalUserId, const FUniqueNetId& UserId, const FOnlinePartyTypeId PartyTypeId) const = 0;
 
 	/**
 	 * Get a list of currently joined parties for the user
@@ -945,7 +873,7 @@ public:
 	 *
 	 * @return true if entries found
 	 */
-	virtual bool GetPendingInvites(const FUniqueNetId& LocalUserId, TArray<TSharedRef<IOnlinePartyInvite>>& OutPendingInvitesArray) const = 0;
+	virtual bool GetPendingInvites(const FUniqueNetId& LocalUserId, TArray<TSharedRef<IOnlinePartyJoinInfo>>& OutPendingInvitesArray) const = 0;
 
 	/**
 	 * Get list of users requesting to join the party
@@ -969,6 +897,34 @@ public:
 	 */
 	virtual bool GetPendingInvitedUsers(const FUniqueNetId& LocalUserId, const FOnlinePartyId& PartyId, TArray<TSharedRef<const FUniqueNetId>>& OutPendingInvitedUserArray) const = 0;
 
+	static const FOnlinePartyTypeId::TInternalType PrimaryPartyTypeIdValue = 0x11111111;
+	/**
+	 * @return party type id for the primary party - the primary party is the party that will be addressable via the social panel
+	 */
+	static const FOnlinePartyTypeId GetPrimaryPartyTypeId() { return FOnlinePartyTypeId(PrimaryPartyTypeIdValue); }
+
+	/**
+	 * @return party type id for the primary party - the primary party is the party that will be addressable via the social panel
+	 */
+	static const FOnlinePartyTypeId MakePartyTypeId(const FOnlinePartyTypeId::TInternalType InTypeId) { ensure(InTypeId != PrimaryPartyTypeIdValue); return FOnlinePartyTypeId(InTypeId); }
+
+	/**
+	 * returns the json version of a join info for a current party
+	 *
+	 * @param LocalUserId       - user making the request
+	 * @param PartyId           - party to make the json from
+	 *
+ 	 */
+	virtual FString MakeJoinInfoJson(const FUniqueNetId& LocalUserId, const FOnlinePartyId& PartyId) = 0;
+
+	/**
+	 * returns a valid join info object from a json blob
+	 *
+	 * @param JoinInfoJson       - json blob to convert
+	 *
+	 */
+	virtual TSharedPtr<IOnlinePartyJoinInfo> MakeJoinInfoFromJson(const FString& JoinInfoJson) = 0;
+
 	/**
 	 * Creates a command line token from a IOnlinePartyJoinInfo object
 	 *
@@ -1021,6 +977,13 @@ public:
 	 */
 	DEFINE_ONLINE_DELEGATE_TWO_PARAM(OnPartyJoined, const FUniqueNetId& /*LocalUserId*/, const FOnlinePartyId& /*PartyId*/);
 
+	/**
+	 * notification of when a party is exited
+	 * @param LocalUserId - id associated with this notification
+	 * @param PartyId - id associated with the party
+	 */
+	DEFINE_ONLINE_DELEGATE_TWO_PARAM(OnPartyExited, const FUniqueNetId& /*LocalUserId*/, const FOnlinePartyId& /*PartyId*/);
+
 	/**
 	 * Notification when player promotion is locked out.
 	 *
@@ -1046,25 +1009,25 @@ public:
 	DEFINE_ONLINE_DELEGATE_THREE_PARAM(OnPartyDataReceived, const FUniqueNetId& /*LocalUserId*/, const FOnlinePartyId& /*PartyId*/, const TSharedRef<FOnlinePartyData>& /*PartyData*/);
 
 	/**
-	 * Notification when a member changes in a party
-	 * @param LocalUserId - id associated with this notification
-	 * @param PartyId - id associated with the party
-	 * @param MemberId - id of member that joined
-	 * @param Reason - how the member changed
-	 */
-	DEFINE_ONLINE_DELEGATE_FOUR_PARAM(OnPartyMemberChanged, const FUniqueNetId& /*LocalUserId*/, const FOnlinePartyId& /*PartyId*/, const FUniqueNetId& /*MemberId*/, const EMemberChangedReason /*Reason*/);
-
-	/**
-	 * Notification when a member exits a party
-	 * @param LocalUserId - id associated with this notification
-	 * @param PartyId - id associated with the party
-	 * @param MemberId - id of member that joined
-	 * @param Reason - why the member was removed
-	 */
-	DEFINE_ONLINE_DELEGATE_FOUR_PARAM(OnPartyMemberExited, const FUniqueNetId& /*LocalUserId*/, const FOnlinePartyId& /*PartyId*/, const FUniqueNetId& /*MemberId*/, const EMemberExitedReason /*Reason*/);
-
-	/**
-	 * Notification when a member joins the party
+	* Notification when a member changes in a party
+	* @param LocalUserId - id associated with this notification
+	* @param PartyId - id associated with the party
+	* @param MemberId - id of member that changed
+	* @param Reason - how the member changed
+	*/
+	DEFINE_ONLINE_DELEGATE_FOUR_PARAM(OnPartyMemberChanged, const FUniqueNetId& /*LocalUserId*/, const FOnlinePartyId& /*PartyId*/, const FUniqueNetId& /*MemberId*/, const EMemberChangedReason /*Reason*/);
+
+	/**
+	* Notification when a member exits a party
+	* @param LocalUserId - id associated with this notification
+	* @param PartyId - id associated with the party
+	* @param MemberId - id of member that joined
+	* @param Reason - why the member was removed
+	*/
+	DEFINE_ONLINE_DELEGATE_FOUR_PARAM(OnPartyMemberExited, const FUniqueNetId& /*LocalUserId*/, const FOnlinePartyId& /*PartyId*/, const FUniqueNetId& /*MemberId*/, const EMemberExitedReason /*Reason*/);
+
+	/**
+	 * Notification when a member joins the party
 	 * @param LocalUserId - id associated with this notification
 	 * @param PartyId - id associated with the party
 	 * @param MemberId - id of member that joined
@@ -1087,20 +1050,20 @@ public:
 	DEFINE_ONLINE_DELEGATE_ONE_PARAM(OnPartyInvitesChanged, const FUniqueNetId& /*LocalUserId*/);
 
 	/**
-	 * Notification when a request for an invite has been received
-	 * @param LocalUserId - id associated with this notification
-	 * @param PartyId - id associated with the party
-	 * @param SenderId - id of user that sent the invite
-	 * @param RequestForId - id of user that sender is requesting the invite for - invalid if the sender is requesting the invite
-	 */
+	* Notification when a request for an invite has been received
+	* @param LocalUserId - id associated with this notification
+	* @param PartyId - id associated with the party
+	* @param SenderId - id of user that sent the invite
+	* @param RequestForId - id of user that sender is requesting the invite for - invalid if the sender is requesting the invite
+	*/
 	DEFINE_ONLINE_DELEGATE_FOUR_PARAM(OnPartyInviteRequestReceived, const FUniqueNetId& /*LocalUserId*/, const FOnlinePartyId& /*PartyId*/, const FUniqueNetId& /*SenderId*/, const FUniqueNetId& /*RequestForId*/);
 
 	/**
-	 * Notification when a new invite is received
-	 * @param LocalUserId - id associated with this notification
-	 * @param SenderId - id of member that sent the invite
-	 * @param PartyId - id associated with the party
-	 */
+	* Notification when a new invite is received
+	* @param LocalUserId - id associated with this notification
+	* @param SenderId - id of member that sent the invite
+	* @param PartyId - id associated with the party
+	*/
 	DEFINE_ONLINE_DELEGATE_THREE_PARAM(OnPartyInviteReceived, const FUniqueNetId& /*LocalUserId*/, const FOnlinePartyId& /*PartyId*/, const FUniqueNetId& /*SenderId*/);
 
 	/**
@@ -1110,7 +1073,7 @@ public:
 	 * @param SenderId - id of member that sent the invite
 	 * @param bWasAccepted - true is the invite was accepted
 	 */
-	DEFINE_ONLINE_DELEGATE_FOUR_PARAM(OnPartyInviteResponseReceived, const FUniqueNetId& /*LocalUserId*/, const FOnlinePartyId& /*PartyId*/, const FUniqueNetId& /*SenderId*/, const EInvitationResponse /*bWasAccepted*/);
+	DEFINE_ONLINE_DELEGATE_FOUR_PARAM(OnPartyInviteResponseReceived, const FUniqueNetId& /*LocalUserId*/, const FOnlinePartyId& /*PartyId*/, const FUniqueNetId& /*SenderId*/, const EInvitationResponse /*Response*/);
 
 	/**
 	 * Notification when a new reservation request is received
@@ -1136,7 +1099,137 @@ public:
 
 };
 
-typedef TSharedPtr<IOnlineParty, ESPMode::ThreadSafe> IOnlinePartyPtr;
+enum class ECreatePartyCompletionResult
+{
+	UnknownClientFailure = -100,
+	AlreadyInPartyOfSpecifiedType,
+	AlreadyCreatingParty,
+	AlreadyInParty,
+	FailedToCreateMucRoom,
+	NoResponse,
+	UnknownInternalFailure = 0,
+	Succeeded = 1
+};
+
+enum class EJoinPartyCompletionResult
+{
+	UnknownClientFailure = -100,
+	BadBuild,
+	JoinInfoInvalid,
+	AlreadyInPartyOfSpecifiedType,
+	AlreadyJoiningParty,
+	AlreadyInParty,
+	MessagingFailure,
+	PartyNotInPendingState,
+	ResponseFromUnexpectedUser,
+	NoSpace,
+	NotApproved,
+	RequesteeNotMember,
+	RequesteeNotLeader,
+	NoResponse,
+	UnknownInternalFailure = 0,
+	Succeeded = 1
+};
+
+enum class ELeavePartyCompletionResult
+{
+	UnknownClientFailure = -100,
+	LeavePending,
+	UnknownLocalUser,
+	UnknownParty,
+	NotMember,
+	MessagingFailure,
+	NoResponse,
+	UnknownTransportFailure,
+	UnknownInternalFailure = 0,
+	Succeeded = 1
+};
+
+enum class EUpdateConfigCompletionResult
+{
+	UnknownClientFailure = -100,
+	UnknownParty,
+	LocalMemberNotMember,
+	LocalMemberNotLeader,
+	RemoteMemberNotMember,
+	MessagingFailure,
+	NoResponse,
+	UnknownInternalFailure = 0,
+	Succeeded = 1
+};
+
+enum class ERequestPartyInvitationCompletionResult
+{
+	NotLoggedIn = -100,
+	InvitePending,
+	AlreadyInParty,
+	PartyFull,
+	NoPermission,
+	UnknownInternalFailure = 0,
+	Succeeded = 1
+};
+
+enum class ESendPartyInvitationCompletionResult
+{
+	NotLoggedIn = -100,
+	InvitePending,
+	AlreadyInParty,
+	PartyFull,
+	NoPermission,
+	UnknownInternalFailure = 0,
+	Succeeded = 1
+};
+
+enum class EAcceptPartyInvitationCompletionResult
+{
+	NotLoggedIn = -100,
+	InvitePending,
+	AlreadyInParty,
+	PartyFull,
+	NoPermission,
+	UnknownInternalFailure = 0,
+	Succeeded = 1
+};
+
+enum class ERejectPartyInvitationCompletionResult
+{
+	NotLoggedIn = -100,
+	InvitePending,
+	AlreadyInParty,
+	PartyFull,
+	NoPermission,
+	UnknownInternalFailure = 0,
+	Succeeded = 1
+};
+
+enum class EKickMemberCompletionResult
+{
+	UnknownClientFailure = -100,
+	UnknownParty,
+	LocalMemberNotMember,
+	LocalMemberNotLeader,
+	RemoteMemberNotMember,
+	MessagingFailure,
+	NoResponse,
+	UnknownInternalFailure = 0,
+	Succeeded = 1
+};
+
+enum class EPromoteMemberCompletionResult
+{
+	UnknownClientFailure = -100,
+	UnknownServiceFailure,
+	UnknownParty,
+	LocalMemberNotMember,
+	LocalMemberNotLeader,
+	PromotionAlreadyPending,
+	TargetIsSelf,
+	TargetNotMember,
+	MessagingFailure,
+	NoResponse,
+	UnknownInternalFailure = 0,
+	Succeeded = 1
+};
 
 enum class EInvitationResponse
 {
@@ -1146,31 +1239,8 @@ enum class EInvitationResponse
 	Accepted
 };
 
-inline const TCHAR* ToString(EInvitationResponse Value)
-{
-	switch (Value)
-	{
-	case EInvitationResponse::UnknownFailure:
-	{
-		return TEXT("UnknownFailure");
-	}
-	case EInvitationResponse::BadBuild:
-	{
-		return TEXT("BadBuild");
-	}
-	case EInvitationResponse::Rejected:
-	{
-		return TEXT("Rejected");
-	}
-	case EInvitationResponse::Accepted:
-	{
-		return TEXT("Accepted");
-	}
-	}
-	return TEXT("");
-}
 
-inline const TCHAR* ToString(EMemberChangedReason Value)
+inline const TCHAR* ToString(const EMemberChangedReason Value)
 {
 	switch (Value)
 	{
@@ -1191,7 +1261,7 @@ inline const TCHAR* ToString(EMemberChangedReason Value)
 }
 
 /** @return the stringified version of the enum passed in */
-inline const TCHAR* ToString(EPartyState Value)
+inline const TCHAR* ToString(const EPartyState Value)
 {
 	switch (Value)
 	{
@@ -1203,26 +1273,14 @@ inline const TCHAR* ToString(EPartyState Value)
 	{
 		return TEXT("Create pending");
 	}
-	case EPartyState::DestroyPending:
-	{
-		return TEXT("Destroy pending");
-	}
 	case EPartyState::JoinPending:
 	{
 		return TEXT("Join pending");
 	}
-	case EPartyState::JoinChatPending:
-	{
-		return TEXT("Join Chat pending");
-	}
 	case EPartyState::LeavePending:
 	{
 		return TEXT("Leave pending");
 	}
-	case EPartyState::RemovalPending:
-	{
-		return TEXT("Removal pending");
-	}
 	case EPartyState::Active:
 	{
 		return TEXT("Active");
@@ -1235,7 +1293,7 @@ inline const TCHAR* ToString(EPartyState Value)
 	return TEXT("");
 }
 
-inline const TCHAR* ToString(EMemberExitedReason Value)
+inline const TCHAR* ToString(const EMemberExitedReason Value)
 {
 	switch (Value)
 	{
@@ -1251,10 +1309,6 @@ inline const TCHAR* ToString(EMemberExitedReason Value)
 	{
 		return TEXT("Removed");
 	}
-	case EMemberExitedReason::Disbanded:
-	{
-		return TEXT("Disbanded");
-	}
 	case EMemberExitedReason::Kicked:
 	{
 		return TEXT("Kicked");
@@ -1263,7 +1317,7 @@ inline const TCHAR* ToString(EMemberExitedReason Value)
 	return TEXT("");
 }
 
-inline const TCHAR* ToString(ECreatePartyCompletionResult Value)
+inline const TCHAR* ToString(const ECreatePartyCompletionResult Value)
 {
 	switch (Value)
 	{
@@ -1299,35 +1353,35 @@ inline const TCHAR* ToString(ECreatePartyCompletionResult Value)
 	return TEXT("");
 }
 
-inline const TCHAR* ToString(ESendInviteCompleteResult Value)
+inline const TCHAR* ToString(const ESendPartyInvitationCompletionResult Value)
 {
 	switch (Value)
 	{
-	case ESendInviteCompleteResult::NotLoggedIn:
+	case ESendPartyInvitationCompletionResult::NotLoggedIn:
 	{
 		return TEXT("NotLoggedIn");
 	}
-	case ESendInviteCompleteResult::InvitePending:
+	case ESendPartyInvitationCompletionResult::InvitePending:
 	{
 		return TEXT("InvitePending");
 	}
-	case ESendInviteCompleteResult::AlreadyInParty:
+	case ESendPartyInvitationCompletionResult::AlreadyInParty:
 	{
 		return TEXT("AlreadyInParty");
 	}
-	case ESendInviteCompleteResult::PartyFull:
+	case ESendPartyInvitationCompletionResult::PartyFull:
 	{
 		return TEXT("PartyFull");
 	}
-	case ESendInviteCompleteResult::NoPermission:
+	case ESendPartyInvitationCompletionResult::NoPermission:
 	{
 		return TEXT("NoPermission");
 	}
-	case ESendInviteCompleteResult::UnknownInternalFailure:
+	case ESendPartyInvitationCompletionResult::UnknownInternalFailure:
 	{
 		return TEXT("UnknownInternalFailure");
 	}
-	case ESendInviteCompleteResult::Succeeded:
+	case ESendPartyInvitationCompletionResult::Succeeded:
 	{
 		return TEXT("Succeeded");
 	}
@@ -1335,7 +1389,7 @@ inline const TCHAR* ToString(ESendInviteCompleteResult Value)
 	return TEXT("");
 }
 
-inline const TCHAR* ToString(EJoinPartyCompletionResult Value)
+inline const TCHAR* ToString(const EJoinPartyCompletionResult Value)
 {
 	switch (Value)
 	{
@@ -1355,17 +1409,17 @@ inline const TCHAR* ToString(EJoinPartyCompletionResult Value)
 	{
 		return TEXT("AlreadyInParty");
 	}
-	case EJoinPartyCompletionResult::UnknownLocalUser:
+	case EJoinPartyCompletionResult::MessagingFailure:
 	{
-		return TEXT("UnknownLocalUser");
+		return TEXT("MessagingFailure");
 	}
-	case EJoinPartyCompletionResult::FailedToJoinChatRoom:
+	case EJoinPartyCompletionResult::PartyNotInPendingState:
 	{
-		return TEXT("FailedToJoinChatRoom");
+		return TEXT("PartyNotInPendingState");
 	}
-	case EJoinPartyCompletionResult::MessagingFailure:
+	case EJoinPartyCompletionResult::ResponseFromUnexpectedUser:
 	{
-		return TEXT("MessagingFailure");
+		return TEXT("ResponseFromUnexpectedUser");
 	}
 	case EJoinPartyCompletionResult::NoSpace:
 	{
@@ -1387,10 +1441,6 @@ inline const TCHAR* ToString(EJoinPartyCompletionResult Value)
 	{
 		return TEXT("NoResponse");
 	}
-	case EJoinPartyCompletionResult::UnknownTransportFailure:
-	{
-		return TEXT("UnknownTransportFailure");
-	}
 	case EJoinPartyCompletionResult::UnknownInternalFailure:
 	{
 		return TEXT("UnknownInternalFailure");
@@ -1403,7 +1453,7 @@ inline const TCHAR* ToString(EJoinPartyCompletionResult Value)
 	return TEXT("");
 }
 
-inline const TCHAR* ToString(ELeavePartyCompletionResult Value)
+inline const TCHAR* ToString(const ELeavePartyCompletionResult Value)
 {
 	switch (Value)
 	{
@@ -1411,9 +1461,9 @@ inline const TCHAR* ToString(ELeavePartyCompletionResult Value)
 	{
 		return TEXT("UnknownClientFailure");
 	}
-	case ELeavePartyCompletionResult::AlreadyLeavingParty:
+	case ELeavePartyCompletionResult::LeavePending:
 	{
-		return TEXT("AlreadyLeavingParty");
+		return TEXT("LeavePending");
 	}
 	case ELeavePartyCompletionResult::UnknownLocalUser:
 	{
@@ -1447,7 +1497,51 @@ inline const TCHAR* ToString(ELeavePartyCompletionResult Value)
 	return TEXT("");
 }
 
-inline const TCHAR* ToString(EKickMemberCompletionResult Value)
+inline const TCHAR* ToString(const EUpdateConfigCompletionResult Value)
+{
+	switch (Value)
+	{
+	case EUpdateConfigCompletionResult::UnknownClientFailure:
+	{
+		return TEXT("UnknownClientFailure");
+	}
+	case EUpdateConfigCompletionResult::UnknownParty:
+	{
+		return TEXT("UnknownParty");
+	}
+	case EUpdateConfigCompletionResult::LocalMemberNotMember:
+	{
+		return TEXT("LocalMemberNotMember");
+	}
+	case EUpdateConfigCompletionResult::LocalMemberNotLeader:
+	{
+		return TEXT("LocalMemberNotLeader");
+	}
+	case EUpdateConfigCompletionResult::RemoteMemberNotMember:
+	{
+		return TEXT("RemoteMemberNotMember");
+	}
+	case EUpdateConfigCompletionResult::MessagingFailure:
+	{
+		return TEXT("MessagingFailure");
+	}
+	case EUpdateConfigCompletionResult::NoResponse:
+	{
+		return TEXT("NoResponse");
+	}
+	case EUpdateConfigCompletionResult::UnknownInternalFailure:
+	{
+		return TEXT("UnknownInternalFailure");
+	}
+	case EUpdateConfigCompletionResult::Succeeded:
+	{
+		return TEXT("Succeeded");
+	}
+	}
+	return TEXT("");
+}
+
+inline const TCHAR* ToString(const EKickMemberCompletionResult Value)
 {
 	switch (Value)
 	{
@@ -1491,6 +1585,219 @@ inline const TCHAR* ToString(EKickMemberCompletionResult Value)
 	return TEXT("");
 }
 
+inline const TCHAR* ToString(const EPromoteMemberCompletionResult Value)
+{
+	switch (Value)
+	{
+	case EPromoteMemberCompletionResult::UnknownClientFailure:
+	{
+		return TEXT("UnknownClientFailure");
+	}
+	case EPromoteMemberCompletionResult::UnknownParty:
+	{
+		return TEXT("UnknownParty");
+	}
+	case EPromoteMemberCompletionResult::LocalMemberNotMember:
+	{
+		return TEXT("LocalMemberNotMember");
+	}
+	case EPromoteMemberCompletionResult::LocalMemberNotLeader:
+	{
+		return TEXT("LocalMemberNotLeader");
+	}
+	case EPromoteMemberCompletionResult::TargetIsSelf:
+	{
+		return TEXT("TargetIsSelf");
+	}
+	case EPromoteMemberCompletionResult::TargetNotMember:
+	{
+		return TEXT("TargetNotMember");
+	}
+	case EPromoteMemberCompletionResult::MessagingFailure:
+	{
+		return TEXT("MessagingFailure");
+	}
+	case EPromoteMemberCompletionResult::NoResponse:
+	{
+		return TEXT("NoResponse");
+	}
+	case EPromoteMemberCompletionResult::UnknownInternalFailure:
+	{
+		return TEXT("UnknownInternalFailure");
+	}
+	case EPromoteMemberCompletionResult::Succeeded:
+	{
+		return TEXT("Succeeded");
+	}
+	}
+	return TEXT("");
+}
+
+inline const TCHAR* ToString(const PartySystemPermissions::EPresencePermissions Value)
+{
+	switch (Value)
+	{
+	case PartySystemPermissions::EPresencePermissions::DoNotPublish:
+	{
+		return TEXT("DoNotPublish");
+	}
+	case PartySystemPermissions::EPresencePermissions::LeaderPublishIdNonePublishKey:
+	{
+		return TEXT("LeaderPublishIdNonePublishKey");
+	}
+	case PartySystemPermissions::EPresencePermissions::LeaderPublishIdLeaderPublishKey:
+	{
+		return TEXT("LeaderPublishIdLeaderPublishKey");
+	}
+	case PartySystemPermissions::EPresencePermissions::FriendPublishIdNonePublishKey:
+	{
+		return TEXT("FriendPublishIdNonePublishKey");
+	}
+	case PartySystemPermissions::EPresencePermissions::FriendPublishIdLeaderPublishKey:
+	{
+		return TEXT("FriendPublishIdLeaderPublishKey");
+	}
+	case PartySystemPermissions::EPresencePermissions::FriendPublishIdFriendPublishKey:
+	{
+		return TEXT("FriendPublishIdFriendPublishKey");
+	}
+	case PartySystemPermissions::EPresencePermissions::AnyonePublishIdNonePublishKey:
+	{
+		return TEXT("AnyonePublishIdNonePublishKey");
+	}
+	case PartySystemPermissions::EPresencePermissions::AnyonePublishIdLeaderPublishKey:
+	{
+		return TEXT("AnyonePublishIdLeaderPublishKey");
+	}
+	case PartySystemPermissions::EPresencePermissions::AnyonePublishIdFriendPublishKey:
+	{
+		return TEXT("AnyonePublishIdFriendPublishKey");
+	}
+	case PartySystemPermissions::EPresencePermissions::AnyonePublishIdAnyonePublishKey:
+	{
+		return TEXT("AnyonePublishIdAnyonePublishKey");
+	}
+	}
+	return TEXT("");
+}
+
+inline const TCHAR* ToString(const PartySystemPermissions::EInvitePermissions Value)
+{
+	switch (Value)
+	{
+	case PartySystemPermissions::EInvitePermissions::Leader:
+	{
+		return TEXT("Leader");
+	}
+	case PartySystemPermissions::EInvitePermissions::Friends:
+	{
+		return TEXT("Friends");
+	}
+	case PartySystemPermissions::EInvitePermissions::Anyone:
+	{
+		return TEXT("Anyone");
+	}
+	}
+	return TEXT("");
+}
+
+inline const TCHAR* ToString(const EJoinRequestAction Value)
+{
+	switch (Value)
+	{
+	case EJoinRequestAction::Manual:
+	{
+		return TEXT("Manual");
+	}
+	case EJoinRequestAction::AutoApprove:
+	{
+		return TEXT("AutoApprove");
+	}
+	case EJoinRequestAction::AutoReject:
+	{
+		return TEXT("AutoReject");
+	}
+	}
+	return TEXT("");
+}
+
+inline const TCHAR* ToString(const EInvitationResponse Value)
+{
+	switch (Value)
+	{
+	case EInvitationResponse::UnknownFailure:
+	{
+		return TEXT("UnknownFailure");
+	}
+	case EInvitationResponse::BadBuild:
+	{
+		return TEXT("BadBuild");
+	}
+	case EInvitationResponse::Rejected:
+	{
+		return TEXT("Rejected");
+	}
+	case EInvitationResponse::Accepted:
+	{
+		return TEXT("Accepted");
+	}
+	}
+	return TEXT("");
+}
+
+inline FString ToDebugString(const FPartyConfiguration& PartyConfiguration)
+{
+	return FString::Printf(TEXT("JoinRequestAction(%s) RemoveOnDisconnect(%d) Publish(%s) Invite(%s) Accepting(%d) Not Accepting Reason(%d) MaxMembers: %d Nickname: %s Description: %s Password: %s"),
+			ToString(PartyConfiguration.JoinRequestAction),
+			PartyConfiguration.bShouldRemoveOnDisconnection,
+			ToString(PartyConfiguration.PresencePermissions),
+			ToString(PartyConfiguration.InvitePermissions),
+			PartyConfiguration.bIsAcceptingMembers,
+			PartyConfiguration.NotAcceptingMembersReason,
+			PartyConfiguration.MaxMembers,
+			*PartyConfiguration.Nickname,
+			*PartyConfiguration.Description,
+			PartyConfiguration.Password.IsEmpty() ? TEXT("not set") : *PartyConfiguration.Password
+		);
+}
+
+inline FString ToDebugString(const IOnlinePartyJoinInfo& JoinInfo)
+{
+	return FString::Printf(TEXT("SourceUserId(%s) SourceDisplayName(%s) PartyId(%s) LeaderUserId(%s) LeaderDisplayName(%s) HasKey(%d) HasPassword(%d) IsAcceptingMembers(%d) NotAcceptingReason(%d)"),
+			*(JoinInfo.GetSourceUserId()->ToDebugString()),
+			*(JoinInfo.GetSourceDisplayName()),
+			*(JoinInfo.GetPartyId()->ToDebugString()),
+			JoinInfo.GetLeaderId()->IsValid() ? *(JoinInfo.GetLeaderId()->ToString()) : TEXT("not set"),
+			!JoinInfo.GetLeaderDisplayName().IsEmpty() ? *(JoinInfo.GetLeaderDisplayName()) : TEXT("not set"),
+			JoinInfo.HasKey() ? 1 : 0,
+			JoinInfo.HasPassword() ? 1 : 0,
+			JoinInfo.IsAcceptingMembers() ? 1 : 0,
+			JoinInfo.GetNotAcceptingReason()
+		);
+}
+
+/**
+* Dump state about the party data for debugging
+*/
+inline FString ToDebugString(const FOnlinePartyData& PartyData)
+{
+	FString Result;
+	bool bAddComma = false;
+	for (auto Iterator : PartyData.KeyValAttrs)
+	{
+		if (bAddComma)
+		{
+			Result += TEXT(", ");
+		}
+		else
+		{
+			bAddComma = true;
+		}
+		Result += FString::Printf(TEXT("[%s=%s]"), *(Iterator.Key), *(Iterator.Value.ToString()));
+	}
+	return Result;
+}
+
 inline void FOnlinePartyData::ToJson(FString& JsonString) const
 {
 	JsonString.Empty();
@@ -1498,10 +1805,10 @@ inline void FOnlinePartyData::ToJson(FString& JsonString) const
 	// iterate over key/val attrs and convert each entry to a json string
 	TSharedRef<FJsonObject> JsonObject(new FJsonObject());
 	TArray<TSharedPtr<FJsonValue> > JsonProperties;
-	for (FOnlinePartyAttrs::TConstIterator It(KeyValAttrs); It; ++It)
+	for (auto Iterator : KeyValAttrs)
 	{
-		const FString& PropertyName = It.Key();
-		const FVariantData& PropertyValue = It.Value();
+		const FString& PropertyName = Iterator.Key;
+		const FVariantData& PropertyValue = Iterator.Value;
 
 		TSharedRef<FJsonObject> PropertyJson = PropertyValue.ToJson();
 		PropertyJson->SetStringField(TEXT("Name"), *PropertyName);
@@ -1546,4 +1853,4 @@ inline void FOnlinePartyData::FromJson(const FString& JsonString)
 			}
 		}
 	}
-}
\ No newline at end of file
+}
diff --git a/Engine/Source/Runtime/Online/OnlineSubsystem/Public/Interfaces/OnlinePresenceInterface.h b/Engine/Source/Runtime/Online/OnlineSubsystem/Public/Interfaces/OnlinePresenceInterface.h
--- a/Engine/Source/Runtime/Online/OnlineSubsystem/Public/Interfaces/OnlinePresenceInterface.h
+++ b/Engine/Source/Runtime/Online/OnlineSubsystem/Public/Interfaces/OnlinePresenceInterface.h
@@ -174,9 +174,9 @@ public:
 	*
 	* @param LocalUserId The unique id of the local user
 	* @param User The unique id of the user from which to get presence information.
-	* @param ClientId The id of the application you want to get the presence of
+	* @param AppId The id of the application you want to get the presence of
 	* @param OutPresence If found, a shared pointer to the cached presence data for User will be stored here.
 	* @return Whether the data was found or not.
 	*/
-	virtual EOnlineCachedResult::Type GetCachedPresenceForClient(const FUniqueNetId& LocalUserId, const FUniqueNetId& User, const FString& ClientId, TSharedPtr<FOnlineUserPresence>& OutPresence) = 0;
+	virtual EOnlineCachedResult::Type GetCachedPresenceForApp(const FUniqueNetId& LocalUserId, const FUniqueNetId& User, const FString& AppId, TSharedPtr<FOnlineUserPresence>& OutPresence) = 0;
 };
diff --git a/Engine/Source/Runtime/Online/OnlineSubsystem/Public/OnlineJsonSerializer.h b/Engine/Source/Runtime/Online/OnlineSubsystem/Public/OnlineJsonSerializer.h
--- a/Engine/Source/Runtime/Online/OnlineSubsystem/Public/OnlineJsonSerializer.h
+++ b/Engine/Source/Runtime/Online/OnlineSubsystem/Public/OnlineJsonSerializer.h
@@ -10,10 +10,10 @@
 #define BEGIN_ONLINE_JSON_SERIALIZER \
 	virtual void Serialize(FOnlineJsonSerializerBase& Serializer, bool bFlatObject) override \
 	{ \
-		if (!bFlatObject) Serializer.StartObject();
+		if (!bFlatObject) { Serializer.StartObject(); }
 
 #define END_ONLINE_JSON_SERIALIZER \
-		if (!bFlatObject) Serializer.EndObject(); \
+		if (!bFlatObject) { Serializer.EndObject(); } \
 	}
 
 #define ONLINE_JSON_SERIALIZE(JsonName, JsonValue) \
diff --git a/Engine/Source/Runtime/Online/OnlineSubsystem/Public/OnlineSubsystem.h b/Engine/Source/Runtime/Online/OnlineSubsystem/Public/OnlineSubsystem.h
--- a/Engine/Source/Runtime/Online/OnlineSubsystem/Public/OnlineSubsystem.h
+++ b/Engine/Source/Runtime/Online/OnlineSubsystem/Public/OnlineSubsystem.h
@@ -36,7 +36,7 @@ DECLARE_CYCLE_STAT_EXTERN(TEXT("VoiceInt"), STAT_Voice_Interface, STATGROUP_Onli
 /** Forward declarations of all interface classes */
 typedef TSharedPtr<class IOnlineSession, ESPMode::ThreadSafe> IOnlineSessionPtr;
 typedef TSharedPtr<class IOnlineFriends, ESPMode::ThreadSafe> IOnlineFriendsPtr;
-typedef TSharedPtr<class IOnlineParty, ESPMode::ThreadSafe> IOnlinePartyPtr;
+typedef TSharedPtr<class IOnlinePartySystem, ESPMode::ThreadSafe> IOnlinePartyPtr;
 typedef TSharedPtr<class IOnlineGroups, ESPMode::ThreadSafe> IOnlineGroupsPtr;
 typedef TSharedPtr<class IOnlineSharedCloud, ESPMode::ThreadSafe> IOnlineSharedCloudPtr;
 typedef TSharedPtr<class IOnlineUserCloud, ESPMode::ThreadSafe> IOnlineUserCloudPtr;
diff --git a/Engine/Source/Runtime/Online/OnlineSubsystem/Public/OnlineSubsystemTypes.h b/Engine/Source/Runtime/Online/OnlineSubsystem/Public/OnlineSubsystemTypes.h
--- a/Engine/Source/Runtime/Online/OnlineSubsystem/Public/OnlineSubsystemTypes.h
+++ b/Engine/Source/Runtime/Online/OnlineSubsystem/Public/OnlineSubsystemTypes.h
@@ -1363,4 +1363,53 @@ typedef FString FNotificationTransportId;
 /**
 * Id of a party instance
 */
-typedef FUniqueNetId FOnlinePartyId;
+class FOnlinePartyId : public IOnlinePlatformData, public TSharedFromThis<FOnlinePartyId>
+{
+protected:
+	/** Hidden on purpose */
+	FOnlinePartyId()
+	{
+	}
+
+	/** Hidden on purpose */
+	FOnlinePartyId(const FOnlinePartyId& Src)
+	{
+	}
+
+	/** Hidden on purpose */
+	FOnlinePartyId& operator=(const FOnlinePartyId& Src)
+	{
+		return *this;
+	}
+
+public:
+	virtual ~FOnlinePartyId() {}
+};
+
+/**
+* Id of a party's type
+*/
+class FOnlinePartyTypeId
+{
+public:
+	typedef uint32 TInternalType;
+	explicit FOnlinePartyTypeId(const TInternalType InValue = 0) : Value(InValue) {}
+	FOnlinePartyTypeId(const FOnlinePartyTypeId& Other) : Value(Other.Value) {}
+
+	bool operator==(const FOnlinePartyTypeId Rhs) const { return Value == Rhs.Value; }
+	bool operator!=(const FOnlinePartyTypeId Rhs) const { return Value != Rhs.Value; }
+
+	TInternalType GetValue() const { return Value; }
+
+	friend bool IsValid(const FOnlinePartyTypeId Value);
+
+protected:
+	TInternalType Value;
+};
+
+inline bool IsValid(const FOnlinePartyTypeId Id) { return Id.GetValue() != 0; }
+
+inline uint32 GetTypeHash(const FOnlinePartyTypeId Id)
+{
+	return Id.GetValue();
+}
diff --git a/Engine/Source/Runtime/PacketHandlers/PacketHandler/Private/PacketHandler.cpp b/Engine/Source/Runtime/PacketHandlers/PacketHandler/Private/PacketHandler.cpp
--- a/Engine/Source/Runtime/PacketHandlers/PacketHandler/Private/PacketHandler.cpp
+++ b/Engine/Source/Runtime/PacketHandlers/PacketHandler/Private/PacketHandler.cpp
@@ -28,10 +28,10 @@ BufferedPacket::~BufferedPacket()
 
 // PACKET HANDLER
 PacketHandler::PacketHandler()
-: Time(0.f)
+: bEnabled(false)
+, Time(0.f)
 , State(Handler::State::Uninitialized)
 , ReliabilityComponent(NULL)
-, bEnabled(false)
 {
 	OutgoingPacket.SetAllowResize(true);
 	OutgoingPacket.AllowAppend(true);
diff --git a/Engine/Source/Runtime/RHI/Private/GPUProfiler.cpp b/Engine/Source/Runtime/RHI/Private/GPUProfiler.cpp
--- a/Engine/Source/Runtime/RHI/Private/GPUProfiler.cpp
+++ b/Engine/Source/Runtime/RHI/Private/GPUProfiler.cpp
@@ -396,4 +396,36 @@ uint64 FGPUTiming::GTimingFrequency = 0;
 /** Whether the static variables have been initialized. */
 bool FGPUTiming::GAreGlobalsInitialized = false;
 
+
+
+float FWindowedGPUTimer::GetElapsedAverage(FRHICommandListImmediate& RHICmdList, float &OutAvgTimeInSeconds)
+{
+	if (QueriesFinished < StartQueries.Num())
+	{
+		return 0.0f;
+	}
+
+	float TotalTime = 0;
+
+	// Grab the queries in our window (specified number of frames old) and calculate total time as an average of the winow size
+	//
+	for (int32 i = StartQueries.Num() - WindowSize; i < StartQueries.Num(); i++)
+	{
+		uint64 StartTime, EndTime;
+		bool StartSucceeded = RHICmdList.GetRenderQueryResult(StartQueries[i], StartTime, false);
+		bool EndSucceeded = RHICmdList.GetRenderQueryResult(EndQueries[i], EndTime, false);
+
+		// figure out what the failure rate of the queries is; they fail because the GPU hasn't finished them when we
+		// try to get the data, so if the failure rate is too high, number of frames behind needs to be increased
+		QueriesFailed += StartSucceeded&&EndSucceeded ? -1 : 1;
+		QueriesFailed = FMath::Max<int32>(0, QueriesFailed);
+		TotalTime += (EndTime - StartTime) / 1000000.0f;
+	}
+
+	float FailRate = static_cast<float>(QueriesFailed) / WindowSize;
+	OutAvgTimeInSeconds = TotalTime / WindowSize;
+
+	return FailRate;
+}
+
 #undef LOCTEXT_NAMESPACE
\ No newline at end of file
diff --git a/Engine/Source/Runtime/RHI/Public/GPUProfiler.h b/Engine/Source/Runtime/RHI/Public/GPUProfiler.h
--- a/Engine/Source/Runtime/RHI/Public/GPUProfiler.h
+++ b/Engine/Source/Runtime/RHI/Public/GPUProfiler.h
@@ -6,6 +6,7 @@
 
 #pragma once
 
+
 /** Stats for a single perf event node. */
 class FGPUProfilerEventNodeStats : public FRefCountedObject
 {
@@ -219,3 +220,72 @@ struct RHI_API FGPUProfiler
 	void BeginFrame();
 	void EndFrame();
 };
+
+
+
+/**
+ * Simple moving window averaged GPU timer; create an instance, call Begin() and End() around the block to time, 
+ * then call GetElapsedAverage to get the averaged timings; BufferSize determines the number of queries in the window,
+ * FramesBehind determines how long we wait to grab query results, so we don't have to block on them, so the effective
+ * window size is BufferSize - FramesBehind
+ * The timer keeps track of failed queries; GetElapsedAverage returns a float [0;1] that indicates fail rate: 0 means
+ * no queries have failed, 1 means all queries within the window have failed. If more than 0.1, it's a good indication
+ * that FramesBehind has to be increased.
+ */
+class FWindowedGPUTimer
+{
+public:
+	FWindowedGPUTimer(FRHICommandListImmediate& RHICmdList)
+		:FWindowedGPUTimer(10, 2, RHICmdList)
+	{
+	}
+
+	FWindowedGPUTimer(uint32 BufferSize, uint32 FramesBehind, FRHICommandListImmediate& RHICmdList)
+		:QueriesFinished(0)
+		, QueriesFailed(0)
+	{
+		StartQueries.AddZeroed(BufferSize);
+		EndQueries.AddZeroed(BufferSize);
+		for (uint32 i = 0; i < BufferSize; i++)
+		{
+			StartQueries[i] = RHICmdList.CreateRenderQuery(RQT_AbsoluteTime);;
+			EndQueries[i] = RHICmdList.CreateRenderQuery(RQT_AbsoluteTime);;
+		}
+		WindowSize = BufferSize - FramesBehind;
+	}
+
+	void Begin(FRHICommandListImmediate& RHICmdList)
+	{
+		RotateQueryBuffer(StartQueries);
+		RHICmdList.EndRenderQuery(StartQueries[0]);
+	}
+
+
+	void End(FRHICommandListImmediate& RHICmdList)
+	{
+		RotateQueryBuffer(EndQueries);
+		RHICmdList.EndRenderQuery(EndQueries[0]);
+		QueriesFinished++;
+	}
+
+	void RotateQueryBuffer(TArray<FRenderQueryRHIRef> &QueryArray)
+	{
+		FRenderQueryRHIRef LastQuery = QueryArray.Last();
+
+		for (int32 i = QueryArray.Num() - 1; i > 0; --i)
+		{
+			QueryArray[i] = QueryArray[i - 1];
+		}
+		QueryArray[0] = LastQuery;
+	}
+
+
+	float GetElapsedAverage(FRHICommandListImmediate& RHICmdList, float &OutAvgTimeInSeconds);
+
+private:
+	int32 QueriesFailed;
+	uint32 WindowSize;
+	int32 QueriesFinished;
+	TArray<FRenderQueryRHIRef> StartQueries;
+	TArray<FRenderQueryRHIRef> EndQueries;
+};
diff --git a/Engine/Source/Runtime/Renderer/Private/PostProcess/PostProcessBokehDOF.cpp b/Engine/Source/Runtime/Renderer/Private/PostProcess/PostProcessBokehDOF.cpp
--- a/Engine/Source/Runtime/Renderer/Private/PostProcess/PostProcessBokehDOF.cpp
+++ b/Engine/Source/Runtime/Renderer/Private/PostProcess/PostProcessBokehDOF.cpp
@@ -70,7 +70,6 @@ public:
 	FShaderParameter VisualizeColors;
 	FShaderParameter CursorPos;
 	FShaderResourceParameter MiniFontTexture;
-	FShaderParameter CircleDofParams;
 	
 	/** Initialization constructor. */
 	PostProcessVisualizeDOFPS(const ShaderMetaType::CompiledShaderInitializerType& Initializer)
@@ -82,14 +81,13 @@ public:
 		MiniFontTexture.Bind(Initializer.ParameterMap, TEXT("MiniFontTexture"));
 		VisualizeColors.Bind(Initializer.ParameterMap, TEXT("VisualizeColors"));
 		CursorPos.Bind(Initializer.ParameterMap, TEXT("CursorPos"));
-		CircleDofParams.Bind(Initializer.ParameterMap,TEXT("CircleDofParams"));
 	}
 
 	// FShader interface.
 	virtual bool Serialize(FArchive& Ar) override
 	{
 		bool bShaderHasOutdatedParameters = FGlobalShader::Serialize(Ar);
-		Ar << PostprocessParameter << DeferredParameters << MiniFontTexture << DepthOfFieldParams << VisualizeColors << CursorPos << CircleDofParams;
+		Ar << PostprocessParameter << DeferredParameters << MiniFontTexture << DepthOfFieldParams << VisualizeColors << CursorPos;
 		return bShaderHasOutdatedParameters;
 	}
 
@@ -142,8 +140,6 @@ public:
 
 			SetShaderValueArray(Context.RHICmdList, ShaderRHI, VisualizeColors, Colors, 2);
 		}
-
-		SetShaderValue(Context.RHICmdList, ShaderRHI, CircleDofParams, CircleDofCoc(Context.View));
 	}
 
 	static const TCHAR* GetSourceFilename()
diff --git a/Engine/Source/Runtime/Renderer/Private/PostProcess/PostProcessCircleDOF.cpp b/Engine/Source/Runtime/Renderer/Private/PostProcess/PostProcessCircleDOF.cpp
--- a/Engine/Source/Runtime/Renderer/Private/PostProcess/PostProcessCircleDOF.cpp
+++ b/Engine/Source/Runtime/Renderer/Private/PostProcess/PostProcessCircleDOF.cpp
@@ -39,47 +39,57 @@ float ComputeFocalLengthFromFov(const FSceneView& View)
 
 // Convert f-stop and focal distance into projected size in half resolution pixels.
 // Setup depth based blur.
-FVector4 CircleDofCoc(const FSceneView& View)
+FVector4 CircleDofHalfCoc(const FSceneView& View)
 {
-	float FocalLengthInMM = ComputeFocalLengthFromFov(View);
+	static const auto CVar = IConsoleManager::Get().FindTConsoleVariableDataInt(TEXT("r.DepthOfFieldQuality"));
+	bool bDepthOfField = View.Family->EngineShowFlags.DepthOfField && CVar->GetValueOnRenderThread() > 0;
+
+	FVector4 Ret(0, 1, 0, 0);
+
+	if(bDepthOfField && View.FinalPostProcessSettings.DepthOfFieldMethod == DOFM_CircleDOF)
+	{
+		float FocalLengthInMM = ComputeFocalLengthFromFov(View);
 	 
-	// Convert focal distance in world position to mm (from cm to mm)
-	float FocalDistanceInMM = View.FinalPostProcessSettings.DepthOfFieldFocalDistance * 10.0f;
+		// Convert focal distance in world position to mm (from cm to mm)
+		float FocalDistanceInMM = View.FinalPostProcessSettings.DepthOfFieldFocalDistance * 10.0f;
+
+		// Convert f-stop, focal length, and focal distance to
+		// projected circle of confusion size at infinity in mm.
+		//
+		// coc = f * f / (n * (d - f))
+		// where,
+		//   f = focal length
+		//   d = focal distance
+		//   n = fstop (where n is the "n" in "f/n")
+		float Radius = FMath::Square(FocalLengthInMM) / (View.FinalPostProcessSettings.DepthOfFieldFstop * (FocalDistanceInMM - FocalLengthInMM));
+
+		// Scale so that APS-C 24.576 mm = full frame.
+		// Convert mm to pixels.
+		float Width = (float)View.ViewRect.Width();
+		Radius = Radius * Width * (1.0f/24.576f);
+
+		// Convert diameter to radius at half resolution (algorithm radius is at half resolution).
+		Radius *= 0.25f;
+
+		// Comment out for now, allowing settings which the algorithm cannot cleanly do.
+		#if 0
+			// Limit to algorithm max size.
+			if(Radius > 6.0f) 
+			{
+				Radius = 6.0f; 
+			}
+		#endif
 
-	// Convert f-stop, focal length, and focal distance to
-	// projected circle of confusion size at infinity in mm.
-	//
-	// coc = f*f / (n * (d - f))
-	// where,
-	//   f = focal length
-	//   d = focal distance
-	//   n = fstop (where n is the "n" in "f/n")
-	float Radius = FMath::Square(FocalLengthInMM) / (View.FinalPostProcessSettings.DepthOfFieldFstop * (FocalDistanceInMM - FocalLengthInMM));
-
-	// Scale so that APS-C 24.576 mm = full frame.
-	// Convert mm to pixels.
-	float Width = (float)View.ViewRect.Width();
-	Radius = Radius * Width * (1.0f/24.576f);
-
-	// Convert diameter to radius at half resolution (algorithm radius is at half resolution).
-	Radius *= 0.25f;
-
-	// Comment out for now, allowing settings which the algorithm cannot cleanly do.
-	#if 0
-		// Limit to algorithm max size.
-		if(Radius > 6.0f) 
-		{
-			Radius = 6.0f; 
-		}
-	#endif
-
-	// The DepthOfFieldDepthBlurAmount = km at which depth blur is 50%.
-	// Need to convert to cm here.
-	return FVector4(
-		Radius, 
-		1.0f/(View.FinalPostProcessSettings.DepthOfFieldDepthBlurAmount * 100000.0),
-		View.FinalPostProcessSettings.DepthOfFieldDepthBlurRadius * Width / 1920.0f,
-		Width / 1920.0f);
+		// The DepthOfFieldDepthBlurAmount = km at which depth blur is 50%.
+		// Need to convert to cm here.
+		Ret = FVector4(
+			Radius, 
+			1.0f / (View.FinalPostProcessSettings.DepthOfFieldDepthBlurAmount * 100000.0f),
+			View.FinalPostProcessSettings.DepthOfFieldDepthBlurRadius * Width / 1920.0f,
+			Width / 1920.0f);
+	}
+
+	return Ret;
 }
 
 /** Encapsulates the Circle DOF setup pixel shader. */
@@ -106,7 +116,6 @@ public:
 	FPostProcessPassParameters PostprocessParameter;
 	FDeferredPixelShaderParameters DeferredParameters;
 	FShaderParameter DepthOfFieldParams;
-	FShaderParameter CircleDofParams;
 
 	/** Initialization constructor. */
 	FPostProcessCircleDOFSetupPS(const ShaderMetaType::CompiledShaderInitializerType& Initializer)
@@ -115,14 +124,13 @@ public:
 		PostprocessParameter.Bind(Initializer.ParameterMap);
 		DeferredParameters.Bind(Initializer.ParameterMap);
 		DepthOfFieldParams.Bind(Initializer.ParameterMap,TEXT("DepthOfFieldParams"));
-		CircleDofParams.Bind(Initializer.ParameterMap,TEXT("CircleDofParams"));
 	}
 
 	// FShader interface.
 	virtual bool Serialize(FArchive& Ar) override
 	{
 		bool bShaderHasOutdatedParameters = FGlobalShader::Serialize(Ar);
-		Ar << PostprocessParameter << DeferredParameters << DepthOfFieldParams << CircleDofParams;
+		Ar << PostprocessParameter << DeferredParameters << DepthOfFieldParams;
 		return bShaderHasOutdatedParameters;
 	}
 
@@ -143,8 +151,6 @@ public:
 
 			SetShaderValueArray(Context.RHICmdList, ShaderRHI, DepthOfFieldParams, DepthOfFieldParamValues, 2);
 		}
-
-		SetShaderValue(Context.RHICmdList, ShaderRHI, CircleDofParams, CircleDofCoc(Context.View));
 	}
 };
 
@@ -294,7 +300,6 @@ public:
 	FPostProcessPassParameters PostprocessParameter;
 	FDeferredPixelShaderParameters DeferredParameters;
 	FShaderParameter DepthOfFieldParams;
-	FShaderParameter CircleDofParams;
 
 	/** Initialization constructor. */
 	FPostProcessCircleDOFDilatePS(const ShaderMetaType::CompiledShaderInitializerType& Initializer)
@@ -303,14 +308,13 @@ public:
 		PostprocessParameter.Bind(Initializer.ParameterMap);
 		DeferredParameters.Bind(Initializer.ParameterMap);
 		DepthOfFieldParams.Bind(Initializer.ParameterMap,TEXT("DepthOfFieldParams"));
-		CircleDofParams.Bind(Initializer.ParameterMap,TEXT("CircleDofParams"));
 	}
 
 	// FShader interface.
 	virtual bool Serialize(FArchive& Ar) override
 	{
 		bool bShaderHasOutdatedParameters = FGlobalShader::Serialize(Ar);
-		Ar << PostprocessParameter << DeferredParameters << DepthOfFieldParams << CircleDofParams;
+		Ar << PostprocessParameter << DeferredParameters << DepthOfFieldParams;
 		return bShaderHasOutdatedParameters;
 	}
 
@@ -331,8 +335,6 @@ public:
 
 			SetShaderValueArray(Context.RHICmdList, ShaderRHI, DepthOfFieldParams, DepthOfFieldParamValues, 2);
 		}
-
-		SetShaderValue(Context.RHICmdList, ShaderRHI, CircleDofParams, CircleDofCoc(Context.View));
 	}
 };
 
@@ -512,7 +514,6 @@ public:
 	FDeferredPixelShaderParameters DeferredParameters;
 	FShaderParameter DepthOfFieldParams;
 	FShaderParameter RandomOffset;
-	FShaderParameter CircleDofParams;
 
 	/** Initialization constructor. */
 	FPostProcessCircleDOFPS(const ShaderMetaType::CompiledShaderInitializerType& Initializer)
@@ -522,14 +523,13 @@ public:
 		DeferredParameters.Bind(Initializer.ParameterMap);
 		DepthOfFieldParams.Bind(Initializer.ParameterMap,TEXT("DepthOfFieldParams"));
 		RandomOffset.Bind(Initializer.ParameterMap, TEXT("RandomOffset"));
-		CircleDofParams.Bind(Initializer.ParameterMap,TEXT("CircleDofParams"));
 	}
 
 	// FShader interface.
 	virtual bool Serialize(FArchive& Ar) override
 	{
 		bool bShaderHasOutdatedParameters = FGlobalShader::Serialize(Ar);
-		Ar << PostprocessParameter << DeferredParameters << DepthOfFieldParams << RandomOffset << CircleDofParams;
+		Ar << PostprocessParameter << DeferredParameters << DepthOfFieldParams << RandomOffset;
 		return bShaderHasOutdatedParameters;
 	}
 
@@ -565,8 +565,6 @@ public:
 		FVector2D RandomOffsetValue;
 		TemporalRandom2(&RandomOffsetValue, Context.View.Family->FrameNumber);
 		SetShaderValue(Context.RHICmdList, ShaderRHI, RandomOffset, RandomOffsetValue);
-
-		SetShaderValue(Context.RHICmdList, ShaderRHI, CircleDofParams, CircleDofCoc(Context.View));
 	}
 	
 	static const TCHAR* GetSourceFilename()
@@ -744,7 +742,6 @@ public:
 	FDeferredPixelShaderParameters DeferredParameters;
 	FShaderParameter DepthOfFieldUVLimit;
 	FShaderParameter RandomOffset;
-	FShaderParameter CircleDofParams;
 
 	/** Initialization constructor. */
 	FPostProcessCircleDOFRecombinePS(const ShaderMetaType::CompiledShaderInitializerType& Initializer)
@@ -754,14 +751,13 @@ public:
 		DeferredParameters.Bind(Initializer.ParameterMap);
 		DepthOfFieldUVLimit.Bind(Initializer.ParameterMap,TEXT("DepthOfFieldUVLimit"));
 		RandomOffset.Bind(Initializer.ParameterMap, TEXT("RandomOffset"));
-		CircleDofParams.Bind(Initializer.ParameterMap,TEXT("CircleDofParams"));
 	}
 
 	// FShader interface.
 	virtual bool Serialize(FArchive& Ar) override
 	{
 		bool bShaderHasOutdatedParameters = FGlobalShader::Serialize(Ar);
-		Ar << PostprocessParameter << DeferredParameters << DepthOfFieldUVLimit << RandomOffset << CircleDofParams;
+		Ar << PostprocessParameter << DeferredParameters << DepthOfFieldUVLimit << RandomOffset;
 		return bShaderHasOutdatedParameters;
 	}
 
@@ -787,8 +783,6 @@ public:
 		FVector2D RandomOffsetValue;
 		TemporalRandom2(&RandomOffsetValue, Context.View.Family->FrameNumber);
 		SetShaderValue(Context.RHICmdList, ShaderRHI, RandomOffset, RandomOffsetValue);
-
-		SetShaderValue(Context.RHICmdList, ShaderRHI, CircleDofParams, CircleDofCoc(Context. View));
 	}
 	
 	static const TCHAR* GetSourceFilename()
diff --git a/Engine/Source/Runtime/Renderer/Private/PostProcess/PostProcessCircleDOF.h b/Engine/Source/Runtime/Renderer/Private/PostProcess/PostProcessCircleDOF.h
--- a/Engine/Source/Runtime/Renderer/Private/PostProcess/PostProcessCircleDOF.h
+++ b/Engine/Source/Runtime/Renderer/Private/PostProcess/PostProcessCircleDOF.h
@@ -98,4 +98,4 @@ private:
 // @return in mm, assuming this sensor: APS-C 24.576 mm
 float ComputeFocalLengthFromFov(const FSceneView& View);
 
-FVector4 CircleDofCoc(const FSceneView& View);
+FVector4 CircleDofHalfCoc(const FSceneView& View);
diff --git a/Engine/Source/Runtime/Renderer/Private/SceneRendering.cpp b/Engine/Source/Runtime/Renderer/Private/SceneRendering.cpp
--- a/Engine/Source/Runtime/Renderer/Private/SceneRendering.cpp
+++ b/Engine/Source/Runtime/Renderer/Private/SceneRendering.cpp
@@ -15,6 +15,7 @@
 #include "FXSystem.h"
 #include "SceneViewExtension.h"
 #include "PostProcessBusyWait.h"
+#include "PostProcessCircleDOF.h"
 #include "SceneUtils.h"
 #include "LightGrid.h"
 
@@ -825,6 +826,7 @@ TUniformBufferRef<FViewUniformShaderParameters> FViewInfo::CreateUniformBuffer(
 								(((float)ViewRect.Max.X / BufferSize.X) - OneScenePixelUVSize.X) , 
 								(((float)ViewRect.Max.Y / BufferSize.Y) - OneScenePixelUVSize.Y) );
 	ViewUniformShaderParameters.SceneTextureMinMax = SceneTexMinMax;
+	ViewUniformShaderParameters.CircleDOFParams = CircleDofHalfCoc(*this);
 
 	FScene* Scene = (FScene*)Family->Scene;
 	ERHIFeatureLevel::Type FeatureLevel = Scene == nullptr ? GMaxRHIFeatureLevel : Scene->GetFeatureLevel();
