Fix compile error in debug. FWorldContext and GEngine are used in OnlineSubsystemUtils.h but not explicitly included.

This seems like more of a bandaid.

diff --git a/Engine/Source/Runtime/Online/Hotfix/Private/HotfixPrivatePCH.h b/Engine/Source/Runtime/Online/Hotfix/Private/HotfixPrivatePCH.h
--- a/Engine/Source/Runtime/Online/Hotfix/Private/HotfixPrivatePCH.h
+++ b/Engine/Source/Runtime/Online/Hotfix/Private/HotfixPrivatePCH.h
@@ -4,4 +4,5 @@
 
 #include "Core.h"
 #include "CoreUObject.h"
+#include "Engine.h"
 #include "Online.h"
