Fix for non-unity builds

diff --git a/Engine/Source/Runtime/Online/OnlineSubsystemUtils/Public/OnlineSubsystemUtils.h b/Engine/Source/Runtime/Online/OnlineSubsystemUtils/Public/OnlineSubsystemUtils.h
--- a/Engine/Source/Runtime/Online/OnlineSubsystemUtils/Public/OnlineSubsystemUtils.h
+++ b/Engine/Source/Runtime/Online/OnlineSubsystemUtils/Public/OnlineSubsystemUtils.h
@@ -3,6 +3,7 @@
 #pragma once
 
 #include "Core.h"
+#include "Engine.h"
 #include "ModuleManager.h"
 #include "OnlineSubsystemUtilsModule.h"
 #include "Online.h"
