Copying //UE4/Orion-Staging to Dev-Main (Originating from //Orion/Main at CL-2777663)

diff --git a/Engine/Source/Runtime/MovieScene/Public/IMovieSceneSpawnRegister.h b/Engine/Source/Runtime/MovieScene/Public/IMovieSceneSpawnRegister.h
new file mode 100644
--- /dev/null
+++ b/Engine/Source/Runtime/MovieScene/Public/IMovieSceneSpawnRegister.h
@@ -0,0 +1,132 @@
+// Copyright 1998-2015 Epic Games, Inc. All Rights Reserved.
+
+#pragma once
+
+#include "MovieSceneSequenceInstance.h"
+
+class IMovieScenePlayer;
+class FMovieSceneSequenceInstance;
+
+/**
+ * Class responsible for managing spawnables in a movie scene
+ */
+class IMovieSceneSpawnRegister : public TSharedFromThis<IMovieSceneSpawnRegister>
+{
+public:
+	virtual ~IMovieSceneSpawnRegister() {}
+
+	/**
+	 * Spawn an object for the specified GUID, from the specified sequence instance.
+	 *
+	 * @param Object 		ID of the object to spawn
+	 * @param Instance 		Sequence instance that is spawning the object
+	 * @param Player 		Movie scene player that is ultimately responsible for spawning the object
+	 * @return the spawned object, or nullptr on failure
+	 */
+	virtual UObject* SpawnObject(const FGuid& Object, FMovieSceneSequenceInstance& Instance, IMovieScenePlayer& Player) = 0;
+
+	/**
+	 * Destroy a specific previously spawned object
+	 *
+	 * @param Object 		ID of the object to destroy
+	 * @param Instance 		Sequence instance that is destroying the object
+	 * @param Player 		Movie scene player that is ultimately responsible for destroying the object
+	 */
+	virtual void DestroySpawnedObject(const FGuid& Object, FMovieSceneSequenceInstance& Instance, IMovieScenePlayer& Player) = 0;
+
+	/**
+	 * Destroy all objects spawned, and owned, by the specified sequence instance
+	 * @note this only destorys objects that the specified sequence instance has direct ownership of (ownership == InnerSequence)
+	 *
+	 * @param Instance 		Sequence instance that is destroying the objects
+	 * @param Player 		Movie scene player that is ultimately responsible for destroying the objects
+	 */
+	virtual void DestroyObjectsOwnedByInstance(FMovieSceneSequenceInstance& Instance, IMovieScenePlayer& Player) = 0;
+
+	/**
+	 * Destroy all objects that have ever been spawned by the specified sequence instance (regardless of ownership)
+	 *
+	 * @param Instance 		Sequence instance that is destroying the objects
+	 * @param Player 		Movie scene player that is ultimately responsible for destroying the objects
+	 */
+	virtual void DestroyObjectsSpawnedByInstance(FMovieSceneSequenceInstance& Instance, IMovieScenePlayer& Player) = 0;
+
+	/**
+	 * Purge any memory of any objects that are considered externally owned
+
+	 * @param Player 		Movie scene player that is ultimately responsible for destroying the objects
+	 */
+	virtual void ForgetExternallyOwnedSpawnedObjects(IMovieScenePlayer& Player) = 0;
+
+	/**
+	 * Destroy all objects that this spawn register has responsibility for (Ownership != External)
+	 *
+	 * @param Player 		Movie scene player that is ultimately responsible for destroying the objects
+	 */
+	virtual void DestroyAllOwnedObjects(IMovieScenePlayer& Player) = 0;
+
+	/**
+	 * Destroy all objects that we have ever spawned
+	 *
+	 * @param Player 		Movie scene player that is ultimately responsible for destroying the objects
+	 */
+	virtual void DestroyAllObjects(IMovieScenePlayer& Player) = 0;
+
+	/**
+	 * Called before a sequence instance is about to update
+	 *
+	 * @param Instance 		The sequence instance that is about to update
+	 * @param Player 		The player that is playing the movie scene
+	 */
+	virtual void PreUpdateSequenceInstance(FMovieSceneSequenceInstance& Instance, IMovieScenePlayer& Player) { }
+
+	/**
+	 * Called after a sequence instance has updated
+	 *
+	 * @param Instance 		The sequence instance that has updated
+	 * @param Player 		The player that is playing the movie scene
+	 */
+	virtual void PostUpdateSequenceInstance(FMovieSceneSequenceInstance& Instance, IMovieScenePlayer& Player) { }
+};
+
+
+class FNullMovieSceneSpawnRegister : public IMovieSceneSpawnRegister
+{
+public:
+	// These 2 functions ensure because they should never be called for sequences that don't support spawnables
+	virtual UObject* SpawnObject(const FGuid& Object, FMovieSceneSequenceInstance&, IMovieScenePlayer&) override { check(false); return nullptr; }
+	virtual void DestroySpawnedObject(const FGuid& Object, FMovieSceneSequenceInstance&, IMovieScenePlayer&) override { check(false); }
+	virtual void DestroyObjectsOwnedByInstance(FMovieSceneSequenceInstance&, IMovieScenePlayer&) override { }
+	virtual void DestroyObjectsSpawnedByInstance(FMovieSceneSequenceInstance&, IMovieScenePlayer&) override { }
+	virtual void DestroyAllOwnedObjects(IMovieScenePlayer&) override { }
+	virtual void ForgetExternallyOwnedSpawnedObjects(IMovieScenePlayer& Player) override { }
+	virtual void DestroyAllObjects(IMovieScenePlayer&) override { }
+};
+
+/**
+ *	Helper key type for mapping a guid and sequence instance to a specific value
+ */
+struct FMovieSceneSpawnRegisterKey
+{
+	FMovieSceneSpawnRegisterKey(const FGuid& InBindingId, FMovieSceneSequenceInstance& InSequenceInstance)
+		: BindingId(InBindingId)
+		, SequenceInstance(InSequenceInstance.AsShared())
+	{
+	}
+
+	bool operator==(const FMovieSceneSpawnRegisterKey& Other) const
+	{
+		return BindingId == Other.BindingId && SequenceInstance == Other.SequenceInstance;
+	}
+	
+	friend uint32 GetTypeHash(const FMovieSceneSpawnRegisterKey& Key)
+	{
+		return HashCombine(GetTypeHash(Key.BindingId), GetTypeHash(Key.SequenceInstance));
+	}
+
+	/** BindingId of the object binding */
+	FGuid BindingId;
+
+	/** Sequence instance that spawned the object */
+	TWeakPtr<FMovieSceneSequenceInstance> SequenceInstance;
+};
\ No newline at end of file
diff --git a/Engine/Source/Runtime/MovieScene/Public/MovieScene.h b/Engine/Source/Runtime/MovieScene/Public/MovieScene.h
--- a/Engine/Source/Runtime/MovieScene/Public/MovieScene.h
+++ b/Engine/Source/Runtime/MovieScene/Public/MovieScene.h
@@ -28,6 +28,7 @@ struct FMovieSceneExpansionState
 	bool bExpanded;
 };
 
+
 /**
  * Editor only data that needs to be saved between sessions for editing but has no runtime purpose
  */
@@ -50,6 +51,29 @@ struct FMovieSceneEditorData
 };
 
 
+/**
+ * Structure for labels that can be assigned to movie scene tracks.
+ */
+USTRUCT()
+struct FMovieSceneTrackLabels
+{
+	GENERATED_BODY()
+
+	UPROPERTY()
+	TArray<FString> Strings;
+
+	void FromString(const FString& LabelString)
+	{
+		LabelString.ParseIntoArray(Strings, TEXT(" "));
+	}
+
+	FString ToString() const
+	{
+		return FString::Join(Strings, TEXT(" "));
+	}
+};
+
+
 /**
  * Implements a movie scene asset.
  */
@@ -342,6 +366,7 @@ public:
 	*/
 	void ReplaceBinding(const FGuid& OldGuid, const FGuid& NewGuid, const FString& Name);
 
+#if WITH_EDITORONLY_DATA
 	/**
 	 * Get the display name of the object with the specified identifier.
 	 *
@@ -350,10 +375,9 @@ public:
 	 */
 	void SetObjectDisplayName(const FGuid& ObjectId, const FText& DisplayName)
 	{
-#if WITH_EDITORONLY_DATA
 		ObjectsToDisplayNames.Add(ObjectId.ToString(), DisplayName);
-#endif
 	}
+#endif
 
 	/**
 	 * Set the start and end playback positions (playback range) for this movie scene
@@ -373,6 +397,33 @@ public:
 	{
 		return EditorData;
 	}
+
+	/**
+	 * Get all known track labels.
+	 *
+	 * @param OutLabels Will contain the collection of known labels.
+	 * @return The number labels returned.
+	 */
+	int32 GetAllLabels(TArray<FString>& OutLabels) const;
+
+	/**
+	 * Get an object's track label.
+	 *
+	 * @param ObjectId The object to get the label for.
+	 * @return The label string.
+	 */
+	FMovieSceneTrackLabels& GetObjectLabels(const FGuid& ObjectId) 
+	{
+		return ObjectsToLabels.FindOrAdd(ObjectId.ToString());
+	}
+
+	/**
+	 * Check whether the specified track label exists.
+	 *
+	 * @param Label The label to check.
+	 * @return true if the label exists, false otherwise.
+	 */
+	bool LabelExists(const FString& Label) const;
 #endif
 
 protected:
@@ -386,14 +437,13 @@ protected:
 
 protected:
 
-	/**
-	 * Called after this object has been deserialized
-	 */
+	/** Called after this object has been deserialized */
 	virtual void PostLoad() override;
 
-	/**
-	 * Perform legacy upgrade of time ranges
-	 */
+	/** Called before this object is being deserialized. */
+	virtual void PreSave() override;
+
+	/** Perform legacy upgrade of time ranges */
 	void UpgradeTimeRanges();
 
 private:
@@ -427,6 +477,10 @@ private:
 	UPROPERTY()
 	TMap<FString, FText> ObjectsToDisplayNames;
 
+	/** Maps object GUIDs to user defined labels. */
+	UPROPERTY()
+	TMap<FString, FMovieSceneTrackLabels> ObjectsToLabels;
+
 	/** Editor only data that needs to be saved between sessions for editing but has no runtime purpose */
 	UPROPERTY()
 	FMovieSceneEditorData EditorData;
diff --git a/Engine/Source/Runtime/MovieScene/Public/MovieSceneCommonHelpers.h b/Engine/Source/Runtime/MovieScene/Public/MovieSceneCommonHelpers.h
--- a/Engine/Source/Runtime/MovieScene/Public/MovieSceneCommonHelpers.h
+++ b/Engine/Source/Runtime/MovieScene/Public/MovieSceneCommonHelpers.h
@@ -7,8 +7,14 @@
 class MOVIESCENE_API MovieSceneHelpers
 {
 public:
+
+	/**
+	 * Gets the sections that were traversed over between the current time and the previous time, including overlapping sections
+	 */
+	static TArray<UMovieSceneSection*> GetAllTraversedSections( const TArray<UMovieSceneSection*>& Sections, float CurrentTime, float PreviousTime );
+
 	/**
-	 * Gets the sections that were traversed over between the current time and the previous time
+	 * Gets the sections that were traversed over between the current time and the previous time, excluding overlapping sections (highest wins)
 	 */
 	static TArray<UMovieSceneSection*> GetTraversedSections( const TArray<UMovieSceneSection*>& Sections, float CurrentTime, float PreviousTime );
 
@@ -130,7 +136,7 @@ public:
 	 * @param InValue   The value to set
 	 */
 	template <typename ValueType>
-	void SetCurrentValue(const UObject* Object, ValueType InValue)
+	void SetCurrentValue(const UObject* Object, const ValueType& InValue)
 	{
 		FPropertyAndFunction PropAndFunction = RuntimeObjectToFunctionMap.FindRef(Object);
 
diff --git a/Engine/Source/Runtime/MovieScene/Public/MovieScenePossessable.h b/Engine/Source/Runtime/MovieScene/Public/MovieScenePossessable.h
--- a/Engine/Source/Runtime/MovieScene/Public/MovieScenePossessable.h
+++ b/Engine/Source/Runtime/MovieScene/Public/MovieScenePossessable.h
@@ -87,23 +87,23 @@ public:
 	}
 
 	/**
-	 * Get the guid of this possessable's parent spawnable
+	 * Get the guid of this possessable's parent, if applicable
 	 *
 	 * @return The guid.
 	 */
-	const FGuid& GetParentSpawnable() const
+	const FGuid& GetParent() const
 	{
-		return ParentSpawnableGuid;
+		return ParentGuid;
 	}
 
 	/**
-	 * Set the guid of this possessable's parent spawnable
+	 * Set the guid of this possessable's parent
 	 *
-	 * @param InParentSpawnableGuid The guid of this possessable's parent spawnable.
+	 * @param InParentGuid The guid of this possessable's parent.
 	 */
-	void SetParentSpawnable(const FGuid& InParentSpawnableGuid)
+	void SetParent(const FGuid& InParentGuid)
 	{
-		ParentSpawnableGuid = InParentSpawnableGuid;
+		ParentGuid = InParentGuid;
 	}
 
 private:
@@ -127,7 +127,7 @@ private:
 	UPROPERTY()
 	UClass* PossessedObjectClass;
 
-	/** GUID relating to this possessable's parent spawnable, if applicable. */
+	/** GUID relating to this possessable's parent, if applicable. */
 	UPROPERTY()
-	FGuid ParentSpawnableGuid;
+	FGuid ParentGuid;
 };
diff --git a/Engine/Source/Runtime/MovieScene/Public/MovieSceneSection.h b/Engine/Source/Runtime/MovieScene/Public/MovieSceneSection.h
--- a/Engine/Source/Runtime/MovieScene/Public/MovieSceneSection.h
+++ b/Engine/Source/Runtime/MovieScene/Public/MovieSceneSection.h
@@ -190,6 +190,18 @@ public:
 	/** Gets the row index for this section */
 	int32 GetRowIndex() const {return RowIndex;}
 	
+	/** Sets this section's priority over overlapping sections (higher wins) */
+	void SetOverlapPriority(int32 NewPriority)
+	{
+		OverlapPriority = NewPriority;
+	}
+
+	/** Gets this section's priority over overlapping sections (higher wins) */
+	int32 GetOverlapPriority() const
+	{
+		return OverlapPriority;
+	}
+
 	/**
 	 * Adds a key to a rich curve, finding an existing key to modify or adding a new one
 	 *
@@ -256,6 +268,10 @@ private:
 	UPROPERTY()
 	int32 RowIndex;
 
+	/** This section's priority over overlapping sections */
+	UPROPERTY()
+	int32 OverlapPriority;
+
 	/** Toggle whether this section is active/inactive */
 	UPROPERTY(EditAnywhere, Category="Section")
 	uint32 bIsActive : 1;
diff --git a/Engine/Source/Runtime/MovieScene/Public/MovieSceneSequence.h b/Engine/Source/Runtime/MovieScene/Public/MovieSceneSequence.h
--- a/Engine/Source/Runtime/MovieScene/Public/MovieSceneSequence.h
+++ b/Engine/Source/Runtime/MovieScene/Public/MovieSceneSequence.h
@@ -26,9 +26,10 @@ public:
 	 * 
 	 * @param ObjectId The guid used to map to the possessable object.  Note the guid can be bound to multiple objects at once
 	 * @param PossessedObject The runtime object which was possessed.
+	 * @param Context Optional context required to bind the specified object (for instance, a parent spawnable object)
 	 * @see UnbindPossessableObjects
 	 */
-	virtual void BindPossessableObject(const FGuid& ObjectId, UObject& PossessedObject) PURE_VIRTUAL(UMovieSceneSequence::BindPossessableObject,);
+	virtual void BindPossessableObject(const FGuid& ObjectId, UObject& PossessedObject, UObject* Context) PURE_VIRTUAL(UMovieSceneSequence::BindPossessableObject,);
 	
 	/**
 	 * Check whether the given object can be possessed by this animation.
@@ -39,20 +40,13 @@ public:
 	virtual bool CanPossessObject(UObject& Object) const PURE_VIRTUAL(UMovieSceneSequence::CanPossessObject, return false;);
 
 	/**
-	 * Gets the possessed or spawned object for the specified identifier.
+	 * Finds the possessed for the specified identifier.
 	 *
 	 * @param ObjectId The unique identifier of the object.
+	 * @param Context Optional context to use to find the required object (for instance, a parent spawnable object)
 	 * @return The object, or nullptr if not found.
 	 */
-	virtual UObject* FindObject(const FGuid& ObjectId) const PURE_VIRTUAL(UMovieSceneSequence::FindObject, return nullptr;);
-
-	/**
-	 * Find the identifier for a possessed or spawned object.
-	 *
-	 * @param The object to get the identifier for.
-	 * @return The object identifier, or an invalid GUID if not found.
-	 */
-	virtual FGuid FindObjectId( UObject& Object) const PURE_VIRTUAL(UMovieSceneSequence::FindGuidForObject, return FGuid(););
+	virtual UObject* FindPossessableObject(const FGuid& ObjectId, UObject* Context) const PURE_VIRTUAL(UMovieSceneSequence::FindPossessedObject, return nullptr;);
 
 	/**
 	 * Get the movie scene that controls this animation.
@@ -78,28 +72,6 @@ public:
 	 * @return true if objects can be spawned by sequencer, false if only existing objects can be possessed.
 	 */
 	virtual bool AllowsSpawnableObjects() const { return false; }
-	
-	/**
-	 * Spawn an object relating to the specified object ID
-	 * 
-	 * @param ObjectId The ID of the object to spawn
-	 * @return The newly spawned or previously-spawned object, or nullptr
-	 */
-	virtual UObject* SpawnObject(const FGuid& ObjectId) { return nullptr; }
-
-	/**
-	 * Destroy a previously spawned object relating to the specified object ID
-	 * 
-	 * @param ObjectId The ID of the object to destroy
-	 */
-	virtual void DestroySpawnedObject(const FGuid& ObjectId) { }
-
-	/**
-	 * Destroy a previously spawned object relating to the specified GUID
-	 *
-	 * @return true if the object was destroyed, false otherwise
-	 */
-	virtual void DestroyAllSpawnedObjects() {}
 
 	/**
 	 * Unbinds all possessable objects from the provided GUID.
@@ -112,15 +84,6 @@ public:
 public:
 
 #if WITH_EDITOR
-	/**
-	 * 
-	 * Tries to get a display name for the binding represented by the guid.
-	 *
-	 * @param ObjectGuid The guid for the object binding.
-	 * @param DisplayName The display name for the object binding.
-	 * @returns true if DisplayName has been set to a valid display name, otherwise false.
-	 */
-	virtual bool TryGetObjectDisplayName(const FGuid& ObjectId, FText& OutDisplayName) const PURE_VIRTUAL(UMovieSceneSequence::TryGetObjectDisplayName, return false;);
 
 	/**
 	 * Get the display name for this movie sequence
diff --git a/Engine/Source/Runtime/MovieScene/Public/MovieSceneSequenceInstance.h b/Engine/Source/Runtime/MovieScene/Public/MovieSceneSequenceInstance.h
--- a/Engine/Source/Runtime/MovieScene/Public/MovieSceneSequenceInstance.h
+++ b/Engine/Source/Runtime/MovieScene/Public/MovieSceneSequenceInstance.h
@@ -4,7 +4,9 @@
 
 #include "IMovieSceneTrackInstance.h"
 
+class IMovieScenePlayer;
 class IMovieSceneTrackInstance;
+
 class UMovieSceneTrack;
 class UMovieSceneSequence;
 
@@ -26,12 +28,45 @@ public:
 	 * @param InMovieSceneSequence The sequence that this instance represents.
 	 */
 	MOVIESCENE_API FMovieSceneSequenceInstance(const UMovieSceneSequence& InMovieSceneSequence);
+	MOVIESCENE_API ~FMovieSceneSequenceInstance();
+
+	/**
+	 * Find the identifier for a possessed or spawned object.
+	 *
+	 * @param Object The object to get the identifier for.
+	 * @return The object identifier, or an invalid GUID if not found.
+	 */
+	MOVIESCENE_API FGuid FindObjectId(UObject& Object) const;
+
+	/**
+	 * Find the identifier for the parent of the specified object
+	 *
+	 * @param Object The object to get the parent identifier for.
+	 * @return The parent's identifier, or an invalid GUID if not applicable.
+	 */
+	MOVIESCENE_API FGuid FindParentObjectId(UObject& Object) const;
+	
+	/**
+	 * Find the object relating to the specified GUID
+	 *
+	 * @param ObjectId The object identifier
+	 * @return The object, or nullptr
+	 */
+	MOVIESCENE_API UObject* FindObject(const FGuid& ObjectId, const IMovieScenePlayer& Player) const;
+
+	/**
+	 * Find the spawned object relating to the specified GUID
+	 *
+	 * @param ObjectId The object identifier
+	 * @return The object, or nullptr
+	 */
+	MOVIESCENE_API UObject* FindSpawnedObject(const FGuid& ObjectId) const;
 
 	/** Save state of the objects that this movie scene controls. */
-	MOVIESCENE_API void SaveState(class IMovieScenePlayer& Player);
+	MOVIESCENE_API void SaveState(IMovieScenePlayer& Player);
 
 	/** Restore state of the objects that this movie scene controls. */
-	MOVIESCENE_API void RestoreState(class IMovieScenePlayer& Player);
+	MOVIESCENE_API void RestoreState(IMovieScenePlayer& Player);
 
 	/**
 	 * Updates this movie scene.
@@ -39,7 +74,7 @@ public:
 	 * @param Position The local playback position.
 	 * @param Player Movie scene player interface for interaction with runtime data.
 	 */
-	MOVIESCENE_API void Update(float Position, float LastPosition, class IMovieScenePlayer& Player);
+	MOVIESCENE_API void Update(float Position, float LastPosition, IMovieScenePlayer& Player);
 
 	/**
 	 * Refreshes the existing instance.
@@ -47,7 +82,7 @@ public:
 	 * Called when something significant about movie scene data changes (like adding or removing a track).
 	 * Instantiates all new tracks found and removes instances for tracks that no longer exist.
 	 */
-	MOVIESCENE_API void RefreshInstance(class IMovieScenePlayer& Player);
+	MOVIESCENE_API void RefreshInstance(IMovieScenePlayer& Player);
 
 	/**
 	 * Spawn an object relating to the specified object ID
@@ -56,7 +91,7 @@ public:
 	 * @param Player The movie scene player responsible for this instance
 	 * @return The newly spawned or previously-spawned object, or nullptr
 	 */
-	MOVIESCENE_API void SpawnObject(const FGuid& ObjectId, IMovieScenePlayer& Player);
+	MOVIESCENE_API void OnObjectSpawned(const FGuid& ObjectId, UObject& SpawnedObject, IMovieScenePlayer& Player);
 
 	/**
 	 * Destroy a previously spawned object relating to the specified object ID
@@ -64,7 +99,7 @@ public:
 	 * @param ObjectId The ID of the object to destroy
 	 * @param Player The movie scene player responsible for this instance
 	 */
-	MOVIESCENE_API void DestroySpawnedObject(const FGuid& ObjectId, IMovieScenePlayer& Player);
+	MOVIESCENE_API void OnSpawnedObjectDestroyed(const FGuid& ObjectId, IMovieScenePlayer& Player);
 
 	/**
 	 * Get the sequence associated with this instance.
@@ -91,9 +126,9 @@ public:
 
 protected:
 
-	void RefreshInstanceMap(const TArray<UMovieSceneTrack*>& Tracks, const TArray<UObject*>& RuntimeObjects, FMovieSceneInstanceMap& TrackInstances, class IMovieScenePlayer& Player);
+	void RefreshInstanceMap(const TArray<UMovieSceneTrack*>& Tracks, const TArray<UObject*>& RuntimeObjects, FMovieSceneInstanceMap& TrackInstances, IMovieScenePlayer& Player);
 
-	void UpdateInternal(float Position, float LastPosition, class IMovieScenePlayer& Player, EMovieSceneUpdatePass UpdatePass);
+	void UpdateInternal(float Position, float LastPosition, IMovieScenePlayer& Player, EMovieSceneUpdatePass UpdatePass);
 
 	/** Update the object binding instance for the specified object */
 	void UpdateObjectBinding(const FGuid& ObjectId, IMovieScenePlayer& Player);
@@ -126,6 +161,9 @@ private:
 	/** All object binding instances */
 	TMap<FGuid, FMovieSceneObjectBindingInstance> ObjectBindingInstances;
 
+	/** A map of object ID -> spawned object */
+	TMap<FGuid, TWeakObjectPtr<UObject>> SpawnedObjects;
+
 	/** Cached time range for the movie scene */
 	TRange<float> TimeRange;
 };
diff --git a/Engine/Source/Runtime/MovieScene/Public/MovieSceneSpawnable.h b/Engine/Source/Runtime/MovieScene/Public/MovieSceneSpawnable.h
--- a/Engine/Source/Runtime/MovieScene/Public/MovieSceneSpawnable.h
+++ b/Engine/Source/Runtime/MovieScene/Public/MovieSceneSpawnable.h
@@ -4,6 +4,18 @@
 
 #include "MovieSceneSpawnable.generated.h"
 
+UENUM()
+enum class ESpawnOwnership : uint8
+{
+	/** The object's lifetime is managed by the sequence that spawned it */
+	InnerSequence,
+
+	/** The object's lifetime is managed by the outermost sequence */
+	MasterSequence,
+
+	/** Once spawned, the object's lifetime is managed externally. */
+	External,
+};
 
 /**
  * MovieSceneSpawnable describes an object that can be spawned for this MovieScene
@@ -23,6 +35,7 @@ public:
 		: Guid(FGuid::NewGuid())
 		, Name(InitName)
 		, GeneratedClass(InitClass)
+		, Ownership(ESpawnOwnership::InnerSequence)
 	{ }
 
 public:
@@ -88,6 +101,22 @@ public:
 		return ChildPossessables;
 	}
 
+	/**
+	 * Get a value indicating what is responsible for this object once it's spawned
+	 */
+	ESpawnOwnership GetSpawnOwnership() const
+	{
+		return Ownership;
+	}
+
+	/**
+	 * Set a value indicating what is responsible for this object once it's spawned
+	 */
+	void SetSpawnOwnership(ESpawnOwnership InOwnership)
+	{
+		Ownership = InOwnership;
+	}
+
 private:
 
 	/** Unique identifier of the spawnable object. */
@@ -110,4 +139,8 @@ private:
 	// @todo sequencer: This should be a TSet, but they don't duplicate correctly atm
 	UPROPERTY()
 	TArray<FGuid> ChildPossessables;
+
+	/** Property indicating where ownership responsibility for this object lies */
+	UPROPERTY()
+	ESpawnOwnership Ownership;
 };
diff --git a/Engine/Source/Runtime/MovieSceneCapture/Private/AutomatedLevelSequenceCapture.cpp b/Engine/Source/Runtime/MovieSceneCapture/Private/AutomatedLevelSequenceCapture.cpp
--- a/Engine/Source/Runtime/MovieSceneCapture/Private/AutomatedLevelSequenceCapture.cpp
+++ b/Engine/Source/Runtime/MovieSceneCapture/Private/AutomatedLevelSequenceCapture.cpp
@@ -43,6 +43,9 @@ void UAutomatedLevelSequenceCapture::SetLevelSequenceActor(ALevelSequenceActor*
 
 void UAutomatedLevelSequenceCapture::Initialize(TWeakPtr<FSceneViewport> InViewport)
 {
+	// Apply command-line overrides from parent class first. This needs to be called before setting up the capture strategy with the desired frame rate.
+	Super::Initialize(InViewport);
+
 	// Apply command-line overrides
 	{
 		FString LevelSequenceAssetPath;
@@ -126,8 +129,6 @@ void UAutomatedLevelSequenceCapture::Initialize(TWeakPtr<FSceneViewport> InViewp
 	CaptureState = ELevelSequenceCaptureState::DelayBeforeWarmUp;
 	RemainingDelaySeconds = FMath::Max( 0.0f, DelayBeforeWarmUp );
 	CaptureStrategy = MakeShareable(new FFixedTimeStepCaptureStrategy(Settings.FrameRate));
-
-	Super::Initialize(InViewport);
 }
 
 
diff --git a/Engine/Source/Runtime/MovieSceneTracks/Private/Sections/MovieSceneSkeletalAnimationSection.cpp b/Engine/Source/Runtime/MovieSceneTracks/Private/Sections/MovieSceneSkeletalAnimationSection.cpp
--- a/Engine/Source/Runtime/MovieSceneTracks/Private/Sections/MovieSceneSkeletalAnimationSection.cpp
+++ b/Engine/Source/Runtime/MovieSceneTracks/Private/Sections/MovieSceneSkeletalAnimationSection.cpp
@@ -67,7 +67,7 @@ void UMovieSceneSkeletalAnimationSection::GetSnapTimes(TArray<float>& OutSnapTim
 	float SeqLength = (GetSequenceLength() - (GetStartOffset() + GetEndOffset())) / AnimPlayRate;
 
 	// Snap to the repeat times
-	while (CurrentTime <= GetEndTime() && !FMath::IsNearlyZero(GetDuration()))
+	while (CurrentTime <= GetEndTime() && !FMath::IsNearlyZero(GetDuration()) && SeqLength > 0)
 	{
 		if (CurrentTime >= GetStartTime())
 		{
diff --git a/Engine/Source/Runtime/MovieSceneTracks/Private/TrackInstances/MovieSceneShotTrackInstance.cpp b/Engine/Source/Runtime/MovieSceneTracks/Private/TrackInstances/MovieSceneShotTrackInstance.cpp
--- a/Engine/Source/Runtime/MovieSceneTracks/Private/TrackInstances/MovieSceneShotTrackInstance.cpp
+++ b/Engine/Source/Runtime/MovieSceneTracks/Private/TrackInstances/MovieSceneShotTrackInstance.cpp
@@ -35,7 +35,7 @@ void FMovieSceneShotTrackInstance::RefreshInstance(const TArray<UObject*>& Runti
 	for (int32 ShotIndex = 0; ShotIndex < ShotSections.Num(); ++ShotIndex)
 	{
 		UMovieSceneShotSection* ShotSection = CastChecked<UMovieSceneShotSection>(ShotSections[ShotIndex]);
-		AcquireCameraForShot(ShotIndex, ShotSection->GetCameraGuid(), SequenceInstance);
+		AcquireCameraForShot(ShotIndex, ShotSection->GetCameraGuid(), SequenceInstance, Player);
 	}
 }
 
@@ -50,15 +50,19 @@ void FMovieSceneShotTrackInstance::Update(float Position, float LastPosition, co
 	UObject* CameraObject = nullptr;
 	const TArray<UMovieSceneSection*>& ShotSections = ShotTrack->GetAllSections();
 
-	for (int32 ShotIndex = 0; ShotIndex < ShotSections.Num(); ++ShotIndex)
+	UMovieSceneShotSection* NearestShotSection = Cast<UMovieSceneShotSection>(MovieSceneHelpers::FindNearestSectionAtTime(ShotSections, Position));
+	if (NearestShotSection != nullptr)
 	{
-		UMovieSceneShotSection* ShotSection = CastChecked<UMovieSceneShotSection>(ShotSections[ShotIndex]);
-
-		if (ShotSection->GetRange().Contains(Position))
+		for (int32 ShotIndex = 0; ShotIndex < ShotSections.Num(); ++ShotIndex)
 		{
-			CameraObject = AcquireCameraForShot(ShotIndex, ShotSection->GetCameraGuid(), SequenceInstance);
+			UMovieSceneShotSection* ShotSection = CastChecked<UMovieSceneShotSection>(ShotSections[ShotIndex]);
+
+			if (ShotSection == NearestShotSection)
+			{
+				CameraObject = AcquireCameraForShot(ShotIndex, ShotSection->GetCameraGuid(), SequenceInstance, Player);
 
-			break;
+				break;
+			}
 		}
 	}
 
@@ -77,13 +81,13 @@ void FMovieSceneShotTrackInstance::Update(float Position, float LastPosition, co
 /* FMovieSceneShotTrackInstance implementation
  *****************************************************************************/
 
-UObject* FMovieSceneShotTrackInstance::AcquireCameraForShot(int32 ShotIndex, const FGuid& CameraGuid, FMovieSceneSequenceInstance& SequenceInstance)
+UObject* FMovieSceneShotTrackInstance::AcquireCameraForShot(int32 ShotIndex, const FGuid& CameraGuid, FMovieSceneSequenceInstance& SequenceInstance, const IMovieScenePlayer& Player)
 {
 	auto& CachedCamera = CachedCameraObjects[ShotIndex];
 
 	if (!CachedCamera.IsValid())
 	{
-		CachedCamera = SequenceInstance.GetSequence()->FindObject(CameraGuid);
+		CachedCamera = SequenceInstance.FindObject(CameraGuid, Player);
 	}
 
 	return CachedCamera.Get();
diff --git a/Engine/Source/Runtime/MovieSceneTracks/Private/TrackInstances/MovieSceneSkeletalAnimationTrackInstance.cpp b/Engine/Source/Runtime/MovieSceneTracks/Private/TrackInstances/MovieSceneSkeletalAnimationTrackInstance.cpp
--- a/Engine/Source/Runtime/MovieSceneTracks/Private/TrackInstances/MovieSceneSkeletalAnimationTrackInstance.cpp
+++ b/Engine/Source/Runtime/MovieSceneTracks/Private/TrackInstances/MovieSceneSkeletalAnimationTrackInstance.cpp
@@ -8,6 +8,41 @@
 #include "Matinee/MatineeAnimInterface.h"
 #include "MovieSceneCommonHelpers.h"
 
+struct FAnimEvalTimes
+{
+	float Current, Previous;
+
+	/** Map the specified times onto the specified anim section */
+	static FAnimEvalTimes MapTimesToAnimation(float ThisPosition, float LastPosition, UMovieSceneSkeletalAnimationSection* AnimSection)
+	{
+		ThisPosition = FMath::Clamp(ThisPosition, AnimSection->GetStartTime(), AnimSection->GetEndTime());
+		LastPosition = FMath::Clamp(LastPosition, AnimSection->GetStartTime(), AnimSection->GetEndTime());
+
+		ThisPosition -= 1 / 1000.0f;
+
+		float AnimPlayRate = FMath::IsNearlyZero(AnimSection->GetPlayRate()) ? 1.0f : AnimSection->GetPlayRate();
+		float SeqLength = AnimSection->GetSequenceLength() - (AnimSection->GetStartOffset() + AnimSection->GetEndOffset());
+
+		ThisPosition = (ThisPosition - AnimSection->GetStartTime()) * AnimPlayRate;
+		ThisPosition = FMath::Fmod(ThisPosition, SeqLength);
+		ThisPosition += AnimSection->GetStartOffset();
+		if (AnimSection->GetReverse())
+		{
+			ThisPosition = (SeqLength - (ThisPosition - AnimSection->GetStartOffset())) + AnimSection->GetStartOffset();
+		}
+
+		LastPosition = (LastPosition - AnimSection->GetStartTime()) * AnimPlayRate;
+		LastPosition = FMath::Fmod(LastPosition, SeqLength);
+		LastPosition += AnimSection->GetStartOffset();
+		if (AnimSection->GetReverse())
+		{
+			LastPosition = (SeqLength - (LastPosition - AnimSection->GetStartOffset())) + AnimSection->GetStartOffset();
+		}
+
+		FAnimEvalTimes EvalTimes = { ThisPosition, LastPosition };
+		return EvalTimes;
+	}
+};
 
 FMovieSceneSkeletalAnimationTrackInstance::FMovieSceneSkeletalAnimationTrackInstance( UMovieSceneSkeletalAnimationTrack& InAnimationTrack )
 {
@@ -20,7 +55,6 @@ FMovieSceneSkeletalAnimationTrackInstance::~FMovieSceneSkeletalAnimationTrackIns
 	// @todo Sequencer Need to find some way to call PreviewFinishAnimControl (needs the runtime objects)
 }
 
-
 void FMovieSceneSkeletalAnimationTrackInstance::Update( float Position, float LastPosition, const TArray<UObject*>& RuntimeObjects, class IMovieScenePlayer& Player, FMovieSceneSequenceInstance& SequenceInstance, EMovieSceneUpdatePass UpdatePass ) 
 {
 	// @todo Sequencer gameplay update has a different code path than editor update for animation
@@ -35,53 +69,36 @@ void FMovieSceneSkeletalAnimationTrackInstance::Update( float Position, float La
 			// cbb: If there is no overlapping section, evaluate the closest section only if the current time is before it.
 			if (AnimSection == nullptr)
 			{
-				UMovieSceneSection* NearestSection = MovieSceneHelpers::FindNearestSectionAtTime(AnimationTrack->GetAllSections(), Position);
-				if (NearestSection != nullptr)
+				AnimSection = Cast<UMovieSceneSkeletalAnimationSection>(MovieSceneHelpers::FindNearestSectionAtTime(AnimationTrack->GetAllSections(), Position));
+				if (AnimSection != nullptr && (Position >= AnimSection->GetStartTime() || Position <= AnimSection->GetEndTime()))
 				{
-					AnimSection = Cast<UMovieSceneSkeletalAnimationSection>(NearestSection);
-					if (Position < AnimSection->GetStartTime())
-					{
-						Position = AnimSection->GetStartTime();
-					}
-					else if ( Position > AnimSection->GetEndTime() )
-					{
-						Position = AnimSection->GetEndTime();
-					}
-					else
-					{
-						AnimSection = nullptr;
-					}
+					AnimSection = nullptr;
 				}
 			}
 
 			if (AnimSection && AnimSection->IsActive())
 			{
-				Position -= 1 / 1000.0f;
+				const FAnimEvalTimes EvalTimes = FAnimEvalTimes::MapTimesToAnimation(Position, LastPosition, AnimSection);
 
 				int32 ChannelIndex = 0;
 				UAnimSequence* AnimSequence = AnimSection->GetAnimSequence();
 
-				float AnimPlayRate = FMath::IsNearlyZero(AnimSection->GetPlayRate()) ? 1.0f : AnimSection->GetPlayRate();
-				float AnimPosition = (Position - AnimSection->GetStartTime()) * AnimPlayRate;
-				float SeqLength = AnimSection->GetSequenceLength() - (AnimSection->GetStartOffset() + AnimSection->GetEndOffset());
-
-				AnimPosition = FMath::Fmod(AnimPosition, SeqLength);
-				AnimPosition += AnimSection->GetStartOffset();
-
-				if (AnimSection->GetReverse())
-				{
-					AnimPosition = (SeqLength - (AnimPosition - AnimSection->GetStartOffset())) + AnimSection->GetStartOffset();
-				}
-
+				// NOTE: The order of the bLooping and bFireNotifiers parameters are swapped in the preview and non-preview SetAnimPosition functions.
+				const bool bLooping = false;
 				if (GIsEditor && !GWorld->HasBegunPlay())
 				{
+					// Pass an absolute delta time to the preview anim evaluation
+					float Delta = FMath::Abs(EvalTimes.Current - EvalTimes.Previous);
+					const bool bFireNotifies = true;
 					AnimInterface->PreviewBeginAnimControl(nullptr);
-					AnimInterface->PreviewSetAnimPosition(AnimSection->GetSlotName(), ChannelIndex, AnimSequence, AnimPosition, true, false, 0.f);
+					AnimInterface->PreviewSetAnimPosition(AnimSection->GetSlotName(), ChannelIndex, AnimSequence, EvalTimes.Current, bLooping, bFireNotifies, Delta);
 				}
 				else
 				{
+					// Don't fire notifies at runtime since they will be fired through the standard animation tick path.
+					const bool bFireNotifies = false;
 					AnimInterface->BeginAnimControl(nullptr);
-					AnimInterface->SetAnimPosition(AnimSection->GetSlotName(), ChannelIndex, AnimSequence, AnimPosition, true, false);
+					AnimInterface->SetAnimPosition(AnimSection->GetSlotName(), ChannelIndex, AnimSequence, EvalTimes.Current, bFireNotifies, bLooping);
 				}
 			}
 		}
diff --git a/Engine/Source/Runtime/MovieSceneTracks/Private/TrackInstances/MovieSceneSpawnTrackInstance.cpp b/Engine/Source/Runtime/MovieSceneTracks/Private/TrackInstances/MovieSceneSpawnTrackInstance.cpp
--- a/Engine/Source/Runtime/MovieSceneTracks/Private/TrackInstances/MovieSceneSpawnTrackInstance.cpp
+++ b/Engine/Source/Runtime/MovieSceneTracks/Private/TrackInstances/MovieSceneSpawnTrackInstance.cpp
@@ -3,6 +3,7 @@
 #include "MovieSceneTracksPrivatePCH.h"
 #include "MovieSceneSequenceInstance.h"
 
+#include "MovieSceneSequence.h"
 #include "MovieSceneSpawnTrack.h"
 #include "MovieSceneSpawnTrackInstance.h"
 
@@ -14,9 +15,15 @@ FMovieSceneSpawnTrackInstance::FMovieSceneSpawnTrackInstance(UMovieSceneSpawnTra
 
 void FMovieSceneSpawnTrackInstance::Update(float Position, float LastPosition, const TArray<UObject*>& RuntimeObjects, IMovieScenePlayer& Player, FMovieSceneSequenceInstance& SequenceInstance, EMovieSceneUpdatePass UpdatePass)
 {
-	if (!SequenceInstance.GetTimeRange().Contains(Position))
+	IMovieSceneSpawnRegister& SpawnRegister = Player.GetSpawnRegister();
+	FMovieSceneSpawnable* Spawnable = SequenceInstance.GetSequence()->GetMovieScene()->FindSpawnable(Track->GetObjectId());
+
+	TRange<float> Range = SequenceInstance.GetTimeRange();
+
+	// If we're evaluating outside of the instance's time range, and the sequence owns the spawnable, there's no reason to evaluate - it should already be destroyed
+	if (Spawnable && Spawnable->GetSpawnOwnership() == ESpawnOwnership::InnerSequence && !Range.Contains(Position) && !Range.Contains(LastPosition))
 	{
-		SequenceInstance.DestroySpawnedObject(Track->GetObject(), Player);
+		SpawnRegister.DestroySpawnedObject(Track->GetObjectId(), SequenceInstance, Player);
 		return;
 	}
 
@@ -26,26 +33,13 @@ void FMovieSceneSpawnTrackInstance::Update(float Position, float LastPosition, c
 		// Spawn the object if needed
 		if (bIsSpawned && RuntimeObjects.Num() == 0)
 		{
-			SequenceInstance.SpawnObject(Track->GetObject(), Player);
+			SpawnRegister.SpawnObject(Track->GetObjectId(), SequenceInstance, Player);
 		}
 
 		// Destroy the object if needed
 		if (!bIsSpawned && RuntimeObjects.Num() != 0)
 		{
-			SequenceInstance.DestroySpawnedObject(Track->GetObject(), Player);
+			SpawnRegister.DestroySpawnedObject(Track->GetObjectId(), SequenceInstance, Player);
 		}
 	}
-}
-
-void FMovieSceneSpawnTrackInstance::RestoreState(const TArray<UObject*>& RuntimeObjects, IMovieScenePlayer& Player, FMovieSceneSequenceInstance& SequenceInstance)
-{
-	if (RuntimeObjects.Num())
-	{
-		SequenceInstance.DestroySpawnedObject(Track->GetObject(), Player);
-	}
-}
-
-void FMovieSceneSpawnTrackInstance::ClearInstance(IMovieScenePlayer& Player, FMovieSceneSequenceInstance& SequenceInstance)
-{
-	SequenceInstance.DestroySpawnedObject(Track->GetObject(), Player);
 }
\ No newline at end of file
diff --git a/Engine/Source/Runtime/MovieSceneTracks/Private/TrackInstances/MovieSceneVectorTrackInstance.cpp b/Engine/Source/Runtime/MovieSceneTracks/Private/TrackInstances/MovieSceneVectorTrackInstance.cpp
--- a/Engine/Source/Runtime/MovieSceneTracks/Private/TrackInstances/MovieSceneVectorTrackInstance.cpp
+++ b/Engine/Source/Runtime/MovieSceneTracks/Private/TrackInstances/MovieSceneVectorTrackInstance.cpp
@@ -119,7 +119,7 @@ void FMovieSceneVectorTrackInstance::Update( float Position, float LastPosition,
 				FVector Value(Vector.X, Vector.Y, Vector.Z);
 				for(UObject* Object : RuntimeObjects)
 				{
-					PropertyBindings->CallFunction<FVector2D>(Object, &Value);
+					PropertyBindings->CallFunction<FVector>(Object, &Value);
 				}
 				break;
 			}
@@ -127,7 +127,7 @@ void FMovieSceneVectorTrackInstance::Update( float Position, float LastPosition,
 			{
 				for(UObject* Object : RuntimeObjects)
 				{
-					PropertyBindings->CallFunction<FVector2D>(Object, &Vector);
+					PropertyBindings->CallFunction<FVector4>(Object, &Vector);
 				}
 				break;
 			}
diff --git a/Engine/Source/Runtime/MovieSceneTracks/Private/Tracks/MovieSceneParticleTrack.cpp b/Engine/Source/Runtime/MovieSceneTracks/Private/Tracks/MovieSceneParticleTrack.cpp
--- a/Engine/Source/Runtime/MovieSceneTracks/Private/Tracks/MovieSceneParticleTrack.cpp
+++ b/Engine/Source/Runtime/MovieSceneTracks/Private/Tracks/MovieSceneParticleTrack.cpp
@@ -86,5 +86,11 @@ UMovieSceneSection* UMovieSceneParticleTrack::CreateNewSection()
 	return NewObject<UMovieSceneParticleSection>( this );
 }
 
+#if WITH_EDITORONLY_DATA
+FText UMovieSceneParticleTrack::GetDisplayName() const
+{
+	return LOCTEXT("DisplayName", "Particle System");
+}
+#endif
 
 #undef LOCTEXT_NAMESPACE
diff --git a/Engine/Source/Runtime/MovieSceneTracks/Private/Tracks/MovieScenePropertyTrack.cpp b/Engine/Source/Runtime/MovieSceneTracks/Private/Tracks/MovieScenePropertyTrack.cpp
--- a/Engine/Source/Runtime/MovieSceneTracks/Private/Tracks/MovieScenePropertyTrack.cpp
+++ b/Engine/Source/Runtime/MovieSceneTracks/Private/Tracks/MovieScenePropertyTrack.cpp
@@ -125,6 +125,7 @@ UMovieSceneSection* UMovieScenePropertyTrack::FindOrAddSection( float Time )
 
 	// Add a new section that starts and ends at the same time
 	UMovieSceneSection* NewSection = CreateNewSection();
+	NewSection->SetFlags(RF_Transactional);
 	NewSection->SetStartTime( Time );
 	NewSection->SetEndTime( Time );
 
diff --git a/Engine/Source/Runtime/MovieSceneTracks/Public/Sections/MovieScene3DTransformSection.h b/Engine/Source/Runtime/MovieSceneTracks/Public/Sections/MovieScene3DTransformSection.h
--- a/Engine/Source/Runtime/MovieSceneTracks/Public/Sections/MovieScene3DTransformSection.h
+++ b/Engine/Source/Runtime/MovieSceneTracks/Public/Sections/MovieScene3DTransformSection.h
@@ -16,8 +16,13 @@ namespace EKey3DTransformChannel
 		Scale = 0x00000004,
 		All = Translation | Rotation | Scale
 	};
-}
 
+	enum ValueType
+	{
+		Key,
+		Default
+	};
+}
 
 /**
 * Stores information about a transform for the purpose of adding keys to a transform section
@@ -56,14 +61,16 @@ struct FTransformData
 
 struct FTransformKey
 {
-	FTransformKey( EKey3DTransformChannel::Type InChannel, EAxis::Type InAxis, float InValue, bool InbUnwindRotation )
+	FTransformKey( EKey3DTransformChannel::Type InChannel, EKey3DTransformChannel::ValueType InChannelValueType, EAxis::Type InAxis, float InValue, bool InbUnwindRotation )
 	{
 		Channel = InChannel;
+		ChannelValueType = InChannelValueType;
 		Axis = InAxis;
 		Value = InValue;
 		bUnwindRotation = InbUnwindRotation;
 	}
 	EKey3DTransformChannel::Type Channel;
+	EKey3DTransformChannel::ValueType ChannelValueType;
 	EAxis::Type Axis;
 	float Value;
 	bool bUnwindRotation;
diff --git a/Engine/Source/Runtime/MovieSceneTracks/Public/Sections/MovieSceneVectorSection.h b/Engine/Source/Runtime/MovieSceneTracks/Public/Sections/MovieSceneVectorSection.h
--- a/Engine/Source/Runtime/MovieSceneTracks/Public/Sections/MovieSceneVectorSection.h
+++ b/Engine/Source/Runtime/MovieSceneTracks/Public/Sections/MovieSceneVectorSection.h
@@ -14,16 +14,23 @@ enum class EKeyVectorChannel
 	W
 };
 
+enum class EKeyVectorValueType
+{
+	Key,
+	Default
+};
 
 struct MOVIESCENETRACKS_API FVectorKey
 {
-	FVectorKey( EKeyVectorChannel InChannel, float InValue )
+	FVectorKey( EKeyVectorChannel InChannel, float InValue, EKeyVectorValueType InValueType )
 	{
 		Channel = InChannel;
 		Value = InValue;
+		ValueType = InValueType;
 	}
 	EKeyVectorChannel Channel;
 	float Value;
+	EKeyVectorValueType ValueType;
 };
 
 
diff --git a/Engine/Source/Runtime/MovieSceneTracks/Public/TrackInstances/MovieSceneShotTrackInstance.h b/Engine/Source/Runtime/MovieSceneTracks/Public/TrackInstances/MovieSceneShotTrackInstance.h
--- a/Engine/Source/Runtime/MovieSceneTracks/Public/TrackInstances/MovieSceneShotTrackInstance.h
+++ b/Engine/Source/Runtime/MovieSceneTracks/Public/TrackInstances/MovieSceneShotTrackInstance.h
@@ -37,7 +37,7 @@ protected:
 	 * @param SequenceInstance The sequence instance that owns this track instance.
 	 * @return The camera object, or nullptr if not found.
 	 */
-	UObject* AcquireCameraForShot(int32 ShotIndex, const FGuid& CameraGuid, FMovieSceneSequenceInstance& SequenceInstance);
+	UObject* AcquireCameraForShot(int32 ShotIndex, const FGuid& CameraGuid, FMovieSceneSequenceInstance& SequenceInstance, const IMovieScenePlayer& Player);
 
 private:
 	
diff --git a/Engine/Source/Runtime/MovieSceneTracks/Public/TrackInstances/MovieSceneSpawnTrackInstance.h b/Engine/Source/Runtime/MovieSceneTracks/Public/TrackInstances/MovieSceneSpawnTrackInstance.h
--- a/Engine/Source/Runtime/MovieSceneTracks/Public/TrackInstances/MovieSceneSpawnTrackInstance.h
+++ b/Engine/Source/Runtime/MovieSceneTracks/Public/TrackInstances/MovieSceneSpawnTrackInstance.h
@@ -16,12 +16,13 @@ public:
 
 	/** IMovieSceneTrackInstance interface */
 	virtual void SaveState(const TArray<UObject*>& RuntimeObjects, IMovieScenePlayer& Player, FMovieSceneSequenceInstance& SequenceInstance) override {}
-	virtual void RestoreState(const TArray<UObject*>& RuntimeObjects, IMovieScenePlayer& Player, FMovieSceneSequenceInstance& SequenceInstance) override;
+	virtual void RestoreState(const TArray<UObject*>& RuntimeObjects, IMovieScenePlayer& Player, FMovieSceneSequenceInstance& SequenceInstance) override {}
+	virtual void RefreshInstance(const TArray<UObject*>& RuntimeObjects, class IMovieScenePlayer& Player, FMovieSceneSequenceInstance& SequenceInstance) override {}
+	virtual void ClearInstance(IMovieScenePlayer& Player, FMovieSceneSequenceInstance& SequenceInstance) override {}
+	virtual EMovieSceneUpdatePass HasUpdatePasses() override { return MSUP_PreUpdate; }
+
 	virtual void Update(float Position, float LastPosition, const TArray<UObject*>& RuntimeObjects, class IMovieScenePlayer& Player, FMovieSceneSequenceInstance& SequenceInstance, EMovieSceneUpdatePass UpdatePass) override;
 
-	// recreate/destroy the spawned object here?
-	virtual void RefreshInstance(const TArray<UObject*>& RuntimeObjects, class IMovieScenePlayer& Player, FMovieSceneSequenceInstance& SequenceInstance) override {}
-	virtual void ClearInstance(IMovieScenePlayer& Player, FMovieSceneSequenceInstance& SequenceInstance) override;
 	virtual float EvalOrder() override { return -10000.f; }
 
 private:
diff --git a/Engine/Source/Runtime/MovieSceneTracks/Public/Tracks/MovieSceneMaterialTrack.h b/Engine/Source/Runtime/MovieSceneTracks/Public/Tracks/MovieSceneMaterialTrack.h
--- a/Engine/Source/Runtime/MovieSceneTracks/Public/Tracks/MovieSceneMaterialTrack.h
+++ b/Engine/Source/Runtime/MovieSceneTracks/Public/Tracks/MovieSceneMaterialTrack.h
@@ -2,7 +2,7 @@
 
 #pragma once
 
-#include "MovieSceneTrack.h"
+#include "MovieSceneNameableTrack.h"
 #include "MovieSceneParameterSection.h"
 #include "MovieSceneMaterialTrack.generated.h"
 
@@ -12,7 +12,7 @@
  */
 UCLASS( MinimalAPI )
 class UMovieSceneMaterialTrack
-	: public UMovieSceneTrack
+	: public UMovieSceneNameableTrack
 {
 	GENERATED_UCLASS_BODY()
 
diff --git a/Engine/Source/Runtime/MovieSceneTracks/Public/Tracks/MovieSceneParticleTrack.h b/Engine/Source/Runtime/MovieSceneTracks/Public/Tracks/MovieSceneParticleTrack.h
--- a/Engine/Source/Runtime/MovieSceneTracks/Public/Tracks/MovieSceneParticleTrack.h
+++ b/Engine/Source/Runtime/MovieSceneTracks/Public/Tracks/MovieSceneParticleTrack.h
@@ -45,6 +45,10 @@ public:
 	virtual void AddNewSection(float SectionTime);
 	virtual UMovieSceneSection* CreateNewSection() override;
 
+#if WITH_EDITORONLY_DATA
+	virtual FText GetDisplayName() const override;
+#endif
+
 private:
 
 	/** List of all particle sections. */
diff --git a/Engine/Source/Runtime/MovieSceneTracks/Public/Tracks/MovieSceneSkeletalAnimationTrack.h b/Engine/Source/Runtime/MovieSceneTracks/Public/Tracks/MovieSceneSkeletalAnimationTrack.h
--- a/Engine/Source/Runtime/MovieSceneTracks/Public/Tracks/MovieSceneSkeletalAnimationTrack.h
+++ b/Engine/Source/Runtime/MovieSceneTracks/Public/Tracks/MovieSceneSkeletalAnimationTrack.h
@@ -3,7 +3,7 @@
 #pragma once
 
 #include "MovieScene.h"
-#include "MovieSceneTrack.h"
+#include "MovieSceneNameableTrack.h"
 #include "MovieSceneSkeletalAnimationTrack.generated.h"
 
 
@@ -15,7 +15,7 @@ class UMovieSceneSection;
  */
 UCLASS( MinimalAPI )
 class UMovieSceneSkeletalAnimationTrack
-	: public UMovieSceneTrack
+	: public UMovieSceneNameableTrack
 {
 	GENERATED_UCLASS_BODY()
 
diff --git a/Engine/Source/Runtime/MovieSceneTracks/Public/Tracks/MovieSceneSpawnTrack.h b/Engine/Source/Runtime/MovieSceneTracks/Public/Tracks/MovieSceneSpawnTrack.h
--- a/Engine/Source/Runtime/MovieSceneTracks/Public/Tracks/MovieSceneSpawnTrack.h
+++ b/Engine/Source/Runtime/MovieSceneTracks/Public/Tracks/MovieSceneSpawnTrack.h
@@ -2,7 +2,7 @@
 
 #pragma once
 
-#include "MovieSceneTrack.h"
+#include "MovieSceneNameableTrack.h"
 #include "MovieSceneSpawnTrack.generated.h"
 
 
@@ -14,7 +14,7 @@ class UMovieSceneSection;
  */
 UCLASS( MinimalAPI )
 class UMovieSceneSpawnTrack
-	: public UMovieSceneTrack
+	: public UMovieSceneNameableTrack
 {
 public:
 
@@ -27,13 +27,13 @@ public:
 	bool Eval(float Position, float LastPostion, bool& bOutSpawned) const;
 
 	/** Get the object identifier that this spawn track controls */
-	const FGuid& GetObject() const
+	const FGuid& GetObjectId() const
 	{
 		return ObjectGuid;
 	}
 
 	/** Set the object identifier that this spawn track controls */
-	void SetObject(const FGuid& InGuid)
+	void SetObjectId(const FGuid& InGuid)
 	{
 		ObjectGuid = InGuid;
 	}
diff --git a/Engine/Source/Runtime/NetworkReplayStreaming/HttpNetworkReplayStreaming/Private/HttpNetworkReplayStreaming.cpp b/Engine/Source/Runtime/NetworkReplayStreaming/HttpNetworkReplayStreaming/Private/HttpNetworkReplayStreaming.cpp
--- a/Engine/Source/Runtime/NetworkReplayStreaming/HttpNetworkReplayStreaming/Private/HttpNetworkReplayStreaming.cpp
+++ b/Engine/Source/Runtime/NetworkReplayStreaming/HttpNetworkReplayStreaming/Private/HttpNetworkReplayStreaming.cpp
@@ -487,6 +487,36 @@ void FHttpNetworkReplayStreamer::FlushCheckpoint( const uint32 TimeInMS )
 	FlushCheckpointInternal( TimeInMS );
 }
 
+FQueuedGotoFakeCheckpoint::FQueuedGotoFakeCheckpoint() : FQueuedHttpRequest( EQueuedHttpRequestType::UploadingCustomEvent, nullptr )
+{
+}
+
+bool FQueuedGotoFakeCheckpoint::PreProcess( FHttpNetworkReplayStreamer* Streamer, const FString& ServerURL, const FString& SessionName )
+{
+	// Make sure to reset the checkpoint archive (this is how we signify that the engine should start from the beginning of the steam (we don't need a checkpoint for that))
+	Streamer->CheckpointArchive.Buffer.Empty();
+	Streamer->CheckpointArchive.Pos = 0;
+
+	// Completely reset our stream (we're going to start downloading from the start of the checkpoint)
+	Streamer->StreamArchive.Buffer.Empty();
+	Streamer->StreamArchive.Pos				= 0;
+	Streamer->StreamArchive.bAtEndOfReplay	= false;
+
+	// Reset our stream range
+	Streamer->StreamTimeRangeStart	= 0;
+	Streamer->StreamTimeRangeEnd	= 0;
+
+	// Reset chunk index
+	Streamer->StreamChunkIndex		= 0;
+
+	Streamer->GotoCheckpointDelegate.ExecuteIfBound( true, Streamer->LastGotoTimeInMS );
+	Streamer->GotoCheckpointDelegate = FOnCheckpointReadyDelegate();
+
+	Streamer->LastGotoTimeInMS = -1;
+
+	return true;
+}
+
 void FHttpNetworkReplayStreamer::GotoCheckpointIndex( const int32 CheckpointIndex, const FOnCheckpointReadyDelegate& Delegate )
 {
 	if ( GotoCheckpointDelegate.IsBound() || DownloadCheckpointIndex != -1 )
@@ -501,23 +531,8 @@ void FHttpNetworkReplayStreamer::GotoCheckpointIndex( const int32 CheckpointInde
 
 	if ( CheckpointIndex == -1 )
 	{
-		// Make sure to reset the checkpoint archive (this is how we signify that the engine should start from the beginning of the steam (we don't need a checkpoint for that))
-		CheckpointArchive.Buffer.Empty();
-		CheckpointArchive.Pos = 0;
-
-		// Completely reset our stream (we're going to start downloading from the start of the checkpoint)
-		StreamArchive.Buffer.Empty();
-		StreamArchive.Pos				= 0;
-		StreamArchive.bAtEndOfReplay	= false;
-
-		// Reset our stream range
-		StreamTimeRangeStart	= 0;
-		StreamTimeRangeEnd		= 0;
-
-		StreamChunkIndex		= 0;
-
-		Delegate.ExecuteIfBound( true, LastGotoTimeInMS );
-		LastGotoTimeInMS = -1;
+		GotoCheckpointDelegate = Delegate;
+		AddCustomRequestToQueue( TSharedPtr< FQueuedHttpRequest >( new FQueuedGotoFakeCheckpoint() ) );
 		return;
 	}
 
@@ -679,7 +694,7 @@ FQueuedHttpRequestAddEvent::FQueuedHttpRequestAddEvent( const FString& InName, c
 	Meta		= InMeta;
 }
 
-bool FQueuedHttpRequestAddEvent::PreProcess( const FString& ServerURL, const FString& SessionName )
+bool FQueuedHttpRequestAddEvent::PreProcess( FHttpNetworkReplayStreamer* Streamer, const FString& ServerURL, const FString& SessionName )
 {
 	if ( SessionName.IsEmpty() )
 	{
@@ -770,6 +785,23 @@ bool FHttpNetworkReplayStreamer::IsTaskPendingOrInFlight( const EQueuedHttpReque
 	return false;
 }
 
+void FHttpNetworkReplayStreamer::CancelInFlightOrPendingTask( const EQueuedHttpRequestType::Type Type )
+{
+	for ( int i = QueuedHttpRequests.Num() - 1; i >= 0; i-- )
+	{
+		if ( QueuedHttpRequests[i]->Type == Type )
+		{
+			QueuedHttpRequests.RemoveAt( i );
+		}
+	}
+
+	if ( InFlightHttpRequest.IsValid() && InFlightHttpRequest->Type == Type )
+	{
+		InFlightHttpRequest->Request->CancelRequest();
+		InFlightHttpRequest = NULL;
+	}
+}
+
 void FHttpNetworkReplayStreamer::ConditionallyDownloadNextChunk()
 {
 	if ( GotoCheckpointDelegate.IsBound() )
@@ -821,7 +853,7 @@ void FHttpNetworkReplayStreamer::ConditionallyDownloadNextChunk()
 	HttpRequest->SetURL( URL );
 	HttpRequest->SetVerb( TEXT( "GET" ) );
 
-	HttpRequest->OnProcessRequestComplete().BindRaw( this, &FHttpNetworkReplayStreamer::HttpDownloadFinished );
+	HttpRequest->OnProcessRequestComplete().BindRaw( this, &FHttpNetworkReplayStreamer::HttpDownloadFinished, StreamChunkIndex );
 
 	AddRequestToQueue( EQueuedHttpRequestType::DownloadingStream, HttpRequest );
 
@@ -1427,7 +1459,7 @@ void FHttpNetworkReplayStreamer::HttpDownloadHeaderFinished( FHttpRequestPtr Htt
 	}
 }
 
-void FHttpNetworkReplayStreamer::HttpDownloadFinished( FHttpRequestPtr HttpRequest, FHttpResponsePtr HttpResponse, bool bSucceeded )
+void FHttpNetworkReplayStreamer::HttpDownloadFinished( FHttpRequestPtr HttpRequest, FHttpResponsePtr HttpResponse, bool bSucceeded, int32 RequestedStreamChunkIndex )
 {
 	RequestFinished( EStreamerState::StreamingDown, EQueuedHttpRequestType::DownloadingStream, HttpRequest );
 
@@ -1435,6 +1467,15 @@ void FHttpNetworkReplayStreamer::HttpDownloadFinished( FHttpRequestPtr HttpReque
 
 	if ( bSucceeded && HttpResponse->GetResponseCode() == EHttpResponseCodes::Ok )
 	{
+		// Make sure our stream chunk index didn't change under our feet
+		if ( RequestedStreamChunkIndex != StreamChunkIndex )
+		{
+			UE_LOG( LogHttpReplay, Error, TEXT( "FHttpNetworkReplayStreamer::HttpDownloadFinished. StreamChunkIndex changed while request was in flight" ) );
+			StreamArchive.Buffer.Empty();
+			SetLastError( ENetworkReplayError::ServiceUnavailable );
+			return;
+		}
+
 		if ( HttpResponse->GetHeader( TEXT( "NumChunks" ) ) == TEXT( "" ) )
 		{
 			// Assume this is an implcit status update
@@ -1772,12 +1813,18 @@ bool FHttpNetworkReplayStreamer::ProcessNextHttpRequest()
 			return ProcessNextHttpRequest();
 		}
 
-		if ( !QueuedRequest->PreProcess( ServerURL, SessionName ) )
+		if ( !QueuedRequest->PreProcess( this, ServerURL, SessionName ) )
 		{
 			// This request failed, go ahead and process the next request
 			return ProcessNextHttpRequest();
 		}
 
+		// If this task has no http request, immediately go to the next task
+		if ( !QueuedRequest->Request.IsValid() )
+		{
+			return ProcessNextHttpRequest();
+		}
+
 		InFlightHttpRequest = QueuedRequest;
 
 		ProcessRequestInternal( InFlightHttpRequest->Request );
diff --git a/Engine/Source/Runtime/NetworkReplayStreaming/HttpNetworkReplayStreaming/Public/HttpNetworkReplayStreaming.h b/Engine/Source/Runtime/NetworkReplayStreaming/HttpNetworkReplayStreaming/Public/HttpNetworkReplayStreaming.h
--- a/Engine/Source/Runtime/NetworkReplayStreaming/HttpNetworkReplayStreaming/Public/HttpNetworkReplayStreaming.h
+++ b/Engine/Source/Runtime/NetworkReplayStreaming/HttpNetworkReplayStreaming/Public/HttpNetworkReplayStreaming.h
@@ -98,6 +98,8 @@ namespace EQueuedHttpRequestType
 	}
 };
 
+class FHttpNetworkReplayStreamer;
+
 class FQueuedHttpRequest
 {
 public:
@@ -112,7 +114,7 @@ public:
 	EQueuedHttpRequestType::Type		Type;
 	TSharedPtr< class IHttpRequest >	Request;
 
-	virtual bool PreProcess( const FString& ServerURL, const FString& SessionName )
+	virtual bool PreProcess( FHttpNetworkReplayStreamer* Streamer, const FString& ServerURL, const FString& SessionName )
 	{
 		return true;
 	}
@@ -131,7 +133,7 @@ public:
 	{
 	}
 
-	virtual bool PreProcess( const FString& ServerURL, const FString& SessionName ) override;
+	virtual bool PreProcess( FHttpNetworkReplayStreamer* Streamer, const FString& ServerURL, const FString& SessionName ) override;
 
 	FString		Name;
 	uint32		TimeInMS;
@@ -139,6 +141,21 @@ public:
 	FString		Meta;
 };
 
+/**
+* FQueuedGotoFakeCheckpoint
+*/
+class FQueuedGotoFakeCheckpoint : public FQueuedHttpRequest
+{
+public:
+	FQueuedGotoFakeCheckpoint();
+
+	virtual ~FQueuedGotoFakeCheckpoint()
+	{
+	}
+
+	virtual bool PreProcess( FHttpNetworkReplayStreamer* Streamer, const FString& ServerURL, const FString& SessionName ) override;
+};
+
 /**
  * Http network replay streaming manager
  */
@@ -183,6 +200,7 @@ public:
 	void StopUploading();
 	void DownloadHeader();
 	bool IsTaskPendingOrInFlight( const EQueuedHttpRequestType::Type Type ) const;
+	void CancelInFlightOrPendingTask( const EQueuedHttpRequestType::Type Type );
 	void ConditionallyDownloadNextChunk();
 	void RefreshViewer( const bool bFinal );
 	void ConditionallyRefreshViewer();
@@ -211,7 +229,7 @@ public:
 
 	void HttpStartDownloadingFinished( FHttpRequestPtr HttpRequest, FHttpResponsePtr HttpResponse, bool bSucceeded );
 	void HttpDownloadHeaderFinished( FHttpRequestPtr HttpRequest, FHttpResponsePtr HttpResponse, bool bSucceeded );
-	void HttpDownloadFinished( FHttpRequestPtr HttpRequest, FHttpResponsePtr HttpResponse, bool bSucceeded );
+	void HttpDownloadFinished( FHttpRequestPtr HttpRequest, FHttpResponsePtr HttpResponse, bool bSucceeded, int32 RequestedStreamChunkIndex );
 	void HttpDownloadCheckpointFinished( FHttpRequestPtr HttpRequest, FHttpResponsePtr HttpResponse, bool bSucceeded );
 	void HttpRefreshViewerFinished( FHttpRequestPtr HttpRequest, FHttpResponsePtr HttpResponse, bool bSucceeded );
 	void HttpStartUploadingFinished( FHttpRequestPtr HttpRequest, FHttpResponsePtr HttpResponse, bool bSucceeded );
diff --git a/Engine/Source/Runtime/Niagara/Private/NiagaraSequence.cpp b/Engine/Source/Runtime/Niagara/Private/NiagaraSequence.cpp
--- a/Engine/Source/Runtime/Niagara/Private/NiagaraSequence.cpp
+++ b/Engine/Source/Runtime/Niagara/Private/NiagaraSequence.cpp
@@ -17,7 +17,7 @@ UNiagaraSequence::UNiagaraSequence(const FObjectInitializer& ObjectInitializer)
  *****************************************************************************/
 
 
-void UNiagaraSequence::BindPossessableObject(const FGuid& ObjectId, UObject& PossessedObject)
+void UNiagaraSequence::BindPossessableObject(const FGuid& ObjectId, UObject& PossessedObject, UObject* Context)
 {
 }
 
@@ -28,18 +28,12 @@ bool UNiagaraSequence::CanPossessObject(UObject& Object) const
 }
 
 
-UObject* UNiagaraSequence::FindObject(const FGuid& ObjectId) const
+UObject* UNiagaraSequence::FindPossessableObject(const FGuid& ObjectId, UObject* Context) const
 {
 	return nullptr;
 }
 
 
-FGuid UNiagaraSequence::FindObjectId(UObject& Object) const
-{
-	return FGuid();
-}
-
-
 UMovieScene* UNiagaraSequence::GetMovieScene() const
 {
 	return MovieScene;
@@ -52,14 +46,6 @@ UObject* UNiagaraSequence::GetParentObject(UObject* Object) const
 }
 
 
-#if WITH_EDITOR
-bool UNiagaraSequence::TryGetObjectDisplayName(const FGuid& ObjectId, FText& OutDisplayName) const
-{
-	return false;
-}
-#endif
-
-
 void UNiagaraSequence::UnbindPossessableObjects(const FGuid& ObjectId)
 {
 }
diff --git a/Engine/Source/Runtime/Niagara/Public/NiagaraSequence.h b/Engine/Source/Runtime/Niagara/Public/NiagaraSequence.h
--- a/Engine/Source/Runtime/Niagara/Public/NiagaraSequence.h
+++ b/Engine/Source/Runtime/Niagara/Public/NiagaraSequence.h
@@ -19,18 +19,13 @@ public:
 
 	// UMovieSceneSequence overrides
 
-	virtual void BindPossessableObject(const FGuid& ObjectId, UObject& PossessedObject) override;
+	virtual void BindPossessableObject(const FGuid& ObjectId, UObject& PossessedObject, UObject* Context) override;
 	virtual bool CanPossessObject(UObject& Object) const override;
-	virtual UObject* FindObject(const FGuid& ObjectId) const override;
-	virtual FGuid FindObjectId(UObject& Object) const override;
+	virtual UObject* FindPossessableObject(const FGuid& ObjectId, UObject* Context) const override;
 	virtual UMovieScene* GetMovieScene() const override;
 	virtual UObject* GetParentObject(UObject* Object) const override;
 	virtual void UnbindPossessableObjects(const FGuid& ObjectId) override;
 
-#if WITH_EDITOR
-	virtual bool TryGetObjectDisplayName(const FGuid& ObjectId, FText& OutDisplayName) const override;
-#endif
-
 	/** Pointer to the movie scene that controls this sequence. */
 	UPROPERTY()
 	UMovieScene* MovieScene;
diff --git a/Engine/Source/Runtime/Online/Android/OnlineSubsystemGooglePlay/Private/OnlineExternalUIInterfaceGooglePlay.cpp b/Engine/Source/Runtime/Online/Android/OnlineSubsystemGooglePlay/Private/OnlineExternalUIInterfaceGooglePlay.cpp
--- a/Engine/Source/Runtime/Online/Android/OnlineSubsystemGooglePlay/Private/OnlineExternalUIInterfaceGooglePlay.cpp
+++ b/Engine/Source/Runtime/Online/Android/OnlineSubsystemGooglePlay/Private/OnlineExternalUIInterfaceGooglePlay.cpp
@@ -78,6 +78,11 @@ bool FOnlineExternalUIGooglePlay::ShowWebURL(const FString& Url, const FShowWebU
 	return false;
 }
 
+bool FOnlineExternalUIGooglePlay::CloseWebURL()
+{
+	return false;
+}
+
 bool FOnlineExternalUIGooglePlay::ShowProfileUI( const FUniqueNetId& Requestor, const FUniqueNetId& Requestee, const FOnProfileUIClosedDelegate& Delegate )
 {
 	return false;
diff --git a/Engine/Source/Runtime/Online/Android/OnlineSubsystemGooglePlay/Private/OnlineExternalUIInterfaceGooglePlay.h b/Engine/Source/Runtime/Online/Android/OnlineSubsystemGooglePlay/Private/OnlineExternalUIInterfaceGooglePlay.h
--- a/Engine/Source/Runtime/Online/Android/OnlineSubsystemGooglePlay/Private/OnlineExternalUIInterfaceGooglePlay.h
+++ b/Engine/Source/Runtime/Online/Android/OnlineSubsystemGooglePlay/Private/OnlineExternalUIInterfaceGooglePlay.h
@@ -27,6 +27,7 @@ public:
 	virtual bool ShowAchievementsUI(int32 LocalUserNum) override;
 	virtual bool ShowLeaderboardUI(const FString& LeaderboardName) override;
 	virtual bool ShowWebURL(const FString& Url, const FShowWebUrlParams& ShowParams, const FOnShowWebUrlClosedDelegate& Delegate = FOnShowWebUrlClosedDelegate()) override;
+	virtual bool CloseWebURL() override;
 	virtual bool ShowProfileUI(const FUniqueNetId& Requestor, const FUniqueNetId& Requestee, const FOnProfileUIClosedDelegate& Delegate = FOnProfileUIClosedDelegate()) override;
 	virtual bool ShowAccountUpgradeUI(const FUniqueNetId& UniqueId) override;
 };
diff --git a/Engine/Source/Runtime/Online/IOS/OnlineSubsystemIOS/Private/OnlineExternalUIInterfaceIOS.cpp b/Engine/Source/Runtime/Online/IOS/OnlineSubsystemIOS/Private/OnlineExternalUIInterfaceIOS.cpp
--- a/Engine/Source/Runtime/Online/IOS/OnlineSubsystemIOS/Private/OnlineExternalUIInterfaceIOS.cpp
+++ b/Engine/Source/Runtime/Online/IOS/OnlineSubsystemIOS/Private/OnlineExternalUIInterfaceIOS.cpp
@@ -71,6 +71,11 @@ bool FOnlineExternalUIIOS::ShowWebURL(const FString& Url, const FShowWebUrlParam
 	return false;
 }
 
+bool FOnlineExternalUIIOS::CloseWebURL()
+{
+	return false;
+}
+
 bool FOnlineExternalUIIOS::ShowProfileUI(const FUniqueNetId& Requestor, const FUniqueNetId& Requestee, const FOnProfileUIClosedDelegate& Delegate)
 {
 	return false;
diff --git a/Engine/Source/Runtime/Online/IOS/OnlineSubsystemIOS/Private/OnlineExternalUIInterfaceIOS.h b/Engine/Source/Runtime/Online/IOS/OnlineSubsystemIOS/Private/OnlineExternalUIInterfaceIOS.h
--- a/Engine/Source/Runtime/Online/IOS/OnlineSubsystemIOS/Private/OnlineExternalUIInterfaceIOS.h
+++ b/Engine/Source/Runtime/Online/IOS/OnlineSubsystemIOS/Private/OnlineExternalUIInterfaceIOS.h
@@ -22,6 +22,7 @@ public:
 	virtual bool ShowAchievementsUI(int32 LocalUserNum) override;
 	virtual bool ShowLeaderboardUI(const FString& LeaderboardName) override;
 	virtual bool ShowWebURL(const FString& Url, const FShowWebUrlParams& ShowParams, const FOnShowWebUrlClosedDelegate& Delegate = FOnShowWebUrlClosedDelegate()) override;
+	virtual bool CloseWebURL() override;
 	virtual bool ShowProfileUI(const FUniqueNetId& Requestor, const FUniqueNetId& Requestee, const FOnProfileUIClosedDelegate& Delegate = FOnProfileUIClosedDelegate()) override;
 	virtual bool ShowAccountUpgradeUI(const FUniqueNetId& UniqueId) override;
 	//~ End IOnlineExternalUI Interface
diff --git a/Engine/Source/Runtime/Online/OnlineFramework/Party/Private/Party.cpp b/Engine/Source/Runtime/Online/OnlineFramework/Party/Private/Party.cpp
--- a/Engine/Source/Runtime/Online/OnlineFramework/Party/Private/Party.cpp
+++ b/Engine/Source/Runtime/Online/OnlineFramework/Party/Private/Party.cpp
@@ -68,7 +68,7 @@ void UParty::OnLogoutComplete(int32 LocalUserNum, bool bWasSuccessful)
 		for (const auto& PartyKeyValue : JoinedParties)
 		{
 			UPartyGameState* Party = PartyKeyValue.Value;
-			if (Party)
+			if (Party && Party->GetPartyId().IsValid())
 			{
 				UE_LOG(LogParty, Log, TEXT("[%s] Removed"), *(Party->GetPartyId()->ToDebugString()));
 				Party->HandleRemovedFromParty(EMemberExitedReason::Left);
@@ -471,21 +471,18 @@ void UParty::CreatePartyInternal(const FUniqueNetId& InUserId, const FOnlinePart
 
 void UParty::OnCreatePartyInternalComplete(const FUniqueNetId& LocalUserId, const TSharedPtr<const FOnlinePartyId>& InPartyId, const ECreatePartyCompletionResult Result, const FOnlinePartyTypeId InPartyTypeId, UPartyDelegates::FOnCreateUPartyComplete InCompletionDelegate)
 {
-	bool bWasSuccessful = Result == ECreatePartyCompletionResult::Succeeded;
-
-	FString PartyIdDebugString = InPartyId.IsValid() ? InPartyId->ToDebugString() : TEXT("Invalid");
-
-	UE_LOG(LogParty, Display, TEXT("OnCreatePartyInternalComplete() %s %d %s"), *PartyIdDebugString, bWasSuccessful, ToString(Result));
+	const FString PartyIdDebugString = InPartyId.IsValid() ? InPartyId->ToDebugString() : TEXT("Invalid");
+	UE_LOG(LogParty, Display, TEXT("OnCreatePartyInternalComplete() %s %s"), *PartyIdDebugString, ToString(Result));
 
 	ECreatePartyCompletionResult LocalResult = Result;
-	if (bWasSuccessful)
+	if (Result == ECreatePartyCompletionResult::Succeeded)
 	{
 		UWorld* World = GetWorld();
 		IOnlinePartyPtr PartyInt = Online::GetPartyInterface(World);
 		if (PartyInt.IsValid())
 		{
 			TSharedPtr<const FOnlineParty> Party = PartyInt->GetParty(LocalUserId, InPartyTypeId);
-			if (Party.IsValid())
+			if (ensure(Party.IsValid()))
 			{
 				TSubclassOf<UPartyGameState>* PartyGameStateClass = PartyClasses.Find(InPartyTypeId);
 				if (PartyGameStateClass != nullptr)
@@ -514,20 +511,11 @@ void UParty::OnCreatePartyInternalComplete(const FUniqueNetId& LocalUserId, cons
 		{
 			LocalResult = ECreatePartyCompletionResult::UnknownClientFailure;
 		}
-
-		if (LocalResult != ECreatePartyCompletionResult::Succeeded)
-		{
-			bWasSuccessful = false;
-		}
 	}
 
-
-	if (!bWasSuccessful)
+	if (LocalResult != ECreatePartyCompletionResult::Succeeded)
 	{
-		if (LocalResult != ECreatePartyCompletionResult::Succeeded)
-		{
-			UE_LOG(LogParty, Warning, TEXT("Error when creating party %s error=%s"), *PartyIdDebugString, ToString(LocalResult));
-		}
+		UE_LOG(LogParty, Warning, TEXT("Error when creating party %s error=%s"), *PartyIdDebugString, ToString(LocalResult));
 	}
 
 	InCompletionDelegate.ExecuteIfBound(LocalUserId, LocalResult);
@@ -590,13 +578,11 @@ void UParty::JoinPartyInternal(const FUniqueNetId& InUserId, const FPartyDetails
 
 void UParty::OnJoinPartyInternalComplete(const FUniqueNetId& LocalUserId, const FOnlinePartyId& InPartyId, const EJoinPartyCompletionResult Result, int32 DeniedResultCode, const FOnlinePartyTypeId InPartyTypeId, UPartyDelegates::FOnJoinUPartyComplete CompletionDelegate)
 {
-	bool bWasSuccessful = Result == EJoinPartyCompletionResult::Succeeded;
-
-	FString PartyIdDebugString = InPartyId.ToDebugString();
-	UE_LOG(LogParty, Display, TEXT("OnJoinPartyInternalComplete() %s %d %s."), *PartyIdDebugString, bWasSuccessful, ToString(Result));
+	const FString PartyIdDebugString = InPartyId.ToDebugString();
+	UE_LOG(LogParty, Display, TEXT("OnJoinPartyInternalComplete() %s %s."), *PartyIdDebugString, ToString(Result));
 
 	EJoinPartyCompletionResult LocalResult = Result;
-	if (bWasSuccessful)
+	if (Result == EJoinPartyCompletionResult::Succeeded)
 	{
 		UWorld* World = GetWorld();
 		IOnlinePartyPtr PartyInt = Online::GetPartyInterface(World);
@@ -630,21 +616,13 @@ void UParty::OnJoinPartyInternalComplete(const FUniqueNetId& LocalUserId, const
 		{
 			LocalResult = EJoinPartyCompletionResult::UnknownClientFailure;
 		}
-
-		if (LocalResult != EJoinPartyCompletionResult::Succeeded)
-		{
-			bWasSuccessful = false;
-		}
 	}
 
-	if (!bWasSuccessful)
+	if (LocalResult != EJoinPartyCompletionResult::Succeeded)
 	{
-		if (LocalResult != EJoinPartyCompletionResult::Succeeded)
-		{
-			UE_LOG(LogParty, Warning, TEXT("Error when joining party %s error=%s"), *PartyIdDebugString, ToString(LocalResult));
-		}
+		UE_LOG(LogParty, Warning, TEXT("Error when joining party %s error=%s"), *PartyIdDebugString, ToString(LocalResult));
 	}
-
+	
 	int32 OutDeniedResultCode = (Result == EJoinPartyCompletionResult::NotApproved) ? DeniedResultCode : 0;
 	CompletionDelegate.ExecuteIfBound(LocalUserId, LocalResult, OutDeniedResultCode);
 }
@@ -657,17 +635,26 @@ void UParty::LeavePartyInternal(const FUniqueNetId& InUserId, const FOnlineParty
 	IOnlinePartyPtr PartyInt = Online::GetPartyInterface(GetWorld());
 	if (PartyInt.IsValid())
 	{
-		TSharedPtr<const FOnlineParty> Party = PartyInt->GetParty(InUserId, InPartyTypeId);
-		TSharedPtr<const FOnlinePartyId> PartyId = Party.IsValid() ? Party->PartyId : TSharedPtr<const FOnlinePartyId>();
-
-		UPartyGameState* PartyState = PartyId.IsValid() ? GetParty(*PartyId) : nullptr;
+		// Get the party by type here (don't rely on interface structures here because they can be missing during disconnects)
+		UPartyGameState* PartyState = GetParty(InPartyTypeId);
 		if (PartyState)
 		{
 			PartyState->HandleLeavingParty();
-			FOnLeavePartyComplete CompletionDelegate;
-			CompletionDelegate.BindUObject(this, &ThisClass::OnLeavePartyInternalComplete, InPartyTypeId, InCompletionDelegate);
-			PartyInt->LeaveParty(InUserId, *PartyId, CompletionDelegate);
-			Result = ELeavePartyCompletionResult::Succeeded;
+
+			TSharedPtr<const FOnlinePartyId> PartyId = PartyState->GetPartyId();
+			if (ensure(PartyId.IsValid()))
+			{
+				FOnLeavePartyComplete CompletionDelegate;
+				CompletionDelegate.BindUObject(this, &ThisClass::OnLeavePartyInternalComplete, InPartyTypeId, InCompletionDelegate);
+				PartyInt->LeaveParty(InUserId, *PartyId, CompletionDelegate);
+				Result = ELeavePartyCompletionResult::Succeeded;
+			}
+			else
+			{
+				// Manual cleanup here because we can't call the above delegate
+				PartyState->HandleRemovedFromParty(EMemberExitedReason::Left);
+				JoinedParties.Remove(InPartyTypeId);
+			}
 		}
 		else
 		{
@@ -678,7 +665,7 @@ void UParty::LeavePartyInternal(const FUniqueNetId& InUserId, const FOnlineParty
 	else
 	{
 		Result = ELeavePartyCompletionResult::UnknownClientFailure;
-		ErrorMsg = FString::Printf(TEXT("No party interface during JoinParty()"));
+		ErrorMsg = FString::Printf(TEXT("No party interface during LeaveParty()"));
 	}
 
 	if (Result != ELeavePartyCompletionResult::Succeeded)
@@ -690,8 +677,8 @@ void UParty::LeavePartyInternal(const FUniqueNetId& InUserId, const FOnlineParty
 
 void UParty::OnLeavePartyInternalComplete(const FUniqueNetId& LocalUserId, const FOnlinePartyId& InPartyId, const ELeavePartyCompletionResult Result, const FOnlinePartyTypeId InPartyTypeId, UPartyDelegates::FOnLeaveUPartyComplete CompletionDelegate)
 {
-	FString PartyIdDebugString = InPartyId.ToDebugString();
-	UE_LOG(LogParty, Display, TEXT("OnLeavePartyComplete() %s %d %s."), *PartyIdDebugString, Result == ELeavePartyCompletionResult::Succeeded, ToString(Result));
+	const FString PartyIdDebugString = InPartyId.ToDebugString();
+	UE_LOG(LogParty, Display, TEXT("OnLeavePartyComplete() %s %s."), *PartyIdDebugString, ToString(Result));
 
 	UPartyGameState* PartyState = GetParty(InPartyTypeId);
 	if (PartyState)
@@ -716,55 +703,11 @@ void UParty::CreateParty(const FUniqueNetId& InUserId, const FOnlinePartyTypeId
 
 void UParty::OnCreatePartyComplete(const FUniqueNetId& LocalUserId, const ECreatePartyCompletionResult Result, const FOnlinePartyTypeId InPartyTypeId, UPartyDelegates::FOnCreateUPartyComplete InCompletionDelegate)
 {
-	bool bWasSuccessful = Result == ECreatePartyCompletionResult::Succeeded;
-
-	UE_LOG(LogParty, Display, TEXT("OnCreatePartyComplete() type(0x%08x) %d %s"), InPartyTypeId.GetValue(), bWasSuccessful, ToString(Result));
+	UE_LOG(LogParty, Display, TEXT("OnCreatePartyComplete() type(0x%08x) %s"), InPartyTypeId.GetValue(), ToString(Result));
 
 	InCompletionDelegate.ExecuteIfBound(LocalUserId, Result);
 }
 
-EOnCreatePartyCompleteInternalResult UParty::OnCreatePartyCompleteInternal(const FUniqueNetId& LocalUserId, const FOnlinePartyTypeId InPartyTypeId)
-{
-	EOnCreatePartyCompleteInternalResult ReturnResult = EOnCreatePartyCompleteInternalResult::UnknownFailure;
-
-	UWorld* World = GetWorld();
-	IOnlinePartyPtr PartyInt = Online::GetPartyInterface(World);
-	if (PartyInt.IsValid())
-	{
-		TSharedPtr<const FOnlineParty> Party = PartyInt->GetParty(LocalUserId, InPartyTypeId);
-		if (Party.IsValid())
-		{
-			TSubclassOf<UPartyGameState>* PartyGameStateClass = PartyClasses.Find(InPartyTypeId);
-			if (PartyGameStateClass != nullptr)
-			{
-				UPartyGameState* NewParty = NewObject<UPartyGameState>(this, *PartyGameStateClass);
-
-				// Add right away so future delegate broadcasts have this available
-				JoinedParties.Add(InPartyTypeId, NewParty);
-
-				// Initialize and trigger delegates
-				NewParty->InitFromCreate(LocalUserId, Party);
-
-				ReturnResult = EOnCreatePartyCompleteInternalResult::Success;
-			}
-			else
-			{
-				ReturnResult = EOnCreatePartyCompleteInternalResult::InvalidPartyType;
-			}
-		}
-		else
-		{
-			ReturnResult = EOnCreatePartyCompleteInternalResult::InvalidParty;
-		}
-	}
-	else
-	{
-		ReturnResult = EOnCreatePartyCompleteInternalResult::InvalidSystemInterface;
-	}
-
-	return ReturnResult;
-}
-
 void UParty::JoinParty(const FUniqueNetId& InUserId, const FPartyDetails& InPartyDetails, const UPartyDelegates::FOnJoinUPartyComplete& InCompletionDelegate)
 {
 	UPartyDelegates::FOnJoinUPartyComplete CompletionDelegate;
@@ -774,73 +717,10 @@ void UParty::JoinParty(const FUniqueNetId& InUserId, const FPartyDetails& InPart
 
 void UParty::OnJoinPartyComplete(const FUniqueNetId& LocalUserId, EJoinPartyCompletionResult Result, int32 DeniedResultCode, const FOnlinePartyTypeId InPartyTypeId, UPartyDelegates::FOnJoinUPartyComplete InCompletionDelegate)
 {
-	bool bWasSuccessful = Result == EJoinPartyCompletionResult::Succeeded;
-
-	UE_LOG(LogParty, Display, TEXT("OnJoinPartyComplete() type(0x%08x) %d %s"), InPartyTypeId.GetValue(), bWasSuccessful, ToString(Result));
-
-	InCompletionDelegate.ExecuteIfBound(LocalUserId, Result, DeniedResultCode);
-
-	//bool bWasSuccessful = Result == EJoinPartyCompletionResult::Succeeded;
-
-	//FString PartyIdDebugString = InPartyId.ToDebugString();
-	//UE_LOG(LogParty, Display, TEXT("OnJoinPartyComplete() %s %d %s."), *PartyIdDebugString, bWasSuccessful, ToString(Result));
-
-	//EJoinPartyCompletionResult LocalResult = Result;
-	//if (bWasSuccessful)
-	//{
-	//	EOnJoinPartyCompleteInternalResult InternalResult = OnJoinPartyCompleteInternal(LocalUserId, InPartyId);
-
-	//	if (InternalResult != EOnJoinPartyCompleteInternalResult::Success)
-	//	{
-	//		LocalResult = EJoinPartyCompletionResult::UnknownClientFailure;
-	//		UE_LOG(LogParty, Warning, TEXT("Error when joining party %s error=%s"), *PartyIdDebugString, ToString(InternalResult));
-	//	}
-	//}
-
-	//int32 OutDeniedResultCode = (Result == EJoinPartyCompletionResult::NotApproved) ? DeniedResultCode : 0;
-	//CompletionDelegate.ExecuteIfBound(LocalUserId, InPartyId, LocalResult, OutDeniedResultCode);
-}
-
-EOnJoinPartyCompleteInternalResult UParty::OnJoinPartyCompleteInternal(const FUniqueNetId& LocalUserId, const FOnlinePartyTypeId InPartyTypeId)
-{
-	EOnJoinPartyCompleteInternalResult ReturnResult = EOnJoinPartyCompleteInternalResult::UnknownFailure;
-
-	UWorld* World = GetWorld();
-	IOnlinePartyPtr PartyInt = Online::GetPartyInterface(World);
-	if (PartyInt.IsValid())
-	{
-		TSharedPtr<const FOnlineParty> Party = PartyInt->GetParty(LocalUserId, InPartyTypeId);
-		if (Party.IsValid())
-		{
-			TSubclassOf<UPartyGameState>* PartyGameStateClass = PartyClasses.Find(InPartyTypeId);
-			if (PartyGameStateClass != nullptr)
-			{
-				UPartyGameState* NewParty = NewObject<UPartyGameState>(this, *PartyGameStateClass);
-
-				// Add right away so future delegate broadcasts have this available
-				JoinedParties.Add(InPartyTypeId, NewParty);
-
-				// Initialize and trigger delegates
-				NewParty->InitFromJoin(LocalUserId, Party);
+	UE_LOG(LogParty, Display, TEXT("OnJoinPartyComplete() type(0x%08x) %s"), InPartyTypeId.GetValue(), ToString(Result));
 
-				ReturnResult = EOnJoinPartyCompleteInternalResult::Success;
-			}
-			else
-			{
-				ReturnResult = EOnJoinPartyCompleteInternalResult::InvalidPartyType;
-			}
-		}
-		else
-		{
-			ReturnResult = EOnJoinPartyCompleteInternalResult::InvalidParty;
-		}
-	}
-	else
-	{
-		ReturnResult = EOnJoinPartyCompleteInternalResult::InvalidSystemInterface;
-	}
-
-	return ReturnResult;
+	int32 OutDeniedResultCode = (Result == EJoinPartyCompletionResult::NotApproved) ? DeniedResultCode : 0;
+	InCompletionDelegate.ExecuteIfBound(LocalUserId, Result, OutDeniedResultCode);
 }
 
 void UParty::LeaveParty(const FUniqueNetId& InUserId, const FOnlinePartyTypeId InPartyTypeId, const UPartyDelegates::FOnLeaveUPartyComplete& InCompletionDelegate)
@@ -852,43 +732,9 @@ void UParty::LeaveParty(const FUniqueNetId& InUserId, const FOnlinePartyTypeId I
 
 void UParty::OnLeavePartyComplete(const FUniqueNetId& LocalUserId, const ELeavePartyCompletionResult Result, const FOnlinePartyTypeId InPartyTypeId, UPartyDelegates::FOnLeaveUPartyComplete InCompletionDelegate)
 {
-	bool bWasSuccessful = Result == ELeavePartyCompletionResult::Succeeded;
-
-	UE_LOG(LogParty, Display, TEXT("OnLeavePartyComplete() type(0x%08x) %d %s"), InPartyTypeId.GetValue(), bWasSuccessful, ToString(Result));
+	UE_LOG(LogParty, Display, TEXT("OnLeavePartyComplete() type(0x%08x) %s"), InPartyTypeId.GetValue(), ToString(Result));
 
 	InCompletionDelegate.ExecuteIfBound(LocalUserId, Result);
-	//FString PartyIdDebugString = InPartyId.ToDebugString();
-	//UE_LOG(LogParty, Display, TEXT("OnLeavePartyComplete() %s %d %s."), *PartyIdDebugString, Result == ELeavePartyCompletionResult::Succeeded, ToString(Result));
-
-	//ELeavePartyCompletionResult LocalResult = Result;
-	//EOnLeavePartyCompleteInternalResult InternalResult = OnLeavePartyCompleteInternal(LocalUserId, InPartyId);
-
-	//if (InternalResult != EOnLeavePartyCompleteInternalResult::Success)
-	//{
-	//	LocalResult = ELeavePartyCompletionResult::UnknownClientFailure;
-	//	UE_LOG(LogParty, Warning, TEXT("Error when leaving party %s error=%s"), *PartyIdDebugString, ToString(InternalResult));
-	//}
-
-	//CompletionDelegate.ExecuteIfBound(LocalUserId, InPartyId, LocalResult);
-}
-
-EOnLeavePartyCompleteInternalResult UParty::OnLeavePartyCompleteInternal(const FUniqueNetId& LocalUserId, const FOnlinePartyTypeId InPartyTypeId)
-{
-	EOnLeavePartyCompleteInternalResult ReturnResult = EOnLeavePartyCompleteInternalResult::UnknownFailure;
-
-	UPartyGameState* PartyState = GetParty(InPartyTypeId);
-	if (PartyState)
-	{
-		PartyState->HandleRemovedFromParty(EMemberExitedReason::Left);
-		JoinedParties.Remove(InPartyTypeId);
-		ReturnResult = EOnLeavePartyCompleteInternalResult::Success;
-	}
-	else
-	{
-		ReturnResult = EOnLeavePartyCompleteInternalResult::InvalidPartyState;
-	}
-
-	return ReturnResult;
 }
 
 void UParty::GetDefaultPersistentPartySettings(EPartyType& PartyType, bool& bLeaderFriendsOnly, bool& bLeaderInvitesOnly)
@@ -959,45 +805,30 @@ void UParty::CreatePersistentParty(const FUniqueNetId& InUserId, const UPartyDel
 
 void UParty::OnCreatePersistentPartyComplete(const FUniqueNetId& LocalUserId, const ECreatePartyCompletionResult Result, UPartyDelegates::FOnCreateUPartyComplete CompletionDelegate)
 {
-	bool bWasSuccessful = Result == ECreatePartyCompletionResult::Succeeded;
-
-	TSharedPtr<const FOnlineParty> Party;
-	TSharedPtr<const FOnlinePartyId> PartyId;
-	
 	UWorld* World = GetWorld();
+	check(World);
+
 	IOnlinePartyPtr PartyInt = Online::GetPartyInterface(World);
-	if (PartyInt.IsValid())
+	if (ensure(PartyInt.IsValid()))
 	{
-		if (bWasSuccessful)
+		if (Result == ECreatePartyCompletionResult::Succeeded)
 		{
-			Party = PartyInt->GetParty(LocalUserId, IOnlinePartySystem::GetPrimaryPartyTypeId());
-			if (Party.IsValid())
+			TSharedPtr<const FOnlineParty> Party = PartyInt->GetParty(LocalUserId, IOnlinePartySystem::GetPrimaryPartyTypeId());
+			if (ensure(Party.IsValid()))
 			{
-				PartyId = Party->PartyId;
-				PersistentPartyId = PartyId;
+				PersistentPartyId = Party->PartyId;
 			}
 		}
 	}
 
-	FString PartyIdDebugString = PartyId.IsValid() ? PartyId->ToDebugString() : TEXT("Invalid");
-
-	UE_LOG(LogParty, Display, TEXT("OnCreatePersistentPartyComplete() %s %d %s"), *PartyIdDebugString, bWasSuccessful, ToString(Result));
-
-	//ECreatePartyCompletionResult LocalResult = Result;
-	//if (bWasSuccessful)
-	//{
-	//	EOnCreatePartyCompleteInternalResult InternalResult = OnCreatePartyCompleteInternal(LocalUserId, IOnlinePartySystem::GetPrimaryPartyTypeId());
-	//	if (InternalResult != EOnCreatePartyCompleteInternalResult::Success)
-	//	{
-	//		LocalResult = ECreatePartyCompletionResult::UnknownClientFailure;
-	//		UE_LOG(LogParty, Warning, TEXT("Error when creating persistent party %s error=%s"), *PartyIdDebugString, ToString(InternalResult));
-	//	}
-	//}
-
+	const FString PartyIdDebugString = PersistentPartyId.IsValid() ? PersistentPartyId->ToDebugString() : TEXT("Invalid");
+	UE_LOG(LogParty, Display, TEXT("OnCreatePersistentPartyComplete() %s %s"), *PartyIdDebugString, ToString(Result));
+	
 	CompletionDelegate.ExecuteIfBound(LocalUserId, Result);
 
 	if (Result == ECreatePartyCompletionResult::Succeeded)
 	{
+		ensure(PersistentPartyId.IsValid());
 		UPartyGameState* PersistentParty = GetPersistentParty();
 		if (PersistentParty)
 		{
@@ -1042,46 +873,29 @@ void UParty::JoinPersistentParty(const FUniqueNetId& InUserId, const FPartyDetai
 
 void UParty::OnJoinPersistentPartyComplete(const FUniqueNetId& LocalUserId, EJoinPartyCompletionResult Result, int32 DeniedResultCode, UPartyDelegates::FOnJoinUPartyComplete CompletionDelegate)
 {
-	bool bWasSuccessful = (Result == EJoinPartyCompletionResult::Succeeded);
-
-	TSharedPtr<const FOnlineParty> Party;
-	TSharedPtr<const FOnlinePartyId> PartyId;
-
 	UWorld* World = GetWorld();
 	IOnlinePartyPtr PartyInt = Online::GetPartyInterface(World);
 	if (PartyInt.IsValid())
 	{
-		if (bWasSuccessful)
+		if (Result == EJoinPartyCompletionResult::Succeeded)
 		{
-			Party = PartyInt->GetParty(LocalUserId, IOnlinePartySystem::GetPrimaryPartyTypeId());
-			if (Party.IsValid())
+			TSharedPtr<const FOnlineParty> Party = PartyInt->GetParty(LocalUserId, IOnlinePartySystem::GetPrimaryPartyTypeId());
+			if (ensure(Party.IsValid()))
 			{
-				PartyId = Party->PartyId;
-				PersistentPartyId = PartyId;
+				PersistentPartyId = Party->PartyId;
 			}
 		}
 	}
 
-	FString PartyIdDebugString = PartyId.IsValid() ? PartyId->ToDebugString() : TEXT("Invalid");
-
-	UE_LOG(LogParty, Display, TEXT("OnJoinPersistentPartyComplete() %s %d %s"), *PartyIdDebugString, bWasSuccessful, ToString(Result));
-
-	//EJoinPartyCompletionResult LocalResult = Result;
-	//if (bWasSuccessful)
-	//{
-	//	EOnJoinPartyCompleteInternalResult InternalResult = OnJoinPartyCompleteInternal(LocalUserId, IOnlinePartySystem::GetPrimaryPartyTypeId());
-	//	if (InternalResult != EOnJoinPartyCompleteInternalResult::Success)
-	//	{
-	//		LocalResult = EJoinPartyCompletionResult::UnknownInternalFailure;
-	//		UE_LOG(LogParty, Warning, TEXT("Error when creating persistent party %s error=%s"), *PartyIdDebugString, ToString(InternalResult));
-	//	}
-	//}
+	FString PartyIdDebugString = PersistentPartyId.IsValid() ? PersistentPartyId->ToDebugString() : TEXT("Invalid");
+	UE_LOG(LogParty, Display, TEXT("OnJoinPersistentPartyComplete() %s %s"), *PartyIdDebugString, ToString(Result));
 
 	int32 OutDeniedResultCode = (Result == EJoinPartyCompletionResult::NotApproved) ? DeniedResultCode : 0;
 	CompletionDelegate.ExecuteIfBound(LocalUserId, Result, OutDeniedResultCode);
 
 	if (Result == EJoinPartyCompletionResult::Succeeded)
 	{
+		ensure(PersistentPartyId.IsValid());
 		UPartyGameState* PersistentParty = GetPersistentParty();
 		if (PersistentParty)
 		{
@@ -1145,33 +959,18 @@ void UParty::LeavePersistentParty(const FUniqueNetId& InUserId, const UPartyDele
 	else
 	{
 		UE_LOG(LogParty, Warning, TEXT("No party during LeavePersistentParty()"));
-		// mark this as true so the ensure works in the delegate
 		InCompletionDelegate.ExecuteIfBound(InUserId, ELeavePartyCompletionResult::UnknownParty);
 	}
 }
 
 void UParty::OnLeavePersistentPartyComplete(const FUniqueNetId& LocalUserId, const ELeavePartyCompletionResult Result, UPartyDelegates::FOnLeaveUPartyComplete CompletionDelegate)
 {
+	UE_LOG(LogParty, Display, TEXT("OnLeavePersistentPartyComplete() %s"), ToString(Result));
+
 	ensure(bLeavingPersistentParty);
-	//ensure(!PersistentPartyId.IsValid());
 	bLeavingPersistentParty = false;
-
-	bool bWasSuccessful = (Result == ELeavePartyCompletionResult::Succeeded);
-
-	UE_LOG(LogParty, Display, TEXT("OnLeavePersistentPartyComplete() %d %s"), bWasSuccessful, ToString(Result));
-
-	// This might be the last reference of the party id, transfer the shared pointer for this brief period of time
-	// PersistentPartyId is expected to be null now as far as external code is concerned
-	TSharedPtr<const FOnlinePartyId> TempPartyId = PersistentPartyId;
 	PersistentPartyId = nullptr;
 
-	//EOnLeavePartyCompleteInternalResult InternalResult = OnLeavePartyCompleteInternal(LocalUserId, IOnlinePartySystem::GetPrimaryPartyTypeId());
-	//ELeavePartyCompletionResult LocalResult = Result;
-	//if (InternalResult != EOnLeavePartyCompleteInternalResult::Success)
-	//{
-	//	LocalResult = ELeavePartyCompletionResult::UnknownClientFailure;
-	//}
-
 	CompletionDelegate.ExecuteIfBound(LocalUserId, Result);
 
 	TArray<UPartyDelegates::FOnLeaveUPartyComplete> DelegatesCopy = LeavePartyCompleteDelegates;
@@ -1186,47 +985,50 @@ void UParty::OnLeavePersistentPartyComplete(const FUniqueNetId& LocalUserId, con
 
 void UParty::RestorePersistentPartyState()
 {
-	if (!bLeavingPersistentParty)
+	if (!GIsRequestingExit)
 	{
-		UWorld* World = GetWorld();
-		IOnlinePartyPtr PartyInt = Online::GetPartyInterface(World);
-		if (PartyInt.IsValid())
+		if (!bLeavingPersistentParty)
 		{
-			UGameInstance* GameInstance = GetGameInstance();
-			check(GameInstance);
+			UWorld* World = GetWorld();
+			IOnlinePartyPtr PartyInt = Online::GetPartyInterface(World);
+			if (PartyInt.IsValid())
+			{
+				UGameInstance* GameInstance = GetGameInstance();
+				check(GameInstance);
 
-			TSharedPtr<const FUniqueNetId> LocalUserId = GameInstance->GetPrimaryPlayerUniqueId();
-			check(LocalUserId.IsValid() && LocalUserId->IsValid());
+				TSharedPtr<const FUniqueNetId> LocalUserId = GameInstance->GetPrimaryPlayerUniqueId();
+				check(LocalUserId.IsValid() && LocalUserId->IsValid());
 
-			UPartyGameState* PersistentParty = GetPersistentParty();
+				UPartyGameState* PersistentParty = GetPersistentParty();
 
-			// Check for existing party and create a new one if necessary
-			bool bFoundExistingPersistentParty = PersistentParty ? true : false;
-			if (bFoundExistingPersistentParty)
-			{
-				// In a party already, make sure the UI is aware of its state
-				if (PersistentParty->ResetForFrontend())
+				// Check for existing party and create a new one if necessary
+				bool bFoundExistingPersistentParty = PersistentParty ? true : false;
+				if (bFoundExistingPersistentParty)
 				{
-					OnPartyResetForFrontend().Broadcast(PersistentParty);
+					// In a party already, make sure the UI is aware of its state
+					if (PersistentParty->ResetForFrontend())
+					{
+						OnPartyResetForFrontend().Broadcast(PersistentParty);
+					}
+					else
+					{
+						// There was an issue resetting the party, so leave 
+						LeaveAndRestorePersistentParty();
+					}
 				}
 				else
 				{
-					// There was an issue resetting the party, so leave 
-					LeaveAndRestorePersistentParty();
-				}
-			}
-			else
-			{
-				PersistentPartyId = nullptr;
+					PersistentPartyId = nullptr;
 
-				// Create a new party
-				CreatePersistentParty(*LocalUserId);
+					// Create a new party
+					CreatePersistentParty(*LocalUserId);
+				}
 			}
 		}
-	}
-	else
-	{
-		UE_LOG(LogParty, Verbose, TEXT("Can't RestorePersistentPartyState while leaving party, ignoring"));
+		else
+		{
+			UE_LOG(LogParty, Verbose, TEXT("Can't RestorePersistentPartyState while leaving party, ignoring"));
+		}
 	}
 }
 
