Copying //UE4/Orion-Staging to Dev-Main (Originating from //Orion/Main at CL-2777663)

diff --git a/Engine/Source/Runtime/Online/OnlineFramework/Party/Private/PartyGameState.cpp b/Engine/Source/Runtime/Online/OnlineFramework/Party/Private/PartyGameState.cpp
--- a/Engine/Source/Runtime/Online/OnlineFramework/Party/Private/PartyGameState.cpp
+++ b/Engine/Source/Runtime/Online/OnlineFramework/Party/Private/PartyGameState.cpp
@@ -111,6 +111,7 @@ void UPartyGameState::ResetPartySize()
 void UPartyGameState::ResetLocalPlayerState()
 {
 	UWorld* World = GetWorld();
+	check(World);
 	AGameState* GameState = World->GetGameState();
 	check(GameState);
 
@@ -130,45 +131,74 @@ void UPartyGameState::ResetLocalPlayerState()
 
 void UPartyGameState::Init(const FUniqueNetId& LocalUserId, TSharedPtr<const FOnlineParty>& InParty)
 {
-	OwningUserId.SetUniqueNetId(LocalUserId.AsShared());
-	OssParty = InParty;
-	CurrentConfig = *OssParty->Config;
-	bDebugAcceptingMembers = CurrentConfig.bIsAcceptingMembers;
+	if (ensure(InParty.IsValid()))
+	{
+		if (ensure(LocalUserId.IsValid()))
+		{
+			OwningUserId.SetUniqueNetId(LocalUserId.AsShared());
+			OssParty = InParty;
 
-	// Last since it needs the party info/id set first
-	RegisterFrontendDelegates();
+			CurrentConfig = *OssParty->Config;
+			bDebugAcceptingMembers = CurrentConfig.bIsAcceptingMembers;
+
+			// Last since it needs the party info/id set first
+			RegisterFrontendDelegates();
+		}
+		else
+		{
+			UE_LOG(LogParty, Warning, TEXT("Init: Invalid owning user id!"));
+		}
+	}
+	else
+	{
+		UE_LOG(LogParty, Warning, TEXT("Init: Invalid party!"));
+	}
 }
 
 void UPartyGameState::InitFromCreate(const FUniqueNetId& LocalUserId, TSharedPtr<const FOnlineParty>& InParty)
 {
-	Init(LocalUserId, InParty);
+	if (ensure(InParty.IsValid()))
+	{
+		Init(LocalUserId, InParty);
 
-	// Setup initial data for the party
-	ResetPartyState();
+		// Setup initial data for the party
+		ResetPartyState();
 
-	FUniqueNetIdRepl MemberId(LocalUserId.AsShared());
-	UpdatePartyData(MemberId);
+		FUniqueNetIdRepl MemberId(LocalUserId.AsShared());
+		UpdatePartyData(MemberId);
 
-	// Broadcast join
-	UParty* Party = GetPartyOuter();
-	check(Party);
-	Party->OnPartyJoined().Broadcast(this);
+		// Broadcast join
+		UParty* Party = GetPartyOuter();
+		check(Party);
+		Party->OnPartyJoined().Broadcast(this);
 
-	// Update local player data and broadcast member data updates
-	SendLocalPlayerPartyData();
+		// Update local player data and broadcast member data updates
+		SendLocalPlayerPartyData();
+	}
+	else
+	{
+		UE_LOG(LogParty, Warning, TEXT("InitFromCreate: Invalid party!"));
+	}
 }
 
 void UPartyGameState::InitFromJoin(const FUniqueNetId& LocalUserId, TSharedPtr<const FOnlineParty>& InParty)
 {
-	Init(LocalUserId, InParty);
+	if (ensure(InParty.IsValid()))
+	{
+		Init(LocalUserId, InParty);
 
-	// Broadcast join
-	UParty* Party = GetPartyOuter();
-	check(Party);
-	Party->OnPartyJoined().Broadcast(this);
+		// Broadcast join
+		UParty* Party = GetPartyOuter();
+		check(Party);
+		Party->OnPartyJoined().Broadcast(this);
 
-	// Update local player data and broadcast member data updates
-	SendLocalPlayerPartyData();
+		// Update local player data and broadcast member data updates
+		SendLocalPlayerPartyData();
+	}
+	else
+	{
+		UE_LOG(LogParty, Warning, TEXT("InitFromJoin: Invalid party!"));
+	}
 }
 
 void UPartyGameState::PreClientTravel()
@@ -191,80 +221,86 @@ bool UPartyGameState::ResetForFrontend()
 
 	CleanupReservationBeacon();
 
-	if (OssParty.IsValid())
+	if (ensure(OssParty.IsValid()))
 	{
-		UWorld* World = GetWorld();
-		IOnlinePartyPtr PartyInt = Online::GetPartyInterface(World);
-		if (PartyInt.IsValid())
+		if (ensure(OwningUserId.IsValid()))
 		{
-			TSharedPtr<const FOnlineParty> ExistingParty = PartyInt->GetParty(*OwningUserId, *OssParty->PartyId);
-			if (ExistingParty.IsValid())
+			UWorld* World = GetWorld();
+			IOnlinePartyPtr PartyInt = Online::GetPartyInterface(World);
+			if (ensure(PartyInt.IsValid()))
 			{
-				bSuccess = true;
+				TSharedPtr<const FOnlineParty> ExistingParty = PartyInt->GetParty(*OwningUserId, *OssParty->PartyId);
+				if (ExistingParty.IsValid())
+				{
+					bSuccess = true;
 
-				Init(*OwningUserId, ExistingParty);
-				bStayWithPartyOnDisconnect = false;
+					Init(*OwningUserId, ExistingParty);
+					bStayWithPartyOnDisconnect = false;
 
-				ResetPartyState();
-				ResetLocalPlayerState();
+					ResetPartyState();
+					ResetLocalPlayerState();
 
-				TSharedPtr<const FUniqueNetId> PartyLeaderId = GetPartyLeader();
-				bool bIsPartyLeader = (PartyLeaderId.IsValid() && *PartyLeaderId == *OwningUserId);
+					const bool bIsPartyLeader = IsLocalPartyLeader();
 
-				// Refresh local player data
-				SendLocalPlayerPartyData();
+					// Refresh local player data
+					SendLocalPlayerPartyData();
 
-				// Remove members we have, but lower level doesn't know about
-				for (const auto& PartyMember : PartyMembersState)
-				{
-					const FUniqueNetIdRepl& MemberId = PartyMember.Key;
-					if (MemberId.IsValid())
+					// Remove members we have, but lower level doesn't know about
+					for (const auto& PartyMember : PartyMembersState)
 					{
-						TSharedPtr<FOnlinePartyMember> CheckPartyMember = PartyInt->GetPartyMember(*OwningUserId, *OssParty->PartyId, *MemberId);
-						if (!CheckPartyMember.IsValid())
+						const FUniqueNetIdRepl& MemberId = PartyMember.Key;
+						if (MemberId.IsValid())
 						{
-							UE_LOG(LogParty, Verbose, TEXT("[%s] Player %s left during fixup"), *OssParty->PartyId->ToString(), *MemberId.ToString());
-							HandlePartyMemberLeft(*MemberId, EMemberExitedReason::Left);
+							TSharedPtr<FOnlinePartyMember> CheckPartyMember = PartyInt->GetPartyMember(*OwningUserId, *OssParty->PartyId, *MemberId);
+							if (!CheckPartyMember.IsValid())
+							{
+								UE_LOG(LogParty, Verbose, TEXT("[%s] Player %s left during fixup"), *OssParty->PartyId->ToString(), *MemberId.ToString());
+								HandlePartyMemberLeft(*MemberId, EMemberExitedReason::Left);
+							}
 						}
 					}
-				}
-
-				// Add members we don't have, but lower level does
-				TArray<TSharedRef<FOnlinePartyMember>> PartyMembers;
-				PartyInt->GetPartyMembers(*OwningUserId, *OssParty->PartyId, PartyMembers);
-				for (const auto& PartyMember : PartyMembers)
-				{
-					TSharedRef<const FUniqueNetId> MemberId = PartyMember->GetUserId();
 
-					FUniqueNetIdRepl UniqueId(MemberId);
-					UPartyMemberState** CurrentPartyMemberDataPtr = PartyMembersState.Find(UniqueId);
-					UPartyMemberState* CurrentPartyMemberData = CurrentPartyMemberDataPtr ? *CurrentPartyMemberDataPtr : nullptr;
-					if (!CurrentPartyMemberData)
+					// Add members we don't have, but lower level does
+					TArray<TSharedRef<FOnlinePartyMember>> PartyMembers;
+					PartyInt->GetPartyMembers(*OwningUserId, *OssParty->PartyId, PartyMembers);
+					for (const auto& PartyMember : PartyMembers)
 					{
-						TSharedPtr<FOnlinePartyData> PartyMemberData = PartyInt->GetPartyMemberData(*OwningUserId, *OssParty->PartyId, *MemberId);
-						if (PartyMemberData.IsValid())
+						TSharedRef<const FUniqueNetId> MemberId = PartyMember->GetUserId();
+
+						FUniqueNetIdRepl UniqueId(MemberId);
+						UPartyMemberState** CurrentPartyMemberDataPtr = PartyMembersState.Find(UniqueId);
+						UPartyMemberState* CurrentPartyMemberData = CurrentPartyMemberDataPtr ? *CurrentPartyMemberDataPtr : nullptr;
+						if (!CurrentPartyMemberData)
 						{
-							UE_LOG(LogParty, Verbose, TEXT("[%s] Player %s data received during fixup"), *OssParty->PartyId->ToString(), *UniqueId.ToString());
-							HandlePartyMemberDataReceived(*MemberId, PartyMemberData.ToSharedRef());
+							TSharedPtr<FOnlinePartyData> PartyMemberData = PartyInt->GetPartyMemberData(*OwningUserId, *OssParty->PartyId, *MemberId);
+							if (PartyMemberData.IsValid())
+							{
+								UE_LOG(LogParty, Verbose, TEXT("[%s] Player %s data received during fixup"), *OssParty->PartyId->ToString(), *UniqueId.ToString());
+								HandlePartyMemberDataReceived(*MemberId, PartyMemberData.ToSharedRef());
+							}
 						}
 					}
-				}
 
-				if (bIsPartyLeader)
+					if (bIsPartyLeader)
+					{
+						UpdatePartyData(OwningUserId);
+						ResetPartySize();
+						UpdateAcceptingMembers();
+					}
+				}
+				else
 				{
-					UpdatePartyData(OwningUserId);
-					ResetPartySize();
-					UpdateAcceptingMembers();
+					UE_LOG(LogParty, Warning, TEXT("Party interface can't find party during reset!"));
 				}
 			}
 			else
 			{
-				UE_LOG(LogParty, Warning, TEXT("Party interface can't find party during reset!"));
+				UE_LOG(LogParty, Warning, TEXT("Invalid party interface during reset!"));
 			}
 		}
 		else
 		{
-			UE_LOG(LogParty, Warning, TEXT("Invalid party interface during reset!"));
+			UE_LOG(LogParty, Warning, TEXT("Invalid owning user during reset!"));
 		}
 	}
 	else
@@ -288,11 +324,12 @@ UPartyMemberState* UPartyGameState::CreateNewPartyMember(const FUniqueNetId& InM
 {
 	UPartyMemberState* NewPartyMemberState = nullptr;
 
-	if (InMemberId.IsValid())
+	if (ensure(InMemberId.IsValid()))
 	{
 		UWorld* World = GetWorld();
+		check(World);
 		IOnlinePartyPtr PartyInt = Online::GetPartyInterface(World);
-		if (PartyInt.IsValid())
+		if (ensure(PartyInt.IsValid()))
 		{
 			TSharedPtr<FOnlinePartyMember> PartyMember = PartyInt->GetPartyMember(*OwningUserId, *OssParty->PartyId, InMemberId);
 			if (PartyMember.IsValid())
@@ -322,7 +359,7 @@ UPartyMemberState* UPartyGameState::CreateNewPartyMember(const FUniqueNetId& InM
 void UPartyGameState::HandlePartyConfigChanged( const TSharedRef<FPartyConfiguration>& InPartyConfig )
 {
 	UE_LOG(LogParty, VeryVerbose, TEXT("[%s] HandlePartyConfigChanged"), OssParty.IsValid() ? *OssParty->PartyId->ToString() : TEXT("INVALID"));
-	if (OssParty.IsValid())
+	if (ensure(OssParty.IsValid()))
 	{
 		CurrentConfig = *OssParty->Config;
 		bDebugAcceptingMembers = CurrentConfig.bIsAcceptingMembers;
@@ -333,9 +370,9 @@ void UPartyGameState::HandlePartyMemberJoined(const FUniqueNetId& InMemberId)
 {
 	UE_LOG(LogParty, VeryVerbose, TEXT("[%s] HandlePartyMemberJoined %s"), OssParty.IsValid() ? *OssParty->PartyId->ToString() : TEXT("INVALID"), *InMemberId.ToString());
 
- 	TSharedRef<const FUniqueNetId> IdRef = InMemberId.AsShared();
- 	FUniqueNetIdRepl MemberId(IdRef);
- 	if (MemberId.IsValid())
+	TSharedRef<const FUniqueNetId> IdRef = InMemberId.AsShared();
+	FUniqueNetIdRepl MemberId(IdRef);
+	if (MemberId.IsValid())
 	{
 		UPartyMemberState** NewPartyMemberPtr = PartyMembersState.Find(MemberId);
 		UPartyMemberState* NewPartyMember = NewPartyMemberPtr ? *NewPartyMemberPtr : nullptr;
@@ -368,18 +405,24 @@ void UPartyGameState::HandlePartyMemberLeft(const FUniqueNetId& InMemberId, EMem
 {
 	UE_LOG(LogParty, VeryVerbose, TEXT("[%s] HandlePartyMemberLeft %s"), OssParty.IsValid() ? *OssParty->PartyId->ToString() : TEXT("INVALID"), *InMemberId.ToString());
 
-	if (InMemberId.IsValid())
+	if (ensure(InMemberId.IsValid()))
 	{
 		TSharedRef<const FUniqueNetId> IdRef = InMemberId.AsShared();
 		FUniqueNetIdRepl MemberId(IdRef);
 		
 		UParty* Party = GetPartyOuter();
-		check(Party);
-		Party->OnPartyMemberLeaving().Broadcast(this, MemberId, Reason);
+		if (ensure(Party))
+		{
+			Party->OnPartyMemberLeaving().Broadcast(this, MemberId, Reason);
+		}
 
 		PartyMembersState.Remove(MemberId);
-		Party->OnPartyMemberLeft().Broadcast(this, MemberId, Reason);
 
+		if (Party)
+		{ 
+			Party->OnPartyMemberLeft().Broadcast(this, MemberId, Reason);
+		}
+		
 		// Update party join state, will cause a failure on leader promotion currently
 		// because we can't tell the difference between "expected leader" and "actually the new leader"
 		UpdateAcceptingMembers();
@@ -388,16 +431,18 @@ void UPartyGameState::HandlePartyMemberLeft(const FUniqueNetId& InMemberId, EMem
 
 void UPartyGameState::HandlePartyMemberPromoted(const FUniqueNetId& InMemberId)
 {
+	ensure(OssParty.IsValid());
 	UE_LOG(LogParty, VeryVerbose, TEXT("[%s] HandlePartyMemberPromoted %s"), OssParty.IsValid() ? *OssParty->PartyId->ToString() : TEXT("INVALID"), *InMemberId.ToString());
 
-	if (InMemberId.IsValid())
+	if (ensure(InMemberId.IsValid()))
 	{
-		TSharedRef<const FUniqueNetId> IdRef = InMemberId.AsShared();
-		FUniqueNetIdRepl MemberId(IdRef);
-
 		UParty* Party = GetPartyOuter();
-		check(Party);
-		Party->OnPartyMemberPromoted().Broadcast(this, MemberId);
+		if (ensure(Party))
+		{
+			TSharedRef<const FUniqueNetId> IdRef = InMemberId.AsShared();
+			FUniqueNetIdRepl MemberId(IdRef);
+			Party->OnPartyMemberPromoted().Broadcast(this, MemberId);
+		}
 	}
 
 	// Now that the leader is gone and a new leader established, make sure the accepting state is correct
@@ -406,24 +451,22 @@ void UPartyGameState::HandlePartyMemberPromoted(const FUniqueNetId& InMemberId)
 
 void UPartyGameState::ComparePartyData(const FPartyState& OldPartyData, const FPartyState& NewPartyData)
 {
-	TSharedPtr<const FUniqueNetId> PartyLeaderId = GetPartyLeader();
-
 	// Client passenger view delegates, leader won't get these because they are driving
-	if (PartyLeaderId.IsValid() && *PartyLeaderId != *OwningUserId)
+	if (!IsLocalPartyLeader())
 	{
 		if (OldPartyData.PartyType != NewPartyData.PartyType)
 		{
-			OnClientPartyTypeChanged().Broadcast(NewPartyData.PartyType);
+			OnPartyTypeChanged().Broadcast(NewPartyData.PartyType);
 		}
 
 		if (OldPartyData.bLeaderFriendsOnly != NewPartyData.bLeaderFriendsOnly)
 		{
-			OnClientLeaderFriendsOnlyChanged().Broadcast(NewPartyData.bLeaderFriendsOnly);
+			OnLeaderFriendsOnlyChanged().Broadcast(NewPartyData.bLeaderFriendsOnly);
 		}
 
 		if (OldPartyData.bLeaderInvitesOnly != NewPartyData.bLeaderInvitesOnly)
 		{
-			OnClientLeaderInvitesOnlyChanged().Broadcast(NewPartyData.bLeaderInvitesOnly);
+			OnLeaderInvitesOnlyChanged().Broadcast(NewPartyData.bLeaderInvitesOnly);
 		}
 	}
 }
@@ -433,8 +476,9 @@ void UPartyGameState::HandlePartyDataReceived(const TSharedRef<FOnlinePartyData>
 	UE_LOG(LogParty, VeryVerbose, TEXT("[%s] HandlePartyDataReceived"), OssParty.IsValid() ? *OssParty->PartyId->ToString() : TEXT("INVALID"));
 
 	UWorld* World = GetWorld();
+	check(World);
 	IOnlinePartyPtr PartyInt = Online::GetPartyInterface(World);
-	if (PartyInt.IsValid())
+	if (ensure(PartyInt.IsValid()))
 	{
 		check(PartyStateRefDef && PartyStateRef);
 		if (FVariantDataConverter::VariantMapToUStruct(InPartyData->KeyValAttrs, PartyStateRefDef, PartyStateRefScratch, 0, CPF_Transient | CPF_RepSkip))
@@ -446,7 +490,7 @@ void UPartyGameState::HandlePartyDataReceived(const TSharedRef<FOnlinePartyData>
 		}
 		else
 		{
-			UE_LOG(LogParty, Warning, TEXT("[%s] Failed to serialize party data!"));
+			UE_LOG(LogParty, Warning, TEXT("Failed to serialize party data!"));
 		}
 	}
 }
@@ -456,8 +500,9 @@ void UPartyGameState::HandlePartyMemberDataReceived(const FUniqueNetId& InMember
 	UE_LOG(LogParty, VeryVerbose, TEXT("[%s] HandlePartyMemberDataReceived %s"), OssParty.IsValid() ? *OssParty->PartyId->ToString() : TEXT("INVALID"), *InMemberId.ToString());
 
 	UWorld* World = GetWorld();
+	check(World);
 	IOnlinePartyPtr PartyInt = Online::GetPartyInterface(World);
-	if (PartyInt.IsValid())
+	if (ensure(PartyInt.IsValid()))
 	{
 		FUniqueNetIdRepl UniqueId(InMemberId.AsShared());
 		UPartyMemberState** CurrentPartyMemberPtr = PartyMembersState.Find(UniqueId);
@@ -511,8 +556,11 @@ void UPartyGameState::HandleRemovedFromParty(EMemberExitedReason Reason)
 {
 	// Trigger delegate first
 	UParty* Party = GetPartyOuter();
-	check(Party);
-	Party->OnPartyLeft().Broadcast(this, Reason);
+	if (ensure(Party))
+	{
+		Party->OnPartyLeft().Broadcast(this, Reason);
+	}
+
 	// Cleanup
 	OnShutdown();
 }
@@ -542,7 +590,7 @@ void UPartyGameState::HandlePartyJoinRequestReceived(const FUniqueNetId& Recipie
 	EApprovalAction ApprovalAction = EApprovalAction::Deny;
 	EJoinPartyDenialReason DenialReason = EJoinPartyDenialReason::Busy;
 
-	if (OssParty.IsValid())
+	if (IsLocalPartyLeader())
 	{
 		int32 NumPartyMembers = GetPartySize();
 		int32 MaxPartyMembers = CurrentConfig.MaxMembers;
@@ -550,16 +598,7 @@ void UPartyGameState::HandlePartyJoinRequestReceived(const FUniqueNetId& Recipie
 		{
 			if (IsInJoinableGameState())
 			{
-				TSharedPtr<const FUniqueNetId> PartyLeader = GetPartyLeader();
-				if (PartyLeader.IsValid() && *OwningUserId == *PartyLeader)
-				{
-					ApprovalAction = ProcessJoinRequest(RecipientId, SenderId, DenialReason);
-				}
-				else
-				{
-					// Party leader has changed
-					DenialReason = EJoinPartyDenialReason::NotPartyLeader;
-				}
+				ApprovalAction = ProcessJoinRequest(RecipientId, SenderId, DenialReason);
 			}
 			else
 			{
@@ -574,7 +613,8 @@ void UPartyGameState::HandlePartyJoinRequestReceived(const FUniqueNetId& Recipie
 	}
 	else
 	{
-		DenialReason = EJoinPartyDenialReason::Busy;
+		// Party leader has changed
+		DenialReason = EJoinPartyDenialReason::NotPartyLeader;
 	}
 
 	if (ApprovalAction == EApprovalAction::Enqueue ||
@@ -603,7 +643,7 @@ void UPartyGameState::HandlePartyJoinRequestReceived(const FUniqueNetId& Recipie
 		
 		UWorld* World = GetWorld();
 		IOnlinePartyPtr PartyInt = Online::GetPartyInterface(World);
-		if (PartyInt.IsValid())
+		if (PartyInt.IsValid() && OssParty.IsValid())
 		{
 			PartyInt->ApproveJoinRequest(RecipientId, *OssParty->PartyId, SenderId, bApproveRequest, (int32)(!bApproveRequest ? DenialReason : EJoinPartyDenialReason::NoReason));
 		}
@@ -614,7 +654,7 @@ FOnlinePartyTypeId UPartyGameState::GetPartyTypeId() const
 {
 	FOnlinePartyTypeId Result;
 
-	if (OssParty.IsValid())
+	if (ensure(OssParty.IsValid()))
 	{
 		Result = OssParty->PartyTypeId;
 	}
@@ -624,7 +664,7 @@ FOnlinePartyTypeId UPartyGameState::GetPartyTypeId() const
 
 TSharedPtr<const FOnlinePartyId> UPartyGameState::GetPartyId() const
 {
-	if (OssParty.IsValid())
+	if (ensure(OssParty.IsValid()))
 	{
 		return OssParty->PartyId;
 	}
@@ -634,7 +674,7 @@ TSharedPtr<const FOnlinePartyId> UPartyGameState::GetPartyId() const
 
 void UPartyGameState::SetPartyType(EPartyType InPartyType, bool bLeaderFriendsOnly, bool bLeaderInvitesOnly)
 {
-	if (*OwningUserId == *GetPartyLeader())
+	if (IsLocalPartyLeader())
 	{
 		check(PartyStateRef);
 		if (PartyStateRef->PartyType != InPartyType ||
@@ -672,6 +712,11 @@ void UPartyGameState::SetPartyType(EPartyType InPartyType, bool bLeaderFriendsOn
 
 			UpdatePartyConfig(bIsPrivate);
 
+			EPartyType OldPartyType = PartyStateRef->PartyType;
+			bool bOldLeaderFriendsOnly = PartyStateRef->bLeaderFriendsOnly;
+			bool bOldLeaderInvitesOnly = PartyStateRef->bLeaderInvitesOnly;
+
+			// Replicate the party settings to other party members
 			PartyStateRef->PartyType = InPartyType;
 			PartyStateRef->bLeaderFriendsOnly = bLeaderFriendsOnly;
 			PartyStateRef->bLeaderInvitesOnly = bLeaderInvitesOnly;
@@ -679,8 +724,28 @@ void UPartyGameState::SetPartyType(EPartyType InPartyType, bool bLeaderFriendsOn
 
 			// Refresh accepting members, taking everything into account
 			UpdateAcceptingMembers();
+
+			// Notify the local player of the changes
+			if (OldPartyType != InPartyType)
+			{
+				OnPartyTypeChanged().Broadcast(InPartyType);
+			}
+
+			if (bOldLeaderFriendsOnly != bLeaderFriendsOnly)
+			{
+				OnLeaderFriendsOnlyChanged().Broadcast(bLeaderFriendsOnly);
+			}
+
+			if (bOldLeaderInvitesOnly != bLeaderInvitesOnly)
+			{
+				OnLeaderInvitesOnlyChanged().Broadcast(bLeaderInvitesOnly);
+			}
 		}
 	}
+	else
+	{
+		UE_LOG(LogParty, Warning, TEXT("Non party leader trying to set party permissions!"));
+	}
 }
 
 void UPartyGameState::StayWithPartyOnExit(bool bInStayWithParty)
@@ -700,7 +765,7 @@ int32 UPartyGameState::GetPartySize() const
 
 void UPartyGameState::SetPartyMaxSize(int32 NewSize)
 {
-	if (OssParty.IsValid())
+	if (ensure(OssParty.IsValid()))
 	{
 		if (CurrentConfig.MaxMembers != NewSize)
 		{
@@ -711,16 +776,20 @@ void UPartyGameState::SetPartyMaxSize(int32 NewSize)
 	}
 	else
 	{
-		UE_LOG(LogParty, Warning, TEXT("Invalid party info updating party size!"));
+		UE_LOG(LogParty, Warning, TEXT("Invalid party updating party size!"));
 	}
 }
 
 int32 UPartyGameState::GetPartyMaxSize() const
 {
-	if (OssParty.IsValid())
+	if (ensure(OssParty.IsValid()))
 	{
 		return OssParty->Config->MaxMembers;
 	}
+	else
+	{
+		UE_LOG(LogParty, Warning, TEXT("Invalid party getting party size!"));
+	}
 
 	// Invalid party info
 	return INDEX_NONE;
@@ -728,11 +797,9 @@ int32 UPartyGameState::GetPartyMaxSize() const
 
 void UPartyGameState::UpdateAcceptingMembers()
 {
-	EJoinPartyDenialReason DenialReason = EJoinPartyDenialReason::NoReason;
-
-	TSharedPtr<const FUniqueNetId> PartyLeader = GetPartyLeader();
-	if (PartyLeader.IsValid() && *OwningUserId == *PartyLeader)
+	if (IsLocalPartyLeader())
 	{
+		EJoinPartyDenialReason DenialReason = EJoinPartyDenialReason::NoReason;
 		bool bCurrentlyAcceptingMembers = false;
 
 		// Look at game joinability (in game with permission or no game at all)
@@ -746,15 +813,15 @@ void UPartyGameState::UpdateAcceptingMembers()
 				// Look at party joinability
 				switch (PartyStateRef->PartyType)
 				{
-				case EPartyType::Public:
-				case EPartyType::FriendsOnly:
-					bCurrentlyAcceptingMembers = true;
-					break;
-				case EPartyType::Private:
-					// Party is private, invite required
-					DenialReason = EJoinPartyDenialReason::PartyPrivate;
-				default:
-					break;
+					case EPartyType::Public:
+					case EPartyType::FriendsOnly:
+						bCurrentlyAcceptingMembers = true;
+						break;
+					case EPartyType::Private:
+						// Party is private, invite required
+						DenialReason = EJoinPartyDenialReason::PartyPrivate;
+					default:
+						break;
 				}
 			}
 			else
@@ -770,30 +837,32 @@ void UPartyGameState::UpdateAcceptingMembers()
 
 		SetAcceptingMembers(bCurrentlyAcceptingMembers, DenialReason);
 	}
+	else
+	{
+		UE_LOG(LogParty, Warning, TEXT("Non party leader trying to update accepting members!"));
+	}
 }
 
 void UPartyGameState::SetAcceptingMembers(bool bIsAcceptingMembers, EJoinPartyDenialReason DenialReason)
 {
-	if (OssParty.IsValid())
+	if (IsLocalPartyLeader())
 	{
-		if (*OwningUserId == *GetPartyLeader())
+		ensure(OssParty.IsValid());
+		if (CurrentConfig.bIsAcceptingMembers != bIsAcceptingMembers)
 		{
-			if (CurrentConfig.bIsAcceptingMembers != bIsAcceptingMembers)
-			{
-				int32 NumPartyMembers = GetPartySize();
-				int32 MaxPartyMembers = CurrentConfig.MaxMembers;
-				bool bIsRoomInParty = (NumPartyMembers < MaxPartyMembers);
-
-				bDebugAcceptingMembers = bIsAcceptingMembers && bIsRoomInParty;
-				CurrentConfig.bIsAcceptingMembers = bDebugAcceptingMembers;
-				CurrentConfig.NotAcceptingMembersReason = (int32)(bDebugAcceptingMembers ? EJoinPartyDenialReason::NoReason : DenialReason);
-				UpdatePartyConfig();
-			}
+			int32 NumPartyMembers = GetPartySize();
+			int32 MaxPartyMembers = CurrentConfig.MaxMembers;
+			bool bIsRoomInParty = (NumPartyMembers < MaxPartyMembers);
+
+			bDebugAcceptingMembers = bIsAcceptingMembers && bIsRoomInParty;
+			CurrentConfig.bIsAcceptingMembers = bDebugAcceptingMembers;
+			CurrentConfig.NotAcceptingMembersReason = (int32)(bDebugAcceptingMembers ? EJoinPartyDenialReason::NoReason : DenialReason);
+			UpdatePartyConfig();
 		}
 	}
 	else
 	{
-		UE_LOG(LogParty, Warning, TEXT("Invalid party info updating joinability!"));
+		UE_LOG(LogParty, Warning, TEXT("Non party leader trying to set accepting members!"));
 	}
 }
 
@@ -831,56 +900,83 @@ bool UPartyGameState::IsInJoinableGameState() const
 
 bool UPartyGameState::CanInvite() const
 {
-	return OssParty->CanLocalUserInvite(*OwningUserId);
+	if (ensure(OssParty.IsValid()))
+	{
+		return OssParty->CanLocalUserInvite(*OwningUserId);
+	}
+	return false;
 }
 
 void UPartyGameState::UpdatePartyConfig(bool bResetAccessKey)
 {
-	if (OssParty.IsValid())
+	if (IsLocalPartyLeader())
 	{
-		if (*OwningUserId == *GetPartyLeader())
+		ensure(OssParty.IsValid());
+
+		UWorld* World = GetWorld();
+		check(World);
+		IOnlinePartyPtr PartyInt = Online::GetPartyInterface(World);
+		if (ensure(PartyInt.IsValid()))
 		{
-			UWorld* World = GetWorld();
-			IOnlinePartyPtr PartyInt = Online::GetPartyInterface(World);
-			if (PartyInt.IsValid())
-			{
-				FOnUpdatePartyComplete CompletionDelegate;
-				CompletionDelegate.BindUObject(this, &ThisClass::OnUpdatePartyConfigComplete);
-				if (!PartyInt->UpdateParty(*OwningUserId, *OssParty->PartyId, CurrentConfig, bResetAccessKey, CompletionDelegate))
-				{
-					UE_LOG(LogParty, Warning, TEXT("[%s] Failed to update party"), OssParty.IsValid() ? *OssParty->PartyId->ToString() : TEXT("INVALID"));
-				}
-			}
-			else
+			FOnUpdatePartyComplete CompletionDelegate;
+			CompletionDelegate.BindUObject(this, &ThisClass::OnUpdatePartyConfigComplete);
+			if (!PartyInt->UpdateParty(*OwningUserId, *OssParty->PartyId, CurrentConfig, bResetAccessKey, CompletionDelegate))
 			{
-				UE_LOG(LogParty, Warning, TEXT("[%s] Invalid party interface updating party size"), OssParty.IsValid() ? *OssParty->PartyId->ToString() : TEXT("INVALID"));
+				UE_LOG(LogParty, Warning, TEXT("[%s] Failed to update party"), *OssParty->PartyId->ToString());
 			}
 		}
+		else
+		{
+			UE_LOG(LogParty, Warning, TEXT("[%s] Invalid party interface updating party size"), *OssParty->PartyId->ToString());
+		}
 	}
 	else
 	{
-		UE_LOG(LogParty, Warning, TEXT("Invalid party info updating party size!"));
+		UE_LOG(LogParty, Warning, TEXT("Non party leader trying to update party config!"));
 	}
 }
 
 void UPartyGameState::OnUpdatePartyConfigComplete(const FUniqueNetId& LocalUserId, const FOnlinePartyId& InPartyId, const EUpdateConfigCompletionResult Result)
 {
-	FString PartyIdDebugString = InPartyId.ToDebugString();
+	const FString PartyIdDebugString = InPartyId.ToDebugString();
+	UE_LOG(LogParty, Verbose, TEXT("[%s] Party config updated %s"), *PartyIdDebugString, ToString(Result));
 
-	UE_LOG(LogParty, Verbose, TEXT("[%s] Party config updated %d %s"), *PartyIdDebugString, Result == EUpdateConfigCompletionResult::Succeeded, ToString(Result));
-
-	if (OssParty.IsValid())
+	if (ensure(OssParty.IsValid()))
 	{
 		CurrentConfig = *OssParty->Config;
 		bDebugAcceptingMembers = CurrentConfig.bIsAcceptingMembers;
 	}
 }
 
-TSharedPtr<const FUniqueNetId> UPartyGameState::GetPartyLeader() const
+bool UPartyGameState::IsLocalPartyLeader()
 {
-	TSharedPtr<const FUniqueNetId> ReturnValue;
+	if (ensure(OwningUserId.IsValid()))
+	{
+		TSharedPtr<const FUniqueNetId> LeaderId = GetPartyLeader();
+		if (LeaderId.IsValid())
+		{
+			if (*OwningUserId == *LeaderId)
+			{
+				return true;
+			}
+		}
+		else
+		{
+			UE_LOG(LogParty, Warning, TEXT("Unable to determine party leader!"));
+		}
+	}
+	else
+	{
+		UE_LOG(LogParty, Warning, TEXT("Invalid owning user id trying to determine party leader!"));
+	}
 
-	if (OssParty.IsValid())
+	return false;
+}
+
+TSharedPtr<const FUniqueNetId> UPartyGameState::GetPartyLeader() const
+{
+	TSharedPtr<const FUniqueNetId> ReturnValue = nullptr;
+	if (ensure(OssParty.IsValid()))
 	{
 		ReturnValue = OssParty->LeaderId;
 	}
@@ -908,30 +1004,41 @@ void UPartyGameState::OnPartyMemberPromoted(const FUniqueNetId& LocalUserId, con
 {
 	FString PartyIdDebugString = InPartyId.ToDebugString();
 	FString MemberIdDebugString = InMemberId.ToDebugString();
-
-	bool bWasSuccessful = Result == EPromoteMemberCompletionResult::Succeeded;
-
-	UE_LOG(LogParty, Verbose, TEXT("[%s] Player %s promoted %s %d %s"), *PartyIdDebugString, *LocalUserId.ToString(), *MemberIdDebugString, bWasSuccessful, ToString(Result));
+	UE_LOG(LogParty, Verbose, TEXT("[%s] Player %s promoted %s %s"), *PartyIdDebugString, *LocalUserId.ToString(), *MemberIdDebugString, ToString(Result));
 }
 
 void UPartyGameState::PromoteMember(const FUniqueNetIdRepl& NewPartyLeader)
 {
-	if (!bPromotionLockoutState)
+	if (IsLocalPartyLeader())
 	{
-		if (*OwningUserId == *GetPartyLeader())
+		ensure(OssParty.IsValid());
+		if (ensure(NewPartyLeader.IsValid()))
 		{
-			IOnlinePartyPtr PartyInt = Online::GetPartyInterface(GetWorld());
-			if (PartyInt.IsValid() && ensure(OwningUserId.IsValid()) && ensure(NewPartyLeader.IsValid()))
+			if (!bPromotionLockoutState)
+			{
+				UWorld* World = GetWorld();
+				check(World);
+				IOnlinePartyPtr PartyInt = Online::GetPartyInterface(World);
+				if (ensure(PartyInt.IsValid()))
+				{
+					FOnPromotePartyMemberComplete CompletionDelegate;
+					CompletionDelegate.BindUObject(this, &ThisClass::OnPartyMemberPromoted);
+					PartyInt->PromoteMember(*OwningUserId, *OssParty->PartyId, *NewPartyLeader, CompletionDelegate);
+				}
+			}
+			else
 			{
-				FOnPromotePartyMemberComplete CompletionDelegate;
-				CompletionDelegate.BindUObject(this, &ThisClass::OnPartyMemberPromoted);
-				PartyInt->PromoteMember(*OwningUserId, *OssParty->PartyId, *NewPartyLeader, CompletionDelegate);
+				UE_LOG(LogParty, Verbose, TEXT("[%s] Promote member feature locked out."), *OssParty->PartyId->ToString());
 			}
 		}
+		else
+		{
+			UE_LOG(LogParty, Warning, TEXT("Trying to promote invalid party member to leader!"));
+		}
 	}
 	else
 	{
-		UE_LOG(LogParty, Verbose, TEXT("[%s] Promote member feature locked out."), OssParty.IsValid() ? *OssParty->PartyId->ToString() : TEXT("INVALID"));
+		UE_LOG(LogParty, Warning, TEXT("Non party leader trying to promote party leader!"));
 	}
 }
 
@@ -939,35 +1046,55 @@ void UPartyGameState::OnPartyMemberKicked(const FUniqueNetId& LocalUserId, const
 {
 	FString PartyIdDebugString = InPartyId.ToDebugString();
 	FString MemberIdDebugString = InMemberId.ToDebugString();
-
-	bool bWasSuccessful = Result == EKickMemberCompletionResult::Succeeded;
-
-	UE_LOG(LogParty, Verbose, TEXT("[%s] Player %s kicked %s %d %s"), *PartyIdDebugString, *LocalUserId.ToString(), *MemberIdDebugString, bWasSuccessful, ToString(Result));
+	UE_LOG(LogParty, Verbose, TEXT("[%s] Player %s kicked %s %s"), *PartyIdDebugString, *LocalUserId.ToString(), *MemberIdDebugString, ToString(Result));
 }
 
 void UPartyGameState::KickMember(const FUniqueNetIdRepl& PartyMemberToKick)
 {
-	IOnlinePartyPtr PartyInt = Online::GetPartyInterface(GetWorld());
-	if (PartyInt.IsValid() && ensure(OwningUserId.IsValid()) && ensure(PartyMemberToKick.IsValid()))
+	if (IsLocalPartyLeader())
 	{
-		if (*OwningUserId == *GetPartyLeader())
+		ensure(OssParty.IsValid());
+		if (ensure(PartyMemberToKick.IsValid()))
+		{
+			UPartyMemberState** LocalPartyMemberStatePtr = PartyMembersState.Find(PartyMemberToKick);
+			if (LocalPartyMemberStatePtr && *LocalPartyMemberStatePtr)
+			{
+				UWorld* World = GetWorld();
+				check(World);
+				IOnlinePartyPtr PartyInt = Online::GetPartyInterface(World);
+				if (ensure(PartyInt.IsValid()))
+				{
+					FOnKickPartyMemberComplete CompletionDelegate;
+					CompletionDelegate.BindUObject(this, &ThisClass::OnPartyMemberKicked);
+					PartyInt->KickMember(*OwningUserId, *OssParty->PartyId, *PartyMemberToKick, CompletionDelegate);
+				}
+			}
+			else
+			{
+				UE_LOG(LogParty, Warning, TEXT("Trying to kick player that is not in your party!"));
+			}
+		}
+		else
 		{
-			FOnKickPartyMemberComplete CompletionDelegate;
-			CompletionDelegate.BindUObject(this, &ThisClass::OnPartyMemberKicked);
-			PartyInt->KickMember(*OwningUserId, *OssParty->PartyId, *PartyMemberToKick, CompletionDelegate);
+			UE_LOG(LogParty, Warning, TEXT("Trying to kick invalid party member!"));
 		}
 	}
+	else
+	{
+		UE_LOG(LogParty, Warning, TEXT("Non party leader trying to kick party member!"));
+	}
 }
 
 void UPartyGameState::InitLocalPlayerPartyData()
 {
 	UWorld* World = GetWorld();
+	check(World);
 	AGameState* GameState = World->GetGameState();
 	check(GameState);
 
 	for (const APlayerState* PlayerState : GameState->PlayerArray)
 	{
-		if (PlayerState && PlayerState->UniqueId.IsValid())
+		if (PlayerState && PlayerState->UniqueId.IsValid() && PlayerState->GetOuter() != nullptr)
 		{
 			UPartyMemberState** LocalPartyMemberStatePtr = PartyMembersState.Find(PlayerState->UniqueId);
 			UPartyMemberState* LocalPartyMemberState = LocalPartyMemberStatePtr ? *LocalPartyMemberStatePtr : nullptr;
@@ -985,96 +1112,163 @@ void UPartyGameState::InitLocalPlayerPartyData()
 
 void UPartyGameState::SendLocalPlayerPartyData()
 {
-	UWorld* World = GetWorld();
-
 	InitLocalPlayerPartyData();
 
-	IOnlinePartyPtr PartyInt = Online::GetPartyInterface(World);
-	if (PartyInt.IsValid())
+	UWorld* World = GetWorld();
+	check(World);
+	AGameState* GameState = World->GetGameState();
+	check(GameState);
+
+	for (const APlayerState* PlayerState : GameState->PlayerArray)
 	{
-		// Send local player data
-		for (auto PartyMembers : PartyMembersState)
+		if (PlayerState && PlayerState->UniqueId.IsValid() && PlayerState->GetOuter() != nullptr)
 		{
-			const FUniqueNetIdRepl& UniqueId = PartyMembers.Key;
-			const UPartyMemberState* PartyMember = PartyMembers.Value;
-
-			UpdatePartyMemberState(UniqueId, PartyMember);
+			// Send local player data
+			UPartyMemberState** MemberStatePtr = PartyMembersState.Find(PlayerState->UniqueId);
+			if (MemberStatePtr != nullptr && *MemberStatePtr != nullptr)
+			{
+				UpdatePartyMemberState(OwningUserId, *MemberStatePtr);
+			}
 		}
 	}
 }
 
 void UPartyGameState::UpdatePartyData(const FUniqueNetIdRepl& InLocalUserId)
 {
-	UWorld* World = GetWorld();
-	IOnlinePartyPtr PartyInt = Online::GetPartyInterface(World);
-	if (PartyInt.IsValid())
+	if (IsLocalPartyLeader())
 	{
-		FOnlinePartyData PartyData;
-		if (FVariantDataConverter::UStructToVariantMap(PartyStateRefDef, PartyStateRef, PartyData.KeyValAttrs, 0, CPF_Transient | CPF_RepSkip))
+		ensure(OwningUserId == InLocalUserId);
+
+		TSharedPtr<const FOnlinePartyId> PartyId = GetPartyId();
+		if (PartyId.IsValid())
+		{
+			UWorld* World = GetWorld();
+			check(World);
+
+			IOnlinePartyPtr PartyInt = Online::GetPartyInterface(World);
+			if (ensure(PartyInt.IsValid()))
+			{
+				FOnlinePartyData PartyData;
+				if (FVariantDataConverter::UStructToVariantMap(PartyStateRefDef, PartyStateRef, PartyData.KeyValAttrs, 0, CPF_Transient | CPF_RepSkip))
+				{
+					PartyInt->UpdatePartyData(*OwningUserId, *PartyId, PartyData);
+				}
+				else
+				{
+					UE_LOG(LogParty, Warning, TEXT("UpdatePartyData: Failed to update party data!"));
+				}
+			}
+			else
+			{
+				UE_LOG(LogParty, Warning, TEXT("UpdatePartyData: Invalid party interface!"));
+			}
+		}
+		else
 		{
-			PartyInt->UpdatePartyData(*InLocalUserId, *OssParty->PartyId, PartyData);
+			UE_LOG(LogParty, Warning, TEXT("UpdatePartyData: Invalid internal party!"));
 		}
 	}
+	else
+	{
+		UE_LOG(LogParty, Warning, TEXT("Non party leader trying to update party state!"));
+	}
 }
 
 void UPartyGameState::UpdatePartyMemberState(const FUniqueNetIdRepl& InLocalUserId, const UPartyMemberState* InPartyMemberState)
 {
-	UWorld* World = GetWorld();
-	IOnlinePartyPtr PartyInt = Online::GetPartyInterface(World);
-	if (PartyInt.IsValid())
+	if (ensure(InLocalUserId.IsValid()))
 	{
-		FOnlinePartyData PartyMemberData;
-		if (FVariantDataConverter::UStructToVariantMap(InPartyMemberState->MemberStateRefDef, InPartyMemberState->MemberStateRef, PartyMemberData.KeyValAttrs, 0, CPF_Transient | CPF_RepSkip))
+		if (ensure(InPartyMemberState))
+		{
+			TSharedPtr<const FOnlinePartyId> PartyId = GetPartyId();
+			if (PartyId.IsValid())
+			{
+				UWorld* World = GetWorld();
+				check(World);
+
+				IOnlinePartyPtr PartyInt = Online::GetPartyInterface(World);
+				if (ensure(PartyInt.IsValid()))
+				{
+					FOnlinePartyData PartyMemberData;
+					if (FVariantDataConverter::UStructToVariantMap(InPartyMemberState->MemberStateRefDef, InPartyMemberState->MemberStateRef, PartyMemberData.KeyValAttrs, 0, CPF_Transient | CPF_RepSkip))
+					{
+						PartyInt->UpdatePartyMemberData(*InLocalUserId, *PartyId, PartyMemberData);
+					}
+					else
+					{
+						UE_LOG(LogParty, Warning, TEXT("UpdatePartyData: Failed to update party member data!"));
+					}
+				}
+				else
+				{
+					UE_LOG(LogParty, Warning, TEXT("UpdatePartyMemberState: Invalid party interface!"));
+				}
+			}
+			else
+			{
+				UE_LOG(LogParty, Warning, TEXT("UpdatePartyMemberState: Invalid internal party!"));
+			}
+		}
+		else
 		{
-			PartyInt->UpdatePartyMemberData(*InLocalUserId, *OssParty->PartyId, PartyMemberData);
+			UE_LOG(LogParty, Warning, TEXT("UpdatePartyMemberState: NULL party member state!"));
 		}
 	}
+	else
+	{
+		UE_LOG(LogParty, Warning, TEXT("UpdatePartyMemberState: Invalid local user!"));
+	}
 }
 
 void UPartyGameState::ConnectToReservationBeacon()
 {
-	UWorld* World = GetWorld();
-
-	if (OssParty.IsValid())
+	if (IsLocalPartyLeader())
 	{
-		FUniqueNetIdRepl PartyLeader(GetPartyLeader());
-		if (PartyLeader.IsValid() && *OwningUserId == *PartyLeader)
+		FPendingMemberApproval NextApproval;
+		if (PendingApprovals.Peek(NextApproval))
 		{
-			FPendingMemberApproval NextApproval;
-			if (PendingApprovals.Peek(NextApproval))
+			UWorld* World = GetWorld();
+			check(World);
+
+			// Reconnect to the reservation beacon to maintain our place in the game (just until actual joined, holds place for all party members)
+			ReservationBeaconClient = World->SpawnActor<APartyBeaconClient>(ReservationBeaconClientClass);
+			if (ReservationBeaconClient)
 			{
-				// Reconnect to the reservation beacon to maintain our place in the game (just until actual joined, holds place for all party members)
-				ReservationBeaconClient = World->SpawnActor<APartyBeaconClient>(ReservationBeaconClientClass);
-				if (ReservationBeaconClient)
-				{
-					UE_LOG(LogParty, Verbose, TEXT("Created party reservation beacon %s."), *ReservationBeaconClient->GetName());
+				UE_LOG(LogParty, Verbose, TEXT("Created party reservation beacon %s."), *ReservationBeaconClient->GetName());
 
-					ReservationBeaconClient->OnHostConnectionFailure().BindUObject(this, &ThisClass::OnReservationBeaconUpdateConnectionFailure);
-					ReservationBeaconClient->OnReservationRequestComplete().BindUObject(this, &ThisClass::OnReservationBeaconUpdateResponseReceived);
-					ReservationBeaconClient->OnReservationCountUpdate().BindUObject(this, &ThisClass::OnReservationCountUpdate);
+				ReservationBeaconClient->OnHostConnectionFailure().BindUObject(this, &ThisClass::OnReservationBeaconUpdateConnectionFailure);
+				ReservationBeaconClient->OnReservationRequestComplete().BindUObject(this, &ThisClass::OnReservationBeaconUpdateResponseReceived);
+				ReservationBeaconClient->OnReservationCountUpdate().BindUObject(this, &ThisClass::OnReservationCountUpdate);
 
-					FPlayerReservation NewPlayerRes;
-					NewPlayerRes.UniqueId = NextApproval.SenderId;
+				FPlayerReservation NewPlayerRes;
+				NewPlayerRes.UniqueId = NextApproval.SenderId;
 
-					TArray<FPlayerReservation> PlayersToAdd;
-					PlayersToAdd.Add(NewPlayerRes);
+				TArray<FPlayerReservation> PlayersToAdd;
+				PlayersToAdd.Add(NewPlayerRes);
 
-					ReservationBeaconClient->RequestReservationUpdate(CurrentSession, PartyLeader, PlayersToAdd);
-				}
-				else
-				{
-					OnReservationBeaconUpdateConnectionFailure();
-				}
+				FUniqueNetIdRepl PartyLeader(GetPartyLeader());
+				ReservationBeaconClient->RequestReservationUpdate(CurrentSession, PartyLeader, PlayersToAdd);
+			}
+			else
+			{
+				OnReservationBeaconUpdateConnectionFailure();
 			}
 		}
 	}
+	else
+	{
+		UE_LOG(LogParty, Warning, TEXT("ConnectToReservationBeacon: Non party leader trying to connect to reservation beacon!"));
+	}
 }
 
 void UPartyGameState::RejectAllPendingJoinRequests()
 {
 	UWorld* World = GetWorld();
+	check(World);
 	IOnlinePartyPtr PartyInt = Online::GetPartyInterface(World);
-	bool bValidInterface = PartyInt.IsValid() && OssParty.IsValid();
+	TSharedPtr<const FOnlinePartyId> PartyId = GetPartyId();
+
+	bool bValidInterface = PartyInt.IsValid() && PartyId.IsValid();
 
 	bool bEmpty = false;
 	FPendingMemberApproval PendingApproval;
@@ -1084,15 +1278,15 @@ void UPartyGameState::RejectAllPendingJoinRequests()
 		bEmpty = !PendingApprovals.Dequeue(PendingApproval);
 		if (!bEmpty && bValidInterface)
 		{
-			UE_LOG(LogParty, Verbose, TEXT("[%s] Responding to approval request for %s with denied"), OssParty.IsValid() ? *OssParty->PartyId->ToString() : TEXT("INVALID"), *PendingApproval.SenderId.ToString());
-			PartyInt->ApproveJoinRequest(*PendingApproval.RecipientId, *OssParty->PartyId, *PendingApproval.SenderId, false, (int32)EJoinPartyDenialReason::Busy);
+			UE_LOG(LogParty, Verbose, TEXT("[%s] Responding to approval request for %s with denied"), *PartyId->ToString(), *PendingApproval.SenderId.ToString());
+			PartyInt->ApproveJoinRequest(*PendingApproval.RecipientId, *PartyId, *PendingApproval.SenderId, false, (int32)EJoinPartyDenialReason::Busy);
 		}
 	} while (!bEmpty);
 }
 
 void UPartyGameState::OnReservationBeaconUpdateConnectionFailure()
 {
-	UE_LOG(LogParty, Verbose, TEXT("Reservation update beacon failure %s."), ReservationBeaconClient ? *ReservationBeaconClient->GetName() : *FString(TEXT("")));
+	UE_LOG(LogParty, Verbose, TEXT("Reservation update beacon failure %s."), ReservationBeaconClient ? *ReservationBeaconClient->GetName() : *FString());
 
 	// empty the queue, denying all requests
 	RejectAllPendingJoinRequests();
@@ -1107,8 +1301,11 @@ void UPartyGameState::OnReservationBeaconUpdateResponseReceived(EPartyReservatio
 		ReservationResponse == EPartyReservationResult::ReservationDuplicate)
 	{
 		UWorld* World = GetWorld();
+
 		IOnlinePartyPtr PartyInt = Online::GetPartyInterface(World);
-		bool bValidInterface = PartyInt.IsValid() && OssParty.IsValid();
+		TSharedPtr<const FOnlinePartyId> PartyId = GetPartyId();
+
+		bool bValidInterface = PartyInt.IsValid() && PartyId.IsValid();
 
 		// There should be at least the one
 		FPendingMemberApproval PendingApproval;
@@ -1116,8 +1313,8 @@ void UPartyGameState::OnReservationBeaconUpdateResponseReceived(EPartyReservatio
 		{
 			if (bValidInterface)
 			{
-				UE_LOG(LogParty, Verbose, TEXT("[%s] Responding to approval request for %s with approved"), OssParty.IsValid() ? *OssParty->PartyId->ToString() : TEXT("INVALID"), *PendingApproval.SenderId.ToString());
-				PartyInt->ApproveJoinRequest(*PendingApproval.RecipientId, *OssParty->PartyId, *PendingApproval.SenderId, true);
+				UE_LOG(LogParty, Verbose, TEXT("[%s] Responding to approval request for %s with approved"), *PartyId->ToString(), *PendingApproval.SenderId.ToString());
+				PartyInt->ApproveJoinRequest(*PendingApproval.RecipientId, *PartyId, *PendingApproval.SenderId, true);
 			}
 		}
 
diff --git a/Engine/Source/Runtime/Online/OnlineFramework/Party/Public/Party.h b/Engine/Source/Runtime/Online/OnlineFramework/Party/Public/Party.h
--- a/Engine/Source/Runtime/Online/OnlineFramework/Party/Public/Party.h
+++ b/Engine/Source/Runtime/Online/OnlineFramework/Party/Public/Party.h
@@ -244,7 +244,6 @@ public:
 	 * @param JoinCompleteDelegate delegate to call once the pending invite has been joined or in all failure cases
 	 */
 	void AddPendingPartyJoin(const FUniqueNetId& LocalUserId, const FPartyDetails& PartyDetails, const UPartyDelegates::FOnJoinUPartyComplete& JoinCompleteDelegate);
-	//void AddPendingPartyJoin(const FUniqueNetId& LocalUserId, const FPartyDetails& PartyDetails, const FOnJoinPartyComplete& JoinCompleteDelegate);
 
 	/**
 	 * Is any local player in the given party
@@ -439,17 +438,13 @@ private:
 	void OnLeavePartyComplete(const FUniqueNetId& LocalUserId, const ELeavePartyCompletionResult Result, const FOnlinePartyTypeId InPartyTypeId, UPartyDelegates::FOnLeaveUPartyComplete CompletionDelegate);
 	void HandleJoinPersistentPartyFailure();
 
+	/**
+	 * Delegates for managing regular parties
+	 */
 	void OnCreatePartyInternalComplete(const FUniqueNetId& LocalUserId, const TSharedPtr<const FOnlinePartyId>& InPartyId, const ECreatePartyCompletionResult Result, const FOnlinePartyTypeId InPartyTypeId, UPartyDelegates::FOnCreateUPartyComplete CompletionDelegate);
 	void OnJoinPartyInternalComplete(const FUniqueNetId& LocalUserId, const FOnlinePartyId& InPartyId, const EJoinPartyCompletionResult Result, int32 DeniedResultCode, const FOnlinePartyTypeId InPartyTypeId, UPartyDelegates::FOnJoinUPartyComplete CompletionDelegate);
 	void OnLeavePartyInternalComplete(const FUniqueNetId& LocalUserId, const FOnlinePartyId& InPartyId, const ELeavePartyCompletionResult Result, const FOnlinePartyTypeId InPartyTypeId, UPartyDelegates::FOnLeaveUPartyComplete CompletionDelegate);
 
-	/**
-	 * Internal helpers for delegates
-	 */
-	enum class EOnCreatePartyCompleteInternalResult OnCreatePartyCompleteInternal(const FUniqueNetId& LocalUserId, const FOnlinePartyTypeId InPartyTypeId);
-	enum class EOnJoinPartyCompleteInternalResult   OnJoinPartyCompleteInternal(const FUniqueNetId& LocalUserId, const FOnlinePartyTypeId InPartyTypeId);
-	enum class EOnLeavePartyCompleteInternalResult  OnLeavePartyCompleteInternal(const FUniqueNetId& LocalUserId, const FOnlinePartyTypeId InPartyTypeId);
-
 	/**
 	 * Delegates for managing persistent parties
 	 */
@@ -514,100 +509,3 @@ protected:
 
 };
 
-enum class EOnCreatePartyCompleteInternalResult
-{
-	UnknownFailure,
-	InvalidSystemInterface,
-	InvalidPartyId,
-	InvalidParty,
-	InvalidPartyType,
-	Success
-};
-
-enum class EOnJoinPartyCompleteInternalResult
-{
-	UnknownFailure,
-	InvalidSystemInterface,
-	InvalidParty,
-	InvalidPartyType,
-	Success
-};
-
-enum class EOnLeavePartyCompleteInternalResult
-{
-	UnknownFailure,
-	InvalidPartyState,
-	Success
-};
-
-inline const TCHAR* ToString(EOnCreatePartyCompleteInternalResult Value)
-{
-	switch (Value)
-	{
-	case EOnCreatePartyCompleteInternalResult::UnknownFailure:
-	{
-		return TEXT("UnknownFailure");
-	}
-	case EOnCreatePartyCompleteInternalResult::InvalidSystemInterface:
-	{
-		return TEXT("InvalidSystemInterface");
-	}
-	case EOnCreatePartyCompleteInternalResult::InvalidPartyId:
-	{
-		return TEXT("InvalidPartyId");
-	}
-	case EOnCreatePartyCompleteInternalResult::InvalidParty:
-	{
-		return TEXT("InvalidParty");
-	}
-	case EOnCreatePartyCompleteInternalResult::Success:
-	{
-		return TEXT("Success");
-	}
-	}
-	return TEXT("");
-}
-
-inline const TCHAR* ToString(EOnJoinPartyCompleteInternalResult Value)
-{
-	switch (Value)
-	{
-	case EOnJoinPartyCompleteInternalResult::UnknownFailure:
-	{
-		return TEXT("UnknownFailure");
-	}
-	case EOnJoinPartyCompleteInternalResult::InvalidSystemInterface:
-	{
-		return TEXT("InvalidSystemInterface");
-	}
-	case EOnJoinPartyCompleteInternalResult::InvalidParty:
-	{
-		return TEXT("InvalidParty");
-	}
-	case EOnJoinPartyCompleteInternalResult::Success:
-	{
-		return TEXT("Success");
-	}
-	}
-	return TEXT("");
-}
-
-inline const TCHAR* ToString(EOnLeavePartyCompleteInternalResult Value)
-{
-	switch (Value)
-	{
-	case EOnLeavePartyCompleteInternalResult::UnknownFailure:
-	{
-		return TEXT("UnknownFailure");
-	}
-	case EOnLeavePartyCompleteInternalResult::InvalidPartyState:
-	{
-		return TEXT("InvalidPartyState");
-	}
-	case EOnLeavePartyCompleteInternalResult::Success:
-	{
-		return TEXT("Success");
-	}
-	}
-	return TEXT("");
-}
diff --git a/Engine/Source/Runtime/Online/OnlineFramework/Party/Public/PartyModule.h b/Engine/Source/Runtime/Online/OnlineFramework/Party/Public/PartyModule.h
--- a/Engine/Source/Runtime/Online/OnlineFramework/Party/Public/PartyModule.h
+++ b/Engine/Source/Runtime/Online/OnlineFramework/Party/Public/PartyModule.h
@@ -13,7 +13,7 @@ DECLARE_STATS_GROUP(TEXT("Party"), STATGROUP_Party, STATCAT_Advanced);
 DECLARE_CYCLE_STAT_EXTERN(TEXT("PartyStat1"), STAT_PartyStat1, STATGROUP_Party, PARTY_API);
 
 /**
- * Module for Voice capture/compression/decompression implementations
+ * Module for the party framework that builds on top of IOnlineParty
  */
 class FPartyModule : 
 	public IModuleInterface, public FSelfRegisteringExec
@@ -50,13 +50,13 @@ private:
 	// IModuleInterface
 
 	/**
-	 * Called when voice module is loaded
+	 * Called when party module is loaded
 	 * Initialize platform specific parts of voice handling
 	 */
 	virtual void StartupModule() override;
 	
 	/**
-	 * Called when voice module is unloaded
+	 * Called when party module is unloaded
 	 * Shutdown platform specific parts of voice handling
 	 */
 	virtual void ShutdownModule() override;
diff --git a/Engine/Source/Runtime/Online/OnlineFramework/Qos/Private/QosEvaluator.cpp b/Engine/Source/Runtime/Online/OnlineFramework/Qos/Private/QosEvaluator.cpp
--- a/Engine/Source/Runtime/Online/OnlineFramework/Qos/Private/QosEvaluator.cpp
+++ b/Engine/Source/Runtime/Online/OnlineFramework/Qos/Private/QosEvaluator.cpp
@@ -461,34 +461,9 @@ void UQosEvaluator::EndAnalytics(EQosCompletionResult CompletionResult)
 	}
 }
 
-FString UQosEvaluator::GetDefaultRegionString()
+const FString& UQosEvaluator::GetDefaultRegionString()
 {
-	static FString RegionString = TEXT("");
-	if (RegionString.IsEmpty())
-	{
-		FString OverrideRegion;
-		if (FParse::Value(FCommandLine::Get(), TEXT("McpRegion="), OverrideRegion))
-		{
-			// Region specified on command line
-			RegionString = OverrideRegion.ToUpper();
-		}
-		else
-		{
-			FString DefaultRegion;
-			if (GConfig->GetString(TEXT("Qos"), TEXT("DefaultRegion"), DefaultRegion, GGameIni))
-			{
-				// Region specified in ini file
-				RegionString = DefaultRegion.ToUpper();
-			}
-			else
-			{
-				// No Region specified. Assume USA.
-				RegionString = TEXT("USA");
-			}
-		}
-	}
-
-	return RegionString;
+	return FQosInterface::GetDefaultRegionString();
 }
 
 UWorld* UQosEvaluator::GetWorld() const
diff --git a/Engine/Source/Runtime/Online/OnlineFramework/Qos/Private/QosInterface.cpp b/Engine/Source/Runtime/Online/OnlineFramework/Qos/Private/QosInterface.cpp
--- a/Engine/Source/Runtime/Online/OnlineFramework/Qos/Private/QosInterface.cpp
+++ b/Engine/Source/Runtime/Online/OnlineFramework/Qos/Private/QosInterface.cpp
@@ -50,4 +50,62 @@ FQosInterface::~FQosInterface()
 UQosEvaluator* FQosInterface::CreateQosEvaluator()
 {
 	return NewObject<UQosEvaluator>();
+}
+
+const FString& FQosInterface::GetDefaultRegionString()
+{
+	static FString RegionString = TEXT("");
+	if (RegionString.IsEmpty())
+	{
+		FString OverrideRegion;
+		if (FParse::Value(FCommandLine::Get(), TEXT("McpRegion="), OverrideRegion))
+		{
+			// Region specified on command line
+			RegionString = OverrideRegion.ToUpper();
+		}
+		else
+		{
+			FString DefaultRegion;
+			if (GConfig->GetString(TEXT("Qos"), TEXT("DefaultRegion"), DefaultRegion, GGameIni))
+			{
+				// Region specified in ini file
+				RegionString = DefaultRegion.ToUpper();
+			}
+			else
+			{
+				// No Region specified. Assume USA.
+				RegionString = TEXT("USA");
+			}
+		}
+	}
+
+	return RegionString;
+}
+
+const FString& FQosInterface::GetDatacenterId()
+{
+	static bool bDCIDCheck = false;
+	static FString DCIDString = TEXT("");
+	if (!bDCIDCheck)
+	{
+		FString OverrideDCID;
+		if (FParse::Value(FCommandLine::Get(), TEXT("DCID="), OverrideDCID))
+		{
+			// Region specified on command line
+			DCIDString = OverrideDCID.ToUpper();
+		}
+		else
+		{
+			FString DefaultDCID;
+			if (GConfig->GetString(TEXT("Qos"), TEXT("DCID"), DefaultDCID, GGameIni))
+			{
+				// Region specified in ini file
+				DCIDString = DefaultDCID.ToUpper();
+			}
+		}
+
+		bDCIDCheck = true;
+	}
+
+	return DCIDString;
 }
\ No newline at end of file
diff --git a/Engine/Source/Runtime/Online/OnlineFramework/Qos/Public/QosEvaluator.h b/Engine/Source/Runtime/Online/OnlineFramework/Qos/Public/QosEvaluator.h
--- a/Engine/Source/Runtime/Online/OnlineFramework/Qos/Public/QosEvaluator.h
+++ b/Engine/Source/Runtime/Online/OnlineFramework/Qos/Public/QosEvaluator.h
@@ -160,7 +160,7 @@ public:
 	 *
 	 * @return the default region identifier
 	 */
-	static FString GetDefaultRegionString();
+	static const FString& GetDefaultRegionString();
 
 private:
 
diff --git a/Engine/Source/Runtime/Online/OnlineFramework/Qos/Public/QosInterface.h b/Engine/Source/Runtime/Online/OnlineFramework/Qos/Public/QosInterface.h
--- a/Engine/Source/Runtime/Online/OnlineFramework/Qos/Public/QosInterface.h
+++ b/Engine/Source/Runtime/Online/OnlineFramework/Qos/Public/QosInterface.h
@@ -40,5 +40,19 @@ public:
 	 * @return a new instance of a qos evaluator
 	 */
 	UQosEvaluator* CreateQosEvaluator();
+
+	/**
+	 * Get the default region for this instance, checking ini and commandline overrides
+	 *
+	 * @return the default region identifier
+	 */
+	static const FString& GetDefaultRegionString();
+
+	/**
+	 * Get the datacenter id for this instance, checking ini and commandline overrides
+	 *
+	 * @return the default datacenter identifier
+	 */
+	static const FString& GetDatacenterId();
 };
 
diff --git a/Engine/Source/Runtime/Online/OnlineSubsystem/Public/Interfaces/OnlineExternalUIInterface.h b/Engine/Source/Runtime/Online/OnlineSubsystem/Public/Interfaces/OnlineExternalUIInterface.h
--- a/Engine/Source/Runtime/Online/OnlineSubsystem/Public/Interfaces/OnlineExternalUIInterface.h
+++ b/Engine/Source/Runtime/Online/OnlineSubsystem/Public/Interfaces/OnlineExternalUIInterface.h
@@ -37,6 +37,10 @@ struct FShowWebUrlParams
 {
 	/** presented without a frame if embedded enabled */
 	bool bEmbedded;
+	/** Show the built in close button */
+	bool bShowCloseButton;
+	/** Show the built in background */
+	bool bShowBackground;
 	/** x offset in pixels from top left */
 	int32 OffsetX;
 	/** y offset in pixels from top left */
@@ -55,6 +59,8 @@ struct FShowWebUrlParams
 	 */
 	FShowWebUrlParams(bool InbEmbedded, int32 InOffsetX, int32 InOffsetY, int32 InSizeX, int32 InSizeY)
 		: bEmbedded(InbEmbedded)
+		, bShowCloseButton(false)
+		, bShowBackground(false)
 		, OffsetX(InOffsetX)
 		, OffsetY(InOffsetY)
 		, SizeX(InSizeX)
@@ -66,6 +72,8 @@ struct FShowWebUrlParams
 	 */
 	FShowWebUrlParams()
 		: bEmbedded(false)
+		, bShowCloseButton(false)
+		, bShowBackground(false)
 		, OffsetX(0)
 		, OffsetY(0)
 		, SizeX(0)
@@ -135,7 +143,7 @@ public:
 	virtual bool ShowLeaderboardUI(const FString& LeaderboardName) = 0;
 
 	/**
-	 *	Displays a web page in the external UI
+	 * Displays a web page in the external UI
 	 *
 	 * @param WebURL fully formed web address (http://www.google.com)
 	 *
@@ -143,6 +151,13 @@ public:
 	 */
 	virtual bool ShowWebURL(const FString& Url, const FShowWebUrlParams& ShowParams, const FOnShowWebUrlClosedDelegate& Delegate = FOnShowWebUrlClosedDelegate()) = 0;
 
+	/**
+	 * Closes the currently active web external UI
+	 *
+	 * @return true if it was able to show the UI, false if it failed
+	 */
+	virtual bool CloseWebURL() = 0;
+
 	/**
 	 * Displays a user's profile card.
 	 *
diff --git a/Engine/Source/Runtime/Online/OnlineSubsystem/Public/Interfaces/OnlinePartyInterface.h b/Engine/Source/Runtime/Online/OnlineSubsystem/Public/Interfaces/OnlinePartyInterface.h
--- a/Engine/Source/Runtime/Online/OnlineSubsystem/Public/Interfaces/OnlinePartyInterface.h
+++ b/Engine/Source/Runtime/Online/OnlineSubsystem/Public/Interfaces/OnlinePartyInterface.h
@@ -142,6 +142,11 @@ public:
      */
 	virtual bool IsAcceptingMembers() const = 0;
 
+	/**
+	 * @return true if this is a party of one
+	 */
+	virtual bool IsPartyOfOne() const = 0;
+
 	/**
 	 * @return why the party is not accepting members
 	 */
diff --git a/Engine/Source/Runtime/Online/OnlineSubsystem/Public/OnlineSessionSettings.h b/Engine/Source/Runtime/Online/OnlineSubsystem/Public/OnlineSessionSettings.h
--- a/Engine/Source/Runtime/Online/OnlineSubsystem/Public/OnlineSessionSettings.h
+++ b/Engine/Source/Runtime/Online/OnlineSubsystem/Public/OnlineSessionSettings.h
@@ -18,8 +18,10 @@
 #define SETTING_BEACONPORT FName(TEXT("BEACONPORT"))
 /** Server responds to Qos beacon requests (value is int32) */
 #define SETTING_QOS FName(TEXT("QOS"))
-/** String describing the region of the world you are in (value is FString) */
+/** Setting describing the region of the world you are in (value is FString) */
 #define SETTING_REGION FName(TEXT("REGION"))
+/** Setting describing the unique id of a datacenter (value is FString) */
+#define SETTING_DCID FName(TEXT("DCID"))
 
 /** 8 user defined integer params to be used when filtering searches for sessions */
 #define SETTING_CUSTOMSEARCHINT1 FName(TEXT("CUSTOMSEARCHINT1"))
diff --git a/Engine/Source/Runtime/Online/OnlineSubsystemSteam/Private/OnlineExternalUIInterfaceSteam.cpp b/Engine/Source/Runtime/Online/OnlineSubsystemSteam/Private/OnlineExternalUIInterfaceSteam.cpp
--- a/Engine/Source/Runtime/Online/OnlineSubsystemSteam/Private/OnlineExternalUIInterfaceSteam.cpp
+++ b/Engine/Source/Runtime/Online/OnlineSubsystemSteam/Private/OnlineExternalUIInterfaceSteam.cpp
@@ -85,6 +85,11 @@ bool FOnlineExternalUISteam::ShowWebURL(const FString& Url, const FShowWebUrlPar
 	return true;
 }
 
+bool FOnlineExternalUISteam::CloseWebURL()
+{
+	return false;
+}
+
 bool FOnlineExternalUISteam::ShowProfileUI(const FUniqueNetId& Requestor, const FUniqueNetId& Requestee, const FOnProfileUIClosedDelegate& Delegate)
 {
 	SteamFriends()->ActivateGameOverlayToUser(TCHAR_TO_UTF8(TEXT("steamid")), (const FUniqueNetIdSteam&)Requestee);
diff --git a/Engine/Source/Runtime/Online/XMPP/Private/XmppJingle/XmppConnectionJingle.cpp b/Engine/Source/Runtime/Online/XMPP/Private/XmppJingle/XmppConnectionJingle.cpp
--- a/Engine/Source/Runtime/Online/XMPP/Private/XmppJingle/XmppConnectionJingle.cpp
+++ b/Engine/Source/Runtime/Online/XMPP/Private/XmppJingle/XmppConnectionJingle.cpp
@@ -443,47 +443,50 @@ TSharedPtr<class IXmppChat> FXmppConnectionJingle::PrivateChat()
 
 bool FXmppConnectionJingle::Tick(float DeltaTime)
 {
+	ELoginProgress::Type LocalLastLoginState, LocalLoginState;
 	{
 		FScopeLock Lock(&LoginStateLock);
+		LocalLastLoginState = LastLoginState;
+		LocalLoginState = LoginState;
+		LastLoginState = LoginState;
+	}
 
-		if (LastLoginState != LoginState)
+	if (LocalLastLoginState != LocalLoginState)
+	{
+		if (LocalLoginState == ELoginProgress::LoggedIn)
 		{
-			if (LoginState == ELoginProgress::LoggedIn)
+			UE_LOG(LogXmpp, Log, TEXT("Logged IN JID=%s"), *GetUserJid().GetFullPath());
+			if (LocalLastLoginState == ELoginProgress::ProcessingLogin)
 			{
-				UE_LOG(LogXmpp, Log, TEXT("Logged IN JID=%s"), *GetUserJid().GetFullPath());
-				if (LastLoginState == ELoginProgress::ProcessingLogin)
+				OnLoginComplete().Broadcast(GetUserJid(), true, FString());
+
+				// send/obtain initial presence
+				if (Presence().IsValid())
 				{
-					OnLoginComplete().Broadcast(GetUserJid(), true, FString());
-
-					// send/obtain initial presence
-					if (Presence().IsValid())
-					{
-						FXmppUserPresence InitialPresence = Presence()->GetPresence();
-						InitialPresence.bIsAvailable = true;
-						InitialPresence.Status = EXmppPresenceStatus::Offline;
-						Presence()->UpdatePresence(InitialPresence);
-					}
+					FXmppUserPresence InitialPresence = Presence()->GetPresence();
+					InitialPresence.bIsAvailable = true;
+					InitialPresence.Status = EXmppPresenceStatus::Offline;
+					Presence()->UpdatePresence(InitialPresence);
 				}
-				OnLoginChanged().Broadcast(GetUserJid(), EXmppLoginStatus::LoggedIn);
 			}
-			else if (LoginState == ELoginProgress::LoggedOut)
+			OnLoginChanged().Broadcast(GetUserJid(), EXmppLoginStatus::LoggedIn);
+		}
+		else if (LocalLoginState == ELoginProgress::LoggedOut)
+		{
+			UE_LOG(LogXmpp, Log, TEXT("Logged OUT JID=%s"), *GetUserJid().GetFullPath());
+			if (LocalLastLoginState == ELoginProgress::ProcessingLogin)
 			{
-				UE_LOG(LogXmpp, Log, TEXT("Logged OUT JID=%s"), *GetUserJid().GetFullPath());
-				if (LastLoginState == ELoginProgress::ProcessingLogin)
-				{
-					OnLoginComplete().Broadcast(GetUserJid(), false, FString());
-				}
-				else if (LastLoginState == ELoginProgress::ProcessingLogout)
-				{
-					OnLogoutComplete().Broadcast(GetUserJid(), true, FString());
-				}
-				if (LastLoginState == ELoginProgress::LoggedIn ||
-					LastLoginState == ELoginProgress::ProcessingLogout)
-				{
-					OnLoginChanged().Broadcast(GetUserJid(), EXmppLoginStatus::LoggedOut);
-				}
+				OnLoginComplete().Broadcast(GetUserJid(), false, FString());
+			}
+			else if (LocalLastLoginState == ELoginProgress::ProcessingLogout)
+			{
+				OnLogoutComplete().Broadcast(GetUserJid(), true, FString());
+			}
+			if (LocalLastLoginState == ELoginProgress::LoggedIn ||
+				LocalLastLoginState == ELoginProgress::ProcessingLogout)
+			{
+				OnLoginChanged().Broadcast(GetUserJid(), EXmppLoginStatus::LoggedOut);
 			}
-			LastLoginState = LoginState;
 		}
 	}
 
@@ -743,12 +746,7 @@ void FXmppConnectionJingle::Login(const FString& UserId, const FString& Password
 
 void FXmppConnectionJingle::Logout()
 {	
-	FString ErrorStr;
-
-	if (GetLoginStatus() != EXmppLoginStatus::LoggedIn)
-	{
-		ErrorStr = TEXT("not logged in");
-	}
+	FString ErrorStr;	
 	if (PumpThread != NULL)
 	{
 		//@todo sz1 - should be able to reuse existing connection pump for login/logout
@@ -756,8 +754,12 @@ void FXmppConnectionJingle::Logout()
 		PumpThread->Logout();
 #else
 		Shutdown();
-		OnLogoutComplete().Broadcast(GetUserJid(), true, ErrorStr);
 #endif
+		if (GetLoginStatus() != EXmppLoginStatus::LoggedIn)
+		{
+			// we want this to trigger OnLogoutComplete because Shutdown() will not (in the case of not logged in)
+			ErrorStr = TEXT("not logged in");
+		}
 	}
 	else
 	{
diff --git a/Engine/Source/Runtime/Renderer/Private/BasePassRendering.cpp b/Engine/Source/Runtime/Renderer/Private/BasePassRendering.cpp
--- a/Engine/Source/Runtime/Renderer/Private/BasePassRendering.cpp
+++ b/Engine/Source/Runtime/Renderer/Private/BasePassRendering.cpp
@@ -396,6 +396,9 @@ public:
 
 		for( int32 BatchElementIndex = 0, Num = Parameters.Mesh.Elements.Num(); BatchElementIndex < Num; BatchElementIndex++ )
 		{
+			TDrawEvent<FRHICommandList> MeshEvent;
+			BeginMeshDrawEvent(RHICmdList, Parameters.PrimitiveSceneProxy, Parameters.Mesh, MeshEvent);
+
 			DrawingPolicy.SetMeshRenderState(
 				RHICmdList, 
 				View,
diff --git a/Engine/Source/Runtime/Renderer/Private/ConvertToUniformMesh.cpp b/Engine/Source/Runtime/Renderer/Private/ConvertToUniformMesh.cpp
--- a/Engine/Source/Runtime/Renderer/Private/ConvertToUniformMesh.cpp
+++ b/Engine/Source/Runtime/Renderer/Private/ConvertToUniformMesh.cpp
@@ -272,8 +272,6 @@ void FConvertToUniformMeshDrawingPolicy::SetMeshRenderState(
 {
 	const FMeshBatchElement& BatchElement = Mesh.Elements[BatchElementIndex];
 
-	EmitMeshDrawEvents(RHICmdList, PrimitiveSceneProxy, Mesh);
-
 	// Set transforms
 	VertexShader->SetMesh(RHICmdList, VertexFactory,View,PrimitiveSceneProxy,BatchElement, DitheredLODTransitionValue);
 	GeometryShader->SetMesh(RHICmdList, VertexFactory,View,PrimitiveSceneProxy,BatchElement, DitheredLODTransitionValue);
diff --git a/Engine/Source/Runtime/Renderer/Private/DepthRendering.cpp b/Engine/Source/Runtime/Renderer/Private/DepthRendering.cpp
--- a/Engine/Source/Runtime/Renderer/Private/DepthRendering.cpp
+++ b/Engine/Source/Runtime/Renderer/Private/DepthRendering.cpp
@@ -418,6 +418,9 @@ bool FDepthDrawingPolicyFactory::DrawMesh(
 			{
 				if(Mask & 1)
 				{
+					TDrawEvent<FRHICommandList> MeshEvent;
+					BeginMeshDrawEvent(RHICmdList, PrimitiveSceneProxy, Mesh, MeshEvent);
+
 					DrawingPolicy.SetMeshRenderState(RHICmdList, View,PrimitiveSceneProxy,Mesh,BatchElementIndex,bBackFace,DitheredLODTransitionValue,FPositionOnlyDepthDrawingPolicy::ElementDataType(),FDepthDrawingPolicy::ContextDataType());
 					DrawingPolicy.DrawMesh(RHICmdList, Mesh,BatchElementIndex);
 				}
@@ -464,6 +467,9 @@ bool FDepthDrawingPolicyFactory::DrawMesh(
 				{
 					if(Mask & 1)
 					{
+						TDrawEvent<FRHICommandList> MeshEvent;
+						BeginMeshDrawEvent(RHICmdList, PrimitiveSceneProxy, Mesh, MeshEvent);
+
 						DrawingPolicy.SetMeshRenderState(RHICmdList, View,PrimitiveSceneProxy,Mesh,BatchElementIndex,bBackFace,DitheredLODTransitionValue,FMeshDrawingPolicy::ElementDataType(),FDepthDrawingPolicy::ContextDataType());
 						DrawingPolicy.DrawMesh(RHICmdList, Mesh,BatchElementIndex);
 					}
diff --git a/Engine/Source/Runtime/Renderer/Private/DistortionRendering.cpp b/Engine/Source/Runtime/Renderer/Private/DistortionRendering.cpp
--- a/Engine/Source/Runtime/Renderer/Private/DistortionRendering.cpp
+++ b/Engine/Source/Runtime/Renderer/Private/DistortionRendering.cpp
@@ -550,8 +550,6 @@ void TDistortionMeshDrawingPolicy<DistortMeshPolicy>::SetMeshRenderState(
 
 	const FMeshBatchElement& BatchElement = Mesh.Elements[BatchElementIndex];
 
-	EmitMeshDrawEvents(RHICmdList, PrimitiveSceneProxy, Mesh);
-
 	// Set transforms
 	VertexShader->SetMesh(RHICmdList, VertexFactory,View,PrimitiveSceneProxy,BatchElement,DitheredLODTransitionValue);
 
@@ -659,6 +657,9 @@ bool TDistortionMeshDrawingPolicyFactory<DistortMeshPolicy>::DrawDynamicMesh(
 		DrawingPolicy.SetSharedState(RHICmdList, &View, typename TDistortionMeshDrawingPolicy<DistortMeshPolicy>::ContextDataType());
 		for (int32 BatchElementIndex = 0; BatchElementIndex < Mesh.Elements.Num(); BatchElementIndex++)
 		{
+			TDrawEvent<FRHICommandList> MeshEvent;
+			BeginMeshDrawEvent(RHICmdList, PrimitiveSceneProxy, Mesh, MeshEvent);
+
 			DrawingPolicy.SetMeshRenderState(RHICmdList, View,PrimitiveSceneProxy,Mesh,BatchElementIndex,bBackFace,Mesh.DitheredLODTransitionAlpha,typename TDistortionMeshDrawingPolicy<DistortMeshPolicy>::ElementDataType(), typename TDistortionMeshDrawingPolicy<DistortMeshPolicy>::ContextDataType());
 			DrawingPolicy.DrawMesh(RHICmdList, Mesh,BatchElementIndex);
 		}
@@ -709,6 +710,9 @@ bool TDistortionMeshDrawingPolicyFactory<DistortMeshPolicy>::DrawStaticMesh(
 		{
 			if(BatchElementMask & 1)
 			{
+				TDrawEvent<FRHICommandList> MeshEvent;
+				BeginMeshDrawEvent(RHICmdList, PrimitiveSceneProxy, StaticMesh, MeshEvent);
+
 				DrawingPolicy.SetMeshRenderState(RHICmdList, *View,PrimitiveSceneProxy,StaticMesh,BatchElementIndex,bBackFace,DitherValue,
 					typename TDistortionMeshDrawingPolicy<DistortMeshPolicy>::ElementDataType(),
 					typename TDistortionMeshDrawingPolicy<DistortMeshPolicy>::ContextDataType()
diff --git a/Engine/Source/Runtime/Renderer/Private/DrawingPolicy.cpp b/Engine/Source/Runtime/Renderer/Private/DrawingPolicy.cpp
--- a/Engine/Source/Runtime/Renderer/Private/DrawingPolicy.cpp
+++ b/Engine/Source/Runtime/Renderer/Private/DrawingPolicy.cpp
@@ -60,13 +60,11 @@ void FMeshDrawingPolicy::SetMeshRenderState(
 	const ContextDataType PolicyContext
 	) const
 {
-	EmitMeshDrawEvents(RHICmdList, PrimitiveSceneProxy, Mesh);
-
-		// Use bitwise logic ops to avoid branches
+	// Use bitwise logic ops to avoid branches
 	RHICmdList.SetRasterizerState( GetStaticRasterizerState<true>(
-			( Mesh.bWireframe || IsWireframe() ) ? FM_Wireframe : FM_Solid, ( ( IsTwoSided() && !NeedsBackfacePass() ) || Mesh.bDisableBackfaceCulling ) ? CM_None :
-			( ( (View.bReverseCulling ^ bBackFace) ^ Mesh.ReverseCulling ) ? CM_CCW : CM_CW )
-			));
+		( Mesh.bWireframe || IsWireframe() ) ? FM_Wireframe : FM_Solid, ( ( IsTwoSided() && !NeedsBackfacePass() ) || Mesh.bDisableBackfaceCulling ) ? CM_None :
+		( ( (View.bReverseCulling ^ bBackFace) ^ Mesh.ReverseCulling ) ? CM_CCW : CM_CW )
+		));
 }
 
 void FMeshDrawingPolicy::DrawMesh(FRHICommandList& RHICmdList, const FMeshBatch& Mesh, int32 BatchElementIndex) const
diff --git a/Engine/Source/Runtime/Renderer/Private/ForwardBasePassRendering.cpp b/Engine/Source/Runtime/Renderer/Private/ForwardBasePassRendering.cpp
--- a/Engine/Source/Runtime/Renderer/Private/ForwardBasePassRendering.cpp
+++ b/Engine/Source/Runtime/Renderer/Private/ForwardBasePassRendering.cpp
@@ -285,6 +285,9 @@ public:
 
 		for( int32 BatchElementIndex=0;BatchElementIndex<Parameters.Mesh.Elements.Num();BatchElementIndex++ )
 		{
+			TDrawEvent<FRHICommandList> MeshEvent;
+			BeginMeshDrawEvent(RHICmdList, Parameters.PrimitiveSceneProxy, Parameters.Mesh, MeshEvent);
+
 			DrawingPolicy.SetMeshRenderState(
 				RHICmdList, 
 				View,
diff --git a/Engine/Source/Runtime/Renderer/Private/ReflectionEnvironment.cpp b/Engine/Source/Runtime/Renderer/Private/ReflectionEnvironment.cpp
--- a/Engine/Source/Runtime/Renderer/Private/ReflectionEnvironment.cpp
+++ b/Engine/Source/Runtime/Renderer/Private/ReflectionEnvironment.cpp
@@ -890,41 +890,45 @@ void FDeferredShadingSceneRenderer::RenderTiledDeferredImageBasedReflections(FRH
 			bool bHasSphereCaptures = (NumSphereCaptures > 0);
 			bool bNeedsClearCoat = (View.ShadingModelMaskInView & (1 << MSM_ClearCoat)) != 0;
 
+			static const FName TiledReflBeginComputeName(TEXT("ReflectionEnvBeginComputeFence"));
+			static const FName TiledReflEndComputeName(TEXT("ReflectionEnvEndComputeFence"));
+			FComputeFenceRHIRef ReflectionBeginFence = RHICmdList.CreateComputeFence(TiledReflBeginComputeName);
+			FComputeFenceRHIRef ReflectionEndFence = RHICmdList.CreateComputeFence(TiledReflEndComputeName);
+
 			//Grab the async compute commandlist.
 			FRHIAsyncComputeCommandListImmediate& RHICmdListComputeImmediate = FRHICommandListExecutor::GetImmediateAsyncComputeCommandList();
 
-			SCOPED_COMPUTE_EVENTF(RHICmdListComputeImmediate, ReflectionEnvironment, TEXT("ReflectionEnvironment ComputeShader %dx%d Tile:%dx%d Box:%d Sphere:%d ClearCoat:%d"),
-				View.ViewRect.Width(), View.ViewRect.Height(), GReflectionEnvironmentTileSizeX, GReflectionEnvironmentTileSizeY,
-				NumBoxCaptures, NumSphereCaptures, bNeedsClearCoat);
+			{
+				SCOPED_COMPUTE_EVENTF(RHICmdListComputeImmediate, ReflectionEnvironment, TEXT("ReflectionEnvironment ComputeShader %dx%d Tile:%dx%d Box:%d Sphere:%d ClearCoat:%d"),
+					View.ViewRect.Width(), View.ViewRect.Height(), GReflectionEnvironmentTileSizeX, GReflectionEnvironmentTileSizeY,
+					NumBoxCaptures, NumSphereCaptures, bNeedsClearCoat);
 
-			ComputeShader = SelectReflectionEnvironmentTiledDeferredCS(View.ShaderMap, bUseLightmaps, bNeedsClearCoat, bHasBoxCaptures, bHasSphereCaptures);
+				ComputeShader = SelectReflectionEnvironmentTiledDeferredCS(View.ShaderMap, bUseLightmaps, bNeedsClearCoat, bHasBoxCaptures, bHasSphereCaptures);
 
-			static const FName TiledReflBeginComputeName(TEXT("ReflectionEnvBeginComputeFence"));
-			static const FName TiledReflEndComputeName(TEXT("ReflectionEnvEndComputeFence"));
-			FComputeFenceRHIRef ReflectionBeginFence = RHICmdList.CreateComputeFence(TiledReflBeginComputeName);
-			FComputeFenceRHIRef ReflectionEndFence = RHICmdList.CreateComputeFence(TiledReflEndComputeName);			
-			
-			//Really we should write this fence where we transition the final depedency for the reflections.  We may add an RHI command just for writing fences if this
-			//can't be done in the general case.  In the meantime, hack this command a bit to write the fence.
-			RHICmdList.TransitionResources(EResourceTransitionAccess::EReadable, EResourceTransitionPipeline::EGfxToCompute, nullptr, 0, ReflectionBeginFence);
-					
-			//we must wait on the fence written from the Gfx pipe to let us know all our dependencies are ready.
-			RHICmdListComputeImmediate.WaitComputeFence(ReflectionBeginFence);
-
-			//standard compute setup, but on the async commandlist.
-			RHICmdListComputeImmediate.SetComputeShader(ComputeShader->GetComputeShader());
-
-			FUnorderedAccessViewRHIParamRef OutUAV = NewSceneColor->GetRenderTargetItem().UAV;
-			ComputeShader->SetParameters(RHICmdListComputeImmediate, View, SSROutput->GetRenderTargetItem().ShaderResourceTexture, SortData, OutUAV, DynamicBentNormalAO);
-			
-			uint32 GroupSizeX = (View.ViewRect.Size().X + GReflectionEnvironmentTileSizeX - 1) / GReflectionEnvironmentTileSizeX;
-			uint32 GroupSizeY = (View.ViewRect.Size().Y + GReflectionEnvironmentTileSizeY - 1) / GReflectionEnvironmentTileSizeY;
-			DispatchComputeShader(RHICmdListComputeImmediate, ComputeShader, GroupSizeX, GroupSizeY, 1);
 
-			ComputeShader->UnsetParameters(RHICmdListComputeImmediate, OutUAV);
-			
-			//transition the output to readable and write the fence to allow the Gfx pipe to carry on.
-			RHICmdListComputeImmediate.TransitionResources(EResourceTransitionAccess::EReadable, EResourceTransitionPipeline::EComputeToGfx, &OutUAV, 1, ReflectionEndFence);
+
+				//Really we should write this fence where we transition the final depedency for the reflections.  We may add an RHI command just for writing fences if this
+				//can't be done in the general case.  In the meantime, hack this command a bit to write the fence.
+				RHICmdList.TransitionResources(EResourceTransitionAccess::EReadable, EResourceTransitionPipeline::EGfxToCompute, nullptr, 0, ReflectionBeginFence);
+
+				//we must wait on the fence written from the Gfx pipe to let us know all our dependencies are ready.
+				RHICmdListComputeImmediate.WaitComputeFence(ReflectionBeginFence);
+
+				//standard compute setup, but on the async commandlist.
+				RHICmdListComputeImmediate.SetComputeShader(ComputeShader->GetComputeShader());
+
+				FUnorderedAccessViewRHIParamRef OutUAV = NewSceneColor->GetRenderTargetItem().UAV;
+				ComputeShader->SetParameters(RHICmdListComputeImmediate, View, SSROutput->GetRenderTargetItem().ShaderResourceTexture, SortData, OutUAV, DynamicBentNormalAO);
+
+				uint32 GroupSizeX = (View.ViewRect.Size().X + GReflectionEnvironmentTileSizeX - 1) / GReflectionEnvironmentTileSizeX;
+				uint32 GroupSizeY = (View.ViewRect.Size().Y + GReflectionEnvironmentTileSizeY - 1) / GReflectionEnvironmentTileSizeY;
+				DispatchComputeShader(RHICmdListComputeImmediate, ComputeShader, GroupSizeX, GroupSizeY, 1);
+
+				ComputeShader->UnsetParameters(RHICmdListComputeImmediate, OutUAV);
+
+				//transition the output to readable and write the fence to allow the Gfx pipe to carry on.
+				RHICmdListComputeImmediate.TransitionResources(EResourceTransitionAccess::EReadable, EResourceTransitionPipeline::EComputeToGfx, &OutUAV, 1, ReflectionEndFence);
+			}
 
 			//immediately dispatch our async compute commands to the RHI thread to be submitted to the GPU as soon as possible.
 			FRHIAsyncComputeCommandListImmediate::ImmediateDispatch(RHICmdListComputeImmediate);			
diff --git a/Engine/Source/Runtime/Renderer/Private/SceneHitProxyRendering.cpp b/Engine/Source/Runtime/Renderer/Private/SceneHitProxyRendering.cpp
--- a/Engine/Source/Runtime/Renderer/Private/SceneHitProxyRendering.cpp
+++ b/Engine/Source/Runtime/Renderer/Private/SceneHitProxyRendering.cpp
@@ -215,8 +215,6 @@ void FHitProxyDrawingPolicy::SetMeshRenderState(
 	const ContextDataType PolicyContext
 	) const
 {
-	EmitMeshDrawEvents(RHICmdList, PrimitiveSceneProxy, Mesh);
-
 	const FMeshBatchElement& BatchElement = Mesh.Elements[BatchElementIndex];
 
 	VertexShader->SetMesh(RHICmdList, VertexFactory, View, PrimitiveSceneProxy, BatchElement,DitheredLODTransitionValue);
diff --git a/Engine/Source/Runtime/Renderer/Private/ShadowRendering.cpp b/Engine/Source/Runtime/Renderer/Private/ShadowRendering.cpp
--- a/Engine/Source/Runtime/Renderer/Private/ShadowRendering.cpp
+++ b/Engine/Source/Runtime/Renderer/Private/ShadowRendering.cpp
@@ -968,8 +968,6 @@ void FShadowDepthDrawingPolicy<bRenderingReflectiveShadowMaps>::SetMeshRenderSta
 {
 	const FMeshBatchElement& BatchElement = Mesh.Elements[BatchElementIndex];
 
-	EmitMeshDrawEvents(RHICmdList, PrimitiveSceneProxy, Mesh);
-
 	VertexShader->SetMesh(RHICmdList, VertexFactory,View,PrimitiveSceneProxy,BatchElement,DitheredLODTransitionValue);
 
 	if( HullShader && DomainShader )
@@ -1126,6 +1124,9 @@ bool FShadowDepthDrawingPolicyFactory::DrawDynamicMesh(
 				DrawingPolicy.SetSharedState(RHICmdList, &View, PolicyContext);
 				for (int32 BatchElementIndex = 0, Num = Mesh.Elements.Num(); BatchElementIndex < Num; BatchElementIndex++)
 				{
+					TDrawEvent<FRHICommandList> MeshEvent;
+					BeginMeshDrawEvent(RHICmdList, PrimitiveSceneProxy, Mesh, MeshEvent);
+
 					DrawingPolicy.SetMeshRenderState(RHICmdList, View,PrimitiveSceneProxy,Mesh,BatchElementIndex,bBackFace,Mesh.DitheredLODTransitionAlpha,FMeshDrawingPolicy::ElementDataType(),PolicyContext);
 					DrawingPolicy.DrawMesh(RHICmdList, Mesh,BatchElementIndex);
 				}
@@ -1148,6 +1149,9 @@ bool FShadowDepthDrawingPolicyFactory::DrawDynamicMesh(
 				DrawingPolicy.SetSharedState(RHICmdList, &View, PolicyContext);
 				for (int32 BatchElementIndex = 0; BatchElementIndex < Mesh.Elements.Num(); BatchElementIndex++)
 				{
+					TDrawEvent<FRHICommandList> MeshEvent;
+					BeginMeshDrawEvent(RHICmdList, PrimitiveSceneProxy, Mesh, MeshEvent);
+
 					DrawingPolicy.SetMeshRenderState(RHICmdList, View,PrimitiveSceneProxy,Mesh,BatchElementIndex,bBackFace,Mesh.DitheredLODTransitionAlpha,FMeshDrawingPolicy::ElementDataType(),PolicyContext);
 					DrawingPolicy.DrawMesh(RHICmdList, Mesh,BatchElementIndex);
 				}
@@ -1310,6 +1314,9 @@ void DrawMeshElements(FRHICommandList& RHICmdList, FShadowDepthDrawingPolicy<bRe
 	{
 		if(BatchElementMask & 1)
 		{
+			TDrawEvent<FRHICommandList> MeshEvent;
+			BeginMeshDrawEvent(RHICmdList, Mesh->PrimitiveSceneInfo->Proxy, *Mesh, MeshEvent);
+
 			SharedDrawingPolicy.SetMeshRenderState(RHICmdList, View, Mesh->PrimitiveSceneInfo->Proxy, *Mesh, BatchElementIndex, false, DitherValue, FMeshDrawingPolicy::ElementDataType(),PolicyContext);
 			SharedDrawingPolicy.DrawMesh(RHICmdList, *Mesh, BatchElementIndex);
 			INC_DWORD_STAT(STAT_ShadowDynamicPathDrawCalls);
@@ -1749,6 +1756,13 @@ void FProjectedShadowInfo::ModifyViewForShadow(FRHICommandList& RHICmdList, FVie
 
 	// we are going to set this back now because we only want the correct view rect for the uniform buffer. For LOD calculations, we want the rendering viewrect and proj matrix.
 	FoundView->ViewRect = OriginalViewRect;
+
+	extern int32 GPreshadowsForceLowestLOD;
+
+	if (bPreShadow && GPreshadowsForceLowestLOD)
+	{
+		FoundView->DrawDynamicFlags = EDrawDynamicFlags::ForceLowestLOD;
+	}
 }
 
 FViewInfo* FProjectedShadowInfo::FindViewForShadow(FSceneRenderer* SceneRenderer)
@@ -1873,6 +1887,12 @@ void FProjectedShadowInfo::RenderDepth(FRHICommandList& RHICmdList, FSceneRender
 		    FoundView->ViewMatrices.TranslatedViewProjectionMatrix = TranslatedViewMatrix * FoundView->ViewMatrices.ProjMatrix;
 		    FoundView->ViewMatrices.InvTranslatedViewProjectionMatrix = FoundView->ViewMatrices.TranslatedViewProjectionMatrix.Inverse();
 	    }
+
+		if (bPreShadow)
+		{
+			// Restore
+			FoundView->DrawDynamicFlags = EDrawDynamicFlags::None;
+		}
     }
 }
 
diff --git a/Engine/Source/Runtime/Renderer/Private/ShadowSetup.cpp b/Engine/Source/Runtime/Renderer/Private/ShadowSetup.cpp
--- a/Engine/Source/Runtime/Renderer/Private/ShadowSetup.cpp
+++ b/Engine/Source/Runtime/Renderer/Private/ShadowSetup.cpp
@@ -73,6 +73,14 @@ bool ShouldUseCachePreshadows()
 	return CVarCachePreshadows.GetValueOnRenderThread() != 0;
 }
 
+int32 GPreshadowsForceLowestLOD = 0;
+FAutoConsoleVariableRef CVarPreshadowsForceLowestLOD(
+	TEXT("r.Shadow.PreshadowsForceLowestDetailLevel"),
+	GPreshadowsForceLowestLOD,
+	TEXT("When enabled, static meshes render their lowest detail level into preshadow depth maps.  Disabled by default as it causes artifacts with poor quality LODs (tree billboard)."),
+	ECVF_Scalability | ECVF_RenderThreadSafe
+	);
+
 /**
  * This value specifies how much bounds will be expanded when rendering a cached preshadow (0.15 = 15% larger).
  * Larger values result in more cache hits, but lower resolution and pull more objects into the depth pass.
@@ -759,7 +767,11 @@ void FProjectedShadowInfo::AddSubjectPrimitive(FPrimitiveSceneInfo* PrimitiveSce
 					{
 						bool bUseExistingVisibility = false;
 
-						if(!bReflectiveShadowmap) // Don't use existing visibility for RSMs 
+						// Preshadows use the lowest LOD because there is no self shadowing
+						const bool bForceLowestDetailLevel = bReflectiveShadowmap || (bPreShadow && GPreshadowsForceLowestLOD);
+
+						// Don't use existing visibility if we need to use a different LOD in the shadow depth pass
+						if (!bForceLowestDetailLevel)
 						{
 							for (int32 MeshIndex = 0; MeshIndex < PrimitiveSceneInfo->StaticMeshes.Num(); MeshIndex++)
 							{
@@ -788,7 +800,7 @@ void FProjectedShadowInfo::AddSubjectPrimitive(FPrimitiveSceneInfo* PrimitiveSce
 							int32 ForcedLODLevel = (CurrentView.Family->EngineShowFlags.LOD) ? GetCVarForceLOD() : 0;
 
 							// Add the primitive's static mesh elements to the draw lists.
-							if ( bReflectiveShadowmap) 
+							if (bForceLowestDetailLevel) 
 							{
 								int8 LODToRenderScan = -CHAR_MAX;
 								// Force the lowest detail LOD Level in reflective shadow maps.
@@ -973,6 +985,11 @@ void FProjectedShadowInfo::GatherDynamicMeshElements(FSceneRenderer& Renderer, F
 		    int32 Disable = 0; //CVarDisableCullShadows.GetValueOnRenderThread();
 		    FConvexVolume NoCull;
     
+			if (bPreShadow && GPreshadowsForceLowestLOD)
+			{
+				FoundView->DrawDynamicFlags = EDrawDynamicFlags::ForceLowestLOD;
+			}
+
 		    if (IsWholeSceneDirectionalShadow())
 		    {
 			    FoundView->ViewMatrices.PreShadowTranslation = FVector(0,0,0);
@@ -986,6 +1003,8 @@ void FProjectedShadowInfo::GatherDynamicMeshElements(FSceneRenderer& Renderer, F
 			    FoundView->ViewMatrices.GetDynamicMeshElementsShadowCullFrustum = (Disable & 1) ? &NoCull : &CasterFrustum;
 			    GatherDynamicMeshElementsArray(FoundView, Renderer, DynamicSubjectPrimitives, DynamicSubjectMeshElements, ReusedViewsArray);
 		    }
+
+			FoundView->DrawDynamicFlags = EDrawDynamicFlags::None;
     
 		    FoundView->ViewMatrices.GetDynamicMeshElementsShadowCullFrustum = (Disable & 2) ? &NoCull : &ReceiverFrustum;
 		    GatherDynamicMeshElementsArray(FoundView, Renderer, ReceiverPrimitives, DynamicReceiverMeshElements, ReusedViewsArray);
diff --git a/Engine/Source/Runtime/Renderer/Private/StaticMeshDrawList.inl b/Engine/Source/Runtime/Renderer/Private/StaticMeshDrawList.inl
--- a/Engine/Source/Runtime/Renderer/Private/StaticMeshDrawList.inl
+++ b/Engine/Source/Runtime/Renderer/Private/StaticMeshDrawList.inl
@@ -96,11 +96,16 @@ void TStaticMeshDrawList<DrawingPolicyType>::DrawElement(
 			for (uint32 BackFace = 0; BackFace < BackFaceEnd; ++BackFace)
 			{
 				INC_DWORD_STAT(STAT_StaticDrawListMeshDrawCalls);
+				const FPrimitiveSceneProxy* Proxy = Element.Mesh->PrimitiveSceneInfo->Proxy;
+
+				TDrawEvent<FRHICommandList> MeshEvent;
+				BeginMeshDrawEvent(RHICmdList, Proxy, *Element.Mesh, MeshEvent);
+
 				float DitherValue = View.GetDitheredLODTransitionValue(*Element.Mesh);
 				DrawingPolicyLink->DrawingPolicy.SetMeshRenderState(
 					RHICmdList, 
 					View,
-					Element.Mesh->PrimitiveSceneInfo->Proxy,
+					Proxy,
 					*Element.Mesh,
 					BatchElementIndex,
 					!!BackFace,
diff --git a/Engine/Source/Runtime/Renderer/Private/TranslucentLighting.cpp b/Engine/Source/Runtime/Renderer/Private/TranslucentLighting.cpp
--- a/Engine/Source/Runtime/Renderer/Private/TranslucentLighting.cpp
+++ b/Engine/Source/Runtime/Renderer/Private/TranslucentLighting.cpp
@@ -442,6 +442,9 @@ public:
 
 				for (int32 BatchElementIndex = 0; BatchElementIndex < Mesh.Elements.Num(); BatchElementIndex++)
 				{
+					TDrawEvent<FRHICommandList> MeshEvent;
+					BeginMeshDrawEvent(RHICmdList, PrimitiveSceneProxy, Mesh, MeshEvent);
+
 					DrawingPolicy.SetMeshRenderState(RHICmdList, View,PrimitiveSceneProxy,Mesh,BatchElementIndex,bBackFace,Mesh.DitheredLODTransitionAlpha,
 						FTranslucencyShadowDepthDrawingPolicy::ElementDataType(),
 						FTranslucencyShadowDepthDrawingPolicy::ContextDataType(DrawingContext.ShadowInfo)
diff --git a/Engine/Source/Runtime/Renderer/Private/VelocityRendering.cpp b/Engine/Source/Runtime/Renderer/Private/VelocityRendering.cpp
--- a/Engine/Source/Runtime/Renderer/Private/VelocityRendering.cpp
+++ b/Engine/Source/Runtime/Renderer/Private/VelocityRendering.cpp
@@ -470,6 +470,9 @@ bool FVelocityDrawingPolicyFactory::DrawDynamicMesh(
 			DrawingPolicy.SetSharedState(RHICmdList, &View, FVelocityDrawingPolicy::ContextDataType());
 			for (int32 BatchElementIndex = 0, BatchElementCount = Mesh.Elements.Num(); BatchElementIndex < BatchElementCount; ++BatchElementIndex)
 			{
+				TDrawEvent<FRHICommandList> MeshEvent;
+				BeginMeshDrawEvent(RHICmdList, PrimitiveSceneProxy, Mesh, MeshEvent);
+
 				DrawingPolicy.SetMeshRenderState(RHICmdList, View, PrimitiveSceneProxy, Mesh, BatchElementIndex, bBackFace, Mesh.DitheredLODTransitionAlpha, FMeshDrawingPolicy::ElementDataType(), FVelocityDrawingPolicy::ContextDataType());
 				DrawingPolicy.DrawMesh(RHICmdList, Mesh, BatchElementIndex);
 			}
diff --git a/Engine/Source/Runtime/Slate/Private/Framework/Application/SlateApplication.cpp b/Engine/Source/Runtime/Slate/Private/Framework/Application/SlateApplication.cpp
--- a/Engine/Source/Runtime/Slate/Private/Framework/Application/SlateApplication.cpp
+++ b/Engine/Source/Runtime/Slate/Private/Framework/Application/SlateApplication.cpp
@@ -1322,9 +1322,10 @@ void FSlateApplication::Tick()
 		Renderer->ReleaseAccessedResources(/* Flush State */ false);
 	}
 	
-	{
-		const float DeltaTime = GetDeltaTime();
 
+	const float DeltaTime = GetDeltaTime();
+
+	{
 		SCOPE_CYCLE_COUNTER( STAT_SlateMessageTick );
 
 		//FPlatformMisc::BeginNamedEvent(FColor::Magenta, "Slate::Platform");
@@ -1348,6 +1349,9 @@ void FSlateApplication::Tick()
 		//FPlatformMisc::EndNamedEvent();
 	}
 
+	UpdateRetainerWidgetsEvent.Broadcast( DeltaTime );
+
+
 	//FPlatformMisc::BeginNamedEvent(FColor::Magenta, "Slate::UpdateCursorLockRegion");
 	// The widget locking the cursor to its bounds may have been reshaped.
 	// Check if the widget was reshaped and update the cursor lock
@@ -1439,7 +1443,8 @@ void FSlateApplication::Tick()
 		{
 			if ( ReleasedCachedElementLists[CacheIndex]->IsInUse() == false )
 			{
-				delete ReleasedCachedElementLists[CacheIndex];
+				ensure( ReleasedCachedElementLists[CacheIndex].IsUnique() );
+				ReleasedCachedElementLists[CacheIndex].Reset();
 				ReleasedCachedElementLists.RemoveAtSwap(CacheIndex, 1, false);
 				CacheIndex--;
 			}
@@ -3682,10 +3687,10 @@ bool FSlateApplication::TakeScreenshot(const TSharedRef<SWidget>& Widget, const
 
 TSharedPtr< FSlateWindowElementList > FSlateApplication::GetCachableElementList(const TSharedPtr<SWindow>& CurrentWindow, const ILayoutCache* LayoutCache)
 {
-	FCacheElementPools* Pools = CachedElementLists.FindRef(LayoutCache);
-	if ( Pools == nullptr )
+	TSharedPtr<FCacheElementPools> Pools = CachedElementLists.FindRef(LayoutCache);
+	if ( !Pools.IsValid() )
 	{
-		Pools = new FCacheElementPools();
+		Pools = MakeShareable( new FCacheElementPools() );
 		CachedElementLists.Add(LayoutCache, Pools);
 	}
 
@@ -3753,8 +3758,8 @@ bool FSlateApplication::FCacheElementPools::IsInUse() const
 
 void FSlateApplication::ReleaseResourcesForLayoutCache(const ILayoutCache* LayoutCache)
 {
-	FCacheElementPools* Pools = CachedElementLists.FindRef(LayoutCache);
-	if ( Pools != nullptr )
+	TSharedPtr<FCacheElementPools> Pools = CachedElementLists.FindRef(LayoutCache);
+	if ( Pools.IsValid() )
 	{
 		ReleasedCachedElementLists.Add(Pools);
 	}
diff --git a/Engine/Source/Runtime/Slate/Private/Widgets/Input/SEditableLabel.cpp b/Engine/Source/Runtime/Slate/Private/Widgets/Input/SEditableLabel.cpp
--- a/Engine/Source/Runtime/Slate/Private/Widgets/Input/SEditableLabel.cpp
+++ b/Engine/Source/Runtime/Slate/Private/Widgets/Input/SEditableLabel.cpp
@@ -12,6 +12,7 @@ void SEditableLabel::Construct(const FArguments& InArgs)
 	CanEditAttribute = InArgs._CanEdit;
 	OnTextChanged = InArgs._OnTextChanged;
 	TextAttribute = InArgs._Text;
+	WasFocused = false;
 
 	ChildSlot
 	[
@@ -97,19 +98,6 @@ bool SEditableLabel::HasKeyboardFocus() const
 }
 
 
-FReply SEditableLabel::OnFocusReceived(const FGeometry& MyGeometry, const FFocusEvent& InFocusEvent)
-{
-	if ((InFocusEvent.GetCause() == EFocusCause::Navigation) ||
-		(InFocusEvent.GetCause() == EFocusCause::SetDirectly))
-	{
-		EnterTextMode();
-		return FReply::Handled().SetUserFocus(EditableText.ToSharedRef(), InFocusEvent.GetCause());
-	}
-
-	return FReply::Unhandled();
-}
-
-
 FReply SEditableLabel::OnKeyDown(const FGeometry& MyGeometry, const FKeyEvent& InKeyEvent)
 {
 	const FKey Key = InKeyEvent.GetKey();
@@ -120,7 +108,7 @@ FReply SEditableLabel::OnKeyDown(const FGeometry& MyGeometry, const FKeyEvent& I
 		return FReply::Handled().SetUserFocus(AsShared(), EFocusCause::Navigation);
 	}
 
-	if (Key == EKeys::Enter)
+	if (WasFocused && (Key == EKeys::Enter))
 	{
 		EnterTextMode();
 		return FReply::Handled().SetUserFocus(EditableText.ToSharedRef(), EFocusCause::Navigation);
@@ -134,20 +122,24 @@ FReply SEditableLabel::OnMouseButtonDown(const FGeometry& MyGeometry, const FPoi
 {
 	if (MouseEvent.GetEffectingButton() == EKeys::LeftMouseButton)
 	{
-		// @todo Slate: implement some mechanism to not enter edit mode right away on mouse up
+		WasFocused = HasKeyboardFocus();
 	}
 
+	if (WasFocused)
+	{
+		return FReply::Handled();
+	}
+	
 	return FReply::Unhandled();
 }
 
 
 FReply SEditableLabel::OnMouseButtonUp(const FGeometry& MyGeometry, const FPointerEvent& MouseEvent )
 {
-	if (MouseEvent.GetEffectingButton() == EKeys::LeftMouseButton)
+	if (WasFocused && (MouseEvent.GetEffectingButton() == EKeys::LeftMouseButton))
 	{
 		EnterTextMode();
-
-		return FReply::Handled();			
+		return FReply::Handled().SetUserFocus(EditableText.ToSharedRef(), EFocusCause::Mouse);;
 	}
 
 	return FReply::Unhandled();
diff --git a/Engine/Source/Runtime/Slate/Public/Framework/Application/SlateApplication.h b/Engine/Source/Runtime/Slate/Public/Framework/Application/SlateApplication.h
--- a/Engine/Source/Runtime/Slate/Public/Framework/Application/SlateApplication.h
+++ b/Engine/Source/Runtime/Slate/Public/Framework/Application/SlateApplication.h
@@ -393,6 +393,10 @@ public:
 	 */
 	void ExternalModalStop();
 
+	/** Delegate for retainer widgets to know when they should update */
+	DECLARE_EVENT_OneParam(FSlateApplication, FOnUpdateRetainerWidgets, float);
+	FOnUpdateRetainerWidgets& OnUpdateRetainerWidgets()  { return UpdateRetainerWidgetsEvent; }
+
 	/** 
 	 * Removes references to FViewportRHI's.  
 	 * This has to be done explicitly instead of using the FRenderResource mechanism because FViewportRHI's are managed by the game thread.
@@ -1651,9 +1655,13 @@ private:
 		TArray< TSharedPtr< FSlateWindowElementList > > InactiveCachedElementListPool;
 	};
 
-	TMap< const ILayoutCache*, FCacheElementPools* > CachedElementLists;
-	TArray< FCacheElementPools* > ReleasedCachedElementLists;
+	TMap< const ILayoutCache*, TSharedPtr<FCacheElementPools> > CachedElementLists;
+	TArray< TSharedPtr<FCacheElementPools> > ReleasedCachedElementLists;
 
 	/** Configured fkeys to control navigation */
 	FNavigationConfig NavigationConfig;
+
+	/** Delegate for retainer widgets to know when they should update */
+	FOnUpdateRetainerWidgets UpdateRetainerWidgetsEvent;
+
 };
diff --git a/Engine/Source/Runtime/Slate/Public/Widgets/Input/SEditableLabel.h b/Engine/Source/Runtime/Slate/Public/Widgets/Input/SEditableLabel.h
--- a/Engine/Source/Runtime/Slate/Public/Widgets/Input/SEditableLabel.h
+++ b/Engine/Source/Runtime/Slate/Public/Widgets/Input/SEditableLabel.h
@@ -83,7 +83,6 @@ public:
 	// SWidget interface
 
 	virtual bool HasKeyboardFocus() const override;
-	virtual FReply OnFocusReceived(const FGeometry& MyGeometry, const FFocusEvent& InFocusEvent) override;
 	virtual FReply OnKeyDown(const FGeometry& MyGeometry, const FKeyEvent& InKeyEvent) override;
 	virtual FReply OnMouseButtonDown(const FGeometry& MyGeometry, const FPointerEvent& MouseEvent) override;
 	virtual FReply OnMouseButtonUp(const FGeometry& MyGeometry, const FPointerEvent& MouseEvent) override;
@@ -99,6 +98,7 @@ private:
 	TAttribute<bool> CanEditAttribute;
 	FOnTextChanged OnTextChanged;
 	TAttribute<FText> TextAttribute;
+	bool WasFocused;
 
 	TSharedPtr<STextBlock> TextBlock;
 	TSharedPtr<SEditableText> EditableText;
diff --git a/Engine/Source/Runtime/SlateCore/Private/Rendering/SlateDrawBuffer.cpp b/Engine/Source/Runtime/SlateCore/Private/Rendering/SlateDrawBuffer.cpp
--- a/Engine/Source/Runtime/SlateCore/Private/Rendering/SlateDrawBuffer.cpp
+++ b/Engine/Source/Runtime/SlateCore/Private/Rendering/SlateDrawBuffer.cpp
@@ -56,7 +56,10 @@ void FSlateDrawBuffer::ClearBuffer()
 	// Move all the window elements back into the pool.
 	for ( TSharedPtr<FSlateWindowElementList> ExistingList : WindowElementLists )
 	{
-		WindowElementListsPool.Add(ExistingList);
+		if( ExistingList->GetWindow().IsValid() )
+		{
+			WindowElementListsPool.Add(ExistingList);
+		}
 	}
 
 	WindowElementLists.Reset();
diff --git a/Engine/Source/Runtime/SlateNullRenderer/Private/SlateNullRenderer.cpp b/Engine/Source/Runtime/SlateNullRenderer/Private/SlateNullRenderer.cpp
--- a/Engine/Source/Runtime/SlateNullRenderer/Private/SlateNullRenderer.cpp
+++ b/Engine/Source/Runtime/SlateNullRenderer/Private/SlateNullRenderer.cpp
@@ -18,6 +18,7 @@ void FSlateNullRenderer::Destroy()
 FSlateDrawBuffer& FSlateNullRenderer::GetDrawBuffer()
 {
 	static FSlateDrawBuffer StaticDrawBuffer;
+	StaticDrawBuffer.ClearBuffer();
 	return StaticDrawBuffer;
 }
 
diff --git a/Engine/Source/Runtime/SlateRHIRenderer/Private/Slate3DRenderer.cpp b/Engine/Source/Runtime/SlateRHIRenderer/Private/Slate3DRenderer.cpp
--- a/Engine/Source/Runtime/SlateRHIRenderer/Private/Slate3DRenderer.cpp
+++ b/Engine/Source/Runtime/SlateRHIRenderer/Private/Slate3DRenderer.cpp
@@ -26,8 +26,7 @@ FSlate3DRenderer::~FSlate3DRenderer()
 
 FSlateDrawBuffer& FSlate3DRenderer::GetDrawBuffer()
 {
-	FreeBufferIndex = (FreeBufferIndex + 1) % 2;
-
+	FreeBufferIndex = (FreeBufferIndex + 1) % NUM_DRAW_BUFFERS;
 	FSlateDrawBuffer* Buffer = &DrawBuffers[FreeBufferIndex];
 
 	while (!Buffer->Lock())
@@ -54,7 +53,7 @@ void FSlate3DRenderer::DrawWindow_GameThread(FSlateDrawBuffer& DrawBuffer)
 
 	// Need to flush the font cache before we add the elements below to avoid the flush potentially 
 	// deleting the texture resources that will be needed by the render thread
-	FontCache->ConditionalFlushCache();
+	//FontCache->ConditionalFlushCache();
 
 	TArray<TSharedPtr<FSlateWindowElementList>>& WindowElementLists = DrawBuffer.GetWindowElementLists();
 
diff --git a/Engine/Source/Runtime/SlateRHIRenderer/Private/Slate3DRenderer.h b/Engine/Source/Runtime/SlateRHIRenderer/Private/Slate3DRenderer.h
--- a/Engine/Source/Runtime/SlateRHIRenderer/Private/Slate3DRenderer.h
+++ b/Engine/Source/Runtime/SlateRHIRenderer/Private/Slate3DRenderer.h
@@ -16,7 +16,8 @@ public:
 private:
 
 	/** Double buffered draw buffers so that the rendering thread can be rendering windows while the game thread is setting up for next frame */
-	FSlateDrawBuffer DrawBuffers[2];
+	static const int32 NUM_DRAW_BUFFERS = 3;
+	FSlateDrawBuffer DrawBuffers[NUM_DRAW_BUFFERS];
 
 	/** The font services to use for rendering text */
 	TSharedRef<FSlateFontServices> SlateFontServices;
diff --git a/Engine/Source/Runtime/UMG/Private/Animation/WidgetAnimation.cpp b/Engine/Source/Runtime/UMG/Private/Animation/WidgetAnimation.cpp
--- a/Engine/Source/Runtime/UMG/Private/Animation/WidgetAnimation.cpp
+++ b/Engine/Source/Runtime/UMG/Private/Animation/WidgetAnimation.cpp
@@ -100,7 +100,7 @@ void UWidgetAnimation::Initialize(UUserWidget* InPreviewWidget)
  *****************************************************************************/
 
 
-void UWidgetAnimation::BindPossessableObject(const FGuid& ObjectId, UObject& PossessedObject)
+void UWidgetAnimation::BindPossessableObject(const FGuid& ObjectId, UObject& PossessedObject, UObject* Context)
 {
 	UPanelSlot* PossessedSlot = Cast<UPanelSlot>(&PossessedObject);
 
@@ -150,21 +150,7 @@ bool UWidgetAnimation::CanPossessObject(UObject& Object) const
 }
 
 
-FGuid UWidgetAnimation::FindObjectId(UObject& Object) const
-{
-	UPanelSlot* Slot = Cast<UPanelSlot>(&Object);
-
-	if (Slot != nullptr)
-	{
-		// slot guids are tracked by their content.
-		return SlotContentPreviewObjectToIds.FindRef(Slot->Content);
-	}
-
-	return PreviewObjectToIds.FindRef(&Object);
-}
-
-
-UObject* UWidgetAnimation::FindObject(const FGuid& ObjectId) const
+UObject* UWidgetAnimation::FindPossessableObject(const FGuid& ObjectId, UObject* Context) const
 {
 	TWeakObjectPtr<UObject> PreviewObject = IdToPreviewObjects.FindRef(ObjectId);
 
@@ -210,46 +196,6 @@ UObject* UWidgetAnimation::GetParentObject(UObject* Object) const
 }
 
 
-#if WITH_EDITOR
-bool UWidgetAnimation::TryGetObjectDisplayName(const FGuid& ObjectId, FText& OutDisplayName) const
-{
-	// TODO: This gets called every frame for every bound object and could
-	// be a potential performance issue for a really complicated animation.
-
-	TWeakObjectPtr<UObject> PreviewObject = IdToPreviewObjects.FindRef(ObjectId);
-
-	if (PreviewObject.IsValid())
-	{
-		OutDisplayName = FText::FromString(PreviewObject.Get()->GetName());
-		return true;
-	}
-
-	TWeakObjectPtr<UObject> SlotContentPreviewObject = IdToSlotContentPreviewObjects.FindRef(ObjectId);
-
-	if (SlotContentPreviewObject.IsValid())
-	{
-		UWidget* SlotContent = Cast<UWidget>(SlotContentPreviewObject.Get());
-		FText PanelName = SlotContent->Slot != nullptr && SlotContent->Slot->Parent != nullptr
-			? FText::FromString(SlotContent->Slot->Parent->GetName())
-			: LOCTEXT("InvalidPanel", "Invalid Panel");
-		FText ContentName = FText::FromString(SlotContent->GetName());
-		if ( PreviewObjectToIds.Contains( SlotContent ) )
-		{
-			OutDisplayName = FText::Format( LOCTEXT( "SlotObjectWithParent", "{0} Slot" ), PanelName );
-		}
-		else
-		{
-			OutDisplayName = FText::Format(LOCTEXT("SlotObject", "{0} ({1} Slot)"), ContentName, PanelName);
-		}
-
-		return true;
-	}
-
-	return false;
-}
-#endif
-
-
 void UWidgetAnimation::UnbindPossessableObjects(const FGuid& ObjectId)
 {
 	// unbind widgets
diff --git a/Engine/Source/Runtime/UMG/Private/Slate/SMeshWidget.cpp b/Engine/Source/Runtime/UMG/Private/Slate/SMeshWidget.cpp
--- a/Engine/Source/Runtime/UMG/Private/Slate/SMeshWidget.cpp
+++ b/Engine/Source/Runtime/UMG/Private/Slate/SMeshWidget.cpp
@@ -6,6 +6,7 @@
 #include "SlateMaterialBrush.h"
 #include "Runtime/SlateRHIRenderer/Public/Interfaces/ISlateRHIRendererModule.h"
 #include "SlateVectorArtData.h"
+#include "SlateVectorArtInstanceData.h"
 
 
 
@@ -211,4 +212,21 @@ void SMeshWidget::AddReferencedObjects(FReferenceCollector& Collector)
 			Collector.AddReferencedObject(ResourceObject);
 		}
 	}
+}
+
+
+void SMeshWidget::PushUpdate(uint32 VectorArtId, const SMeshWidget& Widget, const FVector2D& Position, float Scale, uint32 BaseAddress)
+
+{
+	FSlateVectorArtInstanceData Data;
+	Data.SetPosition(Position);
+	Data.SetScale(Scale);
+	Data.SetBaseAddress(BaseAddress);
+
+	{
+		TSharedPtr<FSlateInstanceBufferUpdate> PerInstaceUpdate = Widget.BeginPerInstanceBufferUpdateConst(VectorArtId);
+		PerInstaceUpdate->GetData().Empty();
+		PerInstaceUpdate->GetData().Add(Data.GetData());
+		FSlateInstanceBufferUpdate::CommitUpdate(PerInstaceUpdate);
+	}
 }
\ No newline at end of file
