OnlineFramework: fix for compilation on a case-sensitive fs.

- Submitted as PR #1890.

diff --git a/Engine/Source/Runtime/Online/OnlineFramework/Lobby/Public/LobbyBeaconPlayerState.h b/Engine/Source/Runtime/Online/OnlineFramework/Lobby/Public/LobbyBeaconPlayerState.h
--- a/Engine/Source/Runtime/Online/OnlineFramework/Lobby/Public/LobbyBeaconPlayerState.h
+++ b/Engine/Source/Runtime/Online/OnlineFramework/Lobby/Public/LobbyBeaconPlayerState.h
@@ -2,7 +2,7 @@
 
 #pragma once
 
-#include "Gameframework/Info.h"
+#include "GameFramework/Info.h"
 #include "LobbyBeaconPlayerState.generated.h"
 
 class AOnlineBeaconClient;
diff --git a/Engine/Source/Runtime/Online/OnlineFramework/Lobby/Public/LobbyBeaconState.h b/Engine/Source/Runtime/Online/OnlineFramework/Lobby/Public/LobbyBeaconState.h
--- a/Engine/Source/Runtime/Online/OnlineFramework/Lobby/Public/LobbyBeaconState.h
+++ b/Engine/Source/Runtime/Online/OnlineFramework/Lobby/Public/LobbyBeaconState.h
@@ -2,7 +2,7 @@
 
 #pragma once
 
-#include "Gameframework/Info.h"
+#include "GameFramework/Info.h"
 #include "LobbyBeaconState.generated.h"
 
 class AOnlineBeaconClient;
