Copying //UE4/Orion-Staging to //UE4/Main (Origin: //Orion/Dev-General [REDACTED_HANDLE])

diff --git a/Engine/Source/Runtime/GameplayAbilities/Public/GameplayEffect.h b/Engine/Source/Runtime/GameplayAbilities/Public/GameplayEffect.h
--- a/Engine/Source/Runtime/GameplayAbilities/Public/GameplayEffect.h
+++ b/Engine/Source/Runtime/GameplayAbilities/Public/GameplayEffect.h
@@ -1521,6 +1521,7 @@ struct GAMEPLAYABILITIES_API FActiveGameplayEffectsContainer : public FFastArray
 	~FActiveGameplayEffectsContainer();
 
 	UAbilitySystemComponent* Owner;
+	bool OwnerIsNetAuthority;
 
 	FOnGivenActiveGameplayEffectRemoved	OnActiveGameplayEffectRemovedDelegate;
 
@@ -1741,7 +1742,10 @@ private:
 	 */
 	bool InternalExecuteMod(FGameplayEffectSpec& Spec, FGameplayModifierEvaluatedData& ModEvalData);
 
-	bool IsNetAuthority() const;
+	bool IsNetAuthority() const
+	{
+		return OwnerIsNetAuthority;
+	}
 
 	/** Called internally to actually remove a GameplayEffect or to reduce its StackCount. Returns true if we resized our internal GameplayEffect array. */
 	bool InternalRemoveActiveGameplayEffect(int32 Idx, int32 StacksToRemove, bool bPrematureRemoval);
diff --git a/Engine/Source/Runtime/GameplayAbilities/Public/GameplayEffectTypes.h b/Engine/Source/Runtime/GameplayAbilities/Public/GameplayEffectTypes.h
--- a/Engine/Source/Runtime/GameplayAbilities/Public/GameplayEffectTypes.h
+++ b/Engine/Source/Runtime/GameplayAbilities/Public/GameplayEffectTypes.h
@@ -708,7 +708,6 @@ struct GAMEPLAYABILITIES_API FGameplayCueParameters
 	FGameplayCueParameters()
 	: NormalizedMagnitude(0.0f)
 	, RawMagnitude(0.0f)
-	, MatchedTagName(NAME_None)
 	, Location(ForceInitToZero)
 	, Normal(ForceInitToZero)
 	{}
@@ -731,11 +730,11 @@ struct GAMEPLAYABILITIES_API FGameplayCueParameters
 	FGameplayEffectContextHandle EffectContext;
 
 	/** The tag name that matched this specific gameplay cue handler */
-	UPROPERTY(BlueprintReadWrite, Category=GameplayCue)
-	FName MatchedTagName;
+	UPROPERTY(BlueprintReadWrite, Category=GameplayCue, NotReplicated)
+	FGameplayTag MatchedTagName;
 
 	/** The original tag of the gameplay cue */
-	UPROPERTY(BlueprintReadWrite, Category=GameplayCue)
+	UPROPERTY(BlueprintReadWrite, Category=GameplayCue, NotReplicated)
 	FGameplayTag OriginalTag;
 
 	/** The aggregated source tags taken from the effect spec */
@@ -764,6 +763,10 @@ struct GAMEPLAYABILITIES_API FGameplayCueParameters
 	UPROPERTY(BlueprintReadWrite, Category=GameplayCue)
 	TWeakObjectPtr<const UObject> SourceObject;
 
+	/** PhysMat of the hit, if there was a hit. */
+	UPROPERTY(BlueprintReadWrite, Category = GameplayCue)
+	TWeakObjectPtr<const UPhysicalMaterial> PhysicalMaterial;
+
 	bool NetSerialize(FArchive& Ar, class UPackageMap* Map, bool& bOutSuccess);
 
 	bool IsInstigatorLocallyControlled() const;
@@ -835,9 +838,9 @@ namespace EGameplayTagEventType
  * while simultaneously tracking the count of parent tags as well. Events/delegates are fired whenever the tag counts
  * of any tag (explicit or parent) are modified.
  */
+
 struct GAMEPLAYABILITIES_API FGameplayTagCountContainer
-{
-	// Constructor
+{	
 	FGameplayTagCountContainer()
 	{}
 
@@ -848,7 +851,10 @@ struct GAMEPLAYABILITIES_API FGameplayTagCountContainer
 	 * 
 	 * @return True if the count container has a gameplay tag that matches, false if not
 	 */
-	bool HasMatchingGameplayTag(FGameplayTag TagToCheck) const;
+	FORCEINLINE bool HasMatchingGameplayTag(FGameplayTag TagToCheck) const
+	{
+		return GameplayTagCountMap.FindRef(TagToCheck) > 0;
+	}
 
 	/**
 	 * Check if the count container has gameplay tags that matches against all of the specified tags (expands to include parents of asset tags)
@@ -858,7 +864,25 @@ struct GAMEPLAYABILITIES_API FGameplayTagCountContainer
 	 * 
 	 * @return True if the count container matches all of the gameplay tags
 	 */
-	bool HasAllMatchingGameplayTags(const FGameplayTagContainer& TagContainer, bool bCountEmptyAsMatch = true) const;
+	FORCEINLINE bool HasAllMatchingGameplayTags(const FGameplayTagContainer& TagContainer, bool bCountEmptyAsMatch = true) const
+	{
+		// if the TagContainer count is 0 return bCountEmptyAsMatch;
+		if (TagContainer.Num() == 0)
+		{
+			return bCountEmptyAsMatch;
+		}
+
+		bool AllMatch = true;
+		for (const FGameplayTag& Tag : TagContainer)
+		{
+			if (GameplayTagCountMap.FindRef(Tag) <= 0)
+			{
+				AllMatch = false;
+				break;
+			}
+		}		
+		return AllMatch;
+	}
 	
 	/**
 	 * Check if the count container has gameplay tags that matches against any of the specified tags (expands to include parents of asset tags)
@@ -868,7 +892,19 @@ struct GAMEPLAYABILITIES_API FGameplayTagCountContainer
 	 * 
 	 * @return True if the count container matches any of the gameplay tags
 	 */
-	bool HasAnyMatchingGameplayTags(const FGameplayTagContainer& TagContainer, bool bCountEmptyAsMatch = true) const;
+	FORCEINLINE bool HasAnyMatchingGameplayTags(const FGameplayTagContainer& TagContainer, bool bCountEmptyAsMatch = true) const
+	{
+		bool AnyMatch = (bCountEmptyAsMatch && TagContainer.Num() == 0);
+		for (const FGameplayTag& Tag : TagContainer)
+		{
+			if (GameplayTagCountMap.FindRef(Tag) > 0)
+			{
+				AnyMatch = true;
+				break;
+			}
+		}
+		return AnyMatch;
+	}
 	
 	/**
 	 * Update the specified container of tags by the specified delta, potentially causing an additional or removal from the explicit tag list
@@ -876,15 +912,59 @@ struct GAMEPLAYABILITIES_API FGameplayTagCountContainer
 	 * @param Container		Container of tags to update
 	 * @param CountDelta	Delta of the tag count to apply
 	 */
-	void UpdateTagCount(const FGameplayTagContainer& Container, int32 CountDelta);
+	FORCEINLINE void UpdateTagCount(const FGameplayTagContainer& Container, int32 CountDelta)
+	{
+		if (CountDelta != 0)
+		{
+			for (auto TagIt = Container.CreateConstIterator(); TagIt; ++TagIt)
+			{
+				UpdateTagMap_Internal(*TagIt, CountDelta);
+			}
+		}
+	}
 	
 	/**
 	 * Update the specified tag by the specified delta, potentially causing an additional or removal from the explicit tag list
 	 * 
 	 * @param Tag			Tag to update
 	 * @param CountDelta	Delta of the tag count to apply
+	 * 
+	 * @return True if tag was *either* added or removed. (E.g., we had the tag and now dont. or didnt have the tag and now we do. We didn't just change the count (1 count -> 2 count would return false).
 	 */
-	void UpdateTagCount(const FGameplayTag& Tag, int32 CountDelta);
+	FORCEINLINE bool UpdateTagCount(const FGameplayTag& Tag, int32 CountDelta)
+	{
+		if (CountDelta != 0)
+		{
+			return UpdateTagMap_Internal(Tag, CountDelta);
+		}
+
+		return false;
+	}
+
+	/**
+	 * Set the specified tag count to a specific value
+	 * 
+	 * @param Tag			Tag to update
+	 * @param Count			New count of the tag
+	 * 
+	 * @return True if tag was *either* added or removed. (E.g., we had the tag and now dont. or didnt have the tag and now we do. We didn't just change the count (1 count -> 2 count would return false).
+	 */
+	FORCEINLINE bool SetTagCount(const FGameplayTag& Tag, int32 NewCount)
+	{
+		int32 ExistingCount = 0;
+		if (int32* Ptr  = ExplicitTagCountMap.Find(Tag))
+		{
+			ExistingCount = *Ptr;
+		}
+
+		int32 CountDelta = NewCount - ExistingCount;
+		if (CountDelta != 0)
+		{
+			return UpdateTagMap_Internal(Tag, CountDelta);
+		}
+
+		return false;
+	}
 
 	/**
 	 *	Broadcasts the AnyChange event for this tag. This is called when the stack count of the backing gameplay effect change.
@@ -906,10 +986,16 @@ struct GAMEPLAYABILITIES_API FGameplayTagCountContainer
 	 * 
 	 * @return Delegate for when any tag's count changes to or off of zero
 	 */
-	FOnGameplayEffectTagCountChanged& RegisterGenericGameplayEvent();
+	FOnGameplayEffectTagCountChanged& RegisterGenericGameplayEvent()
+	{
+		return OnAnyTagChangeDelegate;
+	}
 
 	/** Simple accessor to the explicit gameplay tag list */
-	const FGameplayTagContainer& GetExplicitGameplayTags() const;
+	const FGameplayTagContainer& GetExplicitGameplayTags() const
+	{
+		return ExplicitTags;
+	}
 
 	void Reset();
 
@@ -927,7 +1013,7 @@ private:
 	/** Map of tag to active count of that tag */
 	TMap<FGameplayTag, int32> GameplayTagCountMap;
 
-	/** Map of tag to explicit count of that tag. Cannot share with above map because it's not safe to merge explicit and generic counts */
+	/** Map of tag to explicit count of that tag. Cannot share with above map because it's not safe to merge explicit and generic counts */	
 	TMap<FGameplayTag, int32> ExplicitTagCountMap;
 
 	/** Delegate fired whenever any tag's count changes to or away from zero */
@@ -937,9 +1023,10 @@ private:
 	FGameplayTagContainer ExplicitTags;
 
 	/** Internal helper function to adjust the explicit tag list & corresponding maps/delegates/etc. as necessary */
-	void UpdateTagMap_Internal(const FGameplayTag& Tag, int32 CountDelta);
+	bool UpdateTagMap_Internal(const FGameplayTag& Tag, int32 CountDelta);
 };
 
+
 // -----------------------------------------------------------
 
 /** Encapsulate require and ignore tags */
@@ -1083,7 +1170,92 @@ struct TStructOpsTypeTraits<FGameplayEffectSpecHandle> : public TStructOpsTypeTr
 {
 	enum
 	{
-		WithCopy = true,		// Necessary so that TSharedPtr<FGameplayAbilityTargetData> Data is copied around
+		WithCopy = true,
+		WithNetSerializer = true,
+		WithIdenticalViaEquality = true,
+	};
+};
+
+// -----------------------------------------------------------
+
+
+USTRUCT()
+struct GAMEPLAYABILITIES_API FMinimapReplicationTagCountMap
+{
+	GENERATED_USTRUCT_BODY()
+
+	FMinimapReplicationTagCountMap()
+	{
+		MapID = 0;
+	}
+
+	void AddTag(const FGameplayTag& Tag)
+	{
+		MapID++;
+		TagMap.FindOrAdd(Tag)++;
+	}
+
+	void RemoveTag(const FGameplayTag& Tag)
+	{
+		MapID++;
+		int32& Count = TagMap.FindOrAdd(Tag);
+		Count--;
+		if (Count == 0)
+		{
+			// Remove from map so that we do not replicate
+			TagMap.Remove(Tag);
+		}
+		else if (Count < 0)
+		{
+			ABILITY_LOG(Error, TEXT("FMinimapReplicationTagCountMap::RemoveTag called on Tag %s and count is now < 0"), *Tag.ToString());
+			Count = 0;
+		}
+	}
+
+	void AddTags(const FGameplayTagContainer& Container)
+	{
+		for (const FGameplayTag& Tag : Container)
+		{
+			AddTag(Tag);
+		}
+	}
+
+	void RemoveTags(const FGameplayTagContainer& Container)
+	{
+		for (const FGameplayTag& Tag : Container)
+		{
+			RemoveTag(Tag);
+		}
+	}
+
+	bool NetSerialize(FArchive& Ar, class UPackageMap* Map, bool& bOutSuccess);
+
+	TMap<FGameplayTag, int32>	TagMap;
+
+	UPROPERTY()
+	class UAbilitySystemComponent* Owner;
+
+	/** Comparison operator */
+	bool operator==(FMinimapReplicationTagCountMap const& Other) const
+	{
+		return (MapID == Other.MapID);
+	}
+
+	/** Comparison operator */
+	bool operator!=(FMinimapReplicationTagCountMap const& Other) const
+	{
+		return !(FMinimapReplicationTagCountMap::operator==(Other));
+	}
+
+	int32 MapID;
+};
+
+template<>
+struct TStructOpsTypeTraits<FMinimapReplicationTagCountMap> : public TStructOpsTypeTraitsBase
+{
+	enum
+	{
+		WithCopy = true,
 		WithNetSerializer = true,
 		WithIdenticalViaEquality = true,
 	};
diff --git a/Engine/Source/Runtime/GameplayTags/Classes/GameplayTagContainer.h b/Engine/Source/Runtime/GameplayTags/Classes/GameplayTagContainer.h
--- a/Engine/Source/Runtime/GameplayTags/Classes/GameplayTagContainer.h
+++ b/Engine/Source/Runtime/GameplayTags/Classes/GameplayTagContainer.h
@@ -4,6 +4,8 @@
 
 #include "GameplayTagContainer.generated.h"
 
+DECLARE_LOG_CATEGORY_EXTERN(LogGameplayTags, Log, All);
+
 DECLARE_STATS_GROUP(TEXT("Gameplay Tags"), STATGROUP_GameplayTags, STATCAT_Advanced);
 
 DECLARE_CYCLE_STAT_EXTERN(TEXT("FGameplayTagContainer::HasTag"), STAT_FGameplayTagContainer_HasTag, STATGROUP_GameplayTags, GAMEPLAYTAGS_API);
@@ -132,6 +134,8 @@ struct GAMEPLAYTAGS_API FGameplayTag
 	/** Overridden for fast serialize */
 	bool NetSerialize(FArchive& Ar, class UPackageMap* Map, bool& bOutSuccess);
 
+	bool NetSerialize_Packed(FArchive& Ar, class UPackageMap* Map, bool& bOutSuccess);
+
 private:
 
 	/** Intentionally private so only the tag manager can use */
@@ -142,6 +146,7 @@ private:
 	FName TagName;
 
 	friend class UGameplayTagsManager;
+	friend struct FGameplayTagContainer;
 };
 
 template<>
@@ -188,6 +193,14 @@ struct GAMEPLAYTAGS_API FGameplayTagContainer
 	bool operator==(FGameplayTagContainer const& Other) const;
 	bool operator!=(FGameplayTagContainer const& Other) const;
 
+	template<class AllocatorType>
+	static FGameplayTagContainer CreateFromArray(TArray<FGameplayTag, AllocatorType>& SourceTags)
+	{
+		FGameplayTagContainer Container;
+		Container.GameplayTags.Append(SourceTags);
+		return Container;
+	}
+
 	/**  Returns a new container containing all of the tags of this container, as well as all of their parent tags */
 	FGameplayTagContainer GetGameplayTagParents() const;
 
diff --git a/Engine/Source/Runtime/GameplayTags/Classes/GameplayTagsManager.h b/Engine/Source/Runtime/GameplayTags/Classes/GameplayTagsManager.h
--- a/Engine/Source/Runtime/GameplayTags/Classes/GameplayTagsManager.h
+++ b/Engine/Source/Runtime/GameplayTags/Classes/GameplayTagsManager.h
@@ -295,13 +295,16 @@ class GAMEPLAYTAGS_API UGameplayTagsManager : public UObject
 	 * @param GameplayTag		tag to get parents of
 	 * @return					pointer to container of parents, if any
 	 */
-	FORCEINLINE_DEBUGGABLE const FGameplayTagContainer* GetAllParentsContainer(const FGameplayTag& GameplayTag)
+	FORCEINLINE const FGameplayTagContainer* GetAllParentsContainer(const FGameplayTag& GameplayTag)
 	{
 		const TSharedPtr<FGameplayTagNode>* TagNode = GameplayTagNodeMap.Find(GameplayTag);
 		if (TagNode)
 		{
 			return &((*TagNode)->Parents);
 		}
+#if !(UE_BUILD_SHIPPING || UE_BUILD_TEST)
+		UE_LOG(LogGameplayTags, Warning, TEXT("::GetAllParentsContainer Failed to find tag %s in GameplayTagNodeMap. Is this tag not in the dictionary?"), *GameplayTag.ToString());
+#endif
 		return nullptr;
 	}
 
@@ -314,7 +317,10 @@ class GAMEPLAYTAGS_API UGameplayTagsManager : public UObject
 	static bool ShouldImportTagsFromINI();
 
 	/** Should use fast replication */
-	static bool ShouldUseFastReplication();
+	bool ShouldUseFastReplication()
+	{
+		return bUseFastReplication;
+	}
 
 	void RedirectTagsForContainer(FGameplayTagContainer& Container, TSet<FName>& DeprecatedTagNamesNotFoundInTagMap);
 
@@ -322,6 +328,33 @@ class GAMEPLAYTAGS_API UGameplayTagsManager : public UObject
 	FName GetTagNameFromNetIndex(FGameplayTagNetIndex Index);
 	FGameplayTagNetIndex GetNetIndexFromTag(const FGameplayTag &InTag);
 
+	/** Cached number of bits we need to replicate tags. That is, Log2(Number of Tags). Will always be <= 16. */
+	static int32 NetIndexTrueBitNum;
+	
+	/** The length in bits of the first segment when net serializing tags. We will serialize NetIndexFirstBitSegment + 1 bit to indicatore "more" (more = second segment that is NetIndexTrueBitNum - NetIndexFirstBitSegment) */
+	static int32 NetIndexFirstBitSegment;
+
+	/** Numbers of bits to use for replicating container size. This can be set via config. */
+	static int32 NumBitsForContainerSize;
+
+	/** This is the actual value for an invalid tag "None". This is computed at runtime as (Total number of tags) + 1 */
+	FGameplayTagNetIndex InvalidTagNetIndex;
+
+	/** Cached runtime value for whether we are using fast replication or not. Initialized from config setting. */
+	bool bUseFastReplication;
+	
+
+#if !(UE_BUILD_SHIPPING || UE_BUILD_TEST)
+	/** Mechanism for tracking what tags are frequently replicated */
+
+	static void PrintReplicationFrequencyReport();
+	static void NotifyTagReplicated(FGameplayTag Tag, bool WasInContainer);
+
+	static TMap<FGameplayTag, int32>	ReplicationCountMap;
+	static TMap<FGameplayTag, int32>	ReplicationCountMap_SingleTags;
+	static TMap<FGameplayTag, int32>	ReplicationCountMap_Containers;
+#endif
+
 private:
 
 	friend class FGameplayTagTest;
@@ -363,6 +396,9 @@ private:
 	/** Map of Names to tags - Internal use only */
 	TMap<FName, FGameplayTag> GameplayTagMap;
 
+	/** Our aggregated, sorted list of commonly replicated tags. These tags are given lower indices to ensure they replicate in the first bit segment. */
+	TArray<FGameplayTag> CommonlyReplicatedTags;
+
 #if WITH_EDITOR
 	// This critical section is to handle and editor-only issue where tag requests come from another thread when async loading from a background thread in FGameplayTagContainer::Serialize.
 	// This class is not generically threadsafe.
diff --git a/Engine/Source/Runtime/GameplayTags/Classes/GameplayTagsSettings.h b/Engine/Source/Runtime/GameplayTags/Classes/GameplayTagsSettings.h
--- a/Engine/Source/Runtime/GameplayTags/Classes/GameplayTagsSettings.h
+++ b/Engine/Source/Runtime/GameplayTags/Classes/GameplayTagsSettings.h
@@ -27,6 +27,13 @@ class GAMEPLAYTAGS_API UGameplayTagsSettings : public UObject
 	UPROPERTY(config, EditAnywhere, Category=GameplayTags)
 	TArray<FString>		GameplayTags;
 
+	/** List of tags most frequently replicated */
+	UPROPERTY(config, EditAnywhere, Category="Advanced Replication")
+	TArray<FString>		CommonlyReplicatedTags;
+
+	UPROPERTY(config, EditAnywhere, Category="Advanced Replication")
+	int32 NetIndexFirstBitSegment;
+
 	/** Sorts tags alphabetically */
 	void SortTags();
 
diff --git a/Engine/Source/Runtime/GameplayTags/Private/GameplayTagContainer.cpp b/Engine/Source/Runtime/GameplayTags/Private/GameplayTagContainer.cpp
--- a/Engine/Source/Runtime/GameplayTags/Private/GameplayTagContainer.cpp
+++ b/Engine/Source/Runtime/GameplayTags/Private/GameplayTagContainer.cpp
@@ -9,7 +9,84 @@ DEFINE_STAT(STAT_FGameplayTagContainer_HasTag);
 DEFINE_STAT(STAT_FGameplayTagContainer_DoesTagContainerMatch);
 DEFINE_STAT(STAT_UGameplayTagsManager_GameplayTagsMatch);
 
+/**
+ *	Replicates a tag in a packed format:
+ *	-A segment of NetIndexFirstBitSegment bits are always replicated.
+ *	-Another bit is replicated to indicate "more"
+ *	-If "more", then another segment of (MaxBits - NetIndexFirstBitSegment) length is replicated.
+ *	
+ *	This format is basically the same as SerializeIntPacked, except that there are only 2 segments and they are not the same size.
+ *	The gameplay tag system is able to exploit knoweledge in what tags are frequently replicated to ensure they appear in the first segment.
+ *	Making frequently replicated tags as cheap as possible. 
+ *	
+ *	
+ *	Setting up your project to take advantage of the packed format.
+ *	-Run a normal networked game on non shipping build. 
+ *	-After some time, run console command "GameplayTags.PrintReport" or set "GameplayTags.PrintReportOnShutdown 1" cvar.
+ *	-This will generate information on the server log about what tags replicate most frequently.
+ *	-Take this list and put it in DefaultGameplayTags.ini.
+ *	-CommonlyReplicatedTags is the ordered list of tags.
+ *	-NetIndexFirstBitSegment is the number of bits (not including the "more" bit) for the first segment.
+ *
+ */
+void SerializeTagNetIndexPacked(FArchive& Ar, FGameplayTagNetIndex& Value, const int32 NetIndexFirstBitSegment, const int32 MaxBits)
+{
+	// Case where we have no segment or the segment is larger than max bits
+	if (NetIndexFirstBitSegment <= 0 || NetIndexFirstBitSegment >= MaxBits)
+	{
+		if (Ar.IsLoading())
+		{
+			Value = 0;
+		}
+		Ar.SerializeBits(&Value, MaxBits);
+		return;
+	}
+
+
+	const uint32 BitMasks[] = {0x00, 0x01, 0x03, 0x07, 0x0f, 0x1f, 0x3f, 0x7f, 0xff, 0x1ff, 0x3ff, 0x7ff, 0xfff, 0x1fff, 0x3fff, 0x7fff, 0xffff};
+	const uint32 MoreBits[] = {0x00, 0x01, 0x02, 0x04, 0x08, 0x10, 0x20, 0x40, 0x80, 0x100, 0x200, 0x400, 0x800, 0x1000, 0x2000, 0x4000, 0x8000};
 
+	const int32 FirstSegment = NetIndexFirstBitSegment;
+	const int32 SecondSegment = MaxBits - NetIndexFirstBitSegment;
+
+	if (Ar.IsSaving())
+	{
+		uint32 Mask = BitMasks[FirstSegment];
+		if (Value > Mask)
+		{
+			uint32 FirstDataSegment = ((Value & Mask) | MoreBits[FirstSegment+1]);
+			uint32 SecondDataSegment = (Value >> FirstSegment);
+
+			uint32 SerializedValue = FirstDataSegment | (SecondDataSegment << (FirstSegment+1));				
+
+			Ar.SerializeBits(&SerializedValue, MaxBits + 1);
+		}
+		else
+		{
+			uint32 SerializedValue = Value;
+			Ar.SerializeBits(&SerializedValue, NetIndexFirstBitSegment + 1);
+		}
+
+	}
+	else
+	{
+		uint32 FirstData = 0;
+		Ar.SerializeBits(&FirstData, FirstSegment + 1);
+		uint32 More = FirstData & MoreBits[FirstSegment+1];
+		if (More)
+		{
+			uint32 SecondData = 0;
+			Ar.SerializeBits(&SecondData, SecondSegment);
+			Value = (SecondData << FirstSegment);
+			Value |= (FirstData & BitMasks[FirstSegment]);
+		}
+		else
+		{
+			Value = FirstData;
+		}
+
+	}
+}
 
 
 /** Helper class to parse/eval query token streams. */
@@ -754,28 +831,56 @@ FString FGameplayTagContainer::ToStringSimple() const
 
 bool FGameplayTagContainer::NetSerialize(FArchive& Ar, class UPackageMap* Map, bool& bOutSuccess)
 {
-	uint8 NumTags;
+	// 1st bit to indicate empty tag container or not (empty tag containers are frequently replicated). Early out if empty.
+	uint8 IsEmpty = GameplayTags.Num() == 0;
+	Ar.SerializeBits(&IsEmpty, 1);
+	if (IsEmpty)
+	{
+		if (GameplayTags.Num() > 0)
+		{
+			GameplayTags.Reset();
+		}
+		bOutSuccess = true;
+		return true;
+	}
+
+	// -------------------------------------------------------
+
 	if (Ar.IsSaving())
 	{
-		NumTags = GameplayTags.Num();
-		Ar << NumTags;
-		for (FGameplayTag& Tag : GameplayTags)
+		uint8 NumTags = GameplayTags.Num();
+		uint8 MaxSize = (1 << UGameplayTagsManager::NumBitsForContainerSize);
+		if (!ensureMsgf(NumTags < MaxSize, TEXT("TagContainer has %d elements when max is %d! Tags: %s"), NumTags, NumTags, *ToStringSimple()))
+		{
+			NumTags = MaxSize - 1;
+		}
+		
+		Ar.SerializeBits(&NumTags, UGameplayTagsManager::NumBitsForContainerSize);
+		for (int32 idx=0; idx < NumTags;++idx)
 		{
-			Tag.NetSerialize(Ar, Map, bOutSuccess);
+			FGameplayTag& Tag = GameplayTags[idx];
+			Tag.NetSerialize_Packed(Ar, Map, bOutSuccess);
+
+#if !(UE_BUILD_SHIPPING || UE_BUILD_TEST)
+			UGameplayTagsManager::NotifyTagReplicated(Tag, true);
+#endif
 		}
 	}
 	else
 	{
-		Ar << NumTags;
+		// No Common Container tags, just replicate this like normal
+		uint8 NumTags = 0;
+		Ar.SerializeBits(&NumTags, UGameplayTagsManager::NumBitsForContainerSize);
+
 		GameplayTags.Empty(NumTags);
 		GameplayTags.AddDefaulted(NumTags);
 		for (uint8 idx = 0; idx < NumTags; ++idx)
 		{
-			GameplayTags[idx].NetSerialize(Ar, Map, bOutSuccess);
+			GameplayTags[idx].NetSerialize_Packed(Ar, Map, bOutSuccess);
 		}
-
 	}
 
+
 	bOutSuccess  = true;
 	return true;
 }
@@ -845,56 +950,51 @@ DECLARE_CYCLE_STAT(TEXT("FGameplayTag::NetSerialize"), STAT_FGameplayTag_NetSeri
 
 bool FGameplayTag::NetSerialize(FArchive& Ar, class UPackageMap* Map, bool& bOutSuccess)
 {
-	SCOPE_CYCLE_COUNTER(STAT_FGameplayTag_NetSerialize);
-
-	UGameplayTagsManager& TagManager = IGameplayTagsModule::GetGameplayTagsManager();
-
-	uint8 bHasName = (TagName != NAME_None);
-	uint8 bHasNetIndex = 0;
-	FGameplayTagNetIndex NetIndex = INVALID_TAGNETINDEX;
-
+#if !(UE_BUILD_SHIPPING || UE_BUILD_TEST)
 	if (Ar.IsSaving())
 	{
-		NetIndex = TagManager.GetNetIndexFromTag(*this);
-		if (NetIndex != INVALID_TAGNETINDEX)
-		{
-			// If we have a valid net index, serialize with that
-			bHasNetIndex = true;
-		}
+		UGameplayTagsManager::NotifyTagReplicated(*this, false);
 	}
+#endif
 
-	// Serialize if we have a name at all or are empty
-	Ar.SerializeBits(&bHasName, 1);
+	NetSerialize_Packed(Ar, Map, bOutSuccess);
 
-	if (bHasName)
+	bOutSuccess = true;
+	return true;
+}
+
+bool FGameplayTag::NetSerialize_Packed(FArchive& Ar, class UPackageMap* Map, bool& bOutSuccess)
+{
+	SCOPE_CYCLE_COUNTER(STAT_FGameplayTag_NetSerialize);
+
+	UGameplayTagsManager& TagManager = IGameplayTagsModule::GetGameplayTagsManager();	
+
+	if (TagManager.ShouldUseFastReplication())
 	{
-		Ar.SerializeBits(&bHasNetIndex, 1);
-		// If we have a net index serialize that, otherwise serialize as a name
-		if (bHasNetIndex)
+		FGameplayTagNetIndex NetIndex = INVALID_TAGNETINDEX;
+
+		if (Ar.IsSaving())
 		{
-			Ar << NetIndex;
+			NetIndex = TagManager.GetNetIndexFromTag(*this);
+			
+			SerializeTagNetIndexPacked(Ar, NetIndex, UGameplayTagsManager::NetIndexFirstBitSegment, UGameplayTagsManager::NetIndexTrueBitNum);
 		}
 		else
 		{
-			Ar << TagName;
-		}
-
-		if (Ar.IsLoading() && bHasNetIndex)
-		{
+			SerializeTagNetIndexPacked(Ar, NetIndex, UGameplayTagsManager::NetIndexFirstBitSegment, UGameplayTagsManager::NetIndexTrueBitNum);
 			TagName = TagManager.GetTagNameFromNetIndex(NetIndex);
 		}
+
 	}
 	else
 	{
-		TagName = NAME_None;
+		Ar << TagName;
 	}
 
 	bOutSuccess = true;
 	return true;
 }
 
-
-
 FGameplayTagQuery::FGameplayTagQuery()
 	: TokenStreamVersion(EGameplayTagQueryStreamVersion::LatestVersion)
 {
@@ -1437,3 +1537,55 @@ void FGameplayTagQueryExpression::EmitTokens(TArray<uint8>& TokenStream, TArray<
 	}
 }
 
+#if !(UE_BUILD_SHIPPING || UE_BUILD_TEST)
+
+static void GameplayTagPrintReplicationMap()
+{
+	UGameplayTagsManager::PrintReplicationFrequencyReport();
+}
+
+FAutoConsoleCommand GameplayTagPrintReplicationMapCmd(
+	TEXT("GameplayTags.PrintReport"), 
+	TEXT( "Prints frequency of gameplay tags" ), 
+	FConsoleCommandDelegate::CreateStatic(GameplayTagPrintReplicationMap)
+);
+
+
+static void TagPackingTest()
+{
+	for (int32 TotalNetIndexBits=1; TotalNetIndexBits <= 16; ++TotalNetIndexBits)
+	{
+		for (int32 NetIndexBitsPerComponent=0; NetIndexBitsPerComponent <= TotalNetIndexBits; NetIndexBitsPerComponent++)
+		{
+			for (int32 NetIndex=0; NetIndex < FMath::Pow(2, TotalNetIndexBits); ++NetIndex)
+			{
+				FGameplayTagNetIndex NI = NetIndex;
+
+				FNetBitWriter	BitWriter(nullptr, 1024 * 8);
+				SerializeTagNetIndexPacked(BitWriter, NI, NetIndexBitsPerComponent, TotalNetIndexBits);
+
+				FNetBitReader	Reader(nullptr, BitWriter.GetData(), BitWriter.GetNumBits());
+
+				FGameplayTagNetIndex NewIndex;
+				SerializeTagNetIndexPacked(Reader, NewIndex, NetIndexBitsPerComponent, TotalNetIndexBits);
+
+				if (ensureAlways((NewIndex == NI)) == false)
+				{
+					NetIndex--;
+					continue;
+				}
+			}
+		}
+	}
+
+	UE_LOG(LogGameplayTags, Warning, TEXT("TagPackingTest completed!"));
+
+}
+
+FAutoConsoleCommand TagPackingTestCmd(
+	TEXT("GameplayTags.PackingTest"), 
+	TEXT( "Prints frequency of gameplay tags" ), 
+	FConsoleCommandDelegate::CreateStatic(TagPackingTest)
+);
+
+#endif
\ No newline at end of file
diff --git a/Engine/Source/Runtime/GameplayTags/Private/GameplayTagsManager.cpp b/Engine/Source/Runtime/GameplayTags/Private/GameplayTagsManager.cpp
--- a/Engine/Source/Runtime/GameplayTags/Private/GameplayTagsManager.cpp
+++ b/Engine/Source/Runtime/GameplayTags/Private/GameplayTagsManager.cpp
@@ -11,6 +11,16 @@
 
 #define LOCTEXT_NAMESPACE "GameplayTagManager"
 
+#if !(UE_BUILD_SHIPPING || UE_BUILD_TEST)
+TMap<FGameplayTag, int32>	UGameplayTagsManager::ReplicationCountMap;
+TMap<FGameplayTag, int32>	UGameplayTagsManager::ReplicationCountMap_SingleTags;
+TMap<FGameplayTag, int32>	UGameplayTagsManager::ReplicationCountMap_Containers;
+#endif
+
+int32 UGameplayTagsManager::NetIndexFirstBitSegment=16;
+int32 UGameplayTagsManager::NetIndexTrueBitNum=16;
+int32 UGameplayTagsManager::NumBitsForContainerSize=6;
+
 UGameplayTagsManager::UGameplayTagsManager(const FObjectInitializer& ObjectInitializer)
 	: Super(ObjectInitializer)
 {
@@ -170,10 +180,32 @@ void UGameplayTagsManager::ConstructGameplayTagTree()
 					TableRow.Tag = TagStr;
 					AddTagTableRow(TableRow);
 				}
+
+				// Grab the commonly replicated tags
+				CommonlyReplicatedTags.Reserve(CommonlyReplicatedTags.Num() + Class->GetDefaultObject<UGameplayTagsSettings>()->CommonlyReplicatedTags.Num());
+				for (const FString& Str : Class->GetDefaultObject<UGameplayTagsSettings>()->CommonlyReplicatedTags)
+				{
+					FGameplayTag Tag = RequestGameplayTag(FName(*Str));
+					if (Tag.IsValid())
+					{
+						CommonlyReplicatedTags.Add(Tag);
+					}
+					else
+					{
+						UE_LOG(LogGameplayTags, Warning, TEXT("%s was found in the CommonlyReplicatedTags list but doesn't appear to be a valid tag!"), *Str);
+					}
+				}
+
+				NetIndexFirstBitSegment = Class->GetDefaultObject<UGameplayTagsSettings>()->NetIndexFirstBitSegment;
 			}
 			GameplayRootTag->GetChildTagNodes().Sort(FCompareFGameplayTagNodeByTag());
 		}
 
+		bUseFastReplication = false;
+		GConfig->GetBool(TEXT("GameplayTags"), TEXT("FastReplication"), bUseFastReplication, GEngineIni);
+
+		GConfig->GetInt(TEXT("GameplayTags"), TEXT("NumBitsForContainerSize"), NumBitsForContainerSize, GEngineIni);
+
 		if (ShouldUseFastReplication())
 		{
 #if STATS
@@ -244,8 +276,34 @@ void UGameplayTagsManager::ConstructNetIndex()
 
 	NetworkGameplayTagNodeIndex.Sort(FCompareFGameplayTagNodeByTag());
 
-	// This is now sorted and it should be the same on both client and server
+	// Put the common indices up front
+	for (int32 CommonIdx=0; CommonIdx < CommonlyReplicatedTags.Num(); ++CommonIdx)
+	{
+		int32 BaseIdx=0;
+		FGameplayTag& Tag = CommonlyReplicatedTags[CommonIdx];
+
+		bool Found = false;
+		for (int32 findidx=0; findidx < NetworkGameplayTagNodeIndex.Num(); ++findidx)
+		{
+			if (NetworkGameplayTagNodeIndex[findidx]->GetCompleteTag() == Tag.GetTagName())
+			{
+				NetworkGameplayTagNodeIndex.Swap(findidx, CommonIdx);
+				Found = true;
+				break;
+			}
+		}
+
+		// A non fatal error should have been thrown when parsing the CommonlyReplicatedTags list. If we make it here, something is seriously wrong.
+		checkf( Found, TEXT("Tag %s not found in NetworkGameplayTagNodeIndex"), *Tag.ToString() );
+	}
 
+	InvalidTagNetIndex = NetworkGameplayTagNodeIndex.Num()+1;
+	NetIndexTrueBitNum = FMath::CeilToInt(FMath::Log2(InvalidTagNetIndex));
+	
+	// This should never be smaller than NetIndexTrueBitNum
+	NetIndexFirstBitSegment = FMath::Min<int64>(NetIndexFirstBitSegment, NetIndexTrueBitNum);
+
+	// This is now sorted and it should be the same on both client and server
 	if (NetworkGameplayTagNodeIndex.Num() >= INVALID_TAGNETINDEX)
 	{
 		ensureMsgf(false, TEXT("Too many tags in dictionary for networking! Remove tags or increase tag net index size"));
@@ -266,8 +324,8 @@ FName UGameplayTagsManager::GetTagNameFromNetIndex(FGameplayTagNetIndex Index)
 {
 	if (Index >= NetworkGameplayTagNodeIndex.Num())
 	{
-		ensureMsgf(false, TEXT("Received invalid tag net index %d! Tag index is out of sync on client!"), Index);
-
+		// Ensure Index is the invalid index. If its higher than that, then something is wrong.
+		ensureMsgf(Index == InvalidTagNetIndex, TEXT("Received invalid tag net index %d! Tag index is out of sync on client!"), Index);
 		return NAME_None;
 	}
 	return NetworkGameplayTagNodeIndex[Index]->GetCompleteTag();
@@ -281,7 +339,8 @@ FGameplayTagNetIndex UGameplayTagsManager::GetNetIndexFromTag(const FGameplayTag
 	{
 		return (*GameplayTagNode)->GetNetIndex();
 	}
-	return INVALID_TAGNETINDEX;
+
+	return InvalidTagNetIndex;
 }
 
 bool UGameplayTagsManager::ShouldImportTagsFromINI()
@@ -291,13 +350,6 @@ bool UGameplayTagsManager::ShouldImportTagsFromINI()
 	return ImportFromINI;
 }
 
-bool UGameplayTagsManager::ShouldUseFastReplication()
-{
-	bool ImportFromINI = false;
-	GConfig->GetBool(TEXT("GameplayTags"), TEXT("FastReplication"), ImportFromINI, GEngineIni);
-	return ImportFromINI;
-}
-
 void UGameplayTagsManager::RedirectTagsForContainer(FGameplayTagContainer& Container, TSet<FName>& DeprecatedTagNamesNotFoundInTagMap)
 {
 	TSet<FName> NamesToRemove;
@@ -460,7 +512,6 @@ int32 UGameplayTagsManager::InsertTagIntoNodeArray(FName Tag, TWeakPtr<FGameplay
 	return InsertionIdx;
 }
 
-
 int32 UGameplayTagsManager::GetBestTagCategoryDescription(FString Tag, FText& OutDescription) const
 {
 	// get all the nodes that make up this tag
@@ -481,6 +532,124 @@ int32 UGameplayTagsManager::GetBestTagCategoryDescription(FString Tag, FText& Ou
 	return -1;
 }
 
+#if !(UE_BUILD_SHIPPING || UE_BUILD_TEST)
+void UGameplayTagsManager::PrintReplicationFrequencyReport()
+{
+	UE_LOG(LogGameplayTags, Warning, TEXT("================================="));
+	UE_LOG(LogGameplayTags, Warning, TEXT("Gameplay Tags Replication Report"));
+
+	UE_LOG(LogGameplayTags, Warning, TEXT("\nTags replicated solo:"));
+	ReplicationCountMap_SingleTags.ValueSort(TGreater<int32>());
+	for (auto& It : ReplicationCountMap_SingleTags)
+	{
+		UE_LOG(LogGameplayTags, Warning, TEXT("%s - %d"), *It.Key.ToString(), It.Value);
+	}
+	
+	// ---------------------------------------
+
+	UE_LOG(LogGameplayTags, Warning, TEXT("\nTags replicated in containers:"));
+	ReplicationCountMap_Containers.ValueSort(TGreater<int32>());
+	for (auto& It : ReplicationCountMap_Containers)
+	{
+		UE_LOG(LogGameplayTags, Warning, TEXT("%s - %d"), *It.Key.ToString(), It.Value);
+	}
+
+	// ---------------------------------------
+
+	UE_LOG(LogGameplayTags, Warning, TEXT("\nAll Tags replicated:"));
+	ReplicationCountMap.ValueSort(TGreater<int32>());
+	for (auto& It : ReplicationCountMap)
+	{
+		UE_LOG(LogGameplayTags, Warning, TEXT("%s - %d"), *It.Key.ToString(), It.Value);
+	}
+
+	TMap<int32, int32> SavingsMap;
+	int32 BaselineCost = 0;
+	for (int32 Bits=1; Bits < NetIndexTrueBitNum; ++Bits)
+	{
+		int32 TotalSavings = 0;
+		BaselineCost = 0;
+
+		FGameplayTagNetIndex ExpectedNetIndex=0;
+		for (auto& It : ReplicationCountMap)
+		{
+			int32 ExpectedCostBits = 0;
+			bool FirstSeg = ExpectedNetIndex < FMath::Pow(2, Bits);
+			if (FirstSeg)
+			{
+				// This would fit in the first Bits segment
+				ExpectedCostBits = Bits+1;
+			}
+			else
+			{
+				// Would go in the second segment, so we pay the +1 cost
+				ExpectedCostBits = NetIndexTrueBitNum+1;
+			}
+
+			int32 Savings = (NetIndexTrueBitNum - ExpectedCostBits) * It.Value;
+			BaselineCost += NetIndexTrueBitNum * It.Value;
+
+			//UE_LOG(LogGameplayTags, Warning, TEXT("[Bits: %d] Tag %s would save %d bits"), Bits, *It.Key.ToString(), Savings);
+			ExpectedNetIndex++;
+			TotalSavings += Savings;
+		}
+
+		SavingsMap.FindOrAdd(Bits) = TotalSavings;
+	}
+
+	SavingsMap.ValueSort(TGreater<int32>());
+	int32 BestBits = 0;
+	for (auto& It : SavingsMap)
+	{
+		if (BestBits == 0)
+		{
+			BestBits = It.Key;
+		}
+
+		UE_LOG(LogGameplayTags, Warning, TEXT("%d bits would save %d (%.2f)"), It.Key, It.Value, (float)It.Value / (float)BaselineCost);
+	}
+
+	UE_LOG(LogGameplayTags, Warning, TEXT("\nSuggested config:"));
+
+	// Write out a nice copy pastable config
+	int32 Count=0;
+	for (auto& It : ReplicationCountMap)
+	{
+		UE_LOG(LogGameplayTags, Warning, TEXT("+CommonlyReplicatedTags=%s"), *It.Key.ToString());
+
+		if (Count == FMath::Pow(2, BestBits))
+		{
+			// Print a blank line out, indicating tags after this are not necessary but still may be useful if the user wants to manually edit the list.
+			UE_LOG(LogGameplayTags, Warning, TEXT(""));
+		}
+
+		if (Count++ >= FMath::Pow(2, BestBits+1))
+		{
+			break;
+		}
+	}
+
+	UE_LOG(LogGameplayTags, Warning, TEXT("NetIndexFirstBitSegment=%d"), BestBits);
+
+	UE_LOG(LogGameplayTags, Warning, TEXT("================================="));
+}
+
+void UGameplayTagsManager::NotifyTagReplicated(FGameplayTag Tag, bool WasInContainer)
+{
+	ReplicationCountMap.FindOrAdd(Tag)++;
+
+	if (WasInContainer)
+	{
+		ReplicationCountMap_Containers.FindOrAdd(Tag)++;
+	}
+	else
+	{
+		ReplicationCountMap_SingleTags.FindOrAdd(Tag)++;
+	}
+	
+}
+#endif
+
 #if WITH_EDITOR
 
 static void RecursiveRootTagSearch(const FString& InFilterString, const TArray<TSharedPtr<FGameplayTagNode> >& GameplayRootTags, TArray< TSharedPtr<FGameplayTagNode> >& OutTagArray)
diff --git a/Engine/Source/Runtime/GameplayTags/Private/GameplayTagsModule.cpp b/Engine/Source/Runtime/GameplayTags/Private/GameplayTagsModule.cpp
--- a/Engine/Source/Runtime/GameplayTags/Private/GameplayTagsModule.cpp
+++ b/Engine/Source/Runtime/GameplayTags/Private/GameplayTagsModule.cpp
@@ -26,9 +26,20 @@ void FGameplayTagsModule::StartupModule()
 	GGameplayTagsManager->ConstructGameplayTagTree();
 }
 
+#if !(UE_BUILD_SHIPPING || UE_BUILD_TEST)
+int32 GameplayTagPrintReportOnShutdown = 0;
+static FAutoConsoleVariableRef CVarGameplayTagPrintReportOnShutdown(TEXT("GameplayTags.PrintReportOnShutdown"), GameplayTagPrintReportOnShutdown, TEXT("Print gameplay tag replication report on shutdown"), ECVF_Default );
+#endif
+
+
 void FGameplayTagsModule::ShutdownModule()
 {
-	// This function may be called during shutdown to clean up your module.  For modules that support dynamic reloading,
-	// we call this function before unloading the module.
+#if !(UE_BUILD_SHIPPING || UE_BUILD_TEST)
+	if (GameplayTagPrintReportOnShutdown)
+	{
+		UGameplayTagsManager::PrintReplicationFrequencyReport();
+	}
+#endif
+
 	GGameplayTagsManager = NULL;
 }
diff --git a/Engine/Source/Runtime/GameplayTags/Private/GameplayTagsModulePrivatePCH.h b/Engine/Source/Runtime/GameplayTags/Private/GameplayTagsModulePrivatePCH.h
--- a/Engine/Source/Runtime/GameplayTags/Private/GameplayTagsModulePrivatePCH.h
+++ b/Engine/Source/Runtime/GameplayTags/Private/GameplayTagsModulePrivatePCH.h
@@ -22,4 +22,3 @@
 
 #include "GameplayTagsModule.h"
 
-DECLARE_LOG_CATEGORY_EXTERN(LogGameplayTags, Log, All);
\ No newline at end of file
diff --git a/Engine/Source/Runtime/GameplayTags/Private/GameplayTagsSettings.cpp b/Engine/Source/Runtime/GameplayTags/Private/GameplayTagsSettings.cpp
--- a/Engine/Source/Runtime/GameplayTags/Private/GameplayTagsSettings.cpp
+++ b/Engine/Source/Runtime/GameplayTags/Private/GameplayTagsSettings.cpp
@@ -6,8 +6,7 @@
 UGameplayTagsSettings::UGameplayTagsSettings(const FObjectInitializer& ObjectInitializer)
 : Super(ObjectInitializer)
 {
-
-
+	NetIndexFirstBitSegment = 16;
 }
 
 #if WITH_EDITOR
diff --git a/Engine/Source/Runtime/HTML5/MapPakDownloader/Private/MapPakDownloader.cpp b/Engine/Source/Runtime/HTML5/MapPakDownloader/Private/MapPakDownloader.cpp
--- a/Engine/Source/Runtime/HTML5/MapPakDownloader/Private/MapPakDownloader.cpp
+++ b/Engine/Source/Runtime/HTML5/MapPakDownloader/Private/MapPakDownloader.cpp
@@ -165,7 +165,7 @@ void FMapPakDownloader::CachePak()
 
 											UE_LOG(LogMapPakDownloader, Warning, TEXT("%s download complete!"), *PakRequest->GetFileName());
 											UE_LOG(LogMapPakDownloader, Warning, TEXT("Mounting..."), *Name);
-											FCoreDelegates::OnMountPak.Execute(Name, 0);
+											FCoreDelegates::OnMountPak.Execute(Name, 0, nullptr);
 											UE_LOG(LogMapPakDownloader, Warning, TEXT("%s Mounted!"), *Name);
 
 											// Get hold of the world.
diff --git a/Engine/Source/Runtime/HTTPChunkInstaller/Private/ChunkInstall.h b/Engine/Source/Runtime/HTTPChunkInstaller/Private/ChunkInstall.h
--- a/Engine/Source/Runtime/HTTPChunkInstaller/Private/ChunkInstall.h
+++ b/Engine/Source/Runtime/HTTPChunkInstaller/Private/ChunkInstall.h
@@ -106,13 +106,13 @@ public:
 // 				}
 				if (FCoreDelegates::OnMountPak.IsBound())
 				{
-					auto bSuccess = FCoreDelegates::OnMountPak.Execute(PakFiles[PakIndex], PakReadOrder);
+					auto bSuccess = FCoreDelegates::OnMountPak.Execute(PakFiles[PakIndex], PakReadOrder, nullptr);
 #if !UE_BUILD_SHIPPING
 					if (!bSuccess)
 					{
 						// This can fail because of the sandbox system - which the pak system doesn't understand.
 						auto SandboxedPath = IFileManager::Get().ConvertToAbsolutePathForExternalAppForRead(*PakFiles[PakIndex]);
-						bSuccess = FCoreDelegates::OnMountPak.Execute(SandboxedPath, PakReadOrder);
+						bSuccess = FCoreDelegates::OnMountPak.Execute(SandboxedPath, PakReadOrder, nullptr);
 					}
 #endif
 					MountedPaks.Add(PakFiles[PakIndex]);
diff --git a/Engine/Source/Runtime/HTTPChunkInstaller/Private/ChunkSetup.h b/Engine/Source/Runtime/HTTPChunkInstaller/Private/ChunkSetup.h
--- a/Engine/Source/Runtime/HTTPChunkInstaller/Private/ChunkSetup.h
+++ b/Engine/Source/Runtime/HTTPChunkInstaller/Private/ChunkSetup.h
@@ -299,13 +299,13 @@ public:
 			{
 				if (FCoreDelegates::OnMountPak.IsBound())
 				{
-					auto bSuccess = FCoreDelegates::OnMountPak.Execute(PakPath, PakReadOrder);
+					auto bSuccess = FCoreDelegates::OnMountPak.Execute(PakPath, PakReadOrder, nullptr);
 #if !UE_BUILD_SHIPPING
 					if (!bSuccess)
 					{
 						// This can fail because of the sandbox system - which the pak system doesn't understand.
 						auto SandboxedPath = IFileManager::Get().ConvertToAbsolutePathForExternalAppForRead(*PakPath);
-						bSuccess = FCoreDelegates::OnMountPak.Execute(SandboxedPath, PakReadOrder);
+						bSuccess = FCoreDelegates::OnMountPak.Execute(SandboxedPath, PakReadOrder, nullptr);
 					}
 #endif
 					//Register the install
diff --git a/Engine/Source/Runtime/Launch/Private/LaunchEngineLoop.cpp b/Engine/Source/Runtime/Launch/Private/LaunchEngineLoop.cpp
--- a/Engine/Source/Runtime/Launch/Private/LaunchEngineLoop.cpp
+++ b/Engine/Source/Runtime/Launch/Private/LaunchEngineLoop.cpp
@@ -874,7 +874,13 @@ int32 FEngineLoop::PreInit( const TCHAR* CmdLine )
 		}
 	}
 
-#endif // UE_EDITOR
+	if (bHasCommandletToken)
+	{
+		// will be reset later once the commandlet class loaded
+		PRIVATE_GIsRunningCommandlet = true;
+	}
+
+#endif // WITH_ENGINE
 
 
 	// trim any whitespace at edges of string - this can happen if the token was quoted with leading or trailing whitespace
@@ -969,7 +975,12 @@ int32 FEngineLoop::PreInit( const TCHAR* CmdLine )
 			}
 			CommandletCommandLine = ParsedCmdLine;
 		}
+	}
 
+	if (bHasCommandletToken)
+	{
+		// will be reset later once the commandlet class loaded
+		PRIVATE_GIsRunningCommandlet = true;
 	}
 
 	if( !bIsNotEditor && GIsGameAgnosticExe )
@@ -1558,7 +1569,10 @@ int32 FEngineLoop::PreInit( const TCHAR* CmdLine )
 	MarkObjectsToDisregardForGC(); 
 	GUObjectArray.CloseDisregardForGC();
 
-	SetIsServerForOnlineSubsystemsDelegate(FQueryIsRunningServer::CreateStatic(&IsServerDelegateForOSS));
+	if (IOnlineSubsystem::IsLoaded())
+	{
+		SetIsServerForOnlineSubsystemsDelegate(FQueryIsRunningServer::CreateStatic(&IsServerDelegateForOSS));
+	}
 
 	SlowTask.EnterProgressFrame(5);
 
@@ -2902,6 +2916,8 @@ bool FEngineLoop::AppInit( )
 	// Now finish initializing the file manager after the command line is set up
 	IFileManager::Get().ProcessCommandLineOptions();
 
+	FPageAllocator::LatchProtectedMode();
+
 #if !UE_BUILD_SHIPPING
 	if (FParse::Param(FCommandLine::Get(), TEXT("BUILDMACHINE")))
 	{
@@ -3099,7 +3115,7 @@ bool FEngineLoop::AppInit( )
 
 	UE_LOG(LogInit, Log, TEXT("Build Configuration: %s"), EBuildConfigurations::ToString(FApp::GetBuildConfiguration()));
 	UE_LOG(LogInit, Log, TEXT("Branch Name: %s"), *FApp::GetBranchName() );
-	UE_LOG(LogInit, Log, TEXT("Command line: %s"), FCommandLine::Get() );
+	UE_LOG(LogInit, Log, TEXT("Command line: %s"), FCommandLine::GetForLogging() );
 	UE_LOG(LogInit, Log, TEXT("Base directory: %s"), FPlatformProcess::BaseDir() );
 	//UE_LOG(LogInit, Log, TEXT("Character set: %s"), sizeof(TCHAR)==1 ? TEXT("ANSI") : TEXT("Unicode") );
 	UE_LOG(LogInit, Log, TEXT("Installed Engine Build: %d"), FApp::IsEngineInstalled() ? 1 : 0);
@@ -3114,25 +3130,28 @@ bool FEngineLoop::AppInit( )
 #endif
 
 #if WITH_ENGINE
-	// Earliest place to init the online subsystems
-	// Code needs GConfigFile to be valid
-	// Must be after FThreadStats::StartThread();
-	// Must be before Render/RHI subsystem D3DCreate()
-	// For platform services that need D3D hooks like Steam
-	// --
-	// Why load HTTP?
-	// Because most, if not all online subsystems will load HTTP themselves. This can cause problems though, as HTTP will be loaded *after* OSS, 
-	// and if OSS holds on to resources allocated by it, this will cause crash (modules are being unloaded in LIFO order with no dependency tracking).
-	// Loading HTTP before OSS works around this problem by making ModuleManager unload HTTP after OSS, at the cost of extra module for the few OSS (like Null) that don't use it.
-	if(FModuleManager::Get().ModuleExists(TEXT("XMPP")))
-	{
-		FModuleManager::Get().LoadModule(TEXT("XMPP"));
-	}
-	FModuleManager::Get().LoadModule(TEXT("HTTP"));
-	FModuleManager::Get().LoadModule(TEXT("OnlineSubsystem"));
-	FModuleManager::Get().LoadModule(TEXT("OnlineSubsystemUtils"));
-	// Also load the console/platform specific OSS which might not necessarily be the default OSS instance
-	IOnlineSubsystem::GetByPlatform();
+	if (!IsRunningCommandlet())
+	{
+		// Earliest place to init the online subsystems
+		// Code needs GConfigFile to be valid
+		// Must be after FThreadStats::StartThread();
+		// Must be before Render/RHI subsystem D3DCreate()
+		// For platform services that need D3D hooks like Steam
+		// --
+		// Why load HTTP?
+		// Because most, if not all online subsystems will load HTTP themselves. This can cause problems though, as HTTP will be loaded *after* OSS, 
+		// and if OSS holds on to resources allocated by it, this will cause crash (modules are being unloaded in LIFO order with no dependency tracking).
+		// Loading HTTP before OSS works around this problem by making ModuleManager unload HTTP after OSS, at the cost of extra module for the few OSS (like Null) that don't use it.
+		if (FModuleManager::Get().ModuleExists(TEXT("XMPP")))
+		{
+			FModuleManager::Get().LoadModule(TEXT("XMPP"));
+		}
+		FModuleManager::Get().LoadModule(TEXT("HTTP"));
+		FModuleManager::Get().LoadModule(TEXT("OnlineSubsystem"));
+		FModuleManager::Get().LoadModule(TEXT("OnlineSubsystemUtils"));
+		// Also load the console/platform specific OSS which might not necessarily be the default OSS instance
+		IOnlineSubsystem::GetByPlatform();
+	}
 #endif
 
 	// Checks.
diff --git a/Engine/Source/Runtime/Launch/Resources/VersionLocked.h b/Engine/Source/Runtime/Launch/Resources/VersionLocked.h
new file mode 100644
--- /dev/null
+++ b/Engine/Source/Runtime/Launch/Resources/VersionLocked.h
@@ -0,0 +1,25 @@
+// Copyright 1998-2015 Epic Games, Inc. All Rights Reserved.
+
+#pragma once
+
+#define VERSION_LOCKED 0
+
+#if VERSION_LOCKED
+
+#define ENGINE_MAJOR_VERSION	4
+#define ENGINE_MINOR_VERSION	11
+#define ENGINE_PATCH_VERSION	0
+
+#define ENGINE_IS_LICENSEE_VERSION 0
+
+#define BUILT_FROM_CHANGELIST 2820843
+#define BRANCH_NAME "++Orion+Release-Next"
+
+#define ENGINE_IS_PROMOTED_BUILD (BUILT_FROM_CHANGELIST > 0)
+
+#define EPIC_COMPANY_NAME  "Epic Games, Inc."
+#define EPIC_COPYRIGHT_STRING "Copyright 1998-2015 Epic Games, Inc. All Rights Reserved."
+#define EPIC_PRODUCT_NAME "Unreal Engine"
+#define EPIC_PRODUCT_IDENTIFIER "UnrealEngine"
+
+#endif
diff --git a/Engine/Source/Runtime/Online/Android/OnlineSubsystemGooglePlay/Private/OnlineSubsystemGooglePlay.cpp b/Engine/Source/Runtime/Online/Android/OnlineSubsystemGooglePlay/Private/OnlineSubsystemGooglePlay.cpp
--- a/Engine/Source/Runtime/Online/Android/OnlineSubsystemGooglePlay/Private/OnlineSubsystemGooglePlay.cpp
+++ b/Engine/Source/Runtime/Online/Android/OnlineSubsystemGooglePlay/Private/OnlineSubsystemGooglePlay.cpp
@@ -46,6 +46,11 @@ IOnlineSessionPtr FOnlineSubsystemGooglePlay::GetSessionInterface() const
 	return nullptr;
 }
 
+IMessageSanitizerPtr FOnlineSubsystemGooglePlay::GetMessageSanitizerInterface() const
+{
+	return nullptr;
+}
+
 IOnlineFriendsPtr FOnlineSubsystemGooglePlay::GetFriendsInterface() const
 {
 	return nullptr;
diff --git a/Engine/Source/Runtime/Online/Android/OnlineSubsystemGooglePlay/Public/OnlineSubsystemGooglePlay.h b/Engine/Source/Runtime/Online/Android/OnlineSubsystemGooglePlay/Public/OnlineSubsystemGooglePlay.h
--- a/Engine/Source/Runtime/Online/Android/OnlineSubsystemGooglePlay/Public/OnlineSubsystemGooglePlay.h
+++ b/Engine/Source/Runtime/Online/Android/OnlineSubsystemGooglePlay/Public/OnlineSubsystemGooglePlay.h
@@ -32,6 +32,7 @@ public:
 	//~ Begin IOnlineSubsystem Interface
 	virtual IOnlineSessionPtr GetSessionInterface() const override;
 	virtual IOnlineFriendsPtr GetFriendsInterface() const override;
+	virtual IMessageSanitizerPtr GetMessageSanitizerInterface() const override;
 	virtual IOnlinePartyPtr GetPartyInterface() const override;
 	virtual IOnlineGroupsPtr GetGroupsInterface() const override;
 	virtual IOnlineSharedCloudPtr GetSharedCloudInterface() const override;
diff --git a/Engine/Source/Runtime/Online/Hotfix/Private/OnlineHotfixManager.cpp b/Engine/Source/Runtime/Online/Hotfix/Private/OnlineHotfixManager.cpp
--- a/Engine/Source/Runtime/Online/Hotfix/Private/OnlineHotfixManager.cpp
+++ b/Engine/Source/Runtime/Online/Hotfix/Private/OnlineHotfixManager.cpp
@@ -9,38 +9,37 @@ DEFINE_LOG_CATEGORY(LogHotfixManager);
 
 FName NAME_HotfixManager(TEXT("HotfixManager"));
 
-struct FHotfixManagerExec :
-	public FSelfRegisteringExec
+class FPakFileVisitor : public IPlatformFile::FDirectoryVisitor
 {
-	virtual bool Exec(UWorld* InWorld, const TCHAR* Cmd, FOutputDevice& Ar) override
+public:
+	virtual bool Visit(const TCHAR* FilenameOrDirectory, bool bIsDirectory) override
 	{
-		if (FParse::Command(&Cmd, TEXT("HOTFIX")))
+		if (!bIsDirectory)
 		{
-			UOnlineHotfixManager* HotfixManager = UOnlineHotfixManager::Get(InWorld);
-			if (HotfixManager != nullptr)
-			{
-				HotfixManager->StartHotfixProcess();
-			}
-			return true;
+			Files.Add(FilenameOrDirectory);
 		}
-		return false;
+		return true;
 	}
+
+	TArray<FString> Files;
 };
-static FHotfixManagerExec HotfixManagerExec;
 
 UOnlineHotfixManager::UOnlineHotfixManager() :
 	Super(),
 	TotalFiles(0),
 	NumDownloaded(0),
 	TotalBytes(0),
-	NumBytes(0)
+	NumBytes(0),
+	bHotfixingInProgress(false),
+	bHotfixNeedsReload(false)
 {
 	OnEnumerateFilesCompleteDelegate = FOnEnumerateFilesCompleteDelegate::CreateUObject(this, &UOnlineHotfixManager::OnEnumerateFilesComplete);
 	OnReadFileProgressDelegate = FOnReadFileProgressDelegate::CreateUObject(this, &UOnlineHotfixManager::OnReadFileProgress);
 	OnReadFileCompleteDelegate = FOnReadFileCompleteDelegate::CreateUObject(this, &UOnlineHotfixManager::OnReadFileComplete);
-	// So we only try to apply files for this platform
-	PlatformPrefix = ANSI_TO_TCHAR(FPlatformProperties::PlatformName());
-	PlatformPrefix += TEXT("_");
+#if !UE_BUILD_SHIPPING
+	bLogMountedPakContents = FParse::Param(FCommandLine::Get(), TEXT("LogHotfixPakContents"));
+#endif
+	GameContentPath = FString() / FApp::GetGameName() / TEXT("Content");
 }
 
 UOnlineHotfixManager* UOnlineHotfixManager::Get(UWorld* World)
@@ -73,18 +72,26 @@ void UOnlineHotfixManager::PostInitProperties()
 #if !UE_BUILD_SHIPPING
 	FParse::Value(FCommandLine::Get(), TEXT("HOTFIXPREFIX="), DebugPrefix);
 #endif
+	// So we only try to apply files for this platform
+	PlatformPrefix = DebugPrefix + ANSI_TO_TCHAR(FPlatformProperties::PlatformName());
+	PlatformPrefix += TEXT("_");
+	// Build the default prefix too
+	DefaultPrefix = DebugPrefix + TEXT("Default");
+
 	Super::PostInitProperties();
 }
 
 void UOnlineHotfixManager::Init()
 {
+	bHotfixingInProgress = true;
+	bHotfixNeedsReload = false;
 	TotalFiles = 0;
 	NumDownloaded = 0;
 	TotalBytes = 0;
 	NumBytes = 0;
 	// Build the name of the loc file that we'll care about
 	// It can change at runtime so build it just before fetching the data
-	GameLocName = FInternationalization::Get().GetCurrentCulture()->GetTwoLetterISOLanguageName() + TEXT("_Game.locres");
+	GameLocName = DebugPrefix + FInternationalization::Get().GetCurrentCulture()->GetTwoLetterISOLanguageName() + TEXT("_Game.locres");
 	OnlineTitleFile = Online::GetTitleFileInterface(OSSName.Len() ? FName(*OSSName, FNAME_Find) : NAME_None);
 	if (OnlineTitleFile.IsValid())
 	{
@@ -106,6 +113,7 @@ void UOnlineHotfixManager::Cleanup()
 		OnlineTitleFile->ClearOnReadFileCompleteDelegate_Handle(OnReadFileCompleteDelegateHandle);
 	}
 	OnlineTitleFile = nullptr;
+	bHotfixingInProgress = false;
 }
 
 void UOnlineHotfixManager::StartHotfixProcess()
@@ -118,6 +126,11 @@ void UOnlineHotfixManager::StartHotfixProcess()
 		return;
 	}
 
+	if (bHotfixingInProgress)
+	{
+		return;
+	}
+
 	Init();
 	// Kick off an enumeration of the files that are available to download
 	if (OnlineTitleFile.IsValid())
@@ -133,33 +146,74 @@ void UOnlineHotfixManager::StartHotfixProcess()
 
 struct FHotfixFileSortPredicate
 {
-	const FString PlatformName;
-
-	FHotfixFileSortPredicate(const FString& InPlatformName) :
-		PlatformName(InPlatformName)
+	struct FHotfixFileNameSortPredicate
 	{
-	}
+		const FString PlatformPrefix;
+		const FString DefaultPrefix;
 
-	bool operator()(const FCloudFileHeader &A, const FCloudFileHeader &B) const
-	{
-		if (A.FileName.EndsWith(TEXT("INI")) && B.FileName.EndsWith(TEXT("INI")))
+		FHotfixFileNameSortPredicate(const FString& InPlatformPrefix, const FString& InDefaultPrefix) :
+			PlatformPrefix(InPlatformPrefix),
+			DefaultPrefix(InDefaultPrefix)
 		{
-			bool bAStartsWithDefault = A.FileName.StartsWith(TEXT("Default"));
-			bool bBStartsWithDefault = B.FileName.StartsWith(TEXT("Default"));
-			bool bAStartsWithPlatformName = A.FileName.StartsWith(PlatformName);
-			bool bBStartsWithPlatformName = B.FileName.StartsWith(PlatformName);
-			// Sort any file with Default in front of the Platform version so they apply in correct order
-			if (bAStartsWithDefault && bBStartsWithPlatformName)
+		}
+
+		bool Compare(const FString& A, const FString& B) const
+		{
+			bool bAIsIni = A.EndsWith(TEXT("INI"));
+			bool bBIsIni = B.EndsWith(TEXT("INI"));
+			if (bAIsIni && bBIsIni)
+			{
+				bool bAStartsWithDefault = A.StartsWith(DefaultPrefix);
+				bool bBStartsWithDefault = B.StartsWith(DefaultPrefix);
+				bool bAStartsWithPlatformName = A.StartsWith(PlatformPrefix);
+				bool bBStartsWithPlatformName = B.StartsWith(PlatformPrefix);
+				// Sort any file with Default in front of the Platform version so they apply in correct order
+				if (bAStartsWithDefault && bBStartsWithPlatformName)
+				{
+					return true;
+				}
+				else if (bAStartsWithPlatformName && bBStartsWithDefault)
+				{
+					return false;
+				}
+				// Sort any Default/Platform INIs before other INIs
+				else if ((bAStartsWithDefault || bAStartsWithPlatformName) && !(bBStartsWithDefault || bBStartsWithPlatformName))
+				{
+					return true;
+				}
+				else if (!(bAStartsWithDefault || bAStartsWithPlatformName) && (bBStartsWithDefault || bBStartsWithPlatformName))
+				{
+					return false;
+				}
+			}
+			// Sort all INI files before non-INI files
+			else if (bAIsIni && !bBIsIni)
 			{
 				return true;
 			}
-			else if (bAStartsWithPlatformName && bBStartsWithDefault)
+			else if (!bAIsIni && bBIsIni)
 			{
 				return false;
 			}
+			// Sort by the string order
+			return A < B;
 		}
-		// Sort by the string order
-		return A < B;
+	};
+	FHotfixFileNameSortPredicate FileNameSorter;
+
+	FHotfixFileSortPredicate(const FString& InPlatformPrefix, const FString& InDefaultPrefix) :
+		FileNameSorter(InPlatformPrefix, InDefaultPrefix)
+	{
+	}
+
+	bool operator()(const FCloudFileHeader &A, const FCloudFileHeader &B) const
+	{
+		return FileNameSorter.Compare(A.FileName, B.FileName);
+	}
+
+	bool operator()(const FString& A, const FString& B) const
+	{
+		return FileNameSorter.Compare(FPaths::GetCleanFilename(A), FPaths::GetCleanFilename(B));
 	}
 };
 
@@ -174,14 +228,22 @@ void UOnlineHotfixManager::OnEnumerateFilesComplete(bool bWasSuccessful)
 		// Get the new header data
 		OnlineTitleFile->GetFileList(HotfixFileList);
 		FilterHotfixFiles();
+		// Reduce the set of work to just the files that changed since last run
+		BuildHotfixFileListDeltas();
 		// Sort after filtering so that the comparison below doesn't fail to different order from the server
-		HotfixFileList.Sort<FHotfixFileSortPredicate>(FHotfixFileSortPredicate(PlatformPrefix));
-		if (LastHotfixFileList.Num() == 0 || LastHotfixFileList != HotfixFileList)
+		ChangedHotfixFileList.Sort<FHotfixFileSortPredicate>(FHotfixFileSortPredicate(PlatformPrefix, DefaultPrefix));
+		// Perform any undo operations needed
+		if (ChangedHotfixFileList.Num() > 0 || RemovedHotfixFileList.Num() > 0)
 		{
+			RestoreBackupIniFiles();
 			UnmountHotfixFiles();
+		}
+		// Read any changed files
+		if (ChangedHotfixFileList.Num() > 0)
+		{
 			// Update our totals for our progress delegates
-			TotalFiles = HotfixFileList.Num();
-			for (auto FileHeader : HotfixFileList)
+			TotalFiles = ChangedHotfixFileList.Num();
+			for (auto& FileHeader : ChangedHotfixFileList)
 			{
 				TotalBytes += FileHeader.FileSize;
 			}
@@ -200,6 +262,48 @@ void UOnlineHotfixManager::OnEnumerateFilesComplete(bool bWasSuccessful)
 	}
 }
 
+void UOnlineHotfixManager::BuildHotfixFileListDeltas()
+{
+	RemovedHotfixFileList.Empty();
+	ChangedHotfixFileList.Empty();
+	// Go through the current list and see if it's changed from the previous attempt
+	for (auto& CurrentHeader : HotfixFileList)
+	{
+		bool bFoundMatch = false;
+		for (auto& LastHeader : LastHotfixFileList)
+		{
+			if (LastHeader == CurrentHeader)
+			{
+				bFoundMatch = true;
+				break;
+			}
+		}
+		if (!bFoundMatch)
+		{
+			// We're different so add to the process list
+			ChangedHotfixFileList.Add(CurrentHeader);
+		}
+	}
+	// Find any files that have been removed from the set of hotfix files
+	for (auto& LastHeader : LastHotfixFileList)
+	{
+		bool bFoundMatch = false;
+		for (auto& CurrentHeader : HotfixFileList)
+		{
+			if (LastHeader.FileName == CurrentHeader.FileName)
+			{
+				bFoundMatch = true;
+				break;
+			}
+		}
+		if (!bFoundMatch)
+		{
+			// We've been removed so add to the removed list
+			RemovedHotfixFileList.Add(LastHeader);
+		}
+	}
+}
+
 void UOnlineHotfixManager::FilterHotfixFiles()
 {
 	for (int32 Idx = 0; Idx < HotfixFileList.Num(); Idx++)
@@ -214,16 +318,16 @@ void UOnlineHotfixManager::FilterHotfixFiles()
 
 void UOnlineHotfixManager::ReadHotfixFiles()
 {
-	if (HotfixFileList.Num())
+	if (ChangedHotfixFileList.Num())
 	{
 		check(OnlineTitleFile.IsValid());
 		// Kick off a read for each file
 		// Do this in two passes so already cached files don't trigger completion
-		for (auto FileHeader : HotfixFileList)
+		for (auto& FileHeader : ChangedHotfixFileList)
 		{
 			PendingHotfixFiles.Add(FileHeader.DLName, FPendingFileDLProgress());
 		}
-		for (auto FileHeader : HotfixFileList)
+		for (auto& FileHeader : ChangedHotfixFileList)
 		{
 			OnlineTitleFile->ReadFile(FileHeader.DLName);
 		}
@@ -274,7 +378,7 @@ void UOnlineHotfixManager::UpdateProgress(uint32 FileCount, uint64 UpdateSize)
 
 void UOnlineHotfixManager::ApplyHotfix()
 {
-	for (auto FileHeader : HotfixFileList)
+	for (auto& FileHeader : ChangedHotfixFileList)
 	{
 		if (!ApplyHotfixProcessing(FileHeader))
 		{
@@ -282,9 +386,11 @@ void UOnlineHotfixManager::ApplyHotfix()
 			TriggerHotfixComplete(EHotfixResult::Failed);
 			return;
 		}
+		// Let anyone listening know we just processed this file
+		TriggerOnHotfixProcessedFileDelegates(FileHeader.FileName, GetCachedDirectory() / FileHeader.DLName);
 	}
 	UE_LOG(LogHotfixManager, Display, TEXT("Hotfix data has been successfully applied"));
-	TriggerHotfixComplete(EHotfixResult::Success);
+	TriggerHotfixComplete((!bHotfixNeedsReload) ? EHotfixResult::Success : EHotfixResult::SuccessNeedsReload);
 }
 
 void UOnlineHotfixManager::TriggerHotfixComplete(EHotfixResult HotfixResult)
@@ -300,21 +406,14 @@ void UOnlineHotfixManager::TriggerHotfixComplete(EHotfixResult HotfixResult)
 
 bool UOnlineHotfixManager::WantsHotfixProcessing(const FCloudFileHeader& FileHeader)
 {
-#if !UE_BUILD_SHIPPING
-	if (!DebugPrefix.IsEmpty() && !FileHeader.FileName.StartsWith(DebugPrefix))
-	{
-		// Allow filtering of files by individual or group for testing
-		return false;
-	}
-#endif
 	const FString Extension = FPaths::GetExtension(FileHeader.FileName);
 	if (Extension == TEXT("INI"))
 	{
-		return FileHeader.FileName.StartsWith(PlatformPrefix) || FileHeader.FileName.StartsWith(TEXT("Default"));
+		return FileHeader.FileName.StartsWith(PlatformPrefix) || FileHeader.FileName.StartsWith(DefaultPrefix);
 	}
 	else if (Extension == TEXT("PAK"))
 	{
-		return FileHeader.FileName.StartsWith(PlatformPrefix);
+		return FileHeader.FileName.Find(PlatformPrefix) != -1;
 	}
 	return FileHeader.FileName == GameLocName;
 }
@@ -348,23 +447,32 @@ bool UOnlineHotfixManager::ApplyHotfixProcessing(const FCloudFileHeader& FileHea
 	return bSuccess;
 }
 
-FConfigFile* UOnlineHotfixManager::GetConfigFile(const FString& IniName)
+FString UOnlineHotfixManager::GetStrippedConfigFileName(const FString& IniName)
 {
 	FString StrippedIniName(IniName);
-#if !UE_BUILD_SHIPPING
-	if (!DebugPrefix.IsEmpty() && StrippedIniName.StartsWith(DebugPrefix))
-	{
-		StrippedIniName = IniName.Right(StrippedIniName.Len() - DebugPrefix.Len());
-	}
-#endif
 	if (StrippedIniName.StartsWith(PlatformPrefix))
 	{
 		StrippedIniName = IniName.Right(StrippedIniName.Len() - PlatformPrefix.Len());
 	}
-	if (StrippedIniName.StartsWith(TEXT("Default")))
+	else if (StrippedIniName.StartsWith(DefaultPrefix))
 	{
 		StrippedIniName = IniName.Right(StrippedIniName.Len() - 7);
 	}
+	else if (StrippedIniName.StartsWith(DebugPrefix))
+	{
+		StrippedIniName = IniName.Right(StrippedIniName.Len() - DebugPrefix.Len());
+	}
+	return StrippedIniName;
+}
+
+FString UOnlineHotfixManager::GetConfigFileNamePath(const FString& IniName)
+{
+	return FPaths::GeneratedConfigDir() + ANSI_TO_TCHAR(FPlatformProperties::PlatformName()) / IniName;
+}
+
+FConfigFile* UOnlineHotfixManager::GetConfigFile(const FString& IniName)
+{
+	FString StrippedIniName(GetStrippedConfigFileName(IniName));
 	FConfigFile* ConfigFile = nullptr;
 	// Look for the first matching INI file entry
 	for (TMap<FString, FConfigFile>::TIterator It(*GConfig); It; ++It)
@@ -378,7 +486,7 @@ FConfigFile* UOnlineHotfixManager::GetConfigFile(const FString& IniName)
 	// If not found, add this file to the config cache
 	if (ConfigFile == nullptr)
 	{
-		const FString IniNameWithPath = FPaths::GeneratedConfigDir() + StrippedIniName;
+		const FString IniNameWithPath = GetConfigFileNamePath(StrippedIniName);
 		FConfigFile Empty;
 		GConfig->SetFile(IniNameWithPath, &Empty);
 		ConfigFile = GConfig->Find(IniNameWithPath, false);
@@ -392,6 +500,8 @@ FConfigFile* UOnlineHotfixManager::GetConfigFile(const FString& IniName)
 bool UOnlineHotfixManager::HotfixIniFile(const FString& FileName, const FString& IniData)
 {
 	FConfigFile* ConfigFile = GetConfigFile(FileName);
+	// Store the original file so we can undo this later
+	BackupIniFile(FileName, ConfigFile);
 	// Merge the string into the config file
 	ConfigFile->CombineFromBuffer(IniData);
 	TArray<UClass*> Classes;
@@ -495,33 +605,81 @@ bool UOnlineHotfixManager::HotfixPakFile(const FCloudFileHeader& FileHeader)
 		return false;
 	}
 	FString PakLocation = FString::Printf(TEXT("%s/%s"), *GetCachedDirectory(), *FileHeader.DLName);
-	if (FCoreDelegates::OnMountPak.Execute(PakLocation, 0))
+	FPakFileVisitor Visitor;
+	if (FCoreDelegates::OnMountPak.Execute(PakLocation, 0, &Visitor))
 	{
 		MountedPakFiles.Add(FileHeader.DLName);
 		UE_LOG(LogHotfixManager, Log, TEXT("Hotfix mounted PAK file (%s)"), *FileHeader.FileName);
-		// @todo joeg - Merge any INI files that are part of this PAK file into the config cache
-		// Do we want/need this? Is it clear whether this will give us the right info?
+		int32 NumInisReloaded = 0;
+		const double StartTime = FPlatformTime::Seconds();
+		TArray<FString> IniList;
+		// Iterate through the pak file's contents for INI and asset reloading
+		for (auto& InternalPakFileName : Visitor.Files)
+		{
+			if (InternalPakFileName.EndsWith(TEXT(".ini")))
+			{
+				IniList.Add(InternalPakFileName);
+			}
+			else if (!bHotfixNeedsReload && InternalPakFileName.EndsWith(FPackageName::GetMapPackageExtension()))
+			{
+				bHotfixNeedsReload = IsMapLoaded(InternalPakFileName);
+			}
+		}
+		// Sort the INIs so they are processed consistently
+		IniList.Sort<FHotfixFileSortPredicate>(FHotfixFileSortPredicate(PlatformPrefix, DefaultPrefix));
+		// Now process the INIs in sorted order
+		for (auto& IniName : IniList)
+		{
+			HotfixPakIniFile(IniName);
+			NumInisReloaded++;
+		}
+		UE_LOG(LogHotfixManager, Log, TEXT("Processing pak file (%s) took %f seconds and resulted in (%d) INIs being reloaded"),
+			*FileHeader.FileName, FPlatformTime::Seconds() - StartTime, NumInisReloaded);
+#if !UE_BUILD_SHIPPING
+		if (bLogMountedPakContents)
+		{
+			UE_LOG(LogHotfixManager, Log, TEXT("Files in pak file (%s):"), *FileHeader.FileName);
+			for (auto& FileName : Visitor.Files)
+			{
+				UE_LOG(LogHotfixManager, Log, TEXT("\t\t%s"), *FileName);
+			}
+		}
+#endif
 		return true;
 	}
 	return false;
 }
 
+bool UOnlineHotfixManager::IsMapLoaded(const FString& MapName)
+{
+	FString MapPackageName(MapName.Left(MapName.Len() - 5));
+	MapPackageName = MapPackageName.Replace(*GameContentPath, TEXT("/Game"));
+	// If this map's UPackage exists, it is currently loaded
+	UPackage* MapPackage = FindObject<UPackage>(ANY_PACKAGE, *MapPackageName, true);
+	return MapPackage != nullptr;
+}
+
 bool UOnlineHotfixManager::HotfixPakIniFile(const FString& FileName)
 {
-	FConfigFile* ConfigFile = GetConfigFile(FileName);
-	ConfigFile->Combine(FileName);
+	FString StrippedName;
+	const double StartTime = FPlatformTime::Seconds();
+	// Need to strip off the PAK path
+	FileName.Split(TEXT("/"), nullptr, &StrippedName, ESearchCase::IgnoreCase, ESearchDir::FromEnd);
+	FConfigFile* ConfigFile = GetConfigFile(StrippedName);
+	if (!ConfigFile->Combine(FString(TEXT("../../../")) + FileName.Replace(*GameContentPath, TEXT("/Game"))))
+	{
+		UE_LOG(LogHotfixManager, Log, TEXT("Hotfix failed to merge INI (%s) found in a PAK file"), *FileName);
+		return false;
+	}
 	UE_LOG(LogHotfixManager, Log, TEXT("Hotfix merged INI (%s) found in a PAK file"), *FileName);
-
-	FName IniFileName(*FileName, FNAME_Find);
 	int32 NumObjectsReloaded = 0;
-	const double StartTime = FPlatformTime::Seconds();
 	// Now that we have a list of classes to update, we can iterate objects and
 	// reload if they match the INI file that was changed
 	for (FObjectIterator It; It; ++It)
 	{
 		UClass* Class = It->GetClass();
 		if (Class->HasAnyClassFlags(CLASS_Config) &&
-			Class->ClassConfigName == IniFileName)
+			Class->ClassConfigName == ConfigFile->Name)
 		{
 			// Force a reload of the config vars
 			It->ReloadConfig();
@@ -535,7 +693,7 @@ bool UOnlineHotfixManager::HotfixPakIniFile(const FString& FileName)
 
 const FString UOnlineHotfixManager::GetFriendlyNameFromDLName(const FString& DLName) const
 {
-	for (auto Header : HotfixFileList)
+	for (auto& Header : HotfixFileList)
 	{
 		if (Header.DLName == DLName)
 		{
@@ -547,12 +705,38 @@ const FString UOnlineHotfixManager::GetFriendlyNameFromDLName(const FString& DLN
 
 void UOnlineHotfixManager::UnmountHotfixFiles()
 {
-	// Unmount any hotfix files since we need to download them again
-	for (auto PakFile : MountedPakFiles)
+	if (MountedPakFiles.Num() == 0)
 	{
-		FCoreDelegates::OnUnmountPak.Execute(PakFile);
+		return;
+	}
+	// Unmount any changed hotfix files since we need to download them again
+	for (auto& FileHeader : ChangedHotfixFileList)
+	{
+		for (int32 Index = 0; Index < MountedPakFiles.Num(); Index++)
+		{
+			if (MountedPakFiles[Index] == FileHeader.DLName)
+			{
+				FCoreDelegates::OnUnmountPak.Execute(MountedPakFiles[Index]);
+				MountedPakFiles.RemoveAt(Index);
+				UE_LOG(LogHotfixManager, Log, TEXT("Hotfix unmounted PAK file (%s) so it can be redownloaded"), *FileHeader.FileName);
+				break;
+			}
+		}
+	}
+	// Unmount any removed hotfix files
+	for (auto& FileHeader : RemovedHotfixFileList)
+	{
+		for (int32 Index = 0; Index < MountedPakFiles.Num(); Index++)
+		{
+			if (MountedPakFiles[Index] == FileHeader.DLName)
+			{
+				FCoreDelegates::OnUnmountPak.Execute(MountedPakFiles[Index]);
+				MountedPakFiles.RemoveAt(Index);
+				UE_LOG(LogHotfixManager, Log, TEXT("Hotfix unmounted PAK file (%s) since it was removed from the hotfix set"), *FileHeader.FileName);
+				break;
+			}
+		}
 	}
-	MountedPakFiles.Empty();
 }
 
 FCloudFileHeader* UOnlineHotfixManager::GetFileHeaderFromDLName(const FString& FileName)
@@ -578,3 +762,160 @@ void UOnlineHotfixManager::OnReadFileProgress(const FString& FileName, uint64 By
 		UpdateProgress(0, Delta);
 	}
 }
+
+void UOnlineHotfixManager::BackupIniFile(const FString& IniName, const FConfigFile* ConfigFile)
+{
+	int32 AddAt = IniBackups.AddDefaulted();
+	IniBackups[AddAt].IniName = GetConfigFileNamePath(GetStrippedConfigFileName(IniName));
+	IniBackups[AddAt].ConfigData = *ConfigFile;
+	// There's a lack of deep copy related to the SourceConfigFile so null it out
+	IniBackups[AddAt].ConfigData.SourceConfigFile = nullptr;
+}
+
+void UOnlineHotfixManager::RestoreBackupIniFiles()
+{
+	if (IniBackups.Num() == 0)
+	{
+		return;
+	}
+	const double StartTime = FPlatformTime::Seconds();
+	TArray<FString> RestoredInis;
+	// Restore any changed INI files and build a list of which ones changed for UObject reloading below
+	for (auto& FileHeader : ChangedHotfixFileList)
+	{
+		if (FileHeader.FileName.EndsWith(TEXT(".INI")))
+		{
+			const FString ProcessedName = GetConfigFileNamePath(GetStrippedConfigFileName(FileHeader.FileName));
+			for (int32 Index = 0; Index < IniBackups.Num(); Index++)
+			{
+				if (IniBackups[Index].IniName == ProcessedName)
+				{
+					GConfig->SetFile(IniBackups[Index].IniName, &IniBackups[Index].ConfigData);
+					IniBackups.RemoveAt(Index);
+					RestoredInis.Add(ProcessedName);
+					break;
+				}
+			}
+		}
+	}
+	// Also restore any files that were previously part of the hotfix and now are not
+	for (auto& FileHeader : RemovedHotfixFileList)
+	{
+		if (FileHeader.FileName.EndsWith(TEXT(".INI")))
+		{
+			const FString ProcessedName = GetConfigFileNamePath(GetStrippedConfigFileName(FileHeader.FileName));
+			for (int32 Index = 0; Index < IniBackups.Num(); Index++)
+			{
+				if (IniBackups[Index].IniName == ProcessedName)
+				{
+					GConfig->SetFile(IniBackups[Index].IniName, &IniBackups[Index].ConfigData);
+					IniBackups.RemoveAt(Index);
+					RestoredInis.Add(ProcessedName);
+					break;
+				}
+			}
+		}
+	}
+	// Finally reload any UObjects that were affected by this
+	uint32 NumObjectsReloaded = 0;
+	if (RestoredInis.Num() > 0)
+	{
+		for (FObjectIterator It; It; ++It)
+		{
+			UClass* Class = It->GetClass();
+			if (Class->HasAnyClassFlags(CLASS_Config))
+			{
+				for (int32 Index = 0; Index < RestoredInis.Num(); Index++)
+				{
+					if (Class->GetConfigName() == RestoredInis[Index])
+					{
+						It->ReloadConfig();
+						NumObjectsReloaded++;
+						break;
+					}
+				}
+			}
+		}
+	}
+	UE_LOG(LogHotfixManager, Log, TEXT("Restoring config for %d changed INI files took %f seconds reloading %d objects"),
+		RestoredInis.Num(), FPlatformTime::Seconds() - StartTime, NumObjectsReloaded);
+}
+
+bool UOnlineHotfixManager::DoesHotfixFileMatchLastVersion(const FCloudFileHeader& Header)
+{
+	// Searches the last hotfix file array to determine if this header matches the previous
+	// This is an optimization to avoid applying files that haven't changed
+	for (auto& LastFileHeader : LastHotfixFileList)
+	{
+		if (Header.FileName == LastFileHeader.FileName && Header.Hash == LastFileHeader.Hash)
+		{
+			return true;
+		}
+	}
+	return false;
+}
+
+struct FHotfixManagerExec :
+	public FSelfRegisteringExec
+{
+	virtual bool Exec(UWorld* InWorld, const TCHAR* Cmd, FOutputDevice& Ar) override
+	{
+		if (FParse::Command(&Cmd, TEXT("HOTFIX")))
+		{
+			UOnlineHotfixManager* HotfixManager = UOnlineHotfixManager::Get(InWorld);
+			if (HotfixManager != nullptr)
+			{
+				HotfixManager->StartHotfixProcess();
+			}
+			return true;
+		}
+		else if (FParse::Command(&Cmd, TEXT("TESTHOTFIXSORT")))
+		{
+			TArray<FCloudFileHeader> TestList;
+			FCloudFileHeader Header;
+			Header.FileName = TEXT("SomeRandom.ini");
+			TestList.Add(Header);
+			Header.FileName = TEXT("pakchunk1-PS4_P.pak");
+			TestList.Add(Header);
+			Header.FileName = TEXT("EN_Game.locres");
+			TestList.Add(Header);
+			Header.FileName = TEXT("DefaultGame.ini");
+			TestList.Add(Header);
+			Header.FileName = TEXT("PS4_DefaultEngine.ini");
+			TestList.Add(Header);
+			Header.FileName = TEXT("DefaultEngine.ini");
+			TestList.Add(Header);
+			Header.FileName = TEXT("pakchunk0-PS4_P.pak");
+			TestList.Add(Header);
+			Header.FileName = TEXT("PS4_DefaultGame.ini");
+			TestList.Add(Header);
+			Header.FileName = TEXT("AnotherRandom.ini");
+			TestList.Add(Header);
+			TestList.Sort<FHotfixFileSortPredicate>(FHotfixFileSortPredicate(TEXT("PS4_"), TEXT("Default")));
+
+			UE_LOG(LogHotfixManager, Log, TEXT("Hotfixing sort is:"));
+			for (auto& FileHeader : TestList)
+			{
+				UE_LOG(LogHotfixManager, Log, TEXT("\t%s"), *FileHeader.FileName);
+			}
+
+			TArray<FString> TestList2;
+			TestList2.Add(TEXT("SomeRandom.ini"));
+			TestList2.Add(TEXT("DefaultGame.ini"));
+			TestList2.Add(TEXT("PS4_DefaultEngine.ini"));
+			TestList2.Add(TEXT("DefaultEngine.ini"));
+			TestList2.Add(TEXT("PS4_DefaultGame.ini"));
+			TestList2.Add(TEXT("AnotherRandom.ini"));
+			TestList2.Sort<FHotfixFileSortPredicate>(FHotfixFileSortPredicate(TEXT("PS4_"), TEXT("Default")));
+
+			UE_LOG(LogHotfixManager, Log, TEXT("Hotfixing PAK INI file sort is:"));
+			for (auto& IniName : TestList2)
+			{
+				UE_LOG(LogHotfixManager, Log, TEXT("\t%s"), *IniName);
+			}
+			return true;
+		}
+		return false;
+	}
+};
+static FHotfixManagerExec HotfixManagerExec;
diff --git a/Engine/Source/Runtime/Online/IOS/OnlineSubsystemIOS/Private/OnlineSubsystemIOS.cpp b/Engine/Source/Runtime/Online/IOS/OnlineSubsystemIOS/Private/OnlineSubsystemIOS.cpp
--- a/Engine/Source/Runtime/Online/IOS/OnlineSubsystemIOS/Private/OnlineSubsystemIOS.cpp
+++ b/Engine/Source/Runtime/Online/IOS/OnlineSubsystemIOS/Private/OnlineSubsystemIOS.cpp
@@ -18,6 +18,11 @@ IOnlineFriendsPtr FOnlineSubsystemIOS::GetFriendsInterface() const
 	return FriendsInterface;
 }
 
+IMessageSanitizerPtr FOnlineSubsystemIOS::GetMessageSanitizerInterface() const
+{
+	return nullptr;
+}
+
 IOnlinePartyPtr FOnlineSubsystemIOS::GetPartyInterface() const
 {
 	return nullptr;
diff --git a/Engine/Source/Runtime/Online/IOS/OnlineSubsystemIOS/Public/OnlineSubsystemIOS.h b/Engine/Source/Runtime/Online/IOS/OnlineSubsystemIOS/Public/OnlineSubsystemIOS.h
--- a/Engine/Source/Runtime/Online/IOS/OnlineSubsystemIOS/Public/OnlineSubsystemIOS.h
+++ b/Engine/Source/Runtime/Online/IOS/OnlineSubsystemIOS/Public/OnlineSubsystemIOS.h
@@ -21,6 +21,7 @@ public:
 	//~ Begin IOnlineSubsystem Interface
 	virtual IOnlineSessionPtr GetSessionInterface() const override;
 	virtual IOnlineFriendsPtr GetFriendsInterface() const override;
+	virtual IMessageSanitizerPtr GetMessageSanitizerInterface() const override;
 	virtual IOnlinePartyPtr GetPartyInterface() const override;
 	virtual IOnlineGroupsPtr GetGroupsInterface() const override;
 	virtual IOnlineSharedCloudPtr GetSharedCloudInterface() const override;
diff --git a/Engine/Source/Runtime/Online/OnlineFramework/Party/Private/PartyGameState.cpp b/Engine/Source/Runtime/Online/OnlineFramework/Party/Private/PartyGameState.cpp
--- a/Engine/Source/Runtime/Online/OnlineFramework/Party/Private/PartyGameState.cpp
+++ b/Engine/Source/Runtime/Online/OnlineFramework/Party/Private/PartyGameState.cpp
@@ -998,6 +998,12 @@ void UPartyGameState::GetAllPartyMembers(TArray<UPartyMemberState*>& PartyMember
 	PartyMembersState.GenerateValueArray(PartyMembers);
 }
 
+FChatRoomId UPartyGameState::GetChatRoomID() const
+{
+	ensure(OssParty.IsValid());
+	return OssParty->RoomId;
+}
+
 void UPartyGameState::OnPartyMemberPromoted(const FUniqueNetId& LocalUserId, const FOnlinePartyId& InPartyId, const FUniqueNetId& InMemberId, const EPromoteMemberCompletionResult Result)
 {
 	FString PartyIdDebugString = InPartyId.ToDebugString();
diff --git a/Engine/Source/Runtime/Online/OnlineFramework/Party/Public/PartyGameState.h b/Engine/Source/Runtime/Online/OnlineFramework/Party/Public/PartyGameState.h
--- a/Engine/Source/Runtime/Online/OnlineFramework/Party/Public/PartyGameState.h
+++ b/Engine/Source/Runtime/Online/OnlineFramework/Party/Public/PartyGameState.h
@@ -277,6 +277,13 @@ public:
 	 */
 	void GetAllPartyMembers(TArray<UPartyMemberState*>& PartyMembers) const;
 
+	/**
+	 * Get the chat room ID for this party
+	 *
+	 * @return The chat room ID
+	 */
+	FChatRoomId GetChatRoomID() const;
+
 	template< typename TPartyMemberState >
 	void GetTypedPartyMembers(TArray< TPartyMemberState* >& PartyMembers) const
 	{
diff --git a/Engine/Source/Runtime/Online/OnlineFramework/Qos/Private/QosEvaluator.cpp b/Engine/Source/Runtime/Online/OnlineFramework/Qos/Private/QosEvaluator.cpp
--- a/Engine/Source/Runtime/Online/OnlineFramework/Qos/Private/QosEvaluator.cpp
+++ b/Engine/Source/Runtime/Online/OnlineFramework/Qos/Private/QosEvaluator.cpp
@@ -7,26 +7,6 @@
 #include "QosStats.h"
 #include "OnlineSubsystemUtils.h"
 
-FOnlineSessionSettingsQos::FOnlineSessionSettingsQos(bool bInIsDedicated) 
-{
-	NumPublicConnections = 1;
-	NumPrivateConnections = 0;
-
-	bIsLANMatch = false;
-	bShouldAdvertise = true;
-	bAllowJoinInProgress = true;
-	bAllowInvites = true;
-	bUsesPresence = false;
-	bAllowJoinViaPresence = true;
-	bAllowJoinViaPresenceFriendsOnly = false;
-
-	FString GameModeStr(GAMEMODE_QOS);
-	Set(SETTING_GAMEMODE, GameModeStr, EOnlineDataAdvertisementType::ViaOnlineService);
-	Set(SETTING_QOS, 1, EOnlineDataAdvertisementType::ViaOnlineService);
-	Set(SETTING_REGION, UQosEvaluator::GetDefaultRegionString(), EOnlineDataAdvertisementType::ViaOnlineService);
-	bIsDedicated = bInIsDedicated;
-}
-
 FOnlineSessionSearchQos::FOnlineSessionSearchQos()
 {
 	bIsLanQuery = false;
@@ -37,6 +17,11 @@ FOnlineSessionSearchQos::FOnlineSessionSearchQos()
 	QuerySettings.Set(SETTING_QOS, 1, EOnlineComparisonOp::Equals);
 }
 
+TSharedPtr<FOnlineSessionSettings> FOnlineSessionSearchQos::GetDefaultSessionSettings() const
+{
+	return MakeShareable(new FOnlineSessionSettingsQos());
+}
+
 void FOnlineSessionSearchQos::SortSearchResults()
 {
 	TMap<FString, int32> RegionCounts;
@@ -73,33 +58,29 @@ void FOnlineSessionSearchQos::SortSearchResults()
 	}
 }
 
-
-UQosEvaluator::UQosEvaluator(const FObjectInitializer& ObjectInitializer) :
-	Super(ObjectInitializer),
-	QosSearchQuery(nullptr),
-	QosBeaconClient(nullptr),
-	bInProgress(false),
-	bCancelOperation(false),
-	AnalyticsProvider(nullptr),
-	QosStats(nullptr)
+FQosEvaluator::FQosEvaluator(UWorld* World) 
+:	ParentWorld(World)
+,	QosSearchQuery(nullptr)
+,	QosBeaconClient(nullptr)
+,	bInProgress(false)
+,	bCancelOperation(false)
+,	AnalyticsProvider(nullptr)
+,	QosStats(nullptr)
 {
-	if (!HasAnyFlags(RF_ClassDefaultObject))
-	{
-		SearchState.DatacenterId = GetDefaultRegionString();
-	}
+	check(World != nullptr);
 }
 
-void UQosEvaluator::SetAnalyticsProvider(TSharedPtr<IAnalyticsProvider> InAnalyticsProvider)
+void FQosEvaluator::SetAnalyticsProvider(TSharedPtr<IAnalyticsProvider> InAnalyticsProvider)
 {
 	AnalyticsProvider = InAnalyticsProvider;
 }
 
-void UQosEvaluator::Cancel()
+void FQosEvaluator::Cancel()
 {
 	bCancelOperation = true;
 }
 
-void UQosEvaluator::FindDatacenters(int32 ControllerId, const FOnQosSearchComplete& InCompletionDelegate)
+void FQosEvaluator::FindDatacenters(int32 ControllerId, const FOnQosSearchComplete& InCompletionDelegate)
 {
 	EQosCompletionResult Result = EQosCompletionResult::Failure;
 
@@ -111,53 +92,53 @@ void UQosEvaluator::FindDatacenters(int32 ControllerId, const FOnQosSearchComple
 
 		StartAnalytics();
 
-		bool bForcedCached = QOSConsoleVariables::CVarForceCached.GetValueOnGameThread() != 0;
-		bool bOverrideTimestamp = QOSConsoleVariables::CVarOverrideTimestamp.GetValueOnGameThread() != 0;
-		if (!bForcedCached && ((CurTimestamp - SearchState.LastDatacenterIdTimestamp >= DATACENTERQUERY_INTERVAL) || bOverrideTimestamp))
+		IOnlineSubsystem* OnlineSub = Online::GetSubsystem(GetWorld());
+		if (OnlineSub)
 		{
-			IOnlineSubsystem* OnlineSub = Online::GetSubsystem(GetWorld());
-			if (OnlineSub)
+			IOnlineSessionPtr SessionInt = OnlineSub->GetSessionInterface();
+			if (SessionInt.IsValid())
 			{
-				IOnlineSessionPtr SessionInt = OnlineSub->GetSessionInterface();
-				if (SessionInt.IsValid())
-				{
-					const TSharedRef<FOnlineSessionSearch> QosSearchParams = MakeShareable(new FOnlineSessionSearchQos());
-					QosSearchQuery = QosSearchParams;
-
-					if ( OnFindDatacentersCompleteDelegateHandle.IsValid() )
-					{
-						SessionInt->ClearOnFindSessionsCompleteDelegate_Handle( OnFindDatacentersCompleteDelegateHandle );
-					}
-					// now make a new bind
-					{
-						FOnFindSessionsCompleteDelegate OnFindDatacentersCompleteDelegate = FOnFindSessionsCompleteDelegate::CreateUObject( this, &ThisClass::OnFindDatacentersComplete, InCompletionDelegate );
-						OnFindDatacentersCompleteDelegateHandle = SessionInt->AddOnFindSessionsCompleteDelegate_Handle( OnFindDatacentersCompleteDelegate );
-					}
+				const TSharedRef<FOnlineSessionSearch> QosSearchParams = MakeShareable(new FOnlineSessionSearchQos());
+				QosSearchQuery = QosSearchParams;
 
-					SessionInt->FindSessions(ControllerId, QosSearchParams);
-					Result = EQosCompletionResult::Success;
+				if ( OnFindDatacentersCompleteDelegateHandle.IsValid() )
+				{
+					SessionInt->ClearOnFindSessionsCompleteDelegate_Handle( OnFindDatacentersCompleteDelegateHandle );
+				}
+				// now make a new bind
+				{
+					FOnFindSessionsCompleteDelegate OnFindDatacentersCompleteDelegate = FOnFindSessionsCompleteDelegate::CreateSP( this, &FQosEvaluator::OnFindDatacentersComplete, InCompletionDelegate );
+					OnFindDatacentersCompleteDelegateHandle = SessionInt->AddOnFindSessionsCompleteDelegate_Handle( OnFindDatacentersCompleteDelegate );
 				}
+
+				SessionInt->FindSessions(ControllerId, QosSearchParams);
+				Result = EQosCompletionResult::Success;
 			}
 		}
-		else
-		{
-			Result = EQosCompletionResult::Cached;
-		}
 
 		if (Result != EQosCompletionResult::Success)
 		{
-			FinalizeDatacenterResult(InCompletionDelegate, Result);
+			TWeakPtr<FQosEvaluator> WeakThisCap(AsShared());
+			GetWorldTimerManager().SetTimerForNextTick(FTimerDelegate::CreateLambda([InCompletionDelegate, Result, WeakThisCap]() {
+				auto StrongThis = WeakThisCap.Pin();
+				if (StrongThis.IsValid())
+				{
+					StrongThis->FinalizeDatacenterResult(InCompletionDelegate, Result, TArray<FQosRegionInfo>());
+				}
+			}));
 		}
 	}
 	else
 	{
 		UE_LOG(LogQos, Log, TEXT("Qos evaluation already in progress, ignoring"));
 		// Just trigger delegate now (Finalize resets state vars)
-		InCompletionDelegate.ExecuteIfBound(Result, SearchState.DatacenterId);
+		GetWorldTimerManager().SetTimerForNextTick(FTimerDelegate::CreateLambda([InCompletionDelegate]() {
+			InCompletionDelegate.ExecuteIfBound(EQosCompletionResult::Failure, TArray<FQosRegionInfo>());
+		}));
 	}
 }
 
-void UQosEvaluator::OnFindDatacentersComplete(bool bWasSuccessful, FOnQosSearchComplete InCompletionDelegate)
+void FQosEvaluator::OnFindDatacentersComplete(bool bWasSuccessful, FOnQosSearchComplete InCompletionDelegate)
 {
 	IOnlineSubsystem* OnlineSub = Online::GetSubsystem(GetWorld());
 	if (OnlineSub)
@@ -177,7 +158,7 @@ void UQosEvaluator::OnFindDatacentersComplete(bool bWasSuccessful, FOnQosSearchC
 			// Mark not in progress so it falls into EvaluateServerPing
 			bInProgress = false;
 
-			FOnQosPingEvalComplete CompletionDelegate = FOnQosPingEvalComplete::CreateUObject(this, &ThisClass::OnEvaluateForDatacenterComplete, InCompletionDelegate);
+			FOnQosPingEvalComplete CompletionDelegate = FOnQosPingEvalComplete::CreateSP(this, &FQosEvaluator::OnEvaluateForDatacenterComplete, InCompletionDelegate);
 			EvaluateServerPing(QosSearchQuery, CompletionDelegate);
 		}
 		else
@@ -192,8 +173,9 @@ void UQosEvaluator::OnFindDatacentersComplete(bool bWasSuccessful, FOnQosSearchC
 	}
 }
 
-void UQosEvaluator::OnEvaluateForDatacenterComplete(EQosCompletionResult Result, FOnQosSearchComplete InCompletionDelegate)
+void FQosEvaluator::OnEvaluateForDatacenterComplete(EQosCompletionResult Result, FOnQosSearchComplete InCompletionDelegate)
 {
+	TArray<FQosRegionInfo> RegionInfo;
 	if (Result == EQosCompletionResult::Success)
 	{
 		if (QosSearchQuery.IsValid() && QosSearchQuery->SearchResults.Num() > 0)
@@ -236,17 +218,14 @@ void UQosEvaluator::OnEvaluateForDatacenterComplete(EQosCompletionResult Result,
 				}
 			}
 
-			int32 BestAvgPing = MAX_QUERY_PING;
 			for (auto& QosResult : RegionAggregates)
 			{
 				if (QosResult.Value.NumResults > 0)
 				{
-					QosResult.Value.AvgPing = QosResult.Value.TotalPingInMs / QosResult.Value.NumResults;
-					if (QosResult.Value.AvgPing < BestAvgPing)
-					{
-						BestAvgPing = QosResult.Value.AvgPing;
-						SearchState.DatacenterId = QosResult.Key;
-					}
+					FQosRegionInfo Region;
+					Region.RegionId = QosResult.Key;
+					Region.PingMs = QosResult.Value.TotalPingInMs / QosResult.Value.NumResults;
+					RegionInfo.Add(Region);
 				}
 
 				UE_LOG(LogQos, Verbose, TEXT("Region[%s] Avg: %d Num: %d"), *QosResult.Key, QosResult.Value.AvgPing, QosResult.Value.NumResults);
@@ -256,8 +235,6 @@ void UQosEvaluator::OnEvaluateForDatacenterComplete(EQosCompletionResult Result,
 					QosStats->RecordRegionInfo(QosResult.Key, QosResult.Value.AvgPing, QosResult.Value.NumResults);
 				}
 			}
-
-			SearchState.LastDatacenterIdTimestamp = FPlatformTime::Seconds();
 		}
 		else
 		{
@@ -265,39 +242,28 @@ void UQosEvaluator::OnEvaluateForDatacenterComplete(EQosCompletionResult Result,
 		}
 	}
 
-	FinalizeDatacenterResult(InCompletionDelegate, Result);
+	FinalizeDatacenterResult(InCompletionDelegate, Result, RegionInfo);
 }
 
-void UQosEvaluator::FinalizeDatacenterResult(const FOnQosSearchComplete& InCompletionDelegate, EQosCompletionResult CompletionResult)
+void FQosEvaluator::FinalizeDatacenterResult(const FOnQosSearchComplete& InCompletionDelegate, EQosCompletionResult CompletionResult, const TArray<FQosRegionInfo>& RegionInfo)
 {
-	bool bForceDefault = QOSConsoleVariables::CVarForceDefaultRegion.GetValueOnGameThread() != 0;
-	FString ForceRegion = QOSConsoleVariables::CVarForceRegion.GetValueOnGameThread();
-
-	if (!ForceRegion.IsEmpty())
-	{
-		SearchState.DatacenterId = ForceRegion.ToUpper();
-	}
-	else if (bForceDefault)
-	{
-		SearchState.DatacenterId = GetDefaultRegionString();
-	}
-
-	UE_LOG(LogQos, Log, TEXT("Datacenter evaluation complete. Result: %s Region: %s%s"), 
-		ToString(CompletionResult), *SearchState.DatacenterId, ForceRegion.IsEmpty() ? (bForceDefault ? TEXT("[Default]") : TEXT("")) : TEXT("[Forced]")
-		);
-
+	UE_LOG(LogQos, Log, TEXT("Datacenter evaluation complete. Result: %s "), ToString(CompletionResult));
 
 	EndAnalytics(CompletionResult);
 
 	// Broadcast this data next frame
-	FString& DatacenterId = SearchState.DatacenterId;
-	FTimerDelegate TimerDelegate = FTimerDelegate::CreateLambda(
-		[InCompletionDelegate, CompletionResult, DatacenterId, this](){ InCompletionDelegate.ExecuteIfBound(CompletionResult, DatacenterId); ResetSearchVars(); }
-	);
-	GetWorldTimerManager().SetTimerForNextTick(TimerDelegate);
+	TWeakPtr<FQosEvaluator> WeakThisCap(AsShared());
+	GetWorldTimerManager().SetTimerForNextTick(FTimerDelegate::CreateLambda([InCompletionDelegate, CompletionResult, RegionInfo, WeakThisCap]() {
+		InCompletionDelegate.ExecuteIfBound(CompletionResult, RegionInfo); 
+		auto StrongThis = WeakThisCap.Pin();
+		if (StrongThis.IsValid())
+		{
+			StrongThis->ResetSearchVars();
+		}
+	}));
 }
 
-void UQosEvaluator::EvaluateServerPing(TSharedPtr<FOnlineSessionSearch>& SearchResults, const FOnQosPingEvalComplete& InCompletionDelegate)
+void FQosEvaluator::EvaluateServerPing(TSharedPtr<FOnlineSessionSearch>& SearchResults, const FOnQosPingEvalComplete& InCompletionDelegate)
 {
 	if (!bInProgress)
 	{
@@ -327,7 +293,7 @@ void UQosEvaluator::EvaluateServerPing(TSharedPtr<FOnlineSessionSearch>& SearchR
 	}
 }
 
-void UQosEvaluator::ContinuePingServers()
+void FQosEvaluator::ContinuePingServers()
 {
 	if (!bCancelOperation)
 	{
@@ -338,8 +304,8 @@ void UQosEvaluator::ContinuePingServers()
 			UWorld* World = GetWorld();
 			QosBeaconClient = World->SpawnActor<AQosBeaconClient>(AQosBeaconClient::StaticClass());
 
-			QosBeaconClient->OnQosRequestComplete().BindUObject(this, &ThisClass::OnQosRequestComplete);
-			QosBeaconClient->OnHostConnectionFailure().BindUObject(this, &ThisClass::OnQosConnectionFailure);
+			QosBeaconClient->OnQosRequestComplete().BindSP(this, &FQosEvaluator::OnQosRequestComplete);
+			QosBeaconClient->OnHostConnectionFailure().BindSP(this, &FQosEvaluator::OnQosConnectionFailure);
 			QosBeaconClient->SendQosRequest(QosSearchQuery->SearchResults[CurrentSearchPass.CurrentSessionIdx]);
 		}
 		else
@@ -355,19 +321,25 @@ void UQosEvaluator::ContinuePingServers()
 	}
 }
 
-void UQosEvaluator::FinalizePingServers(EQosCompletionResult Result)
+void FQosEvaluator::FinalizePingServers(EQosCompletionResult Result)
 {
 	UE_LOG(LogQos, Log, TEXT("Ping evaluation complete. Result:%s"), ToString(Result));
 
-	FTimerDelegate TimerDelegate = FTimerDelegate::CreateLambda(
-		[Result, this](){ OnQosPingEvalComplete.ExecuteIfBound(Result); ResetSearchVars(); }
-	);
+	TWeakPtr<FQosEvaluator> WeakThisCap(AsShared());
+	FTimerDelegate TimerDelegate = FTimerDelegate::CreateLambda([Result, WeakThisCap]() {
+		auto StrongThis = WeakThisCap.Pin();
+		if (StrongThis.IsValid())
+		{
+			StrongThis->OnQosPingEvalComplete.ExecuteIfBound(Result);
+			StrongThis->ResetSearchVars();
+		}
+	});
 	GetWorldTimerManager().SetTimerForNextTick(TimerDelegate);
 }
 
-void UQosEvaluator::OnQosRequestComplete(EQosResponseType QosResponse, int32 ResponseTime)
+void FQosEvaluator::OnQosRequestComplete(EQosResponseType QosResponse, int32 ResponseTime)
 {
-	if (QosBeaconClient)
+	if (QosBeaconClient.IsValid())
 	{
 		DestroyClientBeacons();
 	}
@@ -392,16 +364,16 @@ void UQosEvaluator::OnQosRequestComplete(EQosResponseType QosResponse, int32 Res
 		}
 		
 		// Cancel operation will occur next tick if applicable
-		GetWorldTimerManager().SetTimerForNextTick(this, &ThisClass::ContinuePingServers);
+		GetWorldTimerManager().SetTimerForNextTick(FTimerDelegate::CreateSP(this, &FQosEvaluator::ContinuePingServers));
 	}
 }
 
-void UQosEvaluator::OnQosConnectionFailure()
+void FQosEvaluator::OnQosConnectionFailure()
 {
 	OnQosRequestComplete(EQosResponseType::Failure, MAX_QUERY_PING);
 }
 
-void UQosEvaluator::ResetSearchVars()
+void FQosEvaluator::ResetSearchVars()
 {
 	bInProgress = false;
 	bCancelOperation = false;
@@ -410,18 +382,18 @@ void UQosEvaluator::ResetSearchVars()
 	QosSearchQuery = nullptr;
 }
 
-void UQosEvaluator::DestroyClientBeacons()
+void FQosEvaluator::DestroyClientBeacons()
 {
-	if (QosBeaconClient)
+	if (QosBeaconClient.IsValid())
 	{
 		QosBeaconClient->OnQosRequestComplete().Unbind();
 		QosBeaconClient->OnHostConnectionFailure().Unbind();
 		QosBeaconClient->DestroyBeacon();
-		QosBeaconClient = nullptr;
+		QosBeaconClient.Reset();
 	}
 }
 
-void UQosEvaluator::StartAnalytics()
+void FQosEvaluator::StartAnalytics()
 {
 	if (AnalyticsProvider.IsValid())
 	{
@@ -431,38 +403,19 @@ void UQosEvaluator::StartAnalytics()
 	}
 }
 
-void UQosEvaluator::EndAnalytics(EQosCompletionResult CompletionResult)
+void FQosEvaluator::EndAnalytics(EQosCompletionResult CompletionResult)
 {
 	if (QosStats.IsValid())
 	{
 		if (CompletionResult != EQosCompletionResult::Canceled)
 		{
-			bool bForceDefault = QOSConsoleVariables::CVarForceDefaultRegion.GetValueOnGameThread() != 0;
-			FString ForceRegion = QOSConsoleVariables::CVarForceRegion.GetValueOnGameThread();
-
 			EDatacenterResultType ResultType = EDatacenterResultType::Failure;
 			if (CompletionResult != EQosCompletionResult::Failure)
 			{
-				if (CompletionResult == EQosCompletionResult::Cached)
-				{
-					ResultType = EDatacenterResultType::Cached;
-				}
-				else
-				{
-					ResultType = EDatacenterResultType::Normal;
-				}
-
-				if (!ForceRegion.IsEmpty())
-				{
-					ResultType = EDatacenterResultType::Forced;
-				}
-				else if (bForceDefault)
-				{
-					ResultType = EDatacenterResultType::Default;
-				}
+				ResultType = EDatacenterResultType::Normal;
 			}
 
-			QosStats->EndQosPass(ResultType, SearchState.DatacenterId);
+			QosStats->EndQosPass(ResultType);
 			QosStats->Upload(AnalyticsProvider);
 		}
 
@@ -470,20 +423,14 @@ void UQosEvaluator::EndAnalytics(EQosCompletionResult CompletionResult)
 	}
 }
 
-const FString& UQosEvaluator::GetDefaultRegionString()
-{
-	return FQosInterface::GetDefaultRegionString();
-}
-
-UWorld* UQosEvaluator::GetWorld() const
+UWorld* FQosEvaluator::GetWorld() const
 {
-	check(GEngine);
-	check(GEngine->GameViewport);
-	UGameInstance* GameInstance = GEngine->GameViewport->GetGameInstance();
-	return GameInstance->GetWorld();
+	UWorld* World = ParentWorld.Get();
+	check(World);
+	return World;
 }
 
-FTimerManager& UQosEvaluator::GetWorldTimerManager() const
+FTimerManager& FQosEvaluator::GetWorldTimerManager() const
 {
 	return GetWorld()->GetTimerManager();
 }
