Copying //UE4/Orion-Staging to //UE4/Main (Origin: //Orion/Dev-General [REDACTED_HANDLE])

diff --git a/Engine/Source/Runtime/Online/OnlineFramework/Qos/Private/QosInterface.cpp b/Engine/Source/Runtime/Online/OnlineFramework/Qos/Private/QosInterface.cpp
--- a/Engine/Source/Runtime/Online/OnlineFramework/Qos/Private/QosInterface.cpp
+++ b/Engine/Source/Runtime/Online/OnlineFramework/Qos/Private/QosInterface.cpp
@@ -3,109 +3,246 @@
 #include "QosPrivatePCH.h"
 #include "QosInterface.h"
 #include "QosEvaluator.h"
+#include "QosModule.h"
 
-namespace QOSConsoleVariables
+FOnlineSessionSettingsQos::FOnlineSessionSettingsQos(bool bInIsDedicated)
 {
-	// CVars
-	TAutoConsoleVariable<FString> CVarForceRegion(
-		TEXT("QOS.ForceRegion"),
-		TEXT(""),
-		TEXT("Force a given region id\n")
-		TEXT("String region id"),
-		ECVF_Default);
-
-	TAutoConsoleVariable<int32> CVarForceDefaultRegion(
-		TEXT("QOS.UseDefaultRegion"),
-		0,
-		TEXT("Forces default region\n")
-		TEXT("1 enables default region, 0 disables."),
-		ECVF_Default);
-
-	TAutoConsoleVariable<int32> CVarForceCached(
-		TEXT("QOS.ForceCached"),
-		0,
-		TEXT("Forced cached value\n")
-		TEXT("1 Enables, 0 Disables."),
-		ECVF_Default);
-
-	TAutoConsoleVariable<int32> CVarOverrideTimestamp(
-		TEXT("QOS.ForceEval"),
-		0,
-		TEXT("Ignore Qos datacenter caching\n")
-		TEXT("1 Enables, 0 Disables."),
-		ECVF_Default);
-	
-}
+	NumPublicConnections = 1;
+	NumPrivateConnections = 0;
 
-FQosInterface::FQosInterface()
-{
+	bIsLANMatch = false;
+	bShouldAdvertise = true;
+	bAllowJoinInProgress = true;
+	bAllowInvites = true;
+	bUsesPresence = false;
+	bAllowJoinViaPresence = true;
+	bAllowJoinViaPresenceFriendsOnly = false;
 
+	FString GameModeStr(GAMEMODE_QOS);
+	Set(SETTING_GAMEMODE, GameModeStr, EOnlineDataAdvertisementType::ViaOnlineService);
+	Set(SETTING_QOS, 1, EOnlineDataAdvertisementType::ViaOnlineService);
+	Set(SETTING_REGION, FQosInterface::Get()->GetRegionId(), EOnlineDataAdvertisementType::ViaOnlineService);
+	bIsDedicated = bInIsDedicated;
 }
 
-FQosInterface::~FQosInterface()
+//static 
+TSharedRef<FQosInterface> FQosInterface::Get()
 {
-
+	return FQosModule::Get().GetQosInterface();
 }
 
-UQosEvaluator* FQosInterface::CreateQosEvaluator()
+FQosInterface::FQosInterface()
+: MinimumPingMs(250)
+, QosEvalResult(EQosCompletionResult::Invalid)
 {
-	return NewObject<UQosEvaluator>();
+	check(GConfig);
+	GConfig->GetInt(TEXT("Qos"), TEXT("MinimumPingMs"), MinimumPingMs, GGameIni);
+	GConfig->GetString(TEXT("Qos"), TEXT("ForceRegionId"), ForceRegionId, GGameIni);
+
+	// get a forced region id from the command line as an override
+	FParse::Value(FCommandLine::Get(), TEXT("McpRegion="), ForceRegionId);
 }
 
-const FString& FQosInterface::GetDefaultRegionString()
+// static
+FString FQosInterface::GetDatacenterId()
 {
-	static FString RegionString = TEXT("");
-	if (RegionString.IsEmpty())
+	struct FDcidInfo
 	{
-		FString OverrideRegion;
-		if (FParse::Value(FCommandLine::Get(), TEXT("McpRegion="), OverrideRegion))
-		{
-			// Region specified on command line
-			RegionString = OverrideRegion.ToUpper();
-		}
-		else
+		FDcidInfo()
 		{
-			FString DefaultRegion;
-			if (GConfig->GetString(TEXT("Qos"), TEXT("DefaultRegion"), DefaultRegion, GGameIni))
+			FString OverrideDCID;
+			if (FParse::Value(FCommandLine::Get(), TEXT("DCID="), OverrideDCID))
 			{
-				// Region specified in ini file
-				RegionString = DefaultRegion.ToUpper();
+				// Region specified on command line
+				DCIDString = OverrideDCID.ToUpper();
 			}
 			else
 			{
-				// No Region specified. Assume USA.
-				RegionString = TEXT("USA");
+				FString DefaultDCID;
+				check(GConfig);
+				if (GConfig->GetString(TEXT("Qos"), TEXT("DCID"), DefaultDCID, GGameIni))
+				{
+					// Region specified in ini file
+					DCIDString = DefaultDCID.ToUpper();
+				}
 			}
 		}
+
+		FString DCIDString;
+	};
+	static FDcidInfo DCID;
+	return DCID.DCIDString;
+}
+
+void FQosInterface::BeginQosEvaluation(UWorld* World, const TSharedPtr<IAnalyticsProvider>& AnalyticsProvider, const FSimpleDelegate& OnComplete)
+{
+	check(World);
+
+	// no point doing the qos tests at all if we're forcing this
+	if (!ForceRegionId.IsEmpty())
+	{
+		World->GetTimerManager().SetTimerForNextTick(FTimerDelegate::CreateLambda([OnComplete]() {
+			OnComplete.ExecuteIfBound();
+		}));
+		return;
+	}
+
+	// add to the completion delegate
+	OnQosEvalCompleteDelegate.Add(OnComplete);
+
+	// if we're already evaluating, simply return
+	if (Evaluator.IsValid())
+	{
+		return;
 	}
 
-	return RegionString;
+	// create a new evaluator and start the process of running
+	Evaluator = MakeShareable(new FQosEvaluator(World));
+	Evaluator->SetAnalyticsProvider(AnalyticsProvider);
+	Evaluator->FindDatacenters(0, FOnQosSearchComplete::CreateSP(this, &FQosInterface::OnQosEvaluationComplete));
+}
+
+bool FQosInterface::DidRegionTestSucceed() const
+{
+	return QosEvalResult == EQosCompletionResult::Success;
 }
 
-const FString& FQosInterface::GetDatacenterId()
+void FQosInterface::OnQosEvaluationComplete(EQosCompletionResult Result, const TArray<FQosRegionInfo>& RegionInfo)
 {
-	static bool bDCIDCheck = false;
-	static FString DCIDString = TEXT("");
-	if (!bDCIDCheck)
+	// toss this object
+	Evaluator.Reset();
+	QosEvalResult = Result;
+	SelectedRegion.Empty();
+	RegionOptions.Empty(RegionInfo.Num());
+
+	// treat lack of any regions as a failure
+	if (RegionInfo.Num() <= 0)
+	{
+		QosEvalResult = EQosCompletionResult::Failure;
+	}
+
+	if (QosEvalResult == EQosCompletionResult::Success)
 	{
-		FString OverrideDCID;
-		if (FParse::Value(FCommandLine::Get(), TEXT("DCID="), OverrideDCID))
+		// copy over region info that meets our minimum ping
+		for (const FQosRegionInfo& Region : RegionInfo)
 		{
-			// Region specified on command line
-			DCIDString = OverrideDCID.ToUpper();
+			if (MinimumPingMs <= 0 || Region.PingMs <= MinimumPingMs)
+			{
+				UE_LOG(LogQos, Log, TEXT("Region: %s Ping: %d ACCEPTED"), *Region.RegionId, Region.PingMs);
+				RegionOptions.Add(Region);
+			}
+			else
+			{
+				UE_LOG(LogQos, Log, TEXT("Region: %s Ping: %d REJECTED"), *Region.RegionId, Region.PingMs);
+			}
 		}
-		else
+
+		if (RegionOptions.Num() > 0)
 		{
-			FString DefaultDCID;
-			if (GConfig->GetString(TEXT("Qos"), TEXT("DCID"), DefaultDCID, GGameIni))
+			// try to select the default region
+			if (!SetSelectedRegion(GetSavedRegionId()))
 			{
-				// Region specified in ini file
-				DCIDString = DefaultDCID.ToUpper();
+				// try to select the lowest ping
+				int32 BestPing = INT_MAX;
+				FString BestRegionId;
+				for (const FQosRegionInfo& Region : RegionOptions)
+				{
+					if (Region.PingMs < BestPing)
+					{
+						BestPing = Region.PingMs;
+						BestRegionId = Region.RegionId;
+					}
+				}
+				SetSelectedRegion(BestRegionId);
 			}
 		}
+	}
+	
+	// fire notifications
+	TArray<FSimpleDelegate> NotifyList = OnQosEvalCompleteDelegate;
+	OnQosEvalCompleteDelegate.Empty();
+	for (const auto& Callback : NotifyList)
+	{
+		Callback.ExecuteIfBound();
+	}
+}
+
+FString FQosInterface::GetRegionId() const
+{
+	if (!ForceRegionId.IsEmpty())
+	{
+		// we may have updated INI to bypass this process
+		return ForceRegionId;
+	}
+	if (QosEvalResult == EQosCompletionResult::Invalid)
+	{
+		// if we haven't run the evaluator just use the region from settings
+		return GetSavedRegionId();
+	}
+	return SelectedRegion;
+}
+
+const TArray<FQosRegionInfo>& FQosInterface::GetRegionOptions() const
+{
+	return RegionOptions;
+}
+
+void FQosInterface::ForceSelectRegion(const FString& InRegionId)
+{
+	QosEvalResult = EQosCompletionResult::Success;
+	ForceRegionId.Empty(); // remove any override (not typically used)
+
+	// make sure we can select this region
+	FString RegionId = InRegionId.ToUpper();
+	if (!SetSelectedRegion(RegionId))
+	{
+		// if not, add a fake entry and try again
+		FQosRegionInfo RegionInfo;
+		RegionInfo.RegionId = RegionId;
+		RegionInfo.PingMs = 0;
+		RegionOptions.Add(RegionInfo);
+		verify(SetSelectedRegion(RegionId));
+	}
+}
 
-		bDCIDCheck = true;
+bool FQosInterface::SetSelectedRegion(const FString& InRegionId)
+{
+	// make sure we've enumerated
+	if (QosEvalResult != EQosCompletionResult::Success)
+	{
+		// can't select region until we've enumerated them
+		return false;
 	}
+	
+	// make sure it's in the option list
+	FString RegionId = InRegionId.ToUpper();
+	for (const FQosRegionInfo& RegionInfo : RegionOptions)
+	{
+		if (RegionInfo.RegionId == RegionId)
+		{
+			// ok, save it
+			SelectedRegion = RegionId;
+
+			// apply this to saved settings
+			SaveSelectedRegionId();
+			return true;
+		}
+	}
+
+	// can't select a region not in the options list
+	return false;
+}
+
+void FQosInterface::SaveSelectedRegionId()
+{
+	// TODO: write to saved settings
+}
+
+FString FQosInterface::GetSavedRegionId() const
+{
+	// check saved settings
+	// TODO:
 
-	return DCIDString;
+	// if nothing is set, return hardcoded default
+	return TEXT("NONE");
 }
\ No newline at end of file
diff --git a/Engine/Source/Runtime/Online/OnlineFramework/Qos/Private/QosModule.cpp b/Engine/Source/Runtime/Online/OnlineFramework/Qos/Private/QosModule.cpp
--- a/Engine/Source/Runtime/Online/OnlineFramework/Qos/Private/QosModule.cpp
+++ b/Engine/Source/Runtime/Online/OnlineFramework/Qos/Private/QosModule.cpp
@@ -11,18 +11,19 @@ DEFINE_LOG_CATEGORY(LogQos);
 
 void FQosModule::StartupModule()
 {
-	QosInterface = new FQosInterface();
 }
 
 void FQosModule::ShutdownModule()
 {
-	delete QosInterface;
-	QosInterface = nullptr;
 }
 
-FQosInterface* FQosModule::GetQosInterface()
+TSharedRef<FQosInterface> FQosModule::GetQosInterface()
 {
-	return QosInterface;
+	if (!QosInterface.IsValid())
+	{
+		QosInterface = MakeShareable(new FQosInterface());
+	}
+	return QosInterface.ToSharedRef();
 }
 
 bool FQosModule::Exec(UWorld* InWorld, const TCHAR* Cmd, FOutputDevice& Ar)
diff --git a/Engine/Source/Runtime/Online/OnlineFramework/Qos/Private/QosStats.cpp b/Engine/Source/Runtime/Online/OnlineFramework/Qos/Private/QosStats.cpp
--- a/Engine/Source/Runtime/Online/OnlineFramework/Qos/Private/QosStats.cpp
+++ b/Engine/Source/Runtime/Online/OnlineFramework/Qos/Private/QosStats.cpp
@@ -19,7 +19,6 @@ const FString FQosDatacenterStats::QosStats_TotalTime = TEXT("QosStats_TotalTime
 
 // Qos stats
 const FString FQosDatacenterStats::QosStats_DeterminationType = TEXT("QosStats_DeterminationType");
-const FString FQosDatacenterStats::QosStats_DatacenterId = TEXT("QosStats_DatacenterId");
 const FString FQosDatacenterStats::QosStats_NumRegions = TEXT("QosStats_NumRegions");
 const FString FQosDatacenterStats::QosStats_RegionDetails = TEXT("QosStats_RegionDetails");
 const FString FQosDatacenterStats::QosStats_NumResults = TEXT("QosStats_NumResults");
@@ -107,13 +106,12 @@ void FQosDatacenterStats::RecordQosAttempt(const FOnlineSessionSearchResult& Sea
 	}
 }
 
-void FQosDatacenterStats::EndQosPass(EDatacenterResultType Result, const FString& DatacenterId)
+void FQosDatacenterStats::EndQosPass(EDatacenterResultType Result)
 {
 	if (bAnalyticsInProgress)
 	{
 		Finalize();
 		QosData.DeterminationType = Result;
-		QosData.BestDatacenterId = DatacenterId;
 	}
 }
 
@@ -164,7 +162,6 @@ void FQosDatacenterStats::ParseQosResults(TSharedPtr<IAnalyticsProvider>& Analyt
 	QoSAttributes.Add(FAnalyticsEventAttribute(QosStats_TotalTime, QosData.SearchTime.MSecs));
 
 	QoSAttributes.Add(FAnalyticsEventAttribute(QosStats_DeterminationType, ToString(QosData.DeterminationType)));
-	QoSAttributes.Add(FAnalyticsEventAttribute(QosStats_DatacenterId, QosData.BestDatacenterId));
 	QoSAttributes.Add(FAnalyticsEventAttribute(QosStats_NumRegions, QosData.Regions.Num()));
 	QoSAttributes.Add(FAnalyticsEventAttribute(QosStats_NumResults, QosData.NumTotalSearches));
 	QoSAttributes.Add(FAnalyticsEventAttribute(QosStats_NumSuccessCount, QosData.NumSuccessAttempts));
diff --git a/Engine/Source/Runtime/Online/OnlineFramework/Qos/Private/QosStats.h b/Engine/Source/Runtime/Online/OnlineFramework/Qos/Private/QosStats.h
--- a/Engine/Source/Runtime/Online/OnlineFramework/Qos/Private/QosStats.h
+++ b/Engine/Source/Runtime/Online/OnlineFramework/Qos/Private/QosStats.h
@@ -105,8 +105,6 @@ private:
 	{
 		/** Time of the search */
 		FString Timestamp;
-		/** Id of the datacenter chosen */
-		FString BestDatacenterId;
 		/** Way the datacenter was chosen */
 		EDatacenterResultType DeterminationType;
 		/** Time in ms it took to find the search results (exclusive) */
@@ -141,7 +139,6 @@ private:
 
 	// Qos stats
 	static const FString QosStats_DeterminationType;
-	static const FString QosStats_DatacenterId;
 	static const FString QosStats_NumRegions;
 	static const FString QosStats_RegionDetails;
 	static const FString QosStats_NumResults;
@@ -215,9 +212,8 @@ public:
 	 * End recording of a Qos determination
 	 *
 	 * @param Result results of the qos pass
-	 * @param DatacenterId resulting datacenter choice
 	 */
-	void EndQosPass(EDatacenterResultType Result, const FString& DatacenterId);
+	void EndQosPass(EDatacenterResultType Result);
 
 	/**
 	 * Record previously saved stats to an analytics provider
diff --git a/Engine/Source/Runtime/Online/OnlineFramework/Qos/Public/QosEvaluator.h b/Engine/Source/Runtime/Online/OnlineFramework/Qos/Public/QosEvaluator.h
--- a/Engine/Source/Runtime/Online/OnlineFramework/Qos/Public/QosEvaluator.h
+++ b/Engine/Source/Runtime/Online/OnlineFramework/Qos/Public/QosEvaluator.h
@@ -3,11 +3,11 @@
 #pragma once
 
 #include "OnlineSessionInterface.h"
-#include "QosEvaluator.generated.h"
 
 class AQosBeaconClient;
 class FQosDatacenterStats;
 enum class EQosResponseType : uint8;
+struct FQosRegionInfo;
 
 #define GAMEMODE_QOS	TEXT("QOSSERVER")
 
@@ -15,32 +15,18 @@ enum class EQosResponseType : uint8;
 #define DATACENTERQUERY_INTERVAL (5.0 * 60.0)
 
 /** Enum for possible QoS return codes */
-UENUM()
 enum class EQosCompletionResult : uint8
 {
 	/** Incomplete, invalid result */
 	Invalid,
 	/** QoS operation was successful */
 	Success,
-	/** QoS operation was skipped in favor of a cached result */
-	Cached,
 	/** QoS operation ended in failure */
 	Failure,
 	/** QoS operation was canceled */
 	Canceled
 };
 
-/**
- * Generic settings a server runs when hosting a simple QoS response service
- */
-class QOS_API FOnlineSessionSettingsQos : public FOnlineSessionSettings
-{
-public:
-
-	FOnlineSessionSettingsQos(bool bInIsDedicated = true);
-	virtual ~FOnlineSessionSettingsQos() {}
-};
-
 /**
  * Search settings for QoS advertised sessions
  */
@@ -51,41 +37,17 @@ public:
 	virtual ~FOnlineSessionSearchQos() {}
 
 	// FOnlineSessionSearch Interface begin
-	virtual TSharedPtr<FOnlineSessionSettings> GetDefaultSessionSettings() const override { return MakeShareable(new FOnlineSessionSettingsQos()); }
+	virtual TSharedPtr<FOnlineSessionSettings> GetDefaultSessionSettings() const override;
 	virtual void SortSearchResults() override;
 	// FOnlineSessionSearch Interface end
 };
 
-/**
- * Internal state for a given search pass
- */
-USTRUCT()
-struct FQosSearchState
-{
-	GENERATED_USTRUCT_BODY()
-
-	/** Datacenter Id */
-	UPROPERTY()
-	FString DatacenterId;
-	/** Last datacenter lookup timestamp */
-	UPROPERTY()
-	double LastDatacenterIdTimestamp;
-
-	FQosSearchState() :
-		LastDatacenterIdTimestamp(0.0)
-	{}
-};
-
 /**
  * Internal state for a given QoS pass
  */
-USTRUCT()
 struct FQosSearchPass
 {
-	GENERATED_USTRUCT_BODY()
-
 	/** Current search result choice to test */
-	UPROPERTY()
 	int32 CurrentSessionIdx;
 	/** Current array of QoS search results to evaluate */
 	TArray<FOnlineSessionSearchResult> SearchResults;
@@ -106,20 +68,18 @@ DECLARE_DELEGATE_OneParam(FOnQosPingEvalComplete, EQosCompletionResult /** Resul
  * Delegate triggered when all QoS search results have been investigated 
  *
  * @param Result the QoS operation result
- * @param BestDatacenterId the datacenter chosen by the operation (may be forced, cached, or default)
+ * @param RegionInfo The per-region ping information
  */
-DECLARE_DELEGATE_TwoParams(FOnQosSearchComplete, EQosCompletionResult /** Result */, const FString& /** BestDatacenterId */);
+DECLARE_DELEGATE_TwoParams(FOnQosSearchComplete, EQosCompletionResult /** Result */, const TArray<FQosRegionInfo>& /** RegionInfo */);
 
 /**
  * Evaluates QoS metrics to determine the best datacenter under current conditions
  * Additionally capable of generically pinging an array of servers that have a QosBeaconHost active
  */
-UCLASS(config = Engine, notplaceable)
-class QOS_API UQosEvaluator : public UObject
+class QOS_API FQosEvaluator : public TSharedFromThis<FQosEvaluator>
 {
-	GENERATED_UCLASS_BODY()
-
 public:
+	FQosEvaluator(UWorld* World);
 
 	/**
 	 * QoS services
@@ -134,15 +94,6 @@ public:
 	 */
 	void FindDatacenters(int32 ControllerId, const FOnQosSearchComplete& InCompletionDelegate);
 
-	/**
-	 * Evaluate ping reachability for a given list of servers.  Assumes there is a AQosBeaconHost on the machine
-	 * to receive the request and send a response
-	 *
-	 * @param SearchResults array of search results to test
-	 * @param InCompletionDelegate delegate to fire when test is complete
-	 */
-	void EvaluateServerPing(TSharedPtr<FOnlineSessionSearch>& SearchResults, const FOnQosPingEvalComplete& InCompletionDelegate);
-
 	/**
 	 * Is a QoS operation active
 	 *
@@ -155,22 +106,24 @@ public:
 	 */
 	void Cancel();
 
+protected:
+
 	/**
-	 * Get the default region for this instance, checking ini and commandline overrides
+	 * Evaluate ping reachability for a given list of servers.  Assumes there is a AQosBeaconHost on the machine
+	 * to receive the request and send a response
 	 *
-	 * @return the default region identifier
+	 * @param SearchResults array of search results to test
+	 * @param InCompletionDelegate delegate to fire when test is complete
 	 */
-	static const FString& GetDefaultRegionString();
+	void EvaluateServerPing(TSharedPtr<FOnlineSessionSearch>& SearchResults, const FOnQosPingEvalComplete& InCompletionDelegate);
 
 private:
 
-	/** Current QoS data */
-	UPROPERTY()
-	FQosSearchState SearchState;
 	/** Current QoS search/eval state */
-	UPROPERTY()
 	FQosSearchPass CurrentSearchPass;
 
+	TWeakObjectPtr<UWorld> ParentWorld;
+
 	/** QoS search results */
 	TSharedPtr<FOnlineSessionSearch> QosSearchQuery;
 
@@ -179,8 +132,7 @@ private:
 	FOnQosPingEvalComplete OnQosPingEvalComplete;
 
 	/** Beacon for sending QoS requests */
-	UPROPERTY(Transient)
-	AQosBeaconClient* QosBeaconClient;
+	TWeakObjectPtr<AQosBeaconClient> QosBeaconClient;
 
 	/** A QoS operation is in progress */
 	bool bInProgress;
@@ -193,7 +145,7 @@ private:
 	 * @param InCompletionDelegate the delegate to trigger
 	 * @param CompletionResult the QoS operation result
 	 */
-	void FinalizeDatacenterResult(const FOnQosSearchComplete& InCompletionDelegate, EQosCompletionResult CompletionResult);
+	void FinalizeDatacenterResult(const FOnQosSearchComplete& InCompletionDelegate, EQosCompletionResult CompletionResult, const TArray<FQosRegionInfo>& RegionInfo);
 
 	/**
 	 * Continue with the next datacenter endpoint and gather its ping result
@@ -268,7 +220,7 @@ private:
 	void ResetSearchVars();
 
 	/** Quick access to the current world */
-	UWorld* GetWorld() const override;
+	UWorld* GetWorld() const;
 
 	/** Quick access to the world timer manager */
 	FTimerManager& GetWorldTimerManager() const;
@@ -289,10 +241,6 @@ inline const TCHAR* ToString(EQosCompletionResult Result)
 		{
 			return TEXT("Success");
 		}
-		case EQosCompletionResult::Cached:
-		{
-			return TEXT("Cached");
-		}
 		case EQosCompletionResult::Failure:
 		{
 			return TEXT("Failure");
diff --git a/Engine/Source/Runtime/Online/OnlineFramework/Qos/Public/QosInterface.h b/Engine/Source/Runtime/Online/OnlineFramework/Qos/Public/QosInterface.h
--- a/Engine/Source/Runtime/Online/OnlineFramework/Qos/Public/QosInterface.h
+++ b/Engine/Source/Runtime/Online/OnlineFramework/Qos/Public/QosInterface.h
@@ -2,57 +2,116 @@
 
 #pragma once
 
-#include "ModuleManager.h"
-#include "QosModule.h"
+class IAnalyticsProvider;
+class FQosEvaluator;
+enum class EQosCompletionResult : uint8;
 
-class UQosEvaluator;
+/**
+ * Information for a given region
+ */
+struct QOS_API FQosRegionInfo
+{
+	FString RegionId;
+	int32 PingMs;
+};
 
-namespace QOSConsoleVariables
+/**
+ * Generic settings a server runs when hosting a simple QoS response service
+ */
+class QOS_API FOnlineSessionSettingsQos : public FOnlineSessionSettings
 {
-	extern TAutoConsoleVariable<FString> CVarForceRegion;
-	extern TAutoConsoleVariable<int32> CVarForceDefaultRegion;
-	extern TAutoConsoleVariable<int32> CVarOverrideTimestamp;
-	extern TAutoConsoleVariable<int32> CVarForceCached;
-}
+public:
+
+	FOnlineSessionSettingsQos(bool bInIsDedicated = true);
+	virtual ~FOnlineSessionSettingsQos() {}
+};
 
 /**
  * Main Qos interface for actions related to server quality of service
  */
-class QOS_API FQosInterface
+class QOS_API FQosInterface : public TSharedFromThis<FQosInterface>
 {
-
 public:
 
-	FQosInterface();
-	virtual ~FQosInterface();
+	/**
+	 * Get the interface singleton
+	 */
+	static TSharedRef<FQosInterface> Get();
 
-	/** @return the Qos interface */
-	static FQosInterface* Get()
-	{
-		static const FName QosModuleName = TEXT("Qos");
-		FQosModule& QoSModule = FModuleManager::LoadModuleChecked<FQosModule>(QosModuleName);
-		return QoSModule.GetQosInterface();
-	}
+	/**
+	 * Start running the async QoS evaluation 
+	 */
+	void BeginQosEvaluation(UWorld* World, const TSharedPtr<IAnalyticsProvider>& AnalyticsProvider, const FSimpleDelegate& OnComplete);
 
 	/**
-	 * Get a Qos evaluator for making datacenter determinations and ping timings against servers
+	 * Get the region ID for this instance, checking ini and commandline overrides.
+	 * 
+	 * Dedicated servers will have this value specified on the commandline
+	 * 
+	 * Clients pull this value from the settings (or command line) and do a ping test to determine if the setting is viable.
 	 *
-	 * @return a new instance of a qos evaluator
+	 * @return the default region identifier
 	 */
-	UQosEvaluator* CreateQosEvaluator();
+	FString GetRegionId() const;
 
 	/**
-	 * Get the default region for this instance, checking ini and commandline overrides
+	 * Get the list of regions that the client can choose from (returned from search and must meet min ping requirements)
 	 *
-	 * @return the default region identifier
+	 * If this list is empty, the client cannot play.
 	 */
-	static const FString& GetDefaultRegionString();
+	const TArray<FQosRegionInfo>& GetRegionOptions() const;
+
+	/**
+	 * return whether or not the region test succeeded
+	 */
+	bool DidRegionTestSucceed() const;
+
+	/**
+	 * Try to set the selected region ID (must be present in GetRegionOptions)
+	 */
+	bool SetSelectedRegion(const FString& RegionId);
+
+	/**
+	 * Force the selected region creating a fake RegionOption if necessary
+	 */
+	void ForceSelectRegion(const FString& RegionId);
 
 	/**
 	 * Get the datacenter id for this instance, checking ini and commandline overrides
+	 * This is only relevant for dedicated servers (so they can advertise). Client does 
+	 * not search on this (but may choose to prioritize results later)
 	 *
 	 * @return the default datacenter identifier
 	 */
-	static const FString& GetDatacenterId();
+	static FString GetDatacenterId();
+
+	/**
+	 * Legacy static accessor for GetRegionId()
+	 */
+	inline static FString GetDefaultRegionString()
+	{
+		return FQosInterface::Get()->GetRegionId();
+	}
+
+protected:
+	friend class FQosModule;
+	FQosInterface();
+
+private:
+
+	void OnQosEvaluationComplete(EQosCompletionResult Result, const TArray<FQosRegionInfo>& RegionInfo);
+
+	FString GetSavedRegionId() const;
+	void SaveSelectedRegionId();
+
+	int32 MinimumPingMs;
+	FString ForceRegionId;
+
+	TSharedPtr<FQosEvaluator> Evaluator;
+	EQosCompletionResult QosEvalResult;
+	TArray<FQosRegionInfo> RegionOptions;
+	FString SelectedRegion;
+
+	TArray<FSimpleDelegate> OnQosEvalCompleteDelegate;
 };
 
diff --git a/Engine/Source/Runtime/Online/OnlineFramework/Qos/Public/QosModule.h b/Engine/Source/Runtime/Online/OnlineFramework/Qos/Public/QosModule.h
--- a/Engine/Source/Runtime/Online/OnlineFramework/Qos/Public/QosModule.h
+++ b/Engine/Source/Runtime/Online/OnlineFramework/Qos/Public/QosModule.h
@@ -42,7 +42,13 @@ public:
 		return FModuleManager::Get().IsModuleLoaded("Qos");
 	}
 
-	FQosInterface* GetQosInterface();
+	/**
+	 * Get the interface singleton
+	 */
+	TSharedRef<FQosInterface> GetQosInterface();
+
+private:
+	TSharedPtr<FQosInterface> QosInterface;
 
 private:
 
@@ -59,9 +65,5 @@ private:
 	 * Shutdown platform specific parts of template handling
 	 */
 	virtual void ShutdownModule() override;
-
-private:
-
-	FQosInterface* QosInterface;
 };
 
diff --git a/Engine/Source/Runtime/Online/OnlineSubsystem/Private/OnlineSubsystem.cpp b/Engine/Source/Runtime/Online/OnlineSubsystem/Private/OnlineSubsystem.cpp
--- a/Engine/Source/Runtime/Online/OnlineSubsystem/Private/OnlineSubsystem.cpp
+++ b/Engine/Source/Runtime/Online/OnlineSubsystem/Private/OnlineSubsystem.cpp
@@ -49,7 +49,7 @@ int32 GetBuildUniqueId()
 		bStaticCheck = true;
 	}
 
-	int32 Crc = 0;
+	int32 BuildId = 0;
 	if (bUseBuildIdOverride == false)
 	{
 		/** Engine package CRC doesn't change, can't be used as the version - BZ */
@@ -57,20 +57,23 @@ int32 GetBuildUniqueId()
 		// Serialize to a NBO buffer for consistent CRCs across platforms
 		Buffer << GEngineNetVersion;
 		// Now calculate the CRC
-		Crc = static_cast<int32>(FCrc::MemCrc_DEPRECATED((uint8*)Buffer, Buffer.GetByteCount()));
+		uint32 Crc = FCrc::MemCrc32((uint8*)Buffer, Buffer.GetByteCount());
+
+		// make sure it's positive when it's cast back to an int
+		BuildId = static_cast<int32>(Crc & 0x7fffffff);
 	}
 	else
 	{
-		Crc = BuildIdOverride;
+		BuildId = BuildIdOverride;
 	}
 
-	UE_LOG_ONLINE(VeryVerbose, TEXT("GetBuildUniqueId: GEngineNetVersion %d bUseBuildIdOverride %d BuildIdOverride %d Crc %d"),
+	UE_LOG_ONLINE(VeryVerbose, TEXT("GetBuildUniqueId: GEngineNetVersion %d bUseBuildIdOverride %d BuildIdOverride %d BuildId %d"),
 		GEngineNetVersion,
 		bUseBuildIdOverride,
 		BuildIdOverride,
-		Crc);
+		BuildId);
 
-	return Crc;
+	return BuildId;
 }
 
 bool IsPlayerInSessionImpl(IOnlineSession* SessionInt, FName SessionName, const FUniqueNetId& UniqueId)
diff --git a/Engine/Source/Runtime/Online/OnlineSubsystem/Public/Interfaces/IMessageSanitizerInterface.h b/Engine/Source/Runtime/Online/OnlineSubsystem/Public/Interfaces/IMessageSanitizerInterface.h
new file mode 100644
--- /dev/null
+++ b/Engine/Source/Runtime/Online/OnlineSubsystem/Public/Interfaces/IMessageSanitizerInterface.h
@@ -0,0 +1,29 @@
+// Copyright 1998-2015 Epic Games, Inc. All Rights Reserved.
+
+#pragma once
+
+DECLARE_DELEGATE_OneParam(FOnMessageProcessed, const FString& /*SanitizedMessage*/);
+
+struct FSanitizeMessage
+{
+	FSanitizeMessage(FString InRawMessage, FOnMessageProcessed InProcessCompleteDelegate)
+		: RawMessage(InRawMessage)
+		, CompletionDelegate(InProcessCompleteDelegate)
+	{}
+
+	FString RawMessage;
+	FOnMessageProcessed CompletionDelegate;
+};
+
+class IMessageSanitizer
+{
+protected:
+	IMessageSanitizer() {};
+
+public:
+	virtual ~IMessageSanitizer() {};
+	virtual void SanitizeMessage(FString Message, FOnMessageProcessed CompletionDelegate) = 0;
+};
+
+typedef TSharedPtr<IMessageSanitizer, ESPMode::ThreadSafe> IMessageSanitizerPtr;
+
diff --git a/Engine/Source/Runtime/Online/OnlineSubsystem/Public/Interfaces/OnlineIdentityInterface.h b/Engine/Source/Runtime/Online/OnlineSubsystem/Public/Interfaces/OnlineIdentityInterface.h
--- a/Engine/Source/Runtime/Online/OnlineSubsystem/Public/Interfaces/OnlineIdentityInterface.h
+++ b/Engine/Source/Runtime/Online/OnlineSubsystem/Public/Interfaces/OnlineIdentityInterface.h
@@ -135,12 +135,14 @@ public:
 		AccountTypeFailure		=	1 << 3,
 		/** Invalid user */
 		UserNotFound			=	1 << 4,
+		/** User must be logged in */
+		UserNotLoggedIn			=	1 << 5,
 		/** User restricted from chat */
-		ChatRestriction			=	1 << 5,
+		ChatRestriction			=	1 << 6,
 		/** User restricted from User Generated Content */
-		UGCRestriction			=	1 << 6,
+		UGCRestriction			=	1 << 7,
 		/** Platform failed for unknown reason and handles its own dialogs */
-		GenericFailure			=	1 << 7, 
+		GenericFailure			=	1 << 8, 
 	};
 
 	virtual ~IOnlineIdentity() {};
diff --git a/Engine/Source/Runtime/Online/OnlineSubsystem/Public/Interfaces/OnlineStoreInterfaceV2.h b/Engine/Source/Runtime/Online/OnlineSubsystem/Public/Interfaces/OnlineStoreInterfaceV2.h
--- a/Engine/Source/Runtime/Online/OnlineSubsystem/Public/Interfaces/OnlineStoreInterfaceV2.h
+++ b/Engine/Source/Runtime/Online/OnlineSubsystem/Public/Interfaces/OnlineStoreInterfaceV2.h
@@ -13,18 +13,39 @@ typedef FString FUniqueCategoryId;
 class FOnlineStoreOffer
 {
 public:
+	FOnlineStoreOffer()
+	: NumericPrice(-1)
+	, ReleaseDate(0)
+	{
+	}
+	virtual ~FOnlineStoreOffer()
+	{
+	}
+
 	/** Unique offer identifier */
 	FUniqueOfferId OfferId;
+
 	/** Title for display */
 	FText Title;	
 	/** Short description for display */
 	FText Description;
 	/** Full description for display */
 	FText LongDescription;
-	/** Price in text for display */
-	FText DisplayPrice;
+	
+	/** Price in numeric form for comparison/sorting */
+	int32 NumericPrice;
+	/** Price currency code */
+	FString CurrencyCode;
+
 	/** Date the offer was released */
 	FDateTime ReleaseDate;
+
+	/** @return FText suitable for localized display */
+	virtual FText GetDisplayPrice() const
+	{
+		return FText::AsCurrency(NumericPrice / 100.0, CurrencyCode);
+	}
+
 	/** @return True if offer can be purchased */
 	virtual bool IsPurchaseable() const 
 	{ 
@@ -53,27 +74,23 @@ public:
 class FOnlineStoreFilter
 {
 public:
-	/** Namespace for identifying an entire store */
-	FString Namespace;
 	/** Keyword strings to match when filtering items/offers */
 	TArray<FString> Keywords;
 	/** Category paths to match when filtering offers */
 	TArray<FOnlineStoreCategory> IncludeCategories;
 	/** Category paths to exclude when filtering offers */
 	TArray<FOnlineStoreCategory> ExcludeCategories;
-	/** Currency to use for viewing catalog */
-	FString CurrencyCode;
 };
 
 /**
  * Delegate called when available online categories have been queried
  */
-DECLARE_DELEGATE_ThreeParams(FOnQueryOnlineStoreCategoriesComplete, const FUniqueNetId& /*UserId*/, bool /*bWasSuccessful*/, const FString& /*Error*/);
+DECLARE_DELEGATE_TwoParams(FOnQueryOnlineStoreCategoriesComplete, bool /*bWasSuccessful*/, const FString& /*Error*/);
 
 /**
  * Delegate called when online store query completes
  */
-DECLARE_DELEGATE_FourParams(FOnQueryOnlineStoreOffersComplete, const FUniqueNetId& /*UserId*/, bool /*bWasSuccessful*/, const TArray<FUniqueOfferId>& /*OfferIds*/, const FString& /*Error*/);
+DECLARE_DELEGATE_ThreeParams(FOnQueryOnlineStoreOffersComplete, bool /*bWasSuccessful*/, const TArray<FUniqueOfferId>& /*OfferIds*/, const FString& /*Error*/);
 
 /**
  *	Access to available offers for purchase
@@ -83,14 +100,12 @@ class IOnlineStoreV2
 public:
 
 	/**
-	 * Query for available store categories
+	 * Query for available store categories. Delegate callback is guaranteed.
 	 *
 	 * @param UserId user initiating the request
 	 * @param Delegate completion callback
-	 *
-	 * @return true if async operation started
 	 */
-	virtual bool QueryCategories(const FUniqueNetId& UserId, const FOnQueryOnlineStoreCategoriesComplete& Delegate = FOnQueryOnlineStoreCategoriesComplete()) = 0;
+	virtual void QueryCategories(const FUniqueNetId& UserId, const FOnQueryOnlineStoreCategoriesComplete& Delegate = FOnQueryOnlineStoreCategoriesComplete()) = 0;
 
 	/**
 	 * Get currently cached store categories
@@ -100,18 +115,16 @@ public:
 	virtual void GetCategories(TArray<FOnlineStoreCategory>& OutCategories) const = 0;
 
 	/**
-	 * Query for available store offers using a filter
+	 * Query for available store offers using a filter. Delegate callback is guaranteed.
 	 *
 	 * @param UserId user initiating the request
 	 * @param Filter only return offers matching the filter
 	 * @param Delegate completion callback
-	 *
-	 * @return true if async operation started
 	 */
-	virtual bool QueryOffersByFilter(const FUniqueNetId& UserId, const FOnlineStoreFilter& Filter, const FOnQueryOnlineStoreOffersComplete& Delegate = FOnQueryOnlineStoreOffersComplete()) = 0;
+	virtual void QueryOffersByFilter(const FUniqueNetId& UserId, const FOnlineStoreFilter& Filter, const FOnQueryOnlineStoreOffersComplete& Delegate = FOnQueryOnlineStoreOffersComplete()) = 0;
 
 	/**
-	 * Query for available store offers matching the given ids
+	 * Query for available store offers matching the given ids. Delegate callback is guaranteed.
 	 *
 	 * @param UserId user initiating the request
 	 * @param OfferIds only return offers matching these ids
@@ -119,7 +132,7 @@ public:
 	 *
 	 * @return true if async operation started
 	 */
-	virtual bool QueryOffersById(const FUniqueNetId& UserId, const TArray<FUniqueOfferId>& OfferIds, const FOnQueryOnlineStoreOffersComplete& Delegate = FOnQueryOnlineStoreOffersComplete()) = 0;
+	virtual void QueryOffersById(const FUniqueNetId& UserId, const TArray<FUniqueOfferId>& OfferIds, const FOnQueryOnlineStoreOffersComplete& Delegate = FOnQueryOnlineStoreOffersComplete()) = 0;
 
 	/**
 	 * Get currently cached store offers
diff --git a/Engine/Source/Runtime/Online/OnlineSubsystem/Public/OnlineSubsystem.h b/Engine/Source/Runtime/Online/OnlineSubsystem/Public/OnlineSubsystem.h
--- a/Engine/Source/Runtime/Online/OnlineSubsystem/Public/OnlineSubsystem.h
+++ b/Engine/Source/Runtime/Online/OnlineSubsystem/Public/OnlineSubsystem.h
@@ -37,6 +37,7 @@ DECLARE_CYCLE_STAT_EXTERN(TEXT("VoiceInt"), STAT_Voice_Interface, STATGROUP_Onli
 typedef TSharedPtr<class IOnlineSession, ESPMode::ThreadSafe> IOnlineSessionPtr;
 typedef TSharedPtr<class IOnlineFriends, ESPMode::ThreadSafe> IOnlineFriendsPtr;
 typedef TSharedPtr<class IOnlinePartySystem, ESPMode::ThreadSafe> IOnlinePartyPtr;
+typedef TSharedPtr<class IMessageSanitizer, ESPMode::ThreadSafe> IMessageSanitizerPtr;
 typedef TSharedPtr<class IOnlineGroups, ESPMode::ThreadSafe> IOnlineGroupsPtr;
 typedef TSharedPtr<class IOnlineSharedCloud, ESPMode::ThreadSafe> IOnlineSharedCloudPtr;
 typedef TSharedPtr<class IOnlineUserCloud, ESPMode::ThreadSafe> IOnlineUserCloudPtr;
@@ -223,6 +224,12 @@ public:
 	 */
 	virtual IOnlineFriendsPtr GetFriendsInterface() const = 0;
 
+	/**
+	* Get the interface for accessing the message sanitizer service
+	* @return Interface pointer for the appropriate sanitizer service
+	*/
+	virtual IMessageSanitizerPtr GetMessageSanitizerInterface() const = 0;
+
 	/**
 	 * Get the interface for accessing the groups services
 	 * @return Interface pointer for appropriate groups service
diff --git a/Engine/Source/Runtime/Online/OnlineSubsystemAmazon/Public/OnlineSubsystemAmazon.h b/Engine/Source/Runtime/Online/OnlineSubsystemAmazon/Public/OnlineSubsystemAmazon.h
--- a/Engine/Source/Runtime/Online/OnlineSubsystemAmazon/Public/OnlineSubsystemAmazon.h
+++ b/Engine/Source/Runtime/Online/OnlineSubsystemAmazon/Public/OnlineSubsystemAmazon.h
@@ -35,6 +35,7 @@ public:
 
 	virtual IOnlineSessionPtr GetSessionInterface() const override { return nullptr; }
 	virtual IOnlineFriendsPtr GetFriendsInterface() const override { return nullptr; }
+	virtual IMessageSanitizerPtr GetMessageSanitizerInterface() const override { return nullptr; }
 	virtual IOnlinePartyPtr GetPartyInterface() const override { return nullptr; }
 	virtual IOnlineGroupsPtr GetGroupsInterface() const override { return nullptr; }
 	virtual IOnlineSharedCloudPtr GetSharedCloudInterface() const override { return nullptr; }
diff --git a/Engine/Source/Runtime/Online/OnlineSubsystemFacebook/Private/IOS/OnlineSubsystemFacebook.cpp b/Engine/Source/Runtime/Online/OnlineSubsystemFacebook/Private/IOS/OnlineSubsystemFacebook.cpp
--- a/Engine/Source/Runtime/Online/OnlineSubsystemFacebook/Private/IOS/OnlineSubsystemFacebook.cpp
+++ b/Engine/Source/Runtime/Online/OnlineSubsystemFacebook/Private/IOS/OnlineSubsystemFacebook.cpp
@@ -40,6 +40,11 @@ IOnlineFriendsPtr FOnlineSubsystemFacebook::GetFriendsInterface() const
 	return FacebookFriends;
 }
 
+IMessageSanitizerPtr FOnlineSubsystemFacebook::GetMessageSanitizerInterface() const
+{
+	return nullptr;
+}
+
 IOnlineGroupsPtr FOnlineSubsystemFacebook::GetGroupsInterface() const
 {
 	return nullptr;
diff --git a/Engine/Source/Runtime/Online/OnlineSubsystemFacebook/Private/Windows/OnlineSubsystemFacebook.cpp b/Engine/Source/Runtime/Online/OnlineSubsystemFacebook/Private/Windows/OnlineSubsystemFacebook.cpp
--- a/Engine/Source/Runtime/Online/OnlineSubsystemFacebook/Private/Windows/OnlineSubsystemFacebook.cpp
+++ b/Engine/Source/Runtime/Online/OnlineSubsystemFacebook/Private/Windows/OnlineSubsystemFacebook.cpp
@@ -16,6 +16,11 @@ IOnlineFriendsPtr FOnlineSubsystemFacebook::GetFriendsInterface() const
 	return FacebookFriends;
 }
 
+IMessageSanitizerPtr FOnlineSubsystemFacebook::GetMessageSanitizerInterface() const
+{
+	return nullptr;
+}
+
 IOnlinePartyPtr FOnlineSubsystemFacebook::GetPartyInterface() const
 {
 	return nullptr;
diff --git a/Engine/Source/Runtime/Online/OnlineSubsystemNull/Private/OnlineSubsystemNull.cpp b/Engine/Source/Runtime/Online/OnlineSubsystemNull/Private/OnlineSubsystemNull.cpp
--- a/Engine/Source/Runtime/Online/OnlineSubsystemNull/Private/OnlineSubsystemNull.cpp
+++ b/Engine/Source/Runtime/Online/OnlineSubsystemNull/Private/OnlineSubsystemNull.cpp
@@ -20,6 +20,11 @@ IOnlineFriendsPtr FOnlineSubsystemNull::GetFriendsInterface() const
 	return nullptr;
 }
 
+IMessageSanitizerPtr FOnlineSubsystemNull::GetMessageSanitizerInterface() const
+{
+	return nullptr;
+}
+
 IOnlinePartyPtr FOnlineSubsystemNull::GetPartyInterface() const
 {
 	return nullptr;
diff --git a/Engine/Source/Runtime/Online/OnlineSubsystemSteam/Private/OnlineSubsystemSteam.cpp b/Engine/Source/Runtime/Online/OnlineSubsystemSteam/Private/OnlineSubsystemSteam.cpp
--- a/Engine/Source/Runtime/Online/OnlineSubsystemSteam/Private/OnlineSubsystemSteam.cpp
+++ b/Engine/Source/Runtime/Online/OnlineSubsystemSteam/Private/OnlineSubsystemSteam.cpp
@@ -162,6 +162,11 @@ IOnlineFriendsPtr FOnlineSubsystemSteam::GetFriendsInterface() const
 	return FriendInterface;
 }
 
+IMessageSanitizerPtr FOnlineSubsystemSteam::GetMessageSanitizerInterface() const
+{
+	return nullptr;
+}
+
 IOnlineGroupsPtr FOnlineSubsystemSteam::GetGroupsInterface() const
 {
 	return nullptr;
diff --git a/Engine/Source/Runtime/Online/OnlineSubsystemUtils/Classes/LeaderboardFlushCallbackProxy.cpp b/Engine/Source/Runtime/Online/OnlineSubsystemUtils/Classes/LeaderboardFlushCallbackProxy.cpp
--- a/Engine/Source/Runtime/Online/OnlineSubsystemUtils/Classes/LeaderboardFlushCallbackProxy.cpp
+++ b/Engine/Source/Runtime/Online/OnlineSubsystemUtils/Classes/LeaderboardFlushCallbackProxy.cpp
@@ -20,7 +20,7 @@ void ULeaderboardFlushCallbackProxy::TriggerFlush(APlayerController* PlayerContr
 		TSharedPtr<const FUniqueNetId> UserID = PlayerState->UniqueId.GetUniqueNetId();
 		if (UserID.IsValid())
 		{
-			if (IOnlineSubsystem* const OnlineSub = IOnlineSubsystem::Get())
+			if (IOnlineSubsystem* const OnlineSub = IOnlineSubsystem::IsLoaded() ? IOnlineSubsystem::Get() : nullptr)
 			{
 				IOnlineLeaderboardsPtr Leaderboards = OnlineSub->GetLeaderboardsInterface();
 				if (Leaderboards.IsValid())
@@ -77,7 +77,7 @@ void ULeaderboardFlushCallbackProxy::RemoveDelegate()
 {
 	if (!bFailedToEvenSubmit)
 	{
-		if (IOnlineSubsystem* OnlineSub = IOnlineSubsystem::Get())
+		if (IOnlineSubsystem* OnlineSub = IOnlineSubsystem::IsLoaded() ? IOnlineSubsystem::Get() : nullptr)
 		{
 			IOnlineLeaderboardsPtr Leaderboards = OnlineSub->GetLeaderboardsInterface();
 			if (Leaderboards.IsValid())
diff --git a/Engine/Source/Runtime/Online/OnlineSubsystemUtils/Classes/LeaderboardQueryCallbackProxy.cpp b/Engine/Source/Runtime/Online/OnlineSubsystemUtils/Classes/LeaderboardQueryCallbackProxy.cpp
--- a/Engine/Source/Runtime/Online/OnlineSubsystemUtils/Classes/LeaderboardQueryCallbackProxy.cpp
+++ b/Engine/Source/Runtime/Online/OnlineSubsystemUtils/Classes/LeaderboardQueryCallbackProxy.cpp
@@ -21,7 +21,7 @@ void ULeaderboardQueryCallbackProxy::TriggerQuery(APlayerController* PlayerContr
 		TSharedPtr<const FUniqueNetId> UserID = PlayerState->UniqueId.GetUniqueNetId();
 		if (UserID.IsValid())
 		{
-			if (IOnlineSubsystem* const OnlineSub = IOnlineSubsystem::Get())
+			if (IOnlineSubsystem* const OnlineSub = IOnlineSubsystem::IsLoaded() ? IOnlineSubsystem::Get() : nullptr)
 			{
 				IOnlineLeaderboardsPtr Leaderboards = OnlineSub->GetLeaderboardsInterface();
 				if (Leaderboards.IsValid())
@@ -130,7 +130,7 @@ void ULeaderboardQueryCallbackProxy::RemoveDelegate()
 {
 	if (!bFailedToEvenSubmit)
 	{
-		if (IOnlineSubsystem* OnlineSub = IOnlineSubsystem::Get())
+		if (IOnlineSubsystem* OnlineSub = IOnlineSubsystem::IsLoaded() ? IOnlineSubsystem::Get() : nullptr)
 		{
 			IOnlineLeaderboardsPtr Leaderboards = OnlineSub->GetLeaderboardsInterface();
 			if (Leaderboards.IsValid())
diff --git a/Engine/Source/Runtime/Online/OnlineSubsystemUtils/Private/InAppPurchaseCallbackProxy.cpp b/Engine/Source/Runtime/Online/OnlineSubsystemUtils/Private/InAppPurchaseCallbackProxy.cpp
--- a/Engine/Source/Runtime/Online/OnlineSubsystemUtils/Private/InAppPurchaseCallbackProxy.cpp
+++ b/Engine/Source/Runtime/Online/OnlineSubsystemUtils/Private/InAppPurchaseCallbackProxy.cpp
@@ -21,7 +21,7 @@ void UInAppPurchaseCallbackProxy::Trigger(APlayerController* PlayerController, c
 	WorldPtr = (PlayerController != nullptr) ? PlayerController->GetWorld() : nullptr;
 	if (APlayerState* PlayerState = (PlayerController != nullptr) ? PlayerController->PlayerState : nullptr)
 	{
-		if (IOnlineSubsystem* const OnlineSub = IOnlineSubsystem::Get())
+		if (IOnlineSubsystem* const OnlineSub = IOnlineSubsystem::IsLoaded() ? IOnlineSubsystem::Get() : nullptr)
 		{
 			IOnlineStorePtr StoreInterface = OnlineSub->GetStoreInterface();
 			if (StoreInterface.IsValid())
@@ -80,7 +80,7 @@ void UInAppPurchaseCallbackProxy::OnInAppPurchaseComplete(EInAppPurchaseState::T
 void UInAppPurchaseCallbackProxy::OnInAppPurchaseComplete_Delayed()
 {
     /** Cached product details of the purchased product */
-    FInAppPurchaseProductInfo ProductInformation;
+    FInAppPurchaseProductInfo ProductInformation; 
     
     if (SavedPurchaseState == EInAppPurchaseState::Success && PurchaseRequest.IsValid())
     {
@@ -104,7 +104,7 @@ void UInAppPurchaseCallbackProxy::RemoveDelegate()
 {
 	if (!bFailedToEvenSubmit)
 	{
-		if (IOnlineSubsystem* OnlineSub = IOnlineSubsystem::Get())
+		if (IOnlineSubsystem* OnlineSub = IOnlineSubsystem::IsLoaded() ? IOnlineSubsystem::Get() : nullptr)
 		{
 			IOnlineStorePtr InAppPurchases = OnlineSub->GetStoreInterface();
 			if (InAppPurchases.IsValid())
diff --git a/Engine/Source/Runtime/Online/OnlineSubsystemUtils/Private/LeaderboardBlueprintLibrary.cpp b/Engine/Source/Runtime/Online/OnlineSubsystemUtils/Private/LeaderboardBlueprintLibrary.cpp
--- a/Engine/Source/Runtime/Online/OnlineSubsystemUtils/Private/LeaderboardBlueprintLibrary.cpp
+++ b/Engine/Source/Runtime/Online/OnlineSubsystemUtils/Private/LeaderboardBlueprintLibrary.cpp
@@ -19,7 +19,7 @@ bool ULeaderboardBlueprintLibrary::WriteLeaderboardObject(APlayerController* Pla
 		TSharedPtr<const FUniqueNetId> UserId = PlayerState->UniqueId.GetUniqueNetId();
 		if (UserId.IsValid())
 		{
-			if (IOnlineSubsystem* const OnlineSub = IOnlineSubsystem::Get())
+			if (IOnlineSubsystem* const OnlineSub = IOnlineSubsystem::IsLoaded() ? IOnlineSubsystem::Get() : nullptr)
 			{
 				IOnlineLeaderboardsPtr Leaderboards = OnlineSub->GetLeaderboardsInterface();
 				if (Leaderboards.IsValid())
diff --git a/Engine/Source/Runtime/OpenGLDrv/Private/OpenGLTexture.cpp b/Engine/Source/Runtime/OpenGLDrv/Private/OpenGLTexture.cpp
--- a/Engine/Source/Runtime/OpenGLDrv/Private/OpenGLTexture.cpp
+++ b/Engine/Source/Runtime/OpenGLDrv/Private/OpenGLTexture.cpp
@@ -207,6 +207,7 @@ void FOpenGLDynamicRHI::RHIGetTextureMemoryStats(FTextureMemoryStats& OutStats)
 	OutStats.TotalGraphicsMemory = GOpenGLTotalGraphicsMemory ? GOpenGLTotalGraphicsMemory : -1;
 
 	OutStats.AllocatedMemorySize = int64(GCurrentTextureMemorySize) * 1024;
+	OutStats.LargestContiguousAllocation = OutStats.AllocatedMemorySize;
 	OutStats.TexturePoolSize = GTexturePoolSize;
 	OutStats.PendingMemoryAdjustment = 0;
 }
diff --git a/Engine/Source/Runtime/PakFile/Private/IPlatformFilePak.cpp b/Engine/Source/Runtime/PakFile/Private/IPlatformFilePak.cpp
--- a/Engine/Source/Runtime/PakFile/Private/IPlatformFilePak.cpp
+++ b/Engine/Source/Runtime/PakFile/Private/IPlatformFilePak.cpp
@@ -811,9 +811,28 @@ IFileHandle* FPakPlatformFile::CreatePakFileHandle(const TCHAR* Filename, FPakFi
 	return Result;
 }
 
-bool FPakPlatformFile::HandleMountPakDelegate(const FString& PakFilePath, uint32 PakOrder)
+bool FPakPlatformFile::HandleMountPakDelegate(const FString& PakFilePath, uint32 PakOrder, IPlatformFile::FDirectoryVisitor* Visitor)
 {
-	return Mount(*PakFilePath, PakOrder);
+	bool bReturn = Mount(*PakFilePath, PakOrder);
+	if (bReturn && Visitor != nullptr)
+	{
+		TArray<FPakListEntry> Paks;
+		GetMountedPaks(Paks);
+		// Find the single pak we just mounted
+		for (auto Pak : Paks)
+		{
+			if (PakFilePath == Pak.PakFile->GetFilename())
+			{
+				// Get a list of all of the files in the pak
+				for (FPakFile::FFileIterator It(*Pak.PakFile); It; ++It)
+				{
+					Visitor->Visit(*It.Filename(), false);
+				}
+				return true;
+			}
+		}
+	}
+	return bReturn;
 }
 
 bool FPakPlatformFile::HandleUnmountPakDelegate(const FString& PakFilePath)
diff --git a/Engine/Source/Runtime/PakFile/Public/IPlatformFilePak.h b/Engine/Source/Runtime/PakFile/Public/IPlatformFilePak.h
--- a/Engine/Source/Runtime/PakFile/Public/IPlatformFilePak.h
+++ b/Engine/Source/Runtime/PakFile/Public/IPlatformFilePak.h
@@ -866,7 +866,7 @@ class PAKFILE_API FPakPlatformFile : public IPlatformFile
 	/**
 	 * Handler for device delegate to prompt us to load a new pak.	 
 	 */
-	bool HandleMountPakDelegate(const FString& PakFilePath, uint32 PakOrder);
+	bool HandleMountPakDelegate(const FString& PakFilePath, uint32 PakOrder, IPlatformFile::FDirectoryVisitor* Visitor);
 
 	/**
 	 * Handler for device delegate to prompt us to unload a pak.
diff --git a/Engine/Source/Runtime/PerfCounters/Private/PerfCounters.cpp b/Engine/Source/Runtime/PerfCounters/Private/PerfCounters.cpp
--- a/Engine/Source/Runtime/PerfCounters/Private/PerfCounters.cpp
+++ b/Engine/Source/Runtime/PerfCounters/Private/PerfCounters.cpp
@@ -231,6 +231,26 @@ bool FPerfCounters::Tick(float DeltaTime)
 		}
 	}
 
+	// set some internal perf stats ([RCL] FIXME 2015-12-08: move to a better place and make configurable)
+	const float kInternalStatsUpdateInterval = 60.0;	// 60 seconds
+	static float LastTimeStatsUpdated = FPlatformTime::Seconds() - kInternalStatsUpdateInterval * FMath::FRand();
+	float CurrentTime = FPlatformTime::Seconds();
+	if (CurrentTime - LastTimeStatsUpdated > kInternalStatsUpdateInterval)
+	{
+		FPlatformMemoryStats Stats = FPlatformMemory::GetStats();
+		Set(TEXT("AvailablePhysicalMemoryMB"), static_cast<uint64>(Stats.AvailablePhysical / (1024 * 1024)));
+		Set(TEXT("AvailableVirtualMemoryMB"), static_cast<uint64>(Stats.AvailableVirtual / (1024 * 1024)));
+
+		const FString LogFilename = FPlatformOutputDevices::GetAbsoluteLogFilename();
+		uint64 TotalBytesOnLogDrive = 0, FreeBytesOnLogDrive = 0;
+		if (FPlatformMisc::GetDiskTotalAndFreeSpace(LogFilename, TotalBytesOnLogDrive, FreeBytesOnLogDrive))
+		{
+			Set(TEXT("FreeSpaceOnLogFileDiskInMB"), static_cast<uint64>(FreeBytesOnLogDrive / (1024 * 1024)));
+		}
+
+		LastTimeStatsUpdated = CurrentTime;
+	}
+
 	// keep ticking
 	return true;
 }
@@ -279,14 +299,6 @@ bool FPerfCounters::ProcessRequest(uint8* Buffer, int32 BufferLen, FResponse& Re
 			else if (Tokens[1].StartsWith(TEXT("/stats")))
 			{
 				Response.Body = GetAllCountersAsJson();
-
-				// retrieving stats resets them by default, unless ?peek parameter is passed
-				const int kStatsTokenLength = 6; // strlen("/stats");
-				FString TokenRemainder = Tokens[1].Mid(kStatsTokenLength);
-				if (TokenRemainder != TEXT("?peek"))
-				{
-					ResetStatsForNextPeriod();
-				}
 			}
 			else if (Tokens[1].StartsWith(TEXT("/exec?c=")))
 			{
diff --git a/Engine/Source/Runtime/PhysXFormats/Private/PhysXFormats.cpp b/Engine/Source/Runtime/PhysXFormats/Private/PhysXFormats.cpp
--- a/Engine/Source/Runtime/PhysXFormats/Private/PhysXFormats.cpp
+++ b/Engine/Source/Runtime/PhysXFormats/Private/PhysXFormats.cpp
@@ -124,6 +124,16 @@ public:
 		FPhysXOutputStream Buffer(&CookedMeshBuffer);
 		bool Result = PhysXCooking->cookConvexMesh(PConvexMeshDesc, Buffer);
 		
+		if(!Result && !(PConvexMeshDesc.flags & PxConvexFlag::eINFLATE_CONVEX))
+		{
+			//We failed to cook without inflating convex. Let's try again with inflation
+			//This is not ideal since it makes the collision less accurate. It's needed if given verts are extremely close.
+			UE_LOG(LogPhysics, Warning, TEXT("Cooking failed, possibly due to verts too close together. Trying again with inflation."));
+			PConvexMeshDesc.flags |= PxConvexFlag::eINFLATE_CONVEX;
+			Result = PhysXCooking->cookConvexMesh(PConvexMeshDesc, Buffer);
+			
+		}
+
 		// Return default cooking params to normal
 		if (bDeformableMesh)
 		{
diff --git a/Engine/Source/Runtime/Projects/Private/PluginDescriptor.cpp b/Engine/Source/Runtime/Projects/Private/PluginDescriptor.cpp
--- a/Engine/Source/Runtime/Projects/Private/PluginDescriptor.cpp
+++ b/Engine/Source/Runtime/Projects/Private/PluginDescriptor.cpp
@@ -13,6 +13,7 @@ FPluginDescriptor::FPluginDescriptor()
 	, bCanContainContent(false)
 	, bIsBetaVersion(false)
 	, bInstalled(false)
+	, bRequiresBuildPlatform(true)
 { 
 }
 
@@ -102,6 +103,11 @@ bool FPluginDescriptor::Read(const FString& Text, FText& OutFailReason)
 	Object.TryGetBoolField(TEXT("IsBetaVersion"), bIsBetaVersion);
 	Object.TryGetBoolField(TEXT("Installed"), bInstalled);
 
+	if(!Object.TryGetBoolField(TEXT("RequiresBuildPlatform"), bRequiresBuildPlatform))
+	{
+		bRequiresBuildPlatform = true;
+	}
+
 	return true;
 }
 
@@ -145,6 +151,10 @@ FString FPluginDescriptor::ToString() const
 	Writer.WriteValue(TEXT("Installed"), bInstalled);
 
 	FModuleDescriptor::WriteArray(Writer, TEXT("Modules"), Modules);
+	if(!bRequiresBuildPlatform)
+	{
+		Writer.WriteValue(TEXT("RequiresBuildPlatform"), bRequiresBuildPlatform);
+	}
 
 	Writer.WriteObjectEnd();
 	Writer.Close();
@@ -157,6 +167,7 @@ FString FPluginDescriptor::ToString() const
 FPluginReferenceDescriptor::FPluginReferenceDescriptor( const FString& InName, bool bInEnabled, const FString& InMarketplaceURL )
 	: Name(InName)
 	, bEnabled(bInEnabled)
+	, bOptional(false)
 	, MarketplaceURL(InMarketplaceURL)
 { }
 
@@ -200,7 +211,10 @@ bool FPluginReferenceDescriptor::Read( const FJsonObject& Object, FText& OutFail
 		OutFailReason = LOCTEXT("PluginReferenceWithoutEnabled", "Plugin references must have an 'Enabled' field");
 		return false;
 	}
-	
+
+	// Read the optional field
+	Object.TryGetBoolField(TEXT("Optional"), bOptional);
+
 	// Read the metadata for users that don't have the plugin installed
 	Object.TryGetStringField(TEXT("Description"), Description);
 	Object.TryGetStringField(TEXT("MarketplaceURL"), MarketplaceURL);
@@ -247,6 +261,11 @@ void FPluginReferenceDescriptor::Write( TJsonWriter<>& Writer ) const
 	Writer.WriteValue(TEXT("Name"), Name);
 	Writer.WriteValue(TEXT("Enabled"), bEnabled);
 
+	if (bEnabled && bOptional)
+	{
+		Writer.WriteValue(TEXT("Optional"), bOptional);
+	}
+
 	if (Description.Len() > 0)
 	{
 		Writer.WriteValue(TEXT("Description"), Description);
diff --git a/Engine/Source/Runtime/Projects/Private/PluginManager.cpp b/Engine/Source/Runtime/Projects/Private/PluginManager.cpp
--- a/Engine/Source/Runtime/Projects/Private/PluginManager.cpp
+++ b/Engine/Source/Runtime/Projects/Private/PluginManager.cpp
@@ -344,7 +344,7 @@ bool FPluginManager::ConfigureEnabledPlugins()
 			TArray<FPluginReferenceDescriptor> PluginsCopy = Project->Plugins;
 			for(const FPluginReferenceDescriptor& Plugin: PluginsCopy)
 			{
-				if(Plugin.bEnabled && !FindPluginInstance(Plugin.Name).IsValid())
+				if(Plugin.bEnabled && !FindPluginInstance(Plugin.Name).IsValid() && !Plugin.bOptional)
 				{
 					FText Caption(LOCTEXT("PluginMissingCaption", "Plugin missing"));
 					if(Plugin.MarketplaceURL.Len() > 0)
@@ -489,7 +489,7 @@ bool FPluginManager::ConfigureEnabledPlugins()
 					{
 						if (FCoreDelegates::OnMountPak.IsBound())
 						{
-							FCoreDelegates::OnMountPak.Execute(PakPath, 0);
+							FCoreDelegates::OnMountPak.Execute(PakPath, 0, nullptr);
 						}
 					}
 				}
diff --git a/Engine/Source/Runtime/Projects/Public/PluginDescriptor.h b/Engine/Source/Runtime/Projects/Public/PluginDescriptor.h
--- a/Engine/Source/Runtime/Projects/Public/PluginDescriptor.h
+++ b/Engine/Source/Runtime/Projects/Public/PluginDescriptor.h
@@ -79,6 +79,9 @@ struct PROJECTS_API FPluginDescriptor
 	/** Signifies that the plugin was installed on top of the engine */
 	bool bInstalled;
 
+	/** For plugins that are under a platform folder (eg. /PS4/), determines whether compiling the plugin requires the build platform and/or SDK to be available */
+	bool bRequiresBuildPlatform;
+
 	/** Constructor. */
 	FPluginDescriptor();
 
@@ -105,6 +108,9 @@ struct PROJECTS_API FPluginReferenceDescriptor
 
 	/** Whether it should be enabled by default */
 	bool bEnabled;
+
+	/** Whether this plugin is optional, and the game should silently ignore it not being present */
+	bool bOptional;
 	
 	/** Description of the plugin for users that do not have it installed. */
 	FString Description;
diff --git a/Engine/Source/Runtime/RHI/Private/GPUProfiler.cpp b/Engine/Source/Runtime/RHI/Private/GPUProfiler.cpp
--- a/Engine/Source/Runtime/RHI/Private/GPUProfiler.cpp
+++ b/Engine/Source/Runtime/RHI/Private/GPUProfiler.cpp
@@ -43,6 +43,25 @@ static TAutoConsoleVariable<int32> GProfileGPUTransitions(
 	TEXT("Allows profileGPU to display resource transition events."),
 	ECVF_Default);
 
+enum class EGPUProfileSortMode
+{
+	EChronological,
+	ETimeElapsed,
+	ENumPrims,
+	ENumVerts,
+	EMax
+};
+
+static TAutoConsoleVariable<int32> GProfileGPUSort(
+	TEXT("r.ProfileGPU.Sort"),	
+	0,
+	TEXT("Sorts the TTY Dump independently at each level of the tree in various modes.\n")
+	TEXT("0 : Chronological\n")
+	TEXT("1 : By time elapsed\n")
+	TEXT("2 : By number of prims\n")
+	TEXT("3 : By number of verts\n"),
+	ECVF_Default);
+
 struct FNodeStatsCompare
 {
 	/** Sorts nodes by descending durations. */
@@ -59,32 +78,45 @@ static void GatherStatsEventNode(FGPUProfilerEventNode* Node, int32 Depth, TMap<
 	if (Node->NumDraws > 0 || Node->Children.Num() > 0)
 	{
 		Node->TimingResult = Node->GetTiming() * 1000.0f;
+		Node->NumTotalDraws = Node->NumDraws;
+		Node->NumTotalPrimitives = Node->NumPrimitives;
+		Node->NumTotalVertices = Node->NumVertices;
+
+		FGPUProfilerEventNode* Parent = Node->Parent;		
+		while (Parent)
+		{
+			Parent->NumTotalDraws += Node->NumDraws;
+			Parent->NumTotalPrimitives += Node->NumPrimitives;
+			Parent->NumTotalVertices += Node->NumVertices;
+
+			Parent = Parent->Parent;
+		}
+
+		for (int32 ChildIndex = 0; ChildIndex < Node->Children.Num(); ChildIndex++)
+		{
+			// Traverse children
+			GatherStatsEventNode(Node->Children[ChildIndex], Depth + 1, EventHistogram);
+		}
 
 		FGPUProfilerEventNodeStats* FoundHistogramBucket = EventHistogram.Find(Node->Name);
 		if (FoundHistogramBucket)
 		{
-			FoundHistogramBucket->NumDraws += Node->NumDraws;
-			FoundHistogramBucket->NumPrimitives += Node->NumPrimitives;
-			FoundHistogramBucket->NumVertices += Node->NumVertices;
+			FoundHistogramBucket->NumDraws += Node->NumTotalDraws;
+			FoundHistogramBucket->NumPrimitives += Node->NumTotalPrimitives;
+			FoundHistogramBucket->NumVertices += Node->NumTotalVertices;
 			FoundHistogramBucket->TimingResult += Node->TimingResult;
 			FoundHistogramBucket->NumEvents++;
 		}
 		else
 		{
 			FGPUProfilerEventNodeStats NewNodeStats;
-			NewNodeStats.NumDraws = Node->NumDraws;
-			NewNodeStats.NumPrimitives = Node->NumPrimitives;
-			NewNodeStats.NumVertices = Node->NumVertices;
+			NewNodeStats.NumDraws = Node->NumTotalDraws;
+			NewNodeStats.NumPrimitives = Node->NumTotalPrimitives;
+			NewNodeStats.NumVertices = Node->NumTotalVertices;
 			NewNodeStats.TimingResult = Node->TimingResult;
 			NewNodeStats.NumEvents = 1;
 			EventHistogram.Add(Node->Name, NewNodeStats);
 		}
-
-		for (int32 ChildIndex = 0; ChildIndex < Node->Children.Num(); ChildIndex++)
-		{
-			// Traverse children
-			GatherStatsEventNode(Node->Children[ChildIndex], Depth + 1, EventHistogram);
-		}
 	}
 }
 
@@ -119,13 +151,40 @@ static void DumpStatsEventNode(FGPUProfilerEventNode* Node, float RootResult, in
 				Percent,
 				Node->TimingResult,
 				*Node->Name,
-				Node->NumDraws,
-				Node->NumPrimitives,
-				Node->NumVertices,
+				Node->NumTotalDraws,
+				Node->NumTotalPrimitives,
+				Node->NumTotalVertices,
 				*Extra
 				);
 		}
 
+		struct FCompareGPUProfileNode
+		{
+			EGPUProfileSortMode SortMode;
+			FCompareGPUProfileNode(EGPUProfileSortMode InSortMode)
+				: SortMode(InSortMode)
+			{}
+			FORCEINLINE bool operator()(const FGPUProfilerEventNode* A, const FGPUProfilerEventNode* B) const
+			{
+				switch (SortMode)
+				{
+					case EGPUProfileSortMode::ENumPrims:
+						return B->NumTotalPrimitives < A->NumTotalPrimitives;
+					case EGPUProfileSortMode::ENumVerts:
+						return B->NumTotalVertices < A->NumTotalVertices;
+					case EGPUProfileSortMode::ETimeElapsed:
+					default:
+						return B->TimingResult < A->TimingResult;
+				}
+			}
+		};
+
+		EGPUProfileSortMode SortMode = (EGPUProfileSortMode)FMath::Clamp(GProfileGPUSort.GetValueOnRenderThread(), 0, ((int32)EGPUProfileSortMode::EMax - 1));
+		if (SortMode != EGPUProfileSortMode::EChronological)
+		{
+			Node->Children.Sort(FCompareGPUProfileNode(SortMode));
+		}
+
 		float TotalChildTime = 0;
 		uint32 TotalChildDraws = 0;
 		for (int32 ChildIndex = 0; ChildIndex < Node->Children.Num(); ChildIndex++)
diff --git a/Engine/Source/Runtime/RHI/Private/RHI.cpp b/Engine/Source/Runtime/RHI/Private/RHI.cpp
--- a/Engine/Source/Runtime/RHI/Private/RHI.cpp
+++ b/Engine/Source/Runtime/RHI/Private/RHI.cpp
@@ -225,6 +225,9 @@ int32 GMaxTextureMipCount = MAX_TEXTURE_MIP_COUNT;
 bool GSupportsQuadBufferStereo = false;
 bool GSupportsDepthFetchDuringDepthTest = true;
 FString GRHIAdapterName;
+FString GRHIAdapterInternalDriverVersion;
+FString GRHIAdapterUserDriverVersion;
+FString GRHIAdapterDriverDate;
 uint32 GRHIVendorId = 0;
 bool GSupportsRenderDepthTargetableShaderResources = true;
 bool GSupportsRenderTargetFormat_PF_G8 = true;
@@ -248,6 +251,7 @@ bool GUsingNullRHI = false;
 int32 GDrawUPVertexCheckCount = MAX_int32;
 int32 GDrawUPIndexCheckCount = MAX_int32;
 bool GTriggerGPUProfile = false;
+FString GGPUTraceFileName;
 bool GRHISupportsTextureStreaming = false;
 bool GSupportsDepthBoundsTest = false;
 bool GSupportsEfficientAsyncCompute = false;
diff --git a/Engine/Source/Runtime/RHI/Private/RHICommandList.cpp b/Engine/Source/Runtime/RHI/Private/RHICommandList.cpp
--- a/Engine/Source/Runtime/RHI/Private/RHICommandList.cpp
+++ b/Engine/Source/Runtime/RHI/Private/RHICommandList.cpp
@@ -535,7 +535,8 @@ void FRHICommandListExecutor::LatchBypass()
 
 void FRHICommandListExecutor::CheckNoOutstandingCmdLists()
 {
-	check(GRHICommandList.OutstandingCmdListCount.GetValue() == 2); // else we are attempting to delete resources while there is still a live cmdlist (other than the immediate cmd list) somewhere.
+	// else we are attempting to delete resources while there is still a live cmdlist (other than the immediate cmd list) somewhere.
+	checkf(GRHICommandList.OutstandingCmdListCount.GetValue() == 2, TEXT("Oustanding: %i"), GRHICommandList.OutstandingCmdListCount.GetValue());
 }
 
 bool FRHICommandListExecutor::IsRHIThreadActive()
diff --git a/Engine/Source/Runtime/RHI/Public/DynamicRHI.h b/Engine/Source/Runtime/RHI/Public/DynamicRHI.h
--- a/Engine/Source/Runtime/RHI/Public/DynamicRHI.h
+++ b/Engine/Source/Runtime/RHI/Public/DynamicRHI.h
@@ -518,6 +518,9 @@ public:
 	// FlushType: Wait RHI Thread
 	virtual FGeometryShaderRHIRef RHICreateGeometryShaderWithStreamOutput(const TArray<uint8>& Code, const FStreamOutElementList& ElementList, uint32 NumStrides, const uint32* Strides, int32 RasterizedStream) = 0;
 
+	// Some RHIs can have pending messages/logs for error tracking, or debug modes
+	virtual void FlushPendingLogs() {}
+
 	// FlushType: Wait RHI Thread
 	virtual FComputeShaderRHIRef RHICreateComputeShader(const TArray<uint8>& Code) = 0;
 
diff --git a/Engine/Source/Runtime/RHI/Public/GPUProfiler.h b/Engine/Source/Runtime/RHI/Public/GPUProfiler.h
--- a/Engine/Source/Runtime/RHI/Public/GPUProfiler.h
+++ b/Engine/Source/Runtime/RHI/Public/GPUProfiler.h
@@ -15,6 +15,9 @@ public:
 		NumDraws(0),
 		NumPrimitives(0),
 		NumVertices(0),
+		NumTotalDraws(0),
+		NumTotalPrimitives(0),
+		NumTotalVertices(0),
 		TimingResult(0),
 		NumEvents(0)
 	{
@@ -29,6 +32,15 @@ public:
 	/** Exclusive number of vertices rendered in this event. */
 	uint32 NumVertices;
 
+	/** Inclusive number of draw calls rendered in this event and children. */
+	uint32 NumTotalDraws;
+
+	/** Inclusive number of primitives rendered in this event and children. */
+	uint32 NumTotalPrimitives;
+
+	/** Inclusive number of vertices rendered in this event and children. */
+	uint32 NumTotalVertices;
+
 	/** GPU time spent inside the perf event's begin and end, in ms. */
 	float TimingResult;
 
@@ -40,6 +52,9 @@ public:
 		NumDraws += rhs.NumDraws;
 		NumPrimitives += rhs.NumPrimitives;
 		NumVertices += rhs.NumVertices;
+		NumTotalDraws += rhs.NumDraws;
+		NumTotalPrimitives += rhs.NumPrimitives;
+		NumTotalVertices += rhs.NumVertices;
 		TimingResult += rhs.TimingResult;
 		NumEvents += rhs.NumEvents;
 
diff --git a/Engine/Source/Runtime/RHI/Public/RHI.h b/Engine/Source/Runtime/RHI/Public/RHI.h
--- a/Engine/Source/Runtime/RHI/Public/RHI.h
+++ b/Engine/Source/Runtime/RHI/Public/RHI.h
@@ -73,6 +73,9 @@ extern RHI_API bool GSupportsDepthFetchDuringDepthTest;
  * e.g. "NVIDIA GeForce GTX 670"
  */
 extern RHI_API FString GRHIAdapterName;
+extern RHI_API FString GRHIAdapterInternalDriverVersion;
+extern RHI_API FString GRHIAdapterUserDriverVersion;
+extern RHI_API FString GRHIAdapterDriverDate;
 
 // 0 means not defined yet, use functions like IsRHIDeviceAMD() to access
 extern RHI_API uint32 GRHIVendorId;
@@ -192,6 +195,9 @@ extern RHI_API bool GTriggerGPUProfile;
 /** Whether we are profiling GPU hitches. */
 extern RHI_API bool GTriggerGPUHitchProfile;
 
+/** Non-empty if we are performing a gpu trace. Also says where to place trace file. */
+extern RHI_API FString GGPUTraceFileName;
+
 /** True if the RHI supports texture streaming */
 extern RHI_API bool GRHISupportsTextureStreaming;
 /** Amount of memory allocated by textures. In kilobytes. */
@@ -1134,6 +1140,8 @@ struct FTextureMemoryStats
 
 	// Size of allocated memory, in bytes
 	int64 AllocatedMemorySize;
+	// Size of the largest memory fragment, in bytes
+	int64 LargestContiguousAllocation;
 	// 0 if streaming pool size limitation is disabled, in bytes
 	int64 TexturePoolSize;
 	// Upcoming adjustments to allocated memory, in bytes (async reallocations)
@@ -1146,6 +1154,7 @@ struct FTextureMemoryStats
 		, SharedSystemMemory(-1)
 		, TotalGraphicsMemory(-1)
 		, AllocatedMemorySize(0)
+		, LargestContiguousAllocation(0)
 		, TexturePoolSize(0)
 		, PendingMemoryAdjustment(0)
 	{
diff --git a/Engine/Source/Runtime/RenderCore/Private/RenderingThread.cpp b/Engine/Source/Runtime/RenderCore/Private/RenderingThread.cpp
--- a/Engine/Source/Runtime/RenderCore/Private/RenderingThread.cpp
+++ b/Engine/Source/Runtime/RenderCore/Private/RenderingThread.cpp
@@ -402,6 +402,18 @@ public:
 		GRenderThreadId = 0;
 	}
 
+#if PLATFORM_WINDOWS && !PLATFORM_SEH_EXCEPTIONS_DISABLED
+	static int32 FlushRHILogsAndReportCrash(LPEXCEPTION_POINTERS ExceptionInfo)
+	{
+		if (GDynamicRHI)
+		{
+			GDynamicRHI->FlushPendingLogs();
+		}
+
+		return ReportCrash(ExceptionInfo);
+	}
+#endif
+
 	virtual uint32 Run(void) override
 	{
 		FPlatformProcess::SetupGameOrRenderThread(true);
@@ -416,7 +428,7 @@ public:
 				RenderingThreadMain( TaskGraphBoundSyncEvent );
 			}
 #if !PLATFORM_SEH_EXCEPTIONS_DISABLED
-			__except( ReportCrash( GetExceptionInformation() ) )
+			__except(FlushRHILogsAndReportCrash(GetExceptionInformation()))
 			{
 				GRenderingThreadError = GErrorHist;
 
diff --git a/Engine/Source/Runtime/Renderer/Private/CapsuleShadowRendering.cpp b/Engine/Source/Runtime/Renderer/Private/CapsuleShadowRendering.cpp
--- a/Engine/Source/Runtime/Renderer/Private/CapsuleShadowRendering.cpp
+++ b/Engine/Source/Runtime/Renderer/Private/CapsuleShadowRendering.cpp
@@ -746,7 +746,9 @@ bool FDeferredShadingSceneRenderer::RenderCapsuleDirectShadows(
 				}
 
 				{
-					SCOPED_DRAW_EVENT(RHICmdList, Upsample);
+					SCOPED_DRAW_EVENTF(RHICmdList, Upsample, TEXT("UpsampleDirectCapsuleShadowRendering %dx%d"),
+						ScissorRect.Width(), ScissorRect.Height());
+						
 					FSceneRenderTargets::Get(RHICmdList).BeginRenderingLightAttenuation(RHICmdList);
 
 					RHICmdList.SetViewport(View.ViewRect.Min.X, View.ViewRect.Min.Y, 0.0f, View.ViewRect.Max.X, View.ViewRect.Max.Y, 1.0f);
@@ -1052,7 +1054,8 @@ void FDeferredShadingSceneRenderer::RenderIndirectCapsuleShadows(FRHICommandList
 						}
 			
 						{
-							SCOPED_DRAW_EVENT(RHICmdList, Upsample);
+							SCOPED_DRAW_EVENTF(RHICmdList, Upsample, TEXT("UpsampleIndirectCapsuleShadowRendering %dx%d"),
+								ScissorRect.Width(), ScissorRect.Height());
 
 							FSceneRenderTargets& SceneContext = FSceneRenderTargets::Get(RHICmdList);
 
diff --git a/Engine/Source/Runtime/Renderer/Private/CompositionLighting/PostProcessAmbientOcclusion.cpp b/Engine/Source/Runtime/Renderer/Private/CompositionLighting/PostProcessAmbientOcclusion.cpp
--- a/Engine/Source/Runtime/Renderer/Private/CompositionLighting/PostProcessAmbientOcclusion.cpp
+++ b/Engine/Source/Runtime/Renderer/Private/CompositionLighting/PostProcessAmbientOcclusion.cpp
@@ -20,7 +20,7 @@ static TAutoConsoleVariable<int32> CVarAmbientOcclusionCompute(
 	TEXT("r.AmbientOcclusion.Compute"),
 	0,
 	TEXT("If SSAO should use ComputeShader (not available on all platforms) or PixelShader.\n")
-	TEXT("The Compute Shader version is WIP, not optimized, requires hardware support (not mobile/DX10/OpenGL3),")
+	TEXT("The Compute Shader version is WIP, not optimized, requires hardware support (not mobile/DX10/OpenGL3),\n")
 	TEXT("does not use normals which allows it to run right after EarlyZPass (better performance when used with AyncCompute)\n")
 	TEXT(" 0: PixelShader (default)\n")
 	TEXT(" 1: (WIP) Use ComputeShader if possible, otherwise fall back to '0'\n")
diff --git a/Engine/Source/Runtime/Renderer/Private/CustomDepthRendering.cpp b/Engine/Source/Runtime/Renderer/Private/CustomDepthRendering.cpp
--- a/Engine/Source/Runtime/Renderer/Private/CustomDepthRendering.cpp
+++ b/Engine/Source/Runtime/Renderer/Private/CustomDepthRendering.cpp
@@ -34,7 +34,7 @@ bool FCustomDepthPrimSet::DrawPrims(FRHICommandListImmediate& RHICmdList, const
 			{
 				const FPrimitiveViewRelevance& ViewRelevance = View.PrimitiveViewRelevanceMap[PrimitiveSceneInfo->GetIndex()];
 
-				FDepthDrawingPolicyFactory::ContextType Context(DDM_AllOccluders);
+				FDepthDrawingPolicyFactory::ContextType Context(DDM_AllOpaque);
 
 				if (bWriteCustomStencilValues)
 				{
@@ -65,7 +65,7 @@ bool FCustomDepthPrimSet::DrawPrims(FRHICommandListImmediate& RHICmdList, const
 							bDirty |= FDepthDrawingPolicyFactory::DrawStaticMesh(
 								RHICmdList, 
 								View,
-								FDepthDrawingPolicyFactory::ContextType(DDM_AllOccluders),
+								Context,
 								StaticMesh,
 								StaticMesh.Elements.Num() == 1 ? 1 : View.StaticMeshBatchVisibility[StaticMesh.Id],
 								true,
diff --git a/Engine/Source/Runtime/Renderer/Private/DeferredShadingRenderer.cpp b/Engine/Source/Runtime/Renderer/Private/DeferredShadingRenderer.cpp
--- a/Engine/Source/Runtime/Renderer/Private/DeferredShadingRenderer.cpp
+++ b/Engine/Source/Runtime/Renderer/Private/DeferredShadingRenderer.cpp
@@ -48,12 +48,20 @@ static FAutoConsoleVariableRef CVarEarlyZPassMovable(
 	ECVF_RenderThreadSafe | ECVF_ReadOnly
 	);
 
+TAutoConsoleVariable<int32> CVarCustomDepthOrder(
+	TEXT("r.CustomDepth.Order"),
+	1,	
+	TEXT("When CustomDepth (and CustomStencil) is getting rendered\n")
+	TEXT("  0: Before GBuffer (can be more efficient with AsyncCompute, allows using it in DBuffer pass, no GBuffer blending decals allow GBuffer compression)\n")
+	TEXT("  1: After Base Pass (default)"),
+	ECVF_RenderThreadSafe);
+
 #if !(UE_BUILD_SHIPPING || UE_BUILD_TEST)
 static TAutoConsoleVariable<int32> CVarVisualizeTexturePool(
 	TEXT("r.VisualizeTexturePool"),
 	0,
 	TEXT("Allows to enable the visualize the texture pool (currently only on console).\n")
-	TEXT(" 0: off\n")
+	TEXT(" 0: off (default)\n")
 	TEXT(" 1: on"),
 	ECVF_Cheat | ECVF_RenderThreadSafe);
 #endif
@@ -809,6 +817,12 @@ void ServiceLocalQueue()
 	FTaskGraphInterface::Get().ProcessThreadUntilIdle(ENamedThreads::RenderThread_Local);
 }
 
+// @return 0/1
+static int32 GetCustomDepthPassLocation()
+{		
+	return FMath::Clamp(CVarCustomDepthOrder.GetValueOnRenderThread(), 0, 1);
+}
+
 extern bool IsLpvIndirectPassRequired(const FViewInfo& View);
 
 static TAutoConsoleVariable<float> CVarStallInitViews(
@@ -1018,7 +1032,7 @@ void FDeferredShadingSceneRenderer::Render(FRHICommandListImmediate& RHICmdList)
 	}
 	RenderOcclusion(RHICmdList, bOcclusionBeforeBasePass, bHZBBeforeBasePass);
 	ServiceLocalQueue();
-	
+
 	// Clear LPVs for all views
 	if (FeatureLevel >= ERHIFeatureLevel::SM5)
 	{
@@ -1027,7 +1041,11 @@ void FDeferredShadingSceneRenderer::Render(FRHICommandListImmediate& RHICmdList)
 		ServiceLocalQueue();
 	}
 
-	RenderCustomDepthPassAtLocation(RHICmdList, 0);
+	if(GetCustomDepthPassLocation() == 0)
+	{
+		QUICK_SCOPE_CYCLE_COUNTER(STAT_FDeferredShadingSceneRenderer_CustomDepthPass0);
+		RenderCustomDepthPassAtLocation(RHICmdList, 0);
+	}
 
 	// only temporarily available after early z pass and until base pass
 	check(!SceneContext.DBufferA);
@@ -1129,7 +1147,13 @@ void FDeferredShadingSceneRenderer::Render(FRHICommandListImmediate& RHICmdList)
 		SceneContext.FinishRenderingGBuffer(RHICmdList);
 	}
 
-	RenderCustomDepthPassAtLocation(RHICmdList, 1);
+	if(GetCustomDepthPassLocation() == 1)
+	{
+		QUICK_SCOPE_CYCLE_COUNTER(STAT_FDeferredShadingSceneRenderer_CustomDepthPass1);
+		RenderCustomDepthPassAtLocation(RHICmdList, 1);
+	}
+
+	ServiceLocalQueue();
 
 	// Notify the FX system that opaque primitives have been rendered and we now have a valid depth buffer.
 	if (Scene->FXSystem && Views.IsValidIndex(0))
diff --git a/Engine/Source/Runtime/Renderer/Private/LightGrid.h b/Engine/Source/Runtime/Renderer/Private/LightGrid.h
--- a/Engine/Source/Runtime/Renderer/Private/LightGrid.h
+++ b/Engine/Source/Runtime/Renderer/Private/LightGrid.h
@@ -17,6 +17,13 @@ class FLightGridVertexBuffer : public FVertexBuffer
 public:
 	const static uint32 MaxSize = 512 * 1024;
 	TArray<uint32> CPUData;
+	// from CPUData.Num()
+	uint32 AllocatedSize;
+
+	FLightGridVertexBuffer()
+		: AllocatedSize(0)
+	{
+	}
 
 	/**
 	 * Initializes the vertex buffer from the input
@@ -24,15 +31,21 @@ public:
 	void UpdateGPUFromCPUData()
 	{
 		check( IsInRenderingThread() );
-		InitResource();
+
 		if(CPUData.Num())
 		{
-			check( IsInRenderingThread() );
-
 			const uint32 Size = CPUData.Num();
 
 			check(Size <= MaxSize);
 
+			// this can cause reallocations, we could presize CPUData in the constructor if that is an issue
+			if(Size > AllocatedSize)
+			{
+				// recreate
+				ReleaseResource();
+				InitResource();
+			}
+
 			const uint32 BufferSize = Size * sizeof(uint32);
 
 			uint32* RESTRICT Src = &CPUData[0];
@@ -57,17 +70,24 @@ private:
 	{
 		VertexBufferSRV.SafeRelease();
 		FVertexBuffer::ReleaseRHI();
+		AllocatedSize = 0;
 	}
 
 	/** Initialize RHI resources. */
 	virtual void InitRHI() override
 	{
-		const int32 BufferStride = sizeof(uint32);
-		const int32 BufferSize = MaxSize * BufferStride;
-		uint32 Flags = BUF_Volatile | /*BUF_KeepCPUAccessible | */BUF_ShaderResource;
-		FRHIResourceCreateInfo CreateInfo;
-		VertexBufferRHI = RHICreateVertexBuffer(BufferSize, Flags, CreateInfo);
-		VertexBufferSRV = RHICreateShaderResourceView(VertexBufferRHI, BufferStride, PF_R32_UINT);
+		check(!AllocatedSize);
+
+		if(CPUData.Num())
+		{
+			const int32 BufferStride = sizeof(uint32);
+			const int32 BufferSize = MaxSize * BufferStride;
+			uint32 Flags = BUF_Volatile | /*BUF_KeepCPUAccessible | */BUF_ShaderResource;
+			FRHIResourceCreateInfo CreateInfo;
+			VertexBufferRHI = RHICreateVertexBuffer(BufferSize, Flags, CreateInfo);
+			VertexBufferSRV = RHICreateShaderResourceView(VertexBufferRHI, BufferStride, PF_R32_UINT);
+			AllocatedSize = CPUData.Num();
+		}
 	}
 };
 
diff --git a/Engine/Source/Runtime/Renderer/Private/PostProcess/PostProcessTonemap.cpp b/Engine/Source/Runtime/Renderer/Private/PostProcess/PostProcessTonemap.cpp
--- a/Engine/Source/Runtime/Renderer/Private/PostProcess/PostProcessTonemap.cpp
+++ b/Engine/Source/Runtime/Renderer/Private/PostProcess/PostProcessTonemap.cpp
@@ -1159,8 +1159,6 @@ void FRCPassPostProcessTonemap::Process(FRenderingCompositePassContext& Context)
 	// Set the view family's render target/viewport.
 	SetRenderTarget(Context.RHICmdList, DestRenderTarget.TargetableTexture, FTextureRHIParamRef(), ESimpleRenderTargetMode::EUninitializedColorAndDepth);
 
-	Context.SetViewportAndCallRHI(0, 0, 0.0f, PassOutputs[0].RenderTargetDesc.Extent.X, PassOutputs[0].RenderTargetDesc.Extent.Y, 1.0f);
-		
 	if (Context.HasHmdMesh() && View.StereoPass == eSSP_LEFT_EYE)
 	{
 		// needed when using an hmd mesh instead of a full screen quad because we don't touch all of the pixels in the render target
@@ -1172,6 +1170,8 @@ void FRCPassPostProcessTonemap::Process(FRenderingCompositePassContext& Context)
 		Context.RHICmdList.Clear(true, FLinearColor::Black, false, 1.0f, false, 0, DestRect);
 	}
 
+	Context.SetViewportAndCallRHI(DestRect.Min.X, DestRect.Min.Y, 0.0f, DestRect.Max.X + 1, DestRect.Max.Y + 1, 1.0f );
+
 	// set the state
 	Context.RHICmdList.SetBlendState(TStaticBlendState<>::GetRHI());
 	Context.RHICmdList.SetRasterizerState(TStaticRasterizerState<>::GetRHI());
@@ -1211,11 +1211,11 @@ void FRCPassPostProcessTonemap::Process(FRenderingCompositePassContext& Context)
 
 	DrawPostProcessPass(
 		Context.RHICmdList,
-		DestRect.Min.X, DestRect.Min.Y,
+		0, 0,
 		DestRect.Width(), DestRect.Height(),
 		View.ViewRect.Min.X, View.ViewRect.Min.Y,
 		View.ViewRect.Width(), View.ViewRect.Height(),
-		PassOutputs[0].RenderTargetDesc.Extent,
+		DestRect.Size(),
 		SceneContext.GetBufferSizeXY(),
 		VertexShader,
 		View.StereoPass,
diff --git a/Engine/Source/Runtime/Renderer/Private/PostProcess/PostProcessWeightedSampleSum.cpp b/Engine/Source/Runtime/Renderer/Private/PostProcess/PostProcessWeightedSampleSum.cpp
--- a/Engine/Source/Runtime/Renderer/Private/PostProcess/PostProcessWeightedSampleSum.cpp
+++ b/Engine/Source/Runtime/Renderer/Private/PostProcess/PostProcessWeightedSampleSum.cpp
@@ -547,7 +547,8 @@ void FRCPassPostProcessWeightedSampleSum::Process(FRenderingCompositePassContext
 
 	FIntRect DestRect = FIntRect::DivideAndRoundUp(View.ViewRect, DstScaleFactor);
 
-	SCOPED_DRAW_EVENTF(Context.RHICmdList, PostProcessWeightedSampleSum, TEXT("PostProcessWeightedSampleSum#%d %dx%d"), NumSamples, DestRect.Width(), DestRect.Height());
+	SCOPED_DRAW_EVENTF(Context.RHICmdList, PostProcessWeightedSampleSum, TEXT("PostProcessWeightedSampleSum#%d %dx%d in %dx%d"),
+		NumSamples, DestRect.Width(), DestRect.Height(), DestSize.X, DestSize.Y);
 
 	// compute weights as weighted contributions of the TintValue
 	for(uint32 i = 0; i < NumSamples; ++i)
diff --git a/Engine/Source/Runtime/Renderer/Private/SceneRendering.cpp b/Engine/Source/Runtime/Renderer/Private/SceneRendering.cpp
--- a/Engine/Source/Runtime/Renderer/Private/SceneRendering.cpp
+++ b/Engine/Source/Runtime/Renderer/Private/SceneRendering.cpp
@@ -71,14 +71,6 @@ static TAutoConsoleVariable<int32> CVarInstancedStereo(
 	TEXT("0 to disable instanced stereo (default), 1 to enable."),
 	ECVF_ReadOnly | ECVF_RenderThreadSafe);
 
-TAutoConsoleVariable<int32> CVarCustomDepthOrder(
-	TEXT("r.CustomDepth.Order"),
-	1,	
-	TEXT("When CustomDepth (and CustomStencil) is getting rendered\n")
-	TEXT("  0: Before GBuffer (can be more efficient with AsyncCompute, allows using it in DBuffer pass, no GBuffer blending decals allow GBuffer compression)\n")
-	TEXT("  1: After Base Pass (default)"),
-	ECVF_RenderThreadSafe);
-
 #if !(UE_BUILD_SHIPPING || UE_BUILD_TEST)
 static TAutoConsoleVariable<float> CVarGeneralPurposeTweak(
 	TEXT("r.GeneralPurposeTweak"),
@@ -1531,6 +1523,7 @@ void ServiceLocalQueue();
 
 void FSceneRenderer::RenderCustomDepthPassAtLocation(FRHICommandListImmediate& RHICmdList, int32 Location)
 {		
+	extern TAutoConsoleVariable<int32> CVarCustomDepthOrder;
 	int32 CustomDepthOrder = FMath::Clamp(CVarCustomDepthOrder.GetValueOnRenderThread(), 0, 1);
 
 	if(CustomDepthOrder == Location)
diff --git a/Engine/Source/Runtime/Renderer/Private/ShaderBaseClasses.cpp b/Engine/Source/Runtime/Renderer/Private/ShaderBaseClasses.cpp
--- a/Engine/Source/Runtime/Renderer/Private/ShaderBaseClasses.cpp
+++ b/Engine/Source/Runtime/Renderer/Private/ShaderBaseClasses.cpp
@@ -180,7 +180,7 @@ void FMaterialShader::SetParameters(
 	{
 		UE_LOG(
 			LogShaders,
-			Warning,	// TEMP workaround only!!!!
+			Fatal,	// TEMP workaround only!!!!
 			TEXT("%s shader uniform expression set mismatch for material %s/%s.\n")
 			TEXT("Shader compilation info:                %s\n")
 			TEXT("Material render proxy compilation info: %s\n")
@@ -224,13 +224,30 @@ void FMaterialShader::SetParameters(
 		const TArray<FGuid>& ParameterCollections = UniformExpressionCache->ParameterCollections;
 		const int32 ParameterCollectionsNum = ParameterCollections.Num();
 
+		// For shipping and test builds the assert above will be compiled out, but we're trying to verify that this condition is never hit.
+		if (ParameterCollectionUniformBuffers.Num() < ParameterCollectionsNum)
+		{
+			UE_LOG(LogRenderer, Warning,
+				TEXT("ParameterCollectionUniformBuffers.Num() [%u] < ParameterCollectionsNum [%u], this would crash below on SetUniformBufferParameter.\n")
+				TEXT("RenderProxy=%s Material=%s"),
+				ParameterCollectionUniformBuffers.Num(),
+				ParameterCollectionsNum,
+				*MaterialRenderProxy->GetFriendlyName(),
+				*Material.GetFriendlyName()
+				);
+		}
+
 		check(ParameterCollectionUniformBuffers.Num() >= ParameterCollectionsNum);
 
+		
+
+		int32 NumToSet = FMath::Min(ParameterCollectionUniformBuffers.Num(), ParameterCollections.Num());
+
 		// Find each referenced parameter collection's uniform buffer in the scene and set the parameter
-		for (int32 CollectionIndex = 0; CollectionIndex < ParameterCollectionsNum; CollectionIndex++)
-		{
+		for (int32 CollectionIndex = 0; CollectionIndex < NumToSet; CollectionIndex++)
+		{			
 			FUniformBufferRHIParamRef UniformBuffer = GetParameterCollectionBuffer(ParameterCollections[CollectionIndex], View.Family->Scene);
-			SetUniformBufferParameter(RHICmdList, ShaderRHI,ParameterCollectionUniformBuffers[CollectionIndex],UniformBuffer);
+			SetUniformBufferParameter(RHICmdList, ShaderRHI, ParameterCollectionUniformBuffers[CollectionIndex], UniformBuffer);			
 		}
 	}
 
diff --git a/Engine/Source/Runtime/Renderer/Private/ShadowRendering.cpp b/Engine/Source/Runtime/Renderer/Private/ShadowRendering.cpp
--- a/Engine/Source/Runtime/Renderer/Private/ShadowRendering.cpp
+++ b/Engine/Source/Runtime/Renderer/Private/ShadowRendering.cpp
@@ -4203,7 +4203,7 @@ void FForwardShadingSceneRenderer::RenderModulatedShadowDepthMaps(FRHICommandLis
 	if(bShadowsDepthsAllocated)
 	{
 		// Render the shadow depths.
-		SCOPED_DRAW_EVENT(RHICmdList, ShadowDepthsFromOpaqueProjected);
+		SCOPED_DRAW_EVENT(RHICmdList, ModulatedShadowDepthsFromOpaqueProjected);
 	
 		bool bPerformClear = true;
 		auto SetShadowRenderTargets = [&bPerformClear, &SceneContext](FRHICommandList& InRHICmdList)
diff --git a/Engine/Source/Runtime/Slate/Private/Framework/Application/SlateApplication.cpp b/Engine/Source/Runtime/Slate/Private/Framework/Application/SlateApplication.cpp
--- a/Engine/Source/Runtime/Slate/Private/Framework/Application/SlateApplication.cpp
+++ b/Engine/Source/Runtime/Slate/Private/Framework/Application/SlateApplication.cpp
@@ -1216,15 +1216,18 @@ void FSlateApplication::PrivateDrawWindows( TSharedPtr<SWindow> DrawOnlyThisWind
 		}
 	}
 
-
-	// Some windows may have been destroyed/removed.
-	// Do not attempt to draw any windows that have been removed.
-	TArray<TSharedRef<SWindow>> AllWindows = GatherAllDescendants(SlateWindows);
-	DrawWindowArgs.OutDrawBuffer.GetWindowElementLists().RemoveAll([&]( TSharedPtr<FSlateWindowElementList>& Candidate )
+	// This is potentially dangerous on the movie playback thread that slate sometimes runs on
+	if(!IsInSlateThread())
 	{
-		TSharedPtr<SWindow> CandidateWindow = Candidate->GetWindow();
-		return !CandidateWindow.IsValid() || !AllWindows.Contains(CandidateWindow.ToSharedRef());
-	});
+		// Some windows may have been destroyed/removed.
+		// Do not attempt to draw any windows that have been removed.
+		TArray<TSharedRef<SWindow>> AllWindows = GatherAllDescendants(SlateWindows);
+		DrawWindowArgs.OutDrawBuffer.GetWindowElementLists().RemoveAll([&](TSharedPtr<FSlateWindowElementList>& Candidate)
+		{
+			TSharedPtr<SWindow> CandidateWindow = Candidate->GetWindow();
+			return !CandidateWindow.IsValid() || !AllWindows.Contains(CandidateWindow.ToSharedRef());
+		});
+	}
 
 	{	
 		SLATE_CYCLE_COUNTER_SCOPE(GSlateRendererDrawWindows);
diff --git a/Engine/Source/Runtime/Slate/Private/Widgets/Input/SEditableText.cpp b/Engine/Source/Runtime/Slate/Private/Widgets/Input/SEditableText.cpp
--- a/Engine/Source/Runtime/Slate/Private/Widgets/Input/SEditableText.cpp
+++ b/Engine/Source/Runtime/Slate/Private/Widgets/Input/SEditableText.cpp
@@ -178,6 +178,9 @@ void SEditableText::SetTextFromVirtualKeyboard(const FText& InNewText, ESetTextT
 		// fired on the game thread in Tick.
 		bTextChangedByVirtualKeyboard = true;
 	}
+
+	// Clear focus so that virtual keyboard can be reopened again my clicking on the text box
+	FSlateApplication::Get().ClearKeyboardFocus(EFocusCause::Cleared);
 }
 
 void SEditableText::RestoreOriginalText()
diff --git a/Engine/Source/Runtime/Slate/Private/Widgets/Layout/SSeparator.cpp b/Engine/Source/Runtime/Slate/Private/Widgets/Layout/SSeparator.cpp
--- a/Engine/Source/Runtime/Slate/Private/Widgets/Layout/SSeparator.cpp
+++ b/Engine/Source/Runtime/Slate/Private/Widgets/Layout/SSeparator.cpp
@@ -12,6 +12,7 @@ void SSeparator::Construct( const FArguments& InArgs )
 		SBorder::FArguments()
 		.BorderImage(InArgs._SeparatorImage)
 		.Padding(0.0f)
+		.BorderBackgroundColor(InArgs._ColorAndOpacity)
 	);
 
 	Orientation = InArgs._Orientation;
diff --git a/Engine/Source/Runtime/Slate/Private/Widgets/Layout/SWidgetSwitcher.cpp b/Engine/Source/Runtime/Slate/Private/Widgets/Layout/SWidgetSwitcher.cpp
--- a/Engine/Source/Runtime/Slate/Private/Widgets/Layout/SWidgetSwitcher.cpp
+++ b/Engine/Source/Runtime/Slate/Private/Widgets/Layout/SWidgetSwitcher.cpp
@@ -114,6 +114,7 @@ int32 SWidgetSwitcher::RemoveSlot( TSharedRef<SWidget> WidgetToRemove )
 void SWidgetSwitcher::SetActiveWidgetIndex( int32 Index )
 {
 	WidgetIndex = Index;
+	Invalidate(EInvalidateWidget::LayoutAndVolatility);
 }
 
 
diff --git a/Engine/Source/Runtime/Slate/Private/Widgets/Text/SMultiLineEditableText.cpp b/Engine/Source/Runtime/Slate/Private/Widgets/Text/SMultiLineEditableText.cpp
--- a/Engine/Source/Runtime/Slate/Private/Widgets/Text/SMultiLineEditableText.cpp
+++ b/Engine/Source/Runtime/Slate/Private/Widgets/Text/SMultiLineEditableText.cpp
@@ -502,6 +502,9 @@ void SMultiLineEditableText::SetTextFromVirtualKeyboard(const FText& InNewText,
 		VirtualKeyboardTextCommitType = CommitType;
 		bTextCommittedByVirtualKeyboard = true;
 	}
+
+	// Clear focus so that virtual keyboard can be reopened again my clicking on the text box
+	FSlateApplication::Get().ClearKeyboardFocus(EFocusCause::Cleared);
 }
 
 bool SMultiLineEditableText::SetEditableText(const FText& TextToSet, const bool bForce)
diff --git a/Engine/Source/Runtime/Slate/Public/Widgets/Layout/SSeparator.h b/Engine/Source/Runtime/Slate/Public/Widgets/Layout/SSeparator.h
--- a/Engine/Source/Runtime/Slate/Public/Widgets/Layout/SSeparator.h
+++ b/Engine/Source/Runtime/Slate/Public/Widgets/Layout/SSeparator.h
@@ -10,6 +10,7 @@ public:
 		: _SeparatorImage( FCoreStyle::Get().GetBrush("Separator") )
 		, _Orientation(Orient_Horizontal)
 		, _Thickness(3.0f)
+		, _ColorAndOpacity(FLinearColor::White)
 		{}
 		
 		SLATE_ARGUMENT(const FSlateBrush*, SeparatorImage)
@@ -17,6 +18,9 @@ public:
 		SLATE_ARGUMENT( EOrientation, Orientation)
 
 		SLATE_ARGUMENT(float, Thickness)
+
+		/** Color and opacity */
+		SLATE_ATTRIBUTE(FSlateColor, ColorAndOpacity)
 	SLATE_END_ARGS()
 
 	/**
diff --git a/Engine/Source/Runtime/SlateCore/Private/Fonts/FontCache.cpp b/Engine/Source/Runtime/SlateCore/Private/Fonts/FontCache.cpp
--- a/Engine/Source/Runtime/SlateCore/Private/Fonts/FontCache.cpp
+++ b/Engine/Source/Runtime/SlateCore/Private/Fonts/FontCache.cpp
@@ -527,7 +527,6 @@ FSlateFontCache::FSlateFontCache( TSharedRef<ISlateFontAtlasFactory> InFontAtlas
 	, TextShaper( new FSlateTextShaper( FTGlyphCache.Get(), FTAdvanceCache.Get(), FTKerningPairCache.Get(), CompositeFontCache.Get(), FontRenderer.Get(), this ) )
 	, FontAtlasFactory( InFontAtlasFactory )
 	, bFlushRequested( false )
-	, LastFlushHistoryVersion( FTextLocalizationManager::Get().GetTextRevision() )
 {
 	UE_LOG(LogSlate, Log, TEXT("SlateFontCache - WITH_FREETYPE: %d, WITH_HARFBUZZ: %d"), WITH_FREETYPE, WITH_HARFBUZZ);
 
@@ -691,17 +690,6 @@ FShapedGlyphSequenceRef FSlateFontCache::ShapeUnidirectionalText( const TCHAR* I
 
 FCharacterList& FSlateFontCache::GetCharacterList( const FSlateFontInfo &InFontInfo, float FontScale ) const
 {
-	{
-		// If the active culture has changed, then the localized fallback font may also have changed
-		// We need to flush the cache *immediately* in order to avoid getting bad data back from the character list cache
-		// We don't do this immediate flush in HandleCultureChanged because that may be called from the wrong thread
-		const int32 CurrentHistoryVersion = FTextLocalizationManager::Get().GetTextRevision();
-		if (LastFlushHistoryVersion != CurrentHistoryVersion)
-		{
-			FlushCache();
-		}
-	}
-
 	// Create a key for looking up each character
 	const FSlateFontKey FontKey( InFontInfo, FontScale );
 
@@ -802,6 +790,11 @@ const TSet<FName>& FSlateFontCache::GetFontAttributes( const FFontData& InFontDa
 	return CompositeFontCache->GetFontAttributes(InFontData);
 }
 
+void FSlateFontCache::RequestFlushCache()
+{
+	bFlushRequested = true;
+}
+
 void FSlateFontCache::FlushObject( const UObject* const InObject )
 {
 	if( !InObject )
@@ -872,8 +865,6 @@ void FSlateFontCache::FlushCache() const
 
 		FontAtlases.Empty();
 
-		LastFlushHistoryVersion = FTextLocalizationManager::Get().GetTextRevision();
-
 		UE_LOG(LogSlate, Verbose, TEXT("Slate font cache was flushed"));
 	}
 	else
diff --git a/Engine/Source/Runtime/SlateCore/Private/Fonts/FontCacheCompositeFont.cpp b/Engine/Source/Runtime/SlateCore/Private/Fonts/FontCacheCompositeFont.cpp
--- a/Engine/Source/Runtime/SlateCore/Private/Fonts/FontCacheCompositeFont.cpp
+++ b/Engine/Source/Runtime/SlateCore/Private/Fonts/FontCacheCompositeFont.cpp
@@ -204,6 +204,8 @@ TSharedPtr<FFreeTypeFace> FCompositeFontCache::GetFontFace(const FFontData& InFo
 	TSharedPtr<FFreeTypeFace> FaceAndMemory = FontFaceMap.FindRef(&InFontData);
 	if (!FaceAndMemory.IsValid() && InFontData.BulkDataPtr)
 	{
+		FScopeCycleCounterUObject ContextScope(InFontData.BulkDataPtr);
+
 		int32 LockedFontDataSizeBytes = 0;
 		const void* const LockedFontData = InFontData.BulkDataPtr->Lock(LockedFontDataSizeBytes);
 		if (LockedFontDataSizeBytes > 0)
diff --git a/Engine/Source/Runtime/SlateCore/Private/Fonts/LegacySlateFontInfoCache.cpp b/Engine/Source/Runtime/SlateCore/Private/Fonts/LegacySlateFontInfoCache.cpp
--- a/Engine/Source/Runtime/SlateCore/Private/Fonts/LegacySlateFontInfoCache.cpp
+++ b/Engine/Source/Runtime/SlateCore/Private/Fonts/LegacySlateFontInfoCache.cpp
@@ -10,10 +10,18 @@ FLegacySlateFontInfoCache& FLegacySlateFontInfoCache::Get()
 	if (!Singleton.IsValid())
 	{
 		Singleton = MakeShareable(new FLegacySlateFontInfoCache());
+
+		FInternationalization::Get().OnCultureChanged().AddSP(Singleton.Get(), &FLegacySlateFontInfoCache::HandleCultureChanged);
 	}
 	return *Singleton;
 }
 
+FLegacySlateFontInfoCache::FLegacySlateFontInfoCache()
+	: LocalizedFallbackFontDataHistoryVersion(INDEX_NONE)
+	, LocalizedFallbackFontFrameCounter(0)
+{
+}
+
 TSharedPtr<const FCompositeFont> FLegacySlateFontInfoCache::GetCompositeFont(const FName& InLegacyFontName, const EFontHinting InLegacyFontHinting)
 {
 	if (InLegacyFontName.IsNone())
@@ -82,13 +90,17 @@ const FFontData& FLegacySlateFontInfoCache::GetLocalizedFallbackFontData()
 
 	// The fallback font can change if the active culture is changed
 	const int32 CurrentHistoryVersion = FTextLocalizationManager::Get().GetTextRevision();
+	const uint64 CurrentFrameCounter = GFrameCounter;
 
-	if (!LocalizedFallbackFontData.IsValid() || LocalizedFallbackFontDataHistoryVersion != CurrentHistoryVersion)
+	// Only allow the fallback font to be updated once per-frame, as a culture change mid-frame could cause it to change unexpectedly and invalidate some assumptions that the font cache makes
+	// By only allowing it to update once per-frame, we ensure that the font cache has been flushed (which happens at the end of the frame) before we return a new font
+	if (!LocalizedFallbackFontData.IsValid() || (LocalizedFallbackFontDataHistoryVersion != CurrentHistoryVersion && LocalizedFallbackFontFrameCounter != CurrentFrameCounter))
 	{
 		// Don't allow GC while we perform this allocation
 		FGCScopeGuard GCGuard;
 
 		LocalizedFallbackFontDataHistoryVersion = CurrentHistoryVersion;
+		LocalizedFallbackFontFrameCounter = CurrentFrameCounter;
 
 		const FString FallbackFontPath = FPaths::EngineContentDir() / TEXT("Slate/Fonts/") / (NSLOCTEXT("Slate", "FallbackFont", "DroidSansFallback").ToString() + TEXT(".ttf"));
 		LocalizedFallbackFontData = AllLocalizedFallbackFontData.FindRef(FallbackFontPath);
@@ -153,3 +165,9 @@ void FLegacySlateFontInfoCache::AddReferencedObjects(FReferenceCollector& Collec
 		Collector.AddReferencedObject(TmpPtr);
 	}
 }
+
+void FLegacySlateFontInfoCache::HandleCultureChanged()
+{
+	// We set this to the current frame count, as this will prevent the fallback font being updated for the remainder of this frame (as the culture change may have affected the fallback font used)
+	LocalizedFallbackFontFrameCounter = GFrameCounter;
+}
diff --git a/Engine/Source/Runtime/SlateCore/Private/Fonts/LegacySlateFontInfoCache.h b/Engine/Source/Runtime/SlateCore/Private/Fonts/LegacySlateFontInfoCache.h
--- a/Engine/Source/Runtime/SlateCore/Private/Fonts/LegacySlateFontInfoCache.h
+++ b/Engine/Source/Runtime/SlateCore/Private/Fonts/LegacySlateFontInfoCache.h
@@ -5,10 +5,9 @@
 /**
  * Cache used to efficiently upgrade legacy FSlateFontInfo structs to use a composite font by reducing the amount of duplicate instances that are created
  */
-class FLegacySlateFontInfoCache : public FGCObject
+class FLegacySlateFontInfoCache : public FGCObject, public TSharedFromThis<FLegacySlateFontInfoCache>
 {
 public:
-
 	/**
 	 * Get (or create) the singleton instance of this cache
 	 */
@@ -43,6 +42,13 @@ public:
 	virtual void AddReferencedObjects(FReferenceCollector& Collector) override;
 
 private:
+	FLegacySlateFontInfoCache();
+
+	/**
+	 * Called after the active culture has changed
+	 */
+	void HandleCultureChanged();
+
 	struct FLegacyFontKey
 	{
 		FLegacyFontKey()
@@ -88,6 +94,7 @@ private:
 
 	TMap<FString, TSharedPtr<const FFontData>> AllLocalizedFallbackFontData;
 	int32 LocalizedFallbackFontDataHistoryVersion;
+	uint64 LocalizedFallbackFontFrameCounter;
 
 	FCriticalSection LocalizedFallbackFontDataCS;
 	FCriticalSection LastResortFontCS;
diff --git a/Engine/Source/Runtime/SlateCore/Private/Rendering/DrawElements.cpp b/Engine/Source/Runtime/SlateCore/Private/Rendering/DrawElements.cpp
--- a/Engine/Source/Runtime/SlateCore/Private/Rendering/DrawElements.cpp
+++ b/Engine/Source/Runtime/SlateCore/Private/Rendering/DrawElements.cpp
@@ -664,9 +664,9 @@ void FSlateWindowElementList::EndLogicalLayer()
 	DrawStack.Pop();
 }
 
-FSlateRenderDataHandle::FSlateRenderDataHandle(const ILayoutCache* InCacher, FSlateRenderer* InRenderer)
+FSlateRenderDataHandle::FSlateRenderDataHandle(const ILayoutCache* InCacher, ISlateRenderDataManager* InManager)
 	: Cacher(InCacher)
-	, Renderer(InRenderer)
+	, Manager(InManager)
 	, RenderBatches(nullptr)
 	, UsageCount(0)
 {
@@ -674,15 +674,15 @@ FSlateRenderDataHandle::FSlateRenderDataHandle(const ILayoutCache* InCacher, FSl
 
 FSlateRenderDataHandle::~FSlateRenderDataHandle()
 {
-	if ( Renderer )
+	if ( Manager )
 	{
-		Renderer->ReleaseCachedRenderData(this);
+		Manager->BeginReleasingRenderData(this);
 	}
 }
 
 void FSlateRenderDataHandle::Disconnect()
 {
-	Renderer = nullptr;
+	Manager = nullptr;
 	RenderBatches = nullptr;
 }
 
diff --git a/Engine/Source/Runtime/SlateCore/Private/Rendering/SlateRenderer.cpp b/Engine/Source/Runtime/SlateCore/Private/Rendering/SlateRenderer.cpp
--- a/Engine/Source/Runtime/SlateCore/Private/Rendering/SlateRenderer.cpp
+++ b/Engine/Source/Runtime/SlateCore/Private/Rendering/SlateRenderer.cpp
@@ -71,14 +71,14 @@ void FSlateFontServices::FlushFontCache()
 
 void FSlateFontServices::FlushGameThreadFontCache()
 {
-	GameThreadFontCache->FlushCache();
+	GameThreadFontCache->RequestFlushCache();
 	GameThreadFontMeasure->FlushCache();
 }
 
 
 void FSlateFontServices::FlushRenderThreadFontCache()
 {
-	RenderThreadFontCache->FlushCache();
+	RenderThreadFontCache->RequestFlushCache();
 	RenderThreadFontMeasure->FlushCache();
 }
 
@@ -136,7 +136,7 @@ ISlateAtlasProvider* FSlateRenderer::GetFontAtlasProvider()
 
 TSharedRef<FSlateRenderDataHandle, ESPMode::ThreadSafe> FSlateRenderer::CacheElementRenderData(const ILayoutCache* Cacher, FSlateWindowElementList& ElementList)
 {
-	return MakeShareable(new FSlateRenderDataHandle(Cacher, this));
+	return MakeShareable(new FSlateRenderDataHandle(Cacher, nullptr));
 }
 
 void FSlateRenderer::ReleaseCachingResourcesFor(const ILayoutCache* Cacher)
diff --git a/Engine/Source/Runtime/SlateCore/Public/Fonts/FontCache.h b/Engine/Source/Runtime/SlateCore/Public/Fonts/FontCache.h
--- a/Engine/Source/Runtime/SlateCore/Public/Fonts/FontCache.h
+++ b/Engine/Source/Runtime/SlateCore/Public/Fonts/FontCache.h
@@ -668,9 +668,9 @@ public:
 	const TSet<FName>& GetFontAttributes( const FFontData& InFontData ) const;
 
 	/**
-	 * Clears all cached data from the cache
+	 * Issues a request to clear all cached data from the cache
 	 */
-	void FlushCache() const;
+	void RequestFlushCache();
 
 	/**
 	 * Clears just the cached font data, but leaves the atlases alone
@@ -682,6 +682,11 @@ private:
 	FSlateFontCache(const FSlateFontCache&);
 	FSlateFontCache& operator=(const FSlateFontCache&);
 
+	/**
+	 * Clears all cached data from the cache
+	 */
+	void FlushCache() const;
+
 	/** Called after the active culture has changed */
 	void HandleCultureChanged();
 
@@ -722,7 +727,4 @@ private:
 
 	/** Whether or not we have a pending request to flush the cache when it is safe to do so */
 	volatile mutable bool bFlushRequested;
-
-	/** Active culture version the last time the font cache was flushed */
-	mutable int32 LastFlushHistoryVersion;
 };
diff --git a/Engine/Source/Runtime/SlateCore/Public/Rendering/DrawElements.h b/Engine/Source/Runtime/SlateCore/Public/Rendering/DrawElements.h
--- a/Engine/Source/Runtime/SlateCore/Public/Rendering/DrawElements.h
+++ b/Engine/Source/Runtime/SlateCore/Public/Rendering/DrawElements.h
@@ -506,17 +506,23 @@ struct FShaderParams
 class FSlateRenderer;
 class FSlateRenderBatch;
 
+class ISlateRenderDataManager
+{
+public:
+	virtual void BeginReleasingRenderData(const FSlateRenderDataHandle* RenderHandle) = 0;
+};
+
 class SLATECORE_API FSlateRenderDataHandle : public TSharedFromThis < FSlateRenderDataHandle, ESPMode::ThreadSafe >
 {
 public:
-	FSlateRenderDataHandle(const ILayoutCache* Cacher, FSlateRenderer* InRenderer);
+	FSlateRenderDataHandle(const ILayoutCache* Cacher, ISlateRenderDataManager* InManager);
 
 	virtual ~FSlateRenderDataHandle();
-
+	
 	void Disconnect();
 
 	const ILayoutCache* GetCacher() const { return Cacher; }
-	const FSlateRenderer* GetRenderer() const { return Renderer; }
+	//const FSlateRenderer* GetRenderer() const { return Renderer; }
 
 	void SetRenderBatches(TArray<FSlateRenderBatch>* InRenderBatches) { RenderBatches = InRenderBatches; }
 	TArray<FSlateRenderBatch>* GetRenderBatches() { return RenderBatches; }
@@ -528,7 +534,7 @@ public:
 
 private:
 	const ILayoutCache* Cacher;
-	FSlateRenderer* Renderer;
+	ISlateRenderDataManager* Manager;
 	TArray<FSlateRenderBatch>* RenderBatches;
 
 	volatile int32 UsageCount;
diff --git a/Engine/Source/Runtime/SlateCore/Public/Rendering/SlateRenderer.h b/Engine/Source/Runtime/SlateCore/Public/Rendering/SlateRenderer.h
--- a/Engine/Source/Runtime/SlateCore/Public/Rendering/SlateRenderer.h
+++ b/Engine/Source/Runtime/SlateCore/Public/Rendering/SlateRenderer.h
@@ -408,9 +408,6 @@ public:
 	virtual void MapVirtualScreenBuffer(FMappedTextureBuffer* OutImageData) {}
 	virtual void UnmapVirtualScreenBuffer() {}
 
-protected:
-	virtual void ReleaseCachedRenderData(FSlateRenderDataHandle* RenderHandle) { }
-
 private:
 
 	// Non-copyable
diff --git a/Engine/Source/Runtime/SlateRHIRenderer/Private/SlateRHIRenderer.cpp b/Engine/Source/Runtime/SlateRHIRenderer/Private/SlateRHIRenderer.cpp
--- a/Engine/Source/Runtime/SlateRHIRenderer/Private/SlateRHIRenderer.cpp
+++ b/Engine/Source/Runtime/SlateRHIRenderer/Private/SlateRHIRenderer.cpp
@@ -1262,7 +1262,7 @@ void FSlateRHIRenderer::SetWindowRenderTarget(const SWindow& Window, IViewportRe
 
 TSharedRef<FSlateRenderDataHandle, ESPMode::ThreadSafe> FSlateRHIRenderer::CacheElementRenderData(const ILayoutCache* Cacher, FSlateWindowElementList& ElementList)
 {
-	TSharedRef<FSlateRenderDataHandle, ESPMode::ThreadSafe> RenderDataHandle = MakeShareable(new FSlateRenderDataHandle(Cacher, this));
+	TSharedRef<FSlateRenderDataHandle, ESPMode::ThreadSafe> RenderDataHandle = MakeShareable(new FSlateRenderDataHandle(Cacher, ResourceManager.Get()));
 
 	checkSlow(ElementList.GetChildDrawLayers().Num() == 0);
 
@@ -1319,25 +1319,6 @@ void FSlateRHIRenderer::ReleaseCachingResourcesFor(const ILayoutCache* Cacher)
 		});
 }
 
-void FSlateRHIRenderer::ReleaseCachedRenderData(FSlateRenderDataHandle* InRenderHandle)
-{
-	struct FReleaseCachedRenderDataContext
-	{
-		FSlateRHIRenderingPolicy* RenderPolicy;
-		FSlateRenderDataHandle* RenderDataHandle;
-	};
-	FReleaseCachedRenderDataContext ReleaseCachedRenderDataContext =
-	{
-		RenderingPolicy.Get(),
-		InRenderHandle,
-	};
-	ENQUEUE_UNIQUE_RENDER_COMMAND_ONEPARAMETER(
-		ReleaseCachedRenderData,
-		FReleaseCachedRenderDataContext, Context, ReleaseCachedRenderDataContext,
-	{
-		Context.RenderPolicy->ReleaseCachedRenderData(Context.RenderDataHandle);
-	});
-}
 
 FSlateEndDrawingWindowsCommand::FSlateEndDrawingWindowsCommand(FSlateRHIRenderingPolicy& InPolicy, FSlateDrawBuffer* InDrawBuffer)
 	: Policy(InPolicy)
