Copying //UE4/Orion-Staging to //UE4/Main (Source: //Orion/Dev-General @ 2845681)

diff --git a/Engine/Source/Runtime/Engine/Private/PhysicsEngine/PhysXSupport.h b/Engine/Source/Runtime/Engine/Private/PhysicsEngine/PhysXSupport.h
--- a/Engine/Source/Runtime/Engine/Private/PhysicsEngine/PhysXSupport.h
+++ b/Engine/Source/Runtime/Engine/Private/PhysicsEngine/PhysXSupport.h
@@ -246,7 +246,7 @@ class FPhysXAllocator : public PxAllocatorCallback
 		:	AllocationTypeName(InAllocationTypeName)
 		,	AllocationSize(InAllocationSize)
 		{
-			static_assert(sizeof(FPhysXAllocationHeader) == 32, "FPhysXAllocationHeader size must be 32 bytes.");
+			static_assert((sizeof(FPhysXAllocationHeader) % 16) == 0, "FPhysXAllocationHeader size must multiple of bytes.");
 			MagicPadding();
 		}
 
@@ -272,6 +272,7 @@ class FPhysXAllocator : public PxAllocatorCallback
 		FName AllocationTypeName;
 		size_t	AllocationSize;
 		uint8 Padding[8];	//physx needs 16 byte alignment. Additionally we fill padding with a pattern to see if there's any memory stomps
+		uint8 Padding2[(sizeof(FName) + sizeof(size_t) + sizeof(Padding)) % 16];
 
 		void Validate() const
 		{
diff --git a/Engine/Source/Runtime/Engine/Private/PlayerCameraManager.cpp b/Engine/Source/Runtime/Engine/Private/PlayerCameraManager.cpp
--- a/Engine/Source/Runtime/Engine/Private/PlayerCameraManager.cpp
+++ b/Engine/Source/Runtime/Engine/Private/PlayerCameraManager.cpp
@@ -240,15 +240,15 @@ void APlayerCameraManager::ApplyCameraModifiers(float DeltaTime, FMinimalViewInf
 			}
 		}
 
+		// changes to this are good for a single update, so reset this to 1.f after processing
+		AnimInst->TransientScaleModifier = 1.f;
+
 		// handle animations that have finished
-		if (AnimInst->bFinished && AnimInst->bAutoReleaseWhenFinished)
+		if (AnimInst->bFinished)
 		{
 			ReleaseCameraAnimInst(AnimInst);
 			Idx--;		// we removed this from the ActiveAnims array
 		}
-
-		// changes to this are good for a single update, so reset this to 1.f after processing
-		AnimInst->TransientScaleModifier = 1.f;
 	}
 
 	// need to zero this when we are done with it.  playing another animation
diff --git a/Engine/Source/Runtime/Engine/Private/StaticMesh.cpp b/Engine/Source/Runtime/Engine/Private/StaticMesh.cpp
--- a/Engine/Source/Runtime/Engine/Private/StaticMesh.cpp
+++ b/Engine/Source/Runtime/Engine/Private/StaticMesh.cpp
@@ -1856,6 +1856,7 @@ FUObjectAnnotationSparseBool GStaticMeshesThatNeedMaterialFixup;
 COREUOBJECT_API extern bool GOutputCookingWarnings;
 #endif
 
+
 /**
  *	UStaticMesh::Serialize
  */
@@ -1887,9 +1888,9 @@ void UStaticMesh::Serialize(FArchive& Ar)
 			bHasNavigationData && 
 			(NavCollision == nullptr))
 		{
-			if (Ar.IsPersistent() && Ar.IsLoading() && GOutputCookingWarnings)
+			if (Ar.IsPersistent() && Ar.IsLoading() && (Ar.GetDebugSerializationFlags() & DSF_EnableCookerWarnings))
 			{
-					UE_LOG(LogStaticMesh, Warning, TEXT("Serialized NavCollision but it was null (%s) NavCollision will be created dynamicaly at cook time.  Please resave pacage %s."), *GetName(), *GetOutermost()->GetPathName())
+				UE_LOG(LogStaticMesh, Warning, TEXT("Serialized NavCollision but it was null (%s) NavCollision will be created dynamicaly at cook time.  Please resave package %s."), *GetName(), *GetOutermost()->GetPathName());
 			}
 		}
 #endif
diff --git a/Engine/Source/Runtime/Engine/Private/UnrealEngine.cpp b/Engine/Source/Runtime/Engine/Private/UnrealEngine.cpp
--- a/Engine/Source/Runtime/Engine/Private/UnrealEngine.cpp
+++ b/Engine/Source/Runtime/Engine/Private/UnrealEngine.cpp
@@ -1136,9 +1136,71 @@ void PumpABTest()
 #endif
 }
 
+#if !(UE_BUILD_SHIPPING || UE_BUILD_TEST)
+
+struct FTimedMemReport
+{
+	FTimedMemReport()
+	: TotalTime(0.f)
+	, DumpDelayTime(-1.f)
+	{
+		if(const TCHAR* CommandLine = FCommandLine::Get())
+		{
+			FParse::Value(CommandLine, TEXT("TimedMemoryReport="), DumpDelayTime);
+		}
+	}
+
+	static FTimedMemReport& Get()
+	{
+		static FTimedMemReport Singleton;
+		return Singleton;
+	}
+
+	static void SetDumpDelayParse(const TArray<FString>& Args)
+	{
+		if(Args.Num())
+		{
+			float DumpDelay = FCString::Atof(*Args[0]);
+			Get().SetDumpDelay(DumpDelay);
+		}
+	}
+
+	void SetDumpDelay(float InDumpDelay)
+	{
+		DumpDelayTime = InDumpDelay;
+		TotalTime = 0.f;	//reset time
+	}
+
+	void PumpTimedMemoryReports()
+	{
+		if (DumpDelayTime > 0)
+		{
+			TotalTime += FApp::GetDeltaTime();
+			if (TotalTime > DumpDelayTime)
+			{
+				GEngine->Exec(nullptr, TEXT("memreport"), *GLog);
+				TotalTime = 0.f;
+			}
+		}
+	}
+
+private:
+	float TotalTime;
+	float DumpDelayTime;
+};
+
+static FAutoConsoleCommand SetTimedMemReport(TEXT("TimedMemReport.Delay"), TEXT("Determines how long to wait before getting a memreport. < 0 is off"), FConsoleCommandWithArgsDelegate::CreateStatic(&FTimedMemReport::SetDumpDelayParse), ECVF_Cheat);
+
+
+
+#endif
+
 void UEngine::UpdateTimeAndHandleMaxTickRate()
 {
 	PumpABTest();
+#if !(UE_BUILD_SHIPPING || UE_BUILD_TEST)
+	FTimedMemReport::Get().PumpTimedMemoryReports();
+#endif
 
 	// start at now minus a bit so we don't get a zero delta.
 	static double LastTime = FPlatformTime::Seconds() - 0.0001;
@@ -6967,26 +7029,29 @@ void UEngine::LogPerformanceCapture(UWorld* World, const FString& CaptureName)
 
 	extern ENGINE_API float GAverageFPS;
 
-	const FStatUnitData* StatUnitData = World->GetGameViewport()->GetStatUnitData();
-//	const FStatHitchesData* StatHitchesData = World->GetGameViewport->GetStatHitchesData();
-
-	FAutomationPerformanceSnapshot PerfSnapshot;
-	PerfSnapshot.Changelist = FString::FromInt( ChangeList );
-	PerfSnapshot.BuildConfiguration = EBuildConfigurations::ToString( FApp::GetBuildConfiguration() );
-	PerfSnapshot.MapName = MapName;
-	PerfSnapshot.MatineeName = MatineeName;
-	PerfSnapshot.AverageFPS = FString::Printf( TEXT( "%0.2f" ), GAverageFPS );
-	PerfSnapshot.AverageFrameTime = FString::Printf( TEXT( "%0.2f" ), StatUnitData->FrameTime );
-	PerfSnapshot.AverageGameThreadTime = FString::Printf( TEXT( "%0.2f" ), StatUnitData->GameThreadTime );
-	PerfSnapshot.AverageRenderThreadTime = FString::Printf( TEXT( "%0.2f" ), StatUnitData->RenderThreadTime );
-	PerfSnapshot.AverageGPUTime = FString::Printf( TEXT( "%0.2f" ), StatUnitData->GPUFrameTime );
-	// PerfSnapshot.PercentOfFramesAtLeast60FPS = ???;	// @todo
-	// PerfSnapshot.PercentOfFramesAtLeast60FPS = ???;	// @todo
-		
-	const FString PerfSnapshotAsCommaDelimitedString = PerfSnapshot.ToCommaDelimetedString();
-	
-	FAutomationTestFramework::GetInstance().AddAnalyticsItemToCurrentTest( 
-		FString::Printf( TEXT( "%s,%s" ), *PerfSnapshotAsCommaDelimitedString, *EventType ) );
+	if ((World != nullptr) && (World->GetGameViewport() != nullptr))
+	{
+		const FStatUnitData* StatUnitData = World->GetGameViewport()->GetStatUnitData();
+		//	const FStatHitchesData* StatHitchesData = World->GetGameViewport->GetStatHitchesData();
+
+		FAutomationPerformanceSnapshot PerfSnapshot;
+		PerfSnapshot.Changelist = FString::FromInt(ChangeList);
+		PerfSnapshot.BuildConfiguration = EBuildConfigurations::ToString(FApp::GetBuildConfiguration());
+		PerfSnapshot.MapName = MapName;
+		PerfSnapshot.MatineeName = MatineeName;
+		PerfSnapshot.AverageFPS = FString::Printf(TEXT("%0.2f"), GAverageFPS);
+		PerfSnapshot.AverageFrameTime = FString::Printf(TEXT("%0.2f"), StatUnitData->FrameTime);
+		PerfSnapshot.AverageGameThreadTime = FString::Printf(TEXT("%0.2f"), StatUnitData->GameThreadTime);
+		PerfSnapshot.AverageRenderThreadTime = FString::Printf(TEXT("%0.2f"), StatUnitData->RenderThreadTime);
+		PerfSnapshot.AverageGPUTime = FString::Printf(TEXT("%0.2f"), StatUnitData->GPUFrameTime);
+		// PerfSnapshot.PercentOfFramesAtLeast60FPS = ???;	// @todo
+		// PerfSnapshot.PercentOfFramesAtLeast60FPS = ???;	// @todo
+
+		const FString PerfSnapshotAsCommaDelimitedString = PerfSnapshot.ToCommaDelimetedString();
+
+		FAutomationTestFramework::GetInstance().AddAnalyticsItemToCurrentTest(
+			FString::Printf(TEXT("%s,%s"), *PerfSnapshotAsCommaDelimitedString, *EventType));
+	}
 }
 #endif	// UE_BUILD_SHIPPING
 
diff --git a/Engine/Source/Runtime/Engine/Private/UserActivityTracking.cpp b/Engine/Source/Runtime/Engine/Private/UserActivityTracking.cpp
--- a/Engine/Source/Runtime/Engine/Private/UserActivityTracking.cpp
+++ b/Engine/Source/Runtime/Engine/Private/UserActivityTracking.cpp
@@ -7,3 +7,9 @@ FUserActivityTracking::FOnActivityChanged FUserActivityTracking::OnActivityChang
 
 FUserActivity FUserActivityTracking::UserActivity;
 
+void FUserActivityTracking::SetActivity(const FUserActivity& InUserActivity)
+{ 
+	UserActivity = InUserActivity;
+	OnActivityChanged.Broadcast(UserActivity);
+	FCoreDelegates::UserActivityStringChanged.Broadcast(UserActivity.ActionName);
+}
diff --git a/Engine/Source/Runtime/Engine/Private/UserDefinedEnum.cpp b/Engine/Source/Runtime/Engine/Private/UserDefinedEnum.cpp
--- a/Engine/Source/Runtime/Engine/Private/UserDefinedEnum.cpp
+++ b/Engine/Source/Runtime/Engine/Private/UserDefinedEnum.cpp
@@ -68,7 +68,11 @@ void UUserDefinedEnum::PostLoad()
 {
 	Super::PostLoad();
 	FEnumEditorUtils::UpdateAfterPathChanged(this);
-	FEnumEditorUtils::EnsureAllDisplayNamesExist(this);
+
+	if (GIsEditor && !IsRunningCommandlet())
+	{
+		FEnumEditorUtils::EnsureAllDisplayNamesExist(this);
+	}
 }
 
 void UUserDefinedEnum::PostEditUndo()
diff --git a/Engine/Source/Runtime/Engine/Private/VorbisAudioInfo.cpp b/Engine/Source/Runtime/Engine/Private/VorbisAudioInfo.cpp
--- a/Engine/Source/Runtime/Engine/Private/VorbisAudioInfo.cpp
+++ b/Engine/Source/Runtime/Engine/Private/VorbisAudioInfo.cpp
@@ -83,6 +83,8 @@ FVorbisAudioInfo::FVorbisAudioInfo( void )
 	,	SrcBufferDataSize(0)
 	,	BufferOffset(0)
 { 
+	// Make sure we have properly allocated a VFWrapper
+	check(VFWrapper != NULL);
 }
 
 FVorbisAudioInfo::~FVorbisAudioInfo( void ) 
@@ -94,6 +96,7 @@ FVorbisAudioInfo::~FVorbisAudioInfo( void )
 /** Emulate read from memory functionality */
 size_t FVorbisAudioInfo::Read( void *Ptr, uint32 Size )
 {
+	check(Ptr);
 	size_t BytesToRead = FMath::Min( Size, SrcBufferDataSize - BufferOffset );
 	FMemory::Memcpy( Ptr, SrcBufferData + BufferOffset, BytesToRead );
 	BufferOffset += BytesToRead;
@@ -102,7 +105,9 @@ size_t FVorbisAudioInfo::Read( void *Ptr, uint32 Size )
 
 static size_t OggRead( void *ptr, size_t size, size_t nmemb, void *datasource )
 {
-	FVorbisAudioInfo* OggInfo = ( FVorbisAudioInfo* )datasource;
+	check(ptr);
+	check(datasource);
+	FVorbisAudioInfo* OggInfo = (FVorbisAudioInfo*)datasource;
 	return( OggInfo->Read( ptr, size * nmemb ) );
 }
 
@@ -121,6 +126,9 @@ int FVorbisAudioInfo::Seek( uint32 offset, int whence )
 	case SEEK_END:
 		BufferOffset = SrcBufferDataSize - offset;
 		break;
+	default:
+		checkf(false, TEXT("Uknown seek type"));
+		break;
 	}
 
 	return( BufferOffset );
@@ -176,9 +184,11 @@ bool FVorbisAudioInfo::ReadCompressedInfo( const uint8* InSrcBufferData, uint32
 	Callbacks.tell_func = OggTell;
 
 	// Set up the read from memory variables
-	if (ov_open_callbacks(this, &VFWrapper->vf, NULL, 0, Callbacks) < 0)
+	int Result = ov_open_callbacks(this, &VFWrapper->vf, NULL, 0, Callbacks);
+	if (Result < 0)
 	{
-		return( false );
+		UE_LOG(LogAudio, Error, TEXT("FVorbisAudioInfo::ReadCompressedInfo, ov_open_callbacks error code: %d"), Result);
+		return false;
 	}
 
 	if( QualityInfo )
diff --git a/Engine/Source/Runtime/Engine/Public/ContentStreaming.h b/Engine/Source/Runtime/Engine/Public/ContentStreaming.h
--- a/Engine/Source/Runtime/Engine/Public/ContentStreaming.h
+++ b/Engine/Source/Runtime/Engine/Public/ContentStreaming.h
@@ -24,6 +24,7 @@ DECLARE_MEMORY_STAT_POOL_EXTERN(TEXT("Textures On Disk"),STAT_StreamingTexturesM
 DECLARE_MEMORY_STAT_POOL_EXTERN(TEXT("Lightmaps In Memory"),STAT_LightmapMemorySize,STATGROUP_StreamingDetails,FPlatformMemory::MCR_TexturePool, );
 DECLARE_MEMORY_STAT_POOL_EXTERN(TEXT("Lightmaps On Disk"),STAT_LightmapDiskSize,STATGROUP_StreamingDetails,FPlatformMemory::MCR_TexturePool, );
 DECLARE_MEMORY_STAT_POOL_EXTERN(TEXT("HLOD Textures In Memory"), STAT_HLODTextureMemorySize, STATGROUP_StreamingDetails, FPlatformMemory::MCR_TexturePool, );
+DECLARE_MEMORY_STAT_POOL_EXTERN(TEXT("HLOD Textures On Disk"), STAT_HLODTextureDiskSize, STATGROUP_StreamingDetails, FPlatformMemory::MCR_TexturePool, );
 DECLARE_DWORD_ACCUMULATOR_STAT_EXTERN(TEXT("Requests In Cancelation Phase"),STAT_RequestsInCancelationPhase,STATGROUP_StreamingDetails, );
 DECLARE_DWORD_ACCUMULATOR_STAT_EXTERN(TEXT("Requests In Update Phase"),STAT_RequestsInUpdatePhase,STATGROUP_StreamingDetails, );
 DECLARE_DWORD_ACCUMULATOR_STAT_EXTERN(TEXT("Requests in Finalize Phase"),STAT_RequestsInFinalizePhase,STATGROUP_StreamingDetails, );
@@ -483,6 +484,9 @@ struct ITextureStreamingManager : public IStreamingManager
 
 	/** Resets the max ever required textures.  For possibly when changing resolutions or screen pct. */
 	virtual void ResetMaxEverRequired() = 0;
+
+	/** Set current pause state for texture streaming */
+	virtual void PauseTextureStreaming(bool bInShouldPause) = 0;
 };
 
 /**
@@ -1094,6 +1098,11 @@ struct FStreamingManagerTexture : public ITextureStreamingManager
 	static float BandwidthMinimum;
 #endif
 
+	/** Set current pause state for texture streaming */
+	virtual void PauseTextureStreaming(bool bInShouldPause) override
+	{
+		bPauseTextureStreaming = bInShouldPause;
+	}
 
 protected:
 //BEGIN: Thread-safe functions and data
@@ -1351,6 +1360,8 @@ protected:
 	uint64 TotalLightmapDiskSize;
 	/** Total number of bytes in memory, for all HLOD textures. */
 	uint64 TotalHLODMemorySize;
+	/** Total number of bytes on disk, for all HLOD textures. */
+	uint64 TotalHLODDiskSize;
 	/** Number of mip count increase requests in flight. */
 	uint32 TotalMipCountIncreaseRequestsInFlight;
 	/** Total number of bytes in memory, if all textures were streamed perfectly with a 1.0 fudge factor. */
diff --git a/Engine/Source/Runtime/Engine/Public/EngineStats.h b/Engine/Source/Runtime/Engine/Public/EngineStats.h
--- a/Engine/Source/Runtime/Engine/Public/EngineStats.h
+++ b/Engine/Source/Runtime/Engine/Public/EngineStats.h
@@ -142,6 +142,7 @@ DECLARE_DWORD_ACCUMULATOR_STAT_EXTERN(TEXT("Ping"),STAT_Ping,STATGROUP_Net, );
 DECLARE_DWORD_ACCUMULATOR_STAT_EXTERN(TEXT("Channels"),STAT_Channels,STATGROUP_Net, );
 DECLARE_DWORD_ACCUMULATOR_STAT_EXTERN(TEXT("In Rate (bytes)"),STAT_InRate,STATGROUP_Net, );
 DECLARE_DWORD_ACCUMULATOR_STAT_EXTERN(TEXT("Out Rate (bytes)"),STAT_OutRate,STATGROUP_Net, );
+DECLARE_DWORD_ACCUMULATOR_STAT_EXTERN( TEXT( "Out Saturation (%)" ),STAT_OutSaturation, STATGROUP_Net,);
 DECLARE_DWORD_ACCUMULATOR_STAT_EXTERN(TEXT("In Rate (bytes) - max among clients"),STAT_InRateClientMax,STATGROUP_Net, );
 DECLARE_DWORD_ACCUMULATOR_STAT_EXTERN(TEXT("In Rate (bytes) - min among clients"),STAT_InRateClientMin,STATGROUP_Net, );
 DECLARE_DWORD_ACCUMULATOR_STAT_EXTERN(TEXT("In Rate (bytes) - avg among clients"),STAT_InRateClientAvg,STATGROUP_Net, );
diff --git a/Engine/Source/Runtime/Engine/Public/MaterialShared.h b/Engine/Source/Runtime/Engine/Public/MaterialShared.h
--- a/Engine/Source/Runtime/Engine/Public/MaterialShared.h
+++ b/Engine/Source/Runtime/Engine/Public/MaterialShared.h
@@ -1148,9 +1148,9 @@ public:
 	ENGINE_API bool GetMaterialExpressionSource(FString& OutSource);
 
 	/* Helper function to look at both IsMasked and IsDitheredLODTransition to determine if it writes every pixel */
-	ENGINE_API bool WritesEveryPixel() const
+	ENGINE_API bool WritesEveryPixel( bool bShadowPass = true ) const
 	{
-		return !IsMasked() && !IsDitheredLODTransition() && !IsWireframe();
+		return !IsMasked() && !( IsDitheredLODTransition() && bShadowPass ) && !IsWireframe();
 	}
 
 	/** 
diff --git a/Engine/Source/Runtime/Engine/Public/UserActivityTracking.h b/Engine/Source/Runtime/Engine/Public/UserActivityTracking.h
--- a/Engine/Source/Runtime/Engine/Public/UserActivityTracking.h
+++ b/Engine/Source/Runtime/Engine/Public/UserActivityTracking.h
@@ -7,7 +7,7 @@ class ENGINE_API FUserActivityTracking : FNoncopyable
 public:
 	DECLARE_MULTICAST_DELEGATE_OneParam(FOnActivityChanged, const FUserActivity&);
 
-	static void SetActivity(const FUserActivity& InUserActivity) { UserActivity = InUserActivity; }
+	static void SetActivity(const FUserActivity& InUserActivity);
 	static const FUserActivity& GetUserActivity() { return UserActivity; }
 
 	// Called when the UserActivity changes
diff --git a/Engine/Source/Runtime/GameplayAbilities/Private/AbilitySystemComponent.cpp b/Engine/Source/Runtime/GameplayAbilities/Private/AbilitySystemComponent.cpp
--- a/Engine/Source/Runtime/GameplayAbilities/Private/AbilitySystemComponent.cpp
+++ b/Engine/Source/Runtime/GameplayAbilities/Private/AbilitySystemComponent.cpp
@@ -38,6 +38,9 @@ UAbilitySystemComponent::UAbilitySystemComponent(const FObjectInitializer& Objec
 
 	GenericConfirmInputID = INDEX_NONE;
 	GenericCancelInputID = INDEX_NONE;
+
+	bSuppressGrantAbility = false;
+	bSuppressGameplayCues = false;
 }
 
 UAbilitySystemComponent::~UAbilitySystemComponent()
@@ -634,7 +637,7 @@ FActiveGameplayEffectHandle UAbilitySystemComponent::ApplyGameplayEffectSpecToSe
 	
 
 	// We still probably want to apply tags and stuff even if instant?
-	if (bInvokeGameplayCueApplied && AppliedEffect && !AppliedEffect->bIsInhibited)
+	if (!bSuppressGameplayCues && bInvokeGameplayCueApplied && AppliedEffect && !AppliedEffect->bIsInhibited)
 	{
 		// We both added and activated the GameplayCue here.
 		// On the client, who will invoke the gameplay cue from an OnRep, he will need to look at the StartTime to determine
@@ -665,7 +668,10 @@ FActiveGameplayEffectHandle UAbilitySystemComponent::ApplyGameplayEffectSpecToSe
 		// This is an instant application but we are treating it as an infinite duration for prediction. We should still predict the execute GameplayCUE.
 		// (in non predictive case, this will happen inside ::ExecuteGameplayEffect)
 
-		UAbilitySystemGlobals::Get().GetGameplayCueManager()->InvokeGameplayCueExecuted_FromSpec(this, *OurCopyOfSpec, PredictionKey);
+		if (!bSuppressGameplayCues)
+		{
+			UAbilitySystemGlobals::Get().GetGameplayCueManager()->InvokeGameplayCueExecuted_FromSpec(this, *OurCopyOfSpec, PredictionKey);
+		}
 	}
 	else if (Spec.Def->DurationPolicy == EGameplayEffectDurationType::Instant)
 	{
@@ -889,7 +895,7 @@ FActiveGameplayEffectHandle UAbilitySystemComponent::FindActiveGameplayEffectHan
 void UAbilitySystemComponent::InvokeGameplayCueEvent(const FGameplayEffectSpecForRPC &Spec, EGameplayCueEvent::Type EventType)
 {
 	AActor* ActorAvatar = AbilityActorInfo->AvatarActor.Get();
-	if (ActorAvatar == nullptr)
+	if (ActorAvatar == nullptr && !bSuppressGameplayCues)
 	{
 		// No avatar actor to call this gameplaycue on.
 		return;
@@ -943,7 +949,7 @@ void UAbilitySystemComponent::InvokeGameplayCueEvent(const FGameplayTag Gameplay
 {
 	AActor* ActorAvatar = AbilityActorInfo->AvatarActor.Get();
 	
-	if (ActorAvatar != nullptr)
+	if (ActorAvatar != nullptr && !bSuppressGameplayCues)
 	{
 		UAbilitySystemGlobals::Get().GetGameplayCueManager()->HandleGameplayCues(ActorAvatar, GameplayCueTag, EventType, GameplayCueParameters);
 	}
@@ -1236,6 +1242,14 @@ void UAbilitySystemComponent::ForceReplication()
 	}
 }
 
+void UAbilitySystemComponent::ForceAvatarReplication()
+{
+	if (AvatarActor && AvatarActor->Role == ROLE_Authority)
+	{
+		AvatarActor->ForceNetUpdate();
+	}
+}
+
 bool UAbilitySystemComponent::ReplicateSubobjects(class UActorChannel *Channel, class FOutBunch *Bunch, FReplicationFlags *RepFlags)
 {
 	bool WroteSomething = Super::ReplicateSubobjects(Channel, Bunch, RepFlags);
diff --git a/Engine/Source/Runtime/GameplayAbilities/Private/AttributeSet.cpp b/Engine/Source/Runtime/GameplayAbilities/Private/AttributeSet.cpp
--- a/Engine/Source/Runtime/GameplayAbilities/Private/AttributeSet.cpp
+++ b/Engine/Source/Runtime/GameplayAbilities/Private/AttributeSet.cpp
@@ -232,7 +232,30 @@ bool FScalableFloat::SerializeFromMismatchedTag(const FPropertyTag& Tag, FArchiv
 
 		return true;
 	}
+	else if (Tag.Type == NAME_IntProperty)
+	{
+		int32 OldValue;
+		Ar << OldValue;
+		*this = FScalableFloat((float)OldValue);
+
+		return true;
+	}
+	else if (Tag.Type == NAME_Int8Property)
+	{
+		int8 OldValue;
+		Ar << OldValue;
+		*this = FScalableFloat((float)OldValue);
 
+		return true;
+	}
+	else if (Tag.Type == NAME_Int16Property)
+	{
+		int16 OldValue;
+		Ar << OldValue;
+		*this = FScalableFloat((float)OldValue);
+
+		return true;
+	}
 	return false;
 }
 
diff --git a/Engine/Source/Runtime/GameplayAbilities/Private/GameplayCueManager.cpp b/Engine/Source/Runtime/GameplayAbilities/Private/GameplayCueManager.cpp
--- a/Engine/Source/Runtime/GameplayAbilities/Private/GameplayCueManager.cpp
+++ b/Engine/Source/Runtime/GameplayAbilities/Private/GameplayCueManager.cpp
@@ -318,7 +318,6 @@ void UGameplayCueManager::ReloadObjectLibrary(UWorld* World, const UWorld::Initi
 void UGameplayCueManager::LoadObjectLibrary_Internal()
 {
 	FOnGameplayCueNotifySetLoaded OnLoadDelegate = FOnGameplayCueNotifySetLoaded::CreateUObject(this, &UGameplayCueManager::OnGameplayCueNotifyAsyncLoadComplete);
-
 	InitObjectLibraries(LoadedPaths, GameplayCueNotifyActorObjectLibrary, GameplayCueNotifyStaticObjectLibrary, OnLoadDelegate);
 }
 
@@ -437,40 +436,46 @@ void UGameplayCueManager::BuildCuesToAddToGlobalSet(const TArray<FAssetData>& As
 	}
 }
 
+int32 GameplayCueCheckForTooManyRPCs = 1;
+static FAutoConsoleVariableRef CVarGameplayCueCheckForTooManyRPCs(TEXT("AbilitySystem.GameplayCueCheckForTooManyRPCs"), GameplayCueCheckForTooManyRPCs, TEXT("Warns if gameplay cues are being throttled by network code"), ECVF_Default );
+
 void UGameplayCueManager::CheckForTooManyRPCs(FName FuncName, const FGameplayCuePendingExecute& PendingCue, const FString& CueID, const FGameplayEffectContext* EffectContext)
 {
-	static IConsoleVariable* MaxRPCPerNetUpdateCVar = IConsoleManager::Get().FindConsoleVariable(TEXT("net.MaxRPCPerNetUpdate"));
-	if (MaxRPCPerNetUpdateCVar)
+	if (GameplayCueCheckForTooManyRPCs)
 	{
-		AActor* Owner = PendingCue.OwningComponent ? PendingCue.OwningComponent->GetOwner() : nullptr;
-		UWorld* World = Owner ? Owner->GetWorld() : nullptr;
-		UNetDriver* NetDriver = World ? World->GetNetDriver() : nullptr;
-		if (NetDriver)
+		static IConsoleVariable* MaxRPCPerNetUpdateCVar = IConsoleManager::Get().FindConsoleVariable(TEXT("net.MaxRPCPerNetUpdate"));
+		if (MaxRPCPerNetUpdateCVar)
 		{
-			const int32 MaxRPCs = MaxRPCPerNetUpdateCVar->GetInt();
-			for (UNetConnection* ClientConnection : NetDriver->ClientConnections)
+			AActor* Owner = PendingCue.OwningComponent ? PendingCue.OwningComponent->GetOwner() : nullptr;
+			UWorld* World = Owner ? Owner->GetWorld() : nullptr;
+			UNetDriver* NetDriver = World ? World->GetNetDriver() : nullptr;
+			if (NetDriver)
 			{
-				if (ClientConnection)
+				const int32 MaxRPCs = MaxRPCPerNetUpdateCVar->GetInt();
+				for (UNetConnection* ClientConnection : NetDriver->ClientConnections)
 				{
-					UActorChannel** OwningActorChannelPtr = ClientConnection->ActorChannels.Find(Owner);
-					TSharedRef<FObjectReplicator>* ComponentReplicatorPtr = (OwningActorChannelPtr && *OwningActorChannelPtr) ? (*OwningActorChannelPtr)->ReplicationMap.Find(PendingCue.OwningComponent) : nullptr;
-					if (ComponentReplicatorPtr)
+					if (ClientConnection)
 					{
-						const TArray<FObjectReplicator::FRPCCallInfo>& RemoteFuncInfo = (*ComponentReplicatorPtr)->RemoteFuncInfo;
-						for (const FObjectReplicator::FRPCCallInfo& CallInfo : RemoteFuncInfo)
+						UActorChannel** OwningActorChannelPtr = ClientConnection->ActorChannels.Find(Owner);
+						TSharedRef<FObjectReplicator>* ComponentReplicatorPtr = (OwningActorChannelPtr && *OwningActorChannelPtr) ? (*OwningActorChannelPtr)->ReplicationMap.Find(PendingCue.OwningComponent) : nullptr;
+						if (ComponentReplicatorPtr)
 						{
-							if (CallInfo.FuncName == FuncName)
+							const TArray<FObjectReplicator::FRPCCallInfo>& RemoteFuncInfo = (*ComponentReplicatorPtr)->RemoteFuncInfo;
+							for (const FObjectReplicator::FRPCCallInfo& CallInfo : RemoteFuncInfo)
 							{
-								if (CallInfo.Calls > MaxRPCs)
+								if (CallInfo.FuncName == FuncName)
 								{
-									const FString Instigator = EffectContext ? EffectContext->ToString() : TEXT("None");
-									ABILITY_LOG(Warning, TEXT("Attempted to fire %s when no more RPCs are allowed this net update. Max:%d Cue:%s Instigator:%s Component:%s"), *FuncName.ToString(), MaxRPCs, *CueID, *Instigator, *GetPathNameSafe(PendingCue.OwningComponent));
+									if (CallInfo.Calls > MaxRPCs)
+									{
+										const FString Instigator = EffectContext ? EffectContext->ToString() : TEXT("None");
+										ABILITY_LOG(Warning, TEXT("Attempted to fire %s when no more RPCs are allowed this net update. Max:%d Cue:%s Instigator:%s Component:%s"), *FuncName.ToString(), MaxRPCs, *CueID, *Instigator, *GetPathNameSafe(PendingCue.OwningComponent));
 									
-									// Returning here to only log once per offending RPC.
-									return;
-								}
+										// Returning here to only log once per offending RPC.
+										return;
+									}
 
-								break;
+									break;
+								}
 							}
 						}
 					}
@@ -612,7 +617,7 @@ void UGameplayCueManager::HandleAssetRenamed(const FAssetData& Data, const FStri
 void UGameplayCueManager::VerifyNotifyAssetIsInValidPath(FString Path)
 {
 	bool ValidPath = false;
-	for (FString& str: LoadedPaths)
+	for (FString& str: GetValidGameplayCuePaths())
 	{
 		if (Path.Contains(str))
 		{
@@ -627,7 +632,7 @@ void UGameplayCueManager::VerifyNotifyAssetIsInValidPath(FString Path)
 
 		ABILITY_LOG(Warning, TEXT("Warning: Invalid GameplayCuePath: %s"), *Path);
 		ABILITY_LOG(Warning, TEXT("Valid Paths: "));
-		for (FString& str: LoadedPaths)
+		for (FString& str: GetValidGameplayCuePaths())
 		{
 			ABILITY_LOG(Warning, TEXT("  %s"), *str);
 			MessageTry += FString::Printf(TEXT("\n  %s"), *str);
@@ -641,6 +646,14 @@ void UGameplayCueManager::VerifyNotifyAssetIsInValidPath(FString Path)
 	}
 }
 
+void UGameplayCueManager::LoadAllGameplayCueNotifiesForEditor()
+{
+	// Spft load all valid paths
+	TArray<FString> ValidPaths = GetValidGameplayCuePaths();
+
+	GlobalCueSet->Empty();
+	InitObjectLibraries(ValidPaths, GameplayCueNotifyActorObjectLibrary, GameplayCueNotifyStaticObjectLibrary, FOnGameplayCueNotifySetLoaded());
+}
 #endif
 
 
diff --git a/Engine/Source/Runtime/GameplayAbilities/Private/GameplayEffect.cpp b/Engine/Source/Runtime/GameplayAbilities/Private/GameplayEffect.cpp
--- a/Engine/Source/Runtime/GameplayAbilities/Private/GameplayEffect.cpp
+++ b/Engine/Source/Runtime/GameplayAbilities/Private/GameplayEffect.cpp
@@ -2583,7 +2583,7 @@ void FActiveGameplayEffectsContainer::AddActiveGameplayEffectGrantedTagsAndModif
 	ApplicationImmunityGameplayTagCountContainer.UpdateTagCount(Effect.Spec.Def->GrantedApplicationImmunityTags.IgnoreTags, 1);
 
 	// Grant abilities
-	if (IsNetAuthority())
+	if (IsNetAuthority() && !Owner->bSuppressGrantAbility)
 	{
 		for (FGameplayAbilitySpecDef& AbilitySpecDef : Effect.Spec.GrantedAbilitySpecs)
 		{
@@ -2597,24 +2597,27 @@ void FActiveGameplayEffectsContainer::AddActiveGameplayEffectGrantedTagsAndModif
 	}
 
 	// Update GameplayCue tags and events
-	for (const FGameplayEffectCue& Cue : Effect.Spec.Def->GameplayCues)
+	if (!Owner->bSuppressGameplayCues)
 	{
-		Owner->UpdateTagMap(Cue.GameplayCueTags, 1);
-
-		if (bInvokeGameplayCueEvents)
+		for (const FGameplayEffectCue& Cue : Effect.Spec.Def->GameplayCues)
 		{
-			Owner->InvokeGameplayCueEvent(Effect.Spec, EGameplayCueEvent::OnActive);
-			Owner->InvokeGameplayCueEvent(Effect.Spec, EGameplayCueEvent::WhileActive);
-		}
+			Owner->UpdateTagMap(Cue.GameplayCueTags, 1);
 
-		if (IsNetAuthority() && Owner->bMinimalReplication)
-		{
-			for (const FGameplayTag& CueTag : Cue.GameplayCueTags)
+			if (bInvokeGameplayCueEvents)
+			{
+				Owner->InvokeGameplayCueEvent(Effect.Spec, EGameplayCueEvent::OnActive);
+				Owner->InvokeGameplayCueEvent(Effect.Spec, EGameplayCueEvent::WhileActive);
+			}
+
+			if (IsNetAuthority() && Owner->bMinimalReplication)
 			{
-				// Note: minimal replication does not replicate the effect context with the gameplay cue parameters.
-				// This is just a choice right now. If needed, it may be better to convert the effect context to GC parameters *here*
-				// and pass those into this function
-				Owner->AddGameplayCue_MinimalReplication(CueTag);
+				for (const FGameplayTag& CueTag : Cue.GameplayCueTags)
+				{
+					// Note: minimal replication does not replicate the effect context with the gameplay cue parameters.
+					// This is just a choice right now. If needed, it may be better to convert the effect context to GC parameters *here*
+					// and pass those into this function
+					Owner->AddGameplayCue_MinimalReplication(CueTag);
+				}
 			}
 		}
 	}
@@ -2834,20 +2837,23 @@ void FActiveGameplayEffectsContainer::RemoveActiveGameplayEffectGrantedTagsAndMo
 	}
 
 	// Update GameplayCue tags and events
-	for (const FGameplayEffectCue& Cue : Effect.Spec.Def->GameplayCues)
+	if (!Owner->bSuppressGameplayCues)
 	{
-		Owner->UpdateTagMap(Cue.GameplayCueTags, -1);
-
-		if (bInvokeGameplayCueEvents)
+		for (const FGameplayEffectCue& Cue : Effect.Spec.Def->GameplayCues)
 		{
-			Owner->InvokeGameplayCueEvent(Effect.Spec, EGameplayCueEvent::Removed);			
-		}
+			Owner->UpdateTagMap(Cue.GameplayCueTags, -1);
 
-		if (IsNetAuthority() && Owner->bMinimalReplication)
-		{
-			for (const FGameplayTag& CueTag : Cue.GameplayCueTags)
+			if (bInvokeGameplayCueEvents)
 			{
-				Owner->RemoveGameplayCue_MinimalReplication(CueTag);
+				Owner->InvokeGameplayCueEvent(Effect.Spec, EGameplayCueEvent::Removed);			
+			}
+
+			if (IsNetAuthority() && Owner->bMinimalReplication)
+			{
+				for (const FGameplayTag& CueTag : Cue.GameplayCueTags)
+				{
+					Owner->RemoveGameplayCue_MinimalReplication(CueTag);
+				}
 			}
 		}
 	}
@@ -3659,6 +3665,14 @@ void FActiveGameplayEffectsContainer::CloneFrom(const FActiveGameplayEffectsCont
 			}
 		}
 	}
+
+	// Broadcast dirty on everything so that the UAttributeSet properties get updated
+	for (auto& It : AttributeAggregatorMap)
+	{
+		FAggregatorRef& AggregatorRef = It.Value;
+		AggregatorRef.Get()->BroadcastOnDirty();
+
+	}
 }
 
 // --------------------------------------------------------------------------------------------------------------------------------------------------------
diff --git a/Engine/Source/Runtime/GameplayAbilities/Private/GameplayEffectTypes.cpp b/Engine/Source/Runtime/GameplayAbilities/Private/GameplayEffectTypes.cpp
--- a/Engine/Source/Runtime/GameplayAbilities/Private/GameplayEffectTypes.cpp
+++ b/Engine/Source/Runtime/GameplayAbilities/Private/GameplayEffectTypes.cpp
@@ -410,10 +410,14 @@ bool FGameplayTagCountContainer::UpdateTagMap_Internal(const FGameplayTag& Tag,
 		FDelegateInfo* DelegateInfo = GameplayTagEventMap.Find(CurTag);
 		if (DelegateInfo)
 		{
+			// Prior to calling OnAnyChange delegate, copy our OnNewOrRemove delegate, since things listening to OnAnyChange could add or remove 
+			// to this map causing our pointer to become invalid.
+			FOnGameplayEffectTagCountChanged OnNewOrRemoveLocalCopy = DelegateInfo->OnNewOrRemove;
+
 			DelegateInfo->OnAnyChange.Broadcast(CurTag, NewTagCount);
 			if (SignificantChange)
 			{
-				DelegateInfo->OnNewOrRemove.Broadcast(CurTag, NewTagCount);
+				OnNewOrRemoveLocalCopy.Broadcast(CurTag, NewTagCount);
 			}
 		}
 	}
diff --git a/Engine/Source/Runtime/GameplayAbilities/Public/AbilitySystemComponent.h b/Engine/Source/Runtime/GameplayAbilities/Public/AbilitySystemComponent.h
--- a/Engine/Source/Runtime/GameplayAbilities/Public/AbilitySystemComponent.h
+++ b/Engine/Source/Runtime/GameplayAbilities/Public/AbilitySystemComponent.h
@@ -167,6 +167,9 @@ class GAMEPLAYABILITIES_API UAbilitySystemComponent : public UGameplayTasksCompo
 	/** Force owning actor to update it's replication, to make sure that gameplay cues get sent down quickly. Override to change how aggressive this is */
 	virtual void ForceReplication();
 
+	/** Forces avatar actor to update it's replication. Useful for things like needing to replication for movement / locations reasons. */
+	virtual void ForceAvatarReplication();
+
 	void PreReplication(IRepChangedPropertyTracker& ChangedPropertyTracker) override;
 
 	virtual void GetSubobjectsWithStableNamesForNetworking(TArray<UObject*>& Objs) override;
@@ -1330,6 +1333,12 @@ public:
 	// -
 	virtual AActor* GetAvatarActor(const UGameplayTask* Task) const override;
 
+	/** Suppress all ability granting through GEs on this component */
+	bool bSuppressGrantAbility;
+
+	/** Suppress all GameplayCues on this component */
+	bool bSuppressGameplayCues;
+
 protected:
 
 	/** Actually pushes the final attribute value to the attribute set's property. Should not be called by outside code since this does not go through the attribute aggregator system. */
diff --git a/Engine/Source/Runtime/GameplayAbilities/Public/AttributeSet.h b/Engine/Source/Runtime/GameplayAbilities/Public/AttributeSet.h
--- a/Engine/Source/Runtime/GameplayAbilities/Public/AttributeSet.h
+++ b/Engine/Source/Runtime/GameplayAbilities/Public/AttributeSet.h
@@ -229,7 +229,7 @@ public:
 	//copy operator to prevent duplicate handles
 	void operator=(const FScalableFloat& Src);
 
-	/* Used to upgrade a float property into an FScalableFloat */
+	/* Used to upgrade a float or int8/int16/int32 property into an FScalableFloat */
 	bool SerializeFromMismatchedTag(const FPropertyTag& Tag, FArchive& Ar);
 
 private:
diff --git a/Engine/Source/Runtime/GameplayAbilities/Public/GameplayAbilitySpec.h b/Engine/Source/Runtime/GameplayAbilities/Public/GameplayAbilitySpec.h
--- a/Engine/Source/Runtime/GameplayAbilities/Public/GameplayAbilitySpec.h
+++ b/Engine/Source/Runtime/GameplayAbilities/Public/GameplayAbilitySpec.h
@@ -48,6 +48,11 @@ struct FGameplayAbilitySpecHandle
 		return ::GetTypeHash(SpecHandle.Handle);
 	}
 
+	FString ToString() const
+	{
+		return IsValid() ? FString::FromInt(Handle) : TEXT("Invalid");
+	}
+
 private:
 
 	UPROPERTY()
diff --git a/Engine/Source/Runtime/GameplayAbilities/Public/GameplayCueManager.h b/Engine/Source/Runtime/GameplayAbilities/Public/GameplayCueManager.h
--- a/Engine/Source/Runtime/GameplayAbilities/Public/GameplayCueManager.h
+++ b/Engine/Source/Runtime/GameplayAbilities/Public/GameplayCueManager.h
@@ -148,6 +148,9 @@ class GAMEPLAYABILITIES_API UGameplayCueManager : public UDataAsset
 	virtual class UWorld* GetWorld() const override;
 
 #if WITH_EDITOR
+	/** Called from editor to soft load all gameplay cue notifies for the GameplayCueEditor */
+	void LoadAllGameplayCueNotifiesForEditor();
+
 	bool IsAssetInLoadedPaths(UObject *Object) const;
 
 	/** Handles updating an object library when a new asset is created */
@@ -161,6 +164,9 @@ class GAMEPLAYABILITIES_API UGameplayCueManager : public UDataAsset
 
 	void VerifyNotifyAssetIsInValidPath(FString Path);
 
+	/** returns list of valid gameplay cue paths. Subclasses may override this to specify locations that aren't part of the "always loaded" LoadedPaths array */
+	virtual TArray<FString>	GetValidGameplayCuePaths() { return LoadedPaths; }
+
 	bool RegisteredEditorCallbacks;
 
 	bool bAccelerationMapOutdated;
diff --git a/Engine/Source/Runtime/GameplayAbilities/Public/GameplayEffectTypes.h b/Engine/Source/Runtime/GameplayAbilities/Public/GameplayEffectTypes.h
--- a/Engine/Source/Runtime/GameplayAbilities/Public/GameplayEffectTypes.h
+++ b/Engine/Source/Runtime/GameplayAbilities/Public/GameplayEffectTypes.h
@@ -388,7 +388,7 @@ struct GAMEPLAYABILITIES_API FGameplayEffectContext
 		if (GetHitResult())
 		{
 			// Does a deep copy of the hit result
-			NewContext->AddHitResult(*GetHitResult());
+			NewContext->AddHitResult(*GetHitResult(), true);
 		}
 		return NewContext;
 	}
diff --git a/Engine/Source/Runtime/InputCore/Classes/InputCoreTypes.h b/Engine/Source/Runtime/InputCore/Classes/InputCoreTypes.h
--- a/Engine/Source/Runtime/InputCore/Classes/InputCoreTypes.h
+++ b/Engine/Source/Runtime/InputCore/Classes/InputCoreTypes.h
@@ -336,6 +336,8 @@ struct INPUTCORE_API EKeys
 	static const FKey Gamepad_LeftThumbstick;
 	static const FKey Gamepad_RightThumbstick;
 	static const FKey Gamepad_Special_Left;
+	static const FKey Gamepad_Special_Left_X;
+	static const FKey Gamepad_Special_Left_Y;
 	static const FKey Gamepad_Special_Right;
 	static const FKey Gamepad_FaceButton_Bottom;
 	static const FKey Gamepad_FaceButton_Right;
diff --git a/Engine/Source/Runtime/InputCore/Private/InputCoreTypes.cpp b/Engine/Source/Runtime/InputCore/Private/InputCoreTypes.cpp
--- a/Engine/Source/Runtime/InputCore/Private/InputCoreTypes.cpp
+++ b/Engine/Source/Runtime/InputCore/Private/InputCoreTypes.cpp
@@ -172,6 +172,8 @@ const FKey EKeys::Gamepad_RightTriggerAxis("Gamepad_RightTriggerAxis");
 const FKey EKeys::Gamepad_LeftThumbstick("Gamepad_LeftThumbstick");
 const FKey EKeys::Gamepad_RightThumbstick("Gamepad_RightThumbstick");
 const FKey EKeys::Gamepad_Special_Left("Gamepad_Special_Left");
+const FKey EKeys::Gamepad_Special_Left_X("Gamepad_Special_Left_X");
+const FKey EKeys::Gamepad_Special_Left_Y("Gamepad_Special_Left_Y");
 const FKey EKeys::Gamepad_Special_Right("Gamepad_Special_Right");
 const FKey EKeys::Gamepad_FaceButton_Bottom("Gamepad_FaceButton_Bottom");
 const FKey EKeys::Gamepad_FaceButton_Right("Gamepad_FaceButton_Right");
@@ -647,6 +649,9 @@ void EKeys::Initialize()
 	// PS4-specific
 	AddMenuCategoryDisplayInfo("PS4", LOCTEXT("PS4SubCategory", "PS4"), TEXT("GraphEditor.PadEvent_16x"));
 	AddKey(FKeyDetails(EKeys::PS4_Special, LOCTEXT("PS4_Special", "PS4_Special"), FKeyDetails::NotBlueprintBindableKey));
+	AddKey(FKeyDetails(EKeys::Gamepad_Special_Left_X, LOCTEXT("PS4_Gamepad_Special_Left_X", "PS4 Touchpad Button X Axis"), FKeyDetails::GamepadKey | FKeyDetails::FloatAxis, "PS4"));
+	AddKey(FKeyDetails(EKeys::Gamepad_Special_Left_Y, LOCTEXT("PS4_Gamepad_Special_Left_Y", "PS4 Touchpad Button Y Axis"), FKeyDetails::GamepadKey | FKeyDetails::FloatAxis, "PS4"));
+
 
 	// Steam Controller specific
 	AddMenuCategoryDisplayInfo("Steam", LOCTEXT("SteamSubCateogry", "Steam"), TEXT("GraphEditor.PadEvent_16x"));
diff --git a/Engine/Source/Runtime/Launch/Resources/VersionLocked.h b/Engine/Source/Runtime/Launch/Resources/VersionLocked.h
new file mode 100644
--- /dev/null
+++ b/Engine/Source/Runtime/Launch/Resources/VersionLocked.h
@@ -0,0 +1,25 @@
+// Copyright 1998-2015 Epic Games, Inc. All Rights Reserved.
+
+#pragma once
+
+#define VERSION_LOCKED 0
+
+#if VERSION_LOCKED
+
+#define ENGINE_MAJOR_VERSION	4
+#define ENGINE_MINOR_VERSION	12
+#define ENGINE_PATCH_VERSION	0
+
+#define ENGINE_IS_LICENSEE_VERSION 0
+
+#define BUILT_FROM_CHANGELIST 2839255
+#define BRANCH_NAME "++Orion+Release-Next"
+
+#define ENGINE_IS_PROMOTED_BUILD (BUILT_FROM_CHANGELIST > 0)
+
+#define EPIC_COMPANY_NAME  "Epic Games, Inc."
+#define EPIC_COPYRIGHT_STRING "Copyright 1998-2015 Epic Games, Inc. All Rights Reserved."
+#define EPIC_PRODUCT_NAME "Unreal Engine"
+#define EPIC_PRODUCT_IDENTIFIER "UnrealEngine"
+
+#endif
diff --git a/Engine/Source/Runtime/Linux/LinuxCommonStartup/Private/LinuxCommonStartup.cpp b/Engine/Source/Runtime/Linux/LinuxCommonStartup/Private/LinuxCommonStartup.cpp
--- a/Engine/Source/Runtime/Linux/LinuxCommonStartup/Private/LinuxCommonStartup.cpp
+++ b/Engine/Source/Runtime/Linux/LinuxCommonStartup/Private/LinuxCommonStartup.cpp
@@ -230,6 +230,13 @@ int CommonLinuxMain(int argc, char *argv[], int (*RealMain)(const TCHAR * Comman
 		}
 	}
 
+	// check if a specific return code has been set
+	uint8 OverriddenErrorLevel = 0;
+	if (FPlatformMisc::HasOverriddenReturnCode(&OverriddenErrorLevel))
+	{
+		ErrorLevel = OverriddenErrorLevel;
+	}
+
 	if (ErrorLevel)
 	{
 		printf("Exiting abnormally (error code: %d)\n", ErrorLevel);
diff --git a/Engine/Source/Runtime/MoviePlayer/Private/DefaultGameMoviePlayer.cpp b/Engine/Source/Runtime/MoviePlayer/Private/DefaultGameMoviePlayer.cpp
--- a/Engine/Source/Runtime/MoviePlayer/Private/DefaultGameMoviePlayer.cpp
+++ b/Engine/Source/Runtime/MoviePlayer/Private/DefaultGameMoviePlayer.cpp
@@ -89,8 +89,17 @@ FDefaultGameMoviePlayer::~FDefaultGameMoviePlayer()
 	{
 		// This should not happen if initialize was called correctly.  This is a fallback to ensure that the movie player rendering tickable gets unregistered on the rendering thread correctly
 		Shutdown();
-		FlushRenderingCommands();
 	}
+	else
+	{
+		// Even when uninitialized we must safely unregister the movie player on the render thread
+		ENQUEUE_UNIQUE_RENDER_COMMAND_ONEPARAMETER(UnregisterMoviePlayerTickable, FDefaultGameMoviePlayer*, MoviePlayer, this,
+		{
+			MoviePlayer->Unregister();
+		});
+	}
+
+	FlushRenderingCommands();
 }
 
 void FDefaultGameMoviePlayer::RegisterMovieStreamer(TSharedPtr<IMovieStreamer> InMovieStreamer)
@@ -383,9 +392,9 @@ void FDefaultGameMoviePlayer::Tick( float DeltaTime )
 	check(IsInRenderingThread());
 	if (LoadingScreenWindowPtr.IsValid() && RendererPtr.IsValid())
 	{
+		FScopeLock SyncMechanismLock(&SyncMechanismCriticalSection);
 		if (!IsLoadingFinished() && SyncMechanism)
 		{
-			FScopeLock SyncMechanismLock(&SyncMechanismCriticalSection);
 			if (SyncMechanism->IsSlateDrawPassEnqueued())
 			{
 				GFrameNumberRenderThread++;
diff --git a/Engine/Source/Runtime/NetworkFile/Private/NetworkPlatformFile.cpp b/Engine/Source/Runtime/NetworkFile/Private/NetworkPlatformFile.cpp
--- a/Engine/Source/Runtime/NetworkFile/Private/NetworkPlatformFile.cpp
+++ b/Engine/Source/Runtime/NetworkFile/Private/NetworkPlatformFile.cpp
@@ -116,7 +116,15 @@ bool FNetworkPlatformFile::InitializeInternal(IPlatformFile* Inner, const TCHAR*
 	LocalDirectories.Add(FPaths::GameSavedDir() / TEXT("Logs"));
 	LocalDirectories.Add(FPaths::GameSavedDir() / TEXT("Sandboxes"));
 
-	InnerPlatformFile->GetLowerLevel()->AddLocalDirectories(LocalDirectories);
+	if (InnerPlatformFile->GetLowerLevel())
+	{
+		InnerPlatformFile->GetLowerLevel()->AddLocalDirectories(LocalDirectories);
+	}
+	else
+	{
+		InnerPlatformFile->AddLocalDirectories(LocalDirectories);
+	}
+	
 
 	FNetworkFileArchive Payload(NFS_Messages::Heartbeat); 
 	FArrayReader Out;
diff --git a/Engine/Source/Runtime/NetworkReplayStreaming/HttpNetworkReplayStreaming/Private/HttpNetworkReplayStreaming.cpp b/Engine/Source/Runtime/NetworkReplayStreaming/HttpNetworkReplayStreaming/Private/HttpNetworkReplayStreaming.cpp
--- a/Engine/Source/Runtime/NetworkReplayStreaming/HttpNetworkReplayStreaming/Private/HttpNetworkReplayStreaming.cpp
+++ b/Engine/Source/Runtime/NetworkReplayStreaming/HttpNetworkReplayStreaming/Private/HttpNetworkReplayStreaming.cpp
@@ -586,8 +586,6 @@ void FHttpNetworkReplayStreamer::HttpRequestEventDataFinished(FHttpRequestPtr Ht
 {
 	RequestFinished(StreamerState, EQueuedHttpRequestType::RequestEventData, HttpRequest);
 
-	check(RequestEventDataCompleteDelegate.IsBound());
-
 	if (bSucceeded && HttpResponse->GetResponseCode() == EHttpResponseCodes::Ok)
 	{
 		RequestEventDataCompleteDelegate.ExecuteIfBound(HttpResponse->GetContent(), true);
diff --git a/Engine/Source/Runtime/Online/HTTP/Private/HttpRetrySystem.cpp b/Engine/Source/Runtime/Online/HTTP/Private/HttpRetrySystem.cpp
--- a/Engine/Source/Runtime/Online/HTTP/Private/HttpRetrySystem.cpp
+++ b/Engine/Source/Runtime/Online/HTTP/Private/HttpRetrySystem.cpp
@@ -217,13 +217,13 @@ bool FHttpRetrySystem::FManager::Update(uint32* FileCount, uint32* FailingCount,
 				// Save these for failure case retry checks if we hit a completion state
 				bool bShouldRetry = false;
 				bool bCanRetry = false;
-				if (RequestStatus == EHttpRequestStatus::Failed || RequestStatus == EHttpRequestStatus::Succeeded)
+				if (RequestStatus == EHttpRequestStatus::Failed || RequestStatus == EHttpRequestStatus::Failed_ConnectionError || RequestStatus == EHttpRequestStatus::Succeeded)
 				{
 					bShouldRetry = ShouldRetry(HttpRetryRequestEntry);
 					bCanRetry = CanRetry(HttpRetryRequestEntry);
 				}
 
-				if (RequestStatus == EHttpRequestStatus::Failed || forceFail || (bShouldRetry && bCanRetry))
+				if (RequestStatus == EHttpRequestStatus::Failed || RequestStatus == EHttpRequestStatus::Failed_ConnectionError || forceFail || (bShouldRetry && bCanRetry))
 				{
 					bIsGreen = false;
                     if(HttpRetryRequestEntry.bShouldCancel == false)
diff --git a/Engine/Source/Runtime/Online/Hotfix/Private/OnlineHotfixManager.cpp b/Engine/Source/Runtime/Online/Hotfix/Private/OnlineHotfixManager.cpp
--- a/Engine/Source/Runtime/Online/Hotfix/Private/OnlineHotfixManager.cpp
+++ b/Engine/Source/Runtime/Online/Hotfix/Private/OnlineHotfixManager.cpp
@@ -31,7 +31,8 @@ UOnlineHotfixManager::UOnlineHotfixManager() :
 	TotalBytes(0),
 	NumBytes(0),
 	bHotfixingInProgress(false),
-	bHotfixNeedsReload(false)
+	bHotfixNeedsMapReload(false),
+	ChangedOrRemovedPakCount(0)
 {
 	OnEnumerateFilesCompleteDelegate = FOnEnumerateFilesCompleteDelegate::CreateUObject(this, &UOnlineHotfixManager::OnEnumerateFilesComplete);
 	OnReadFileProgressDelegate = FOnReadFileProgressDelegate::CreateUObject(this, &UOnlineHotfixManager::OnReadFileProgress);
@@ -84,11 +85,12 @@ void UOnlineHotfixManager::PostInitProperties()
 void UOnlineHotfixManager::Init()
 {
 	bHotfixingInProgress = true;
-	bHotfixNeedsReload = false;
+	bHotfixNeedsMapReload = false;
 	TotalFiles = 0;
 	NumDownloaded = 0;
 	TotalBytes = 0;
 	NumBytes = 0;
+	ChangedOrRemovedPakCount = 0;
 	// Build the name of the loc file that we'll care about
 	// It can change at runtime so build it just before fetching the data
 	GameLocName = DebugPrefix + FInternationalization::Get().GetCurrentCulture()->GetTwoLetterISOLanguageName() + TEXT("_Game.locres");
@@ -371,9 +373,6 @@ void UOnlineHotfixManager::UpdateProgress(uint32 FileCount, uint64 UpdateSize)
 	NumBytes += UpdateSize;
 	// Update our progress
 	TriggerOnHotfixProgressDelegates(NumDownloaded, TotalFiles, NumBytes, TotalBytes);
-
-	UE_LOG(LogHotfixManager, Display, TEXT("NumFiles (%d), TotalFiles (%d), NumBytes (%d), TotalBytes(%d)"),
-		NumDownloaded, TotalFiles, NumBytes, TotalBytes);
 }
 
 void UOnlineHotfixManager::ApplyHotfix()
@@ -390,7 +389,18 @@ void UOnlineHotfixManager::ApplyHotfix()
 		TriggerOnHotfixProcessedFileDelegates(FileHeader.FileName, GetCachedDirectory() / FileHeader.DLName);
 	}
 	UE_LOG(LogHotfixManager, Display, TEXT("Hotfix data has been successfully applied"));
-	TriggerHotfixComplete((!bHotfixNeedsReload) ? EHotfixResult::Success : EHotfixResult::SuccessNeedsReload);
+	EHotfixResult Result = EHotfixResult::Success;
+	if (ChangedOrRemovedPakCount > 0)
+	{
+		UE_LOG(LogHotfixManager, Display, TEXT("Hotfix has changed or removed PAK files so a relaunch of the app is needed"));
+		Result = EHotfixResult::SuccessNeedsRelaunch;
+	}
+	else if (bHotfixNeedsMapReload)
+	{
+		UE_LOG(LogHotfixManager, Display, TEXT("Hotfix has detected PAK files containing currently loaded maps, so a level load is needed"));
+		Result = EHotfixResult::SuccessNeedsReload;
+	}
+	TriggerHotfixComplete(Result);
 }
 
 void UOnlineHotfixManager::TriggerHotfixComplete(EHotfixResult HotfixResult)
@@ -620,9 +630,9 @@ bool UOnlineHotfixManager::HotfixPakFile(const FCloudFileHeader& FileHeader)
 			{
 				IniList.Add(InternalPakFileName);
 			}
-			else if (!bHotfixNeedsReload && InternalPakFileName.EndsWith(FPackageName::GetMapPackageExtension()))
+			else if (!bHotfixNeedsMapReload && InternalPakFileName.EndsWith(FPackageName::GetMapPackageExtension()))
 			{
-				bHotfixNeedsReload = IsMapLoaded(InternalPakFileName);
+				bHotfixNeedsMapReload = IsMapLoaded(InternalPakFileName);
 			}
 		}
 		// Sort the INIs so they are processed consistently
@@ -718,6 +728,7 @@ void UOnlineHotfixManager::UnmountHotfixFiles()
 			{
 				FCoreDelegates::OnUnmountPak.Execute(MountedPakFiles[Index]);
 				MountedPakFiles.RemoveAt(Index);
+				ChangedOrRemovedPakCount++;
 				UE_LOG(LogHotfixManager, Log, TEXT("Hotfix unmounted PAK file (%s) so it can be redownloaded"), *FileHeader.FileName);
 				break;
 			}
@@ -732,6 +743,7 @@ void UOnlineHotfixManager::UnmountHotfixFiles()
 			{
 				FCoreDelegates::OnUnmountPak.Execute(MountedPakFiles[Index]);
 				MountedPakFiles.RemoveAt(Index);
+				ChangedOrRemovedPakCount++;
 				UE_LOG(LogHotfixManager, Log, TEXT("Hotfix unmounted PAK file (%s) since it was removed from the hotfix set"), *FileHeader.FileName);
 				break;
 			}
diff --git a/Engine/Source/Runtime/Online/IOS/OnlineSubsystemIOS/Private/OnlineFriendsInterfaceIOS.cpp b/Engine/Source/Runtime/Online/IOS/OnlineSubsystemIOS/Private/OnlineFriendsInterfaceIOS.cpp
--- a/Engine/Source/Runtime/Online/IOS/OnlineSubsystemIOS/Private/OnlineFriendsInterfaceIOS.cpp
+++ b/Engine/Source/Runtime/Online/IOS/OnlineSubsystemIOS/Private/OnlineFriendsInterfaceIOS.cpp
@@ -17,7 +17,7 @@ FString FOnlineFriendIOS::GetRealName() const
 	return Result;
 }
 
-FString FOnlineFriendIOS::GetDisplayName() const
+FString FOnlineFriendIOS::GetDisplayName(const FString& Platform) const
 {
 	FString Result;
 	GetAccountData(TEXT("nickname"), Result);
diff --git a/Engine/Source/Runtime/Online/IOS/OnlineSubsystemIOS/Private/OnlineFriendsInterfaceIOS.h b/Engine/Source/Runtime/Online/IOS/OnlineSubsystemIOS/Private/OnlineFriendsInterfaceIOS.h
--- a/Engine/Source/Runtime/Online/IOS/OnlineSubsystemIOS/Private/OnlineFriendsInterfaceIOS.h
+++ b/Engine/Source/Runtime/Online/IOS/OnlineSubsystemIOS/Private/OnlineFriendsInterfaceIOS.h
@@ -20,7 +20,7 @@ public:
 
 	virtual TSharedRef<const FUniqueNetId> GetUserId() const override;
 	virtual FString GetRealName() const override;
-	virtual FString GetDisplayName() const override;
+	virtual FString GetDisplayName(const FString& Platform = FString()) const override;
 	virtual bool GetUserAttribute(const FString& AttrName, FString& OutAttrValue) const override;
 
 	// FOnlineFriend
diff --git a/Engine/Source/Runtime/Online/OnlineFramework/Party/Private/Party.cpp b/Engine/Source/Runtime/Online/OnlineFramework/Party/Private/Party.cpp
--- a/Engine/Source/Runtime/Online/OnlineFramework/Party/Private/Party.cpp
+++ b/Engine/Source/Runtime/Online/OnlineFramework/Party/Private/Party.cpp
@@ -93,19 +93,29 @@ void UParty::OnLoginStatusChanged(int32 LocalUserNum, ELoginStatus::Type OldStat
 		if (JoinedParties.Num())
 		{
 			UE_LOG(LogParty, Log, TEXT("Party cleanup on logout"));
+			TArray<FOnlinePartyTypeId> PartiesToRemove;
 			for (const auto& PartyKeyValue : JoinedParties)
 			{
-				UPartyGameState* Party = PartyKeyValue.Value;
-				if (Party)
+				PartiesToRemove.Add(PartyKeyValue.Key);
+			}
+			for (const auto& PartyKey : PartiesToRemove)
+			{
+				UPartyGameState** FoundParty = JoinedParties.Find(PartyKey);
+				if (FoundParty)
 				{
-					TSharedPtr<const FOnlinePartyId> PartyId = Party->GetPartyId();
-					FString PartyIdString = PartyId.IsValid() ? PartyId->ToDebugString() : TEXT("");
-					UE_LOG(LogParty, Log, TEXT("[%s] Removed"), *PartyIdString);
-					Party->HandleRemovedFromParty(EMemberExitedReason::Left);
+					UPartyGameState* Party = *FoundParty;
+					if (Party)
+					{
+						TSharedPtr<const FOnlinePartyId> PartyId = Party->GetPartyId();
+						FString PartyIdString = PartyId.IsValid() ? PartyId->ToDebugString() : TEXT("");
+						UE_LOG(LogParty, Log, TEXT("[%s] Removed"), *PartyIdString);
+						Party->HandleRemovedFromParty(EMemberExitedReason::Left);
+					}
+					JoinedParties.Remove(PartyKey);
 				}
 			}
 
-			JoinedParties.Empty();
+			ensure(JoinedParties.Num() == 0);
 		}
 	}
 
diff --git a/Engine/Source/Runtime/Online/OnlineFramework/Party/Private/PartyGameState.cpp b/Engine/Source/Runtime/Online/OnlineFramework/Party/Private/PartyGameState.cpp
--- a/Engine/Source/Runtime/Online/OnlineFramework/Party/Private/PartyGameState.cpp
+++ b/Engine/Source/Runtime/Online/OnlineFramework/Party/Private/PartyGameState.cpp
@@ -706,7 +706,14 @@ void UPartyGameState::SetPartyType(EPartyType InPartyType, bool bLeaderFriendsOn
 			}
 
 			CurrentConfig.PresencePermissions = PresencePermissions;
-			CurrentConfig.InvitePermissions = bLeaderInvitesOnly ? PartySystemPermissions::EInvitePermissions::Leader : PartySystemPermissions::EInvitePermissions::Anyone;
+			if (PartyStateRef->bInvitesDisabled)
+			{
+				CurrentConfig.InvitePermissions = PartySystemPermissions::EInvitePermissions::Noone;
+			}
+			else
+			{
+				CurrentConfig.InvitePermissions = bLeaderInvitesOnly ? PartySystemPermissions::EInvitePermissions::Leader : PartySystemPermissions::EInvitePermissions::Anyone;
+			}
 
 			UpdatePartyConfig(bIsPrivate);
 
@@ -746,6 +753,40 @@ void UPartyGameState::SetPartyType(EPartyType InPartyType, bool bLeaderFriendsOn
 	}
 }
 
+void UPartyGameState::SetInvitesDisabled(bool bInvitesDisabled)
+{
+	if (IsLocalPartyLeader())
+	{
+		check(PartyStateRef);
+		if (PartyStateRef->bInvitesDisabled != bInvitesDisabled)
+		{
+			if (bInvitesDisabled)
+			{
+				CurrentConfig.InvitePermissions = PartySystemPermissions::EInvitePermissions::Noone;
+			}
+			else
+			{
+				CurrentConfig.InvitePermissions = PartyStateRef->bLeaderInvitesOnly ? PartySystemPermissions::EInvitePermissions::Leader : PartySystemPermissions::EInvitePermissions::Anyone;
+			}
+
+			UpdatePartyConfig();
+
+			// Replicate the party settings to other party members
+			PartyStateRef->bInvitesDisabled = bInvitesDisabled;
+			UpdatePartyData(OwningUserId);
+
+			// Refresh accepting members, taking everything into account
+			UpdateAcceptingMembers();
+
+			OnInvitesDisabledChanged().Broadcast(bInvitesDisabled);
+		}
+	}
+	else
+	{
+		UE_LOG(LogParty, Verbose, TEXT("Non party leader trying to set invites disabled!"));
+	}
+}
+
 void UPartyGameState::StayWithPartyOnExit(bool bInStayWithParty)
 {
 	bStayWithPartyOnDisconnect = bInStayWithParty;
@@ -1096,49 +1137,40 @@ void UPartyGameState::KickMember(const FUniqueNetIdRepl& PartyMemberToKick)
 	}
 }
 
-void UPartyGameState::InitLocalPlayerPartyData()
+UPartyMemberState* UPartyGameState::InitPartyMemberStateFromLocalPlayer(const ULocalPlayer* LocalPlayer)
 {
-	UWorld* World = GetWorld();
-	check(World);
-	AGameState* GameState = World->GetGameState();
-	check(GameState);
-
-	for (const APlayerState* PlayerState : GameState->PlayerArray)
+	UPartyMemberState* LocalPartyMemberState = nullptr;
+	TSharedPtr<const FUniqueNetId> UniqueNetId = LocalPlayer->GetPreferredUniqueNetId();
+	if (UniqueNetId.IsValid())
 	{
-		if (PlayerState && PlayerState->UniqueId.IsValid() && PlayerState->GetOuter() != nullptr)
+		UPartyMemberState** LocalPartyMemberStatePtr = PartyMembersState.Find(UniqueNetId);
+		LocalPartyMemberState = LocalPartyMemberStatePtr ? *LocalPartyMemberStatePtr : nullptr;
+		if (!LocalPartyMemberState)
 		{
-			UPartyMemberState** LocalPartyMemberStatePtr = PartyMembersState.Find(PlayerState->UniqueId);
-			UPartyMemberState* LocalPartyMemberState = LocalPartyMemberStatePtr ? *LocalPartyMemberStatePtr : nullptr;
-			if (!LocalPartyMemberState)
+			LocalPartyMemberState = CreateNewPartyMember(*UniqueNetId);
+			if (LocalPartyMemberState)
 			{
-				LocalPartyMemberState = CreateNewPartyMember(*PlayerState->UniqueId);
-				if (LocalPartyMemberState)
-				{
-					PartyMembersState.Add(PlayerState->UniqueId, LocalPartyMemberState);
-				}
+				PartyMembersState.Add(UniqueNetId, LocalPartyMemberState);
 			}
 		}
 	}
+	return LocalPartyMemberState;
 }
 
 void UPartyGameState::SendLocalPlayerPartyData()
 {
-	InitLocalPlayerPartyData();
-
 	UWorld* World = GetWorld();
 	check(World);
-	AGameState* GameState = World->GetGameState();
-	check(GameState);
 
-	for (const APlayerState* PlayerState : GameState->PlayerArray)
+	for (FLocalPlayerIterator It(GEngine, World); It; ++It)
 	{
-		if (PlayerState && PlayerState->UniqueId.IsValid() && PlayerState->GetOuter() != nullptr)
+		ULocalPlayer* LP = Cast<ULocalPlayer>(*It);
+		if (LP)
 		{
-			// Send local player data
-			UPartyMemberState** MemberStatePtr = PartyMembersState.Find(PlayerState->UniqueId);
-			if (MemberStatePtr != nullptr && *MemberStatePtr != nullptr)
+			const UPartyMemberState* MemberStatePtr = InitPartyMemberStateFromLocalPlayer(LP);
+			if (MemberStatePtr != nullptr)
 			{
-				UpdatePartyMemberState(OwningUserId, *MemberStatePtr);
+				UpdatePartyMemberState(OwningUserId, MemberStatePtr);
 			}
 		}
 	}
@@ -1209,7 +1241,7 @@ void UPartyGameState::UpdatePartyMemberState(const FUniqueNetIdRepl& InLocalUser
 					}
 					else
 					{
-						UE_LOG(LogParty, Warning, TEXT("UpdatePartyData: Failed to update party member data!"));
+						UE_LOG(LogParty, Warning, TEXT("UpdatePartyMemberState: Failed to update party member data!"));
 					}
 				}
 				else
diff --git a/Engine/Source/Runtime/Online/OnlineFramework/Party/Public/PartyGameState.h b/Engine/Source/Runtime/Online/OnlineFramework/Party/Public/PartyGameState.h
--- a/Engine/Source/Runtime/Online/OnlineFramework/Party/Public/PartyGameState.h
+++ b/Engine/Source/Runtime/Online/OnlineFramework/Party/Public/PartyGameState.h
@@ -85,10 +85,15 @@ struct FPartyState
 	UPROPERTY()
 	bool bLeaderInvitesOnly;
 
+	/** Are invites allowed at all? */
+	UPROPERTY()
+	bool bInvitesDisabled;
+
 	FPartyState() :
 		PartyType(EPartyType::FriendsOnly),
 		bLeaderFriendsOnly(false),
-		bLeaderInvitesOnly(false)
+		bLeaderInvitesOnly(false),
+		bInvitesDisabled(false)
 	{
 	}
 
@@ -142,12 +147,20 @@ private:
 	DECLARE_MULTICAST_DELEGATE_OneParam(FOnLeaderFriendsOnlyChanged, bool /* bLeaderFriendsOnly */);
 
 	/**
-	 * Delegate fired when a invite permissions change
+	 * Delegate fired when leader invite permissions change
 	 * 
 	 * @param bLeaderInviteOnly only leader invites are allowed
 	 */
 	DECLARE_MULTICAST_DELEGATE_OneParam(FOnLeaderInvitesOnlyChanged, bool /* bLeaderInviteOnly */);
 
+	/**
+	 * Delegate fired when invites disabled changes
+	 * 
+	 * @param bInvitesDisabled invites are disabled
+	 */
+	DECLARE_MULTICAST_DELEGATE_OneParam(FOnInvitesDisabledChanged, bool /* bInvitesDisabled */);
+
+
 protected:
 
 	/** 
@@ -204,6 +217,13 @@ public:
 	 */
 	void SetPartyType(EPartyType InPartyType, bool bLeaderFriendsOnly, bool bLeaderInvitesOnly);
 
+	/**
+	 * Disable (or re-enable) the ability to send party invites
+	 *
+	 * @param bInvitesDisabled invites are allowed to be sent
+	 */
+	void SetInvitesDisabled(bool bInvitesDisabled);
+
 	/**
 	 * Set stay with party on disconnects or not
 	 */
@@ -305,6 +325,8 @@ public:
 	FOnLeaderFriendsOnlyChanged& OnLeaderFriendsOnlyChanged() { return LeaderFriendsOnlyChanged; }
 	/** @return delegate fired when the leader based invite permissions have changed */
 	FOnLeaderInvitesOnlyChanged& OnLeaderInvitesOnlyChanged() { return LeaderInvitesOnlyChanged; }
+	/** @return delegate fired when invite permissions have changed */
+	FOnInvitesDisabledChanged& OnInvitesDisabledChanged() { return InvitesDisabledChanged; }
 	
 protected:
 
@@ -412,6 +434,7 @@ protected:
 	FOnPartyTypeChanged PartyTypeChanged;
 	FOnLeaderFriendsOnlyChanged LeaderFriendsOnlyChanged;
 	FOnLeaderInvitesOnlyChanged LeaderInvitesOnlyChanged;
+	FOnInvitesDisabledChanged InvitesDisabledChanged;
 
 	/**
 	 * Common initialization for a newly instantiated party
@@ -447,9 +470,9 @@ protected:
 	void UpdatePartyMemberState(const FUniqueNetIdRepl& InLocalUserId, const UPartyMemberState* InPartyPlayerData);
 
 	/**
-	 * Create and initialize a party member state value for each local player prior to transmission to other party members
+	 * Create and initialize a party member state value for a local player prior to transmission to other party members
 	 */
-	virtual void InitLocalPlayerPartyData();
+	virtual UPartyMemberState* InitPartyMemberStateFromLocalPlayer(const ULocalPlayer* LocalPlayer);
 
 	/**
 	 * Send all local player data to other party members, typically the initial call
diff --git a/Engine/Source/Runtime/Online/OnlineFramework/Qos/Private/QosEvaluator.cpp b/Engine/Source/Runtime/Online/OnlineFramework/Qos/Private/QosEvaluator.cpp
--- a/Engine/Source/Runtime/Online/OnlineFramework/Qos/Private/QosEvaluator.cpp
+++ b/Engine/Source/Runtime/Online/OnlineFramework/Qos/Private/QosEvaluator.cpp
@@ -217,21 +217,34 @@ void FQosEvaluator::OnEvaluateForDatacenterComplete(EQosCompletionResult Result,
 				}
 			}
 
+			//@TODO: This is a temporary measure to reduce the effect of framerate on data center ping estimation
+			// (due to the use of beacons that are ticked on the game thread for the estimation)
+			// This value can be 0 to disable discounting or something like 1 or 2
+			const float FractionOfFrameTimeToDiscount = 2.0f;
+
+			extern ENGINE_API float GAverageMS;
+			const int32 TimeToDiscount = (int32)(FractionOfFrameTimeToDiscount * GAverageMS);
+
 			for (auto& QosResult : RegionAggregates)
 			{
+				int32 RawAvgPing = MAX_QUERY_PING;
+				int32 AdjustedAvgPing = RawAvgPing;
 				if (QosResult.Value.NumResults > 0)
 				{
+					RawAvgPing = QosResult.Value.TotalPingInMs / QosResult.Value.NumResults;
+					AdjustedAvgPing = FMath::Max<int32>(RawAvgPing - TimeToDiscount, 1);
+
 					FQosRegionInfo Region;
 					Region.RegionId = QosResult.Key;
-					Region.PingMs = QosResult.Value.TotalPingInMs / QosResult.Value.NumResults;
+					Region.PingMs = AdjustedAvgPing;
 					RegionInfo.Add(Region);
 				}
 
-				UE_LOG(LogQos, Verbose, TEXT("Region[%s] Avg: %d Num: %d"), *QosResult.Key, QosResult.Value.AvgPing, QosResult.Value.NumResults);
+				UE_LOG(LogQos, Verbose, TEXT("Region[%s] Avg: %d Num: %d; Adjusted: %d"), *QosResult.Key, RawAvgPing, QosResult.Value.NumResults, AdjustedAvgPing);
 
 				if (QosStats.IsValid())
 				{
-					QosStats->RecordRegionInfo(QosResult.Key, QosResult.Value.AvgPing, QosResult.Value.NumResults);
+					QosStats->RecordRegionInfo(QosResult.Key, AdjustedAvgPing, QosResult.Value.NumResults);
 				}
 			}
 		}
@@ -355,9 +368,12 @@ void FQosEvaluator::OnQosRequestComplete(EQosResponseType QosResponse, int32 Res
 			SearchResult.PingInMs = MAX_QUERY_PING;
 		}
 
+		FString QosRegion;
+		SearchResult.Session.SessionSettings.Get(SETTING_REGION, QosRegion);
+
 		extern ENGINE_API float GAverageFPS;
 		extern ENGINE_API float GAverageMS;
-		UE_LOG(LogQos, Verbose, TEXT("Qos response received: %d ms FPS: %0.2f MS: %0.2f"), ResponseTime, GAverageFPS, GAverageMS);
+		UE_LOG(LogQos, Verbose, TEXT("Qos response received for region %s: %d ms FPS: %0.2f MS: %0.2f"), *QosRegion, ResponseTime, GAverageFPS, GAverageMS);
 
 		if (QosStats.IsValid())
 		{
diff --git a/Engine/Source/Runtime/Online/OnlineSubsystem/Private/OnlineAsyncTaskManager.cpp b/Engine/Source/Runtime/Online/OnlineSubsystem/Private/OnlineAsyncTaskManager.cpp
--- a/Engine/Source/Runtime/Online/OnlineSubsystem/Private/OnlineAsyncTaskManager.cpp
+++ b/Engine/Source/Runtime/Online/OnlineSubsystem/Private/OnlineAsyncTaskManager.cpp
@@ -5,6 +5,19 @@
 
 int32 FOnlineAsyncTaskManager::InvocationCount = 0;
 
+#if !UE_BUILD_SHIPPING
+namespace OSSConsoleVariables
+{
+	/** Time to delay finalization of a task in the out queue */
+	TAutoConsoleVariable<float> CVarDelayAsyncTaskOutQueue(
+		TEXT("OSS.DelayAsyncTaskOutQueue"),
+		0.0f,
+		TEXT("Min total async task time\n")
+		TEXT("Time in secs"),
+		ECVF_Default);
+}
+#endif
+
 /** The default value for the polling interval when not set by config */
 #define POLLING_INTERVAL_MS 50
 
@@ -181,6 +194,10 @@ void FOnlineAsyncTaskManager::GameTick()
 	FOnlineAsyncItem* Item = NULL;
 	int32 CurrentQueueSize = 0;
 
+#if !UE_BUILD_SHIPPING
+	const float TimeToWait = OSSConsoleVariables::CVarDelayAsyncTaskOutQueue.GetValueOnGameThread();
+#endif
+
 	do 
 	{
 		Item = NULL;
@@ -191,12 +208,33 @@ void FOnlineAsyncTaskManager::GameTick()
 			if (CurrentQueueSize > 0)
 			{
 				Item = OutQueue[0];
-				OutQueue.RemoveAt(0);
+
+#if !UE_BUILD_SHIPPING
+				if (Item && Item->GetElapsedTime() >= TimeToWait)
+				{
+					OutQueue.RemoveAt(0);
+				}
+				else
+				{
+					Item = nullptr;
+				}
+#else
+				OutQueue.RemoveAt(0);	
+#endif
 			}
 		}
 
 		if (Item)
 		{
+#if !UE_BUILD_SHIPPING
+			if (TimeToWait > 0.0f)
+			{
+				UE_LOG(LogOnline, Verbose, TEXT("Async task '%s' finalizing after %f seconds"),
+					*Item->ToString(),
+					Item->GetElapsedTime());
+			}
+#endif
+
 			// Finish work and trigger delegates
 			Item->Finalize();
 			Item->TriggerDelegates();
diff --git a/Engine/Source/Runtime/Online/OnlineSubsystem/Public/Interfaces/OnlineExternalUIInterface.h b/Engine/Source/Runtime/Online/OnlineSubsystem/Public/Interfaces/OnlineExternalUIInterface.h
--- a/Engine/Source/Runtime/Online/OnlineSubsystem/Public/Interfaces/OnlineExternalUIInterface.h
+++ b/Engine/Source/Runtime/Online/OnlineSubsystem/Public/Interfaces/OnlineExternalUIInterface.h
@@ -81,6 +81,14 @@ struct FShowWebUrlParams
 	{}
 };
 
+/**
+ * Union of all the platform informational message types we handle (some may be handled by more than one platform)
+ */
+enum class EPlatformMessageType
+{
+	EmptyStore
+};
+
 /** 
  * Interface definition for the online services external UIs
  * Any online service that provides extra UI overlays will implement the relevant functions
@@ -177,6 +185,24 @@ public:
 	*/
 	virtual bool ShowAccountUpgradeUI(const FUniqueNetId& UniqueId) = 0;
 
+	/**
+	 * Displays an informational system dialog.
+	 *
+	 * @param UserId Who to show this dialog for.
+	 * @param MessageId Platform-specific identifier for the message box to display.
+	 */
+	virtual bool ShowPlatformMessageBox(const FUniqueNetId& UserId, EPlatformMessageType MessageType) { return false; }
+
+	/**
+	 * Tell the system that you've entered into a screen considered in-game store.
+	 */
+	virtual void ReportEnterInGameStoreUI() {}
+
+	/**
+	 * Tell the system that you've entered into a screen considered in-game store.
+	 */
+	virtual void ReportExitInGameStoreUI() {}
+
 	/**
 	 * Delegate called when the external UI is opened or closed
 	 *
diff --git a/Engine/Source/Runtime/Online/OnlineSubsystem/Public/Interfaces/OnlinePartyInterface.h b/Engine/Source/Runtime/Online/OnlineSubsystem/Public/Interfaces/OnlinePartyInterface.h
--- a/Engine/Source/Runtime/Online/OnlineSubsystem/Public/Interfaces/OnlinePartyInterface.h
+++ b/Engine/Source/Runtime/Online/OnlineSubsystem/Public/Interfaces/OnlinePartyInterface.h
@@ -341,6 +341,8 @@ namespace PartySystemPermissions
 		Friends,
 		/** Available to anyone */
 		Anyone,
+		/** Noone can send invites */
+		Noone, 
 	};
 }
 
diff --git a/Engine/Source/Runtime/Online/OnlineSubsystem/Public/Interfaces/OnlineStoreInterfaceV2.h b/Engine/Source/Runtime/Online/OnlineSubsystem/Public/Interfaces/OnlineStoreInterfaceV2.h
--- a/Engine/Source/Runtime/Online/OnlineSubsystem/Public/Interfaces/OnlineStoreInterfaceV2.h
+++ b/Engine/Source/Runtime/Online/OnlineSubsystem/Public/Interfaces/OnlineStoreInterfaceV2.h
@@ -31,6 +31,8 @@ public:
 	FText Description;
 	/** Full description for display */
 	FText LongDescription;
+	/** Full pricing info as text for display */
+	FText PriceText;
 	
 	/** Price in numeric form for comparison/sorting */
 	int32 NumericPrice;
@@ -43,7 +45,15 @@ public:
 	/** @return FText suitable for localized display */
 	virtual FText GetDisplayPrice() const
 	{
-		return FText::AsCurrencyBase(NumericPrice, CurrencyCode);
+		if (!PriceText.IsEmpty())
+		{
+			return PriceText;			
+		}
+		else
+		{
+			return FText::AsCurrencyBase(NumericPrice, CurrencyCode);
+		}
+		
 	}
 
 	/** @return True if offer can be purchased */
diff --git a/Engine/Source/Runtime/Online/OnlineSubsystem/Public/OnlineError.h b/Engine/Source/Runtime/Online/OnlineSubsystem/Public/OnlineError.h
--- a/Engine/Source/Runtime/Online/OnlineSubsystem/Public/OnlineError.h
+++ b/Engine/Source/Runtime/Online/OnlineSubsystem/Public/OnlineError.h
@@ -3,7 +3,6 @@
 #pragma once
 
 namespace EOnlineServerConnectionStatus { enum Type; }
-typedef TSharedPtr<class IHttpResponse, ESPMode::ThreadSafe> FHttpResponsePtr;
 
 /** Generic Error response for OSS calls */
 struct ONLINESUBSYSTEM_API FOnlineError
diff --git a/Engine/Source/Runtime/Online/OnlineSubsystem/Public/OnlineSessionSettings.h b/Engine/Source/Runtime/Online/OnlineSubsystem/Public/OnlineSessionSettings.h
--- a/Engine/Source/Runtime/Online/OnlineSubsystem/Public/OnlineSessionSettings.h
+++ b/Engine/Source/Runtime/Online/OnlineSubsystem/Public/OnlineSessionSettings.h
@@ -570,7 +570,11 @@ public:
 
 	/** Constructor */
 	FOnlineSessionSearch() :
-		SearchState(EOnlineAsyncTaskState::NotStarted)
+		SearchState(EOnlineAsyncTaskState::NotStarted),
+		MaxSearchResults(0),
+		bIsLanQuery(false),
+		PingBucketSize(0),
+		TimeoutInSeconds(0.0f)
 	{
 		QuerySettings.Set(SETTING_MAPNAME, FString(TEXT("")), EOnlineComparisonOp::Equals);
 		QuerySettings.Set(SEARCH_DEDICATED_ONLY, false, EOnlineComparisonOp::Equals);
diff --git a/Engine/Source/Runtime/Online/OnlineSubsystem/Public/OnlineSubsystemTypes.h b/Engine/Source/Runtime/Online/OnlineSubsystem/Public/OnlineSubsystemTypes.h
--- a/Engine/Source/Runtime/Online/OnlineSubsystem/Public/OnlineSubsystemTypes.h
+++ b/Engine/Source/Runtime/Online/OnlineSubsystem/Public/OnlineSubsystemTypes.h
@@ -1112,7 +1112,7 @@ public:
 	/**
 	 * @return the nickname of the user if known
 	 */
-	virtual FString GetDisplayName() const = 0;
+	virtual FString GetDisplayName(const FString& Platform = FString()) const = 0;
 	/** 
 	 * @return Any additional user data associated with a registered user
 	 */
diff --git a/Engine/Source/Runtime/Online/OnlineSubsystemAmazon/Private/OnlineIdentityAmazon.cpp b/Engine/Source/Runtime/Online/OnlineSubsystemAmazon/Private/OnlineIdentityAmazon.cpp
--- a/Engine/Source/Runtime/Online/OnlineSubsystemAmazon/Private/OnlineIdentityAmazon.cpp
+++ b/Engine/Source/Runtime/Online/OnlineSubsystemAmazon/Private/OnlineIdentityAmazon.cpp
@@ -15,7 +15,7 @@ FString FUserOnlineAccountAmazon::GetRealName() const
 	return FString();
 }
 
-FString FUserOnlineAccountAmazon::GetDisplayName() const
+FString FUserOnlineAccountAmazon::GetDisplayName(const FString& Platform) const
 {
 	return FString();
 }
diff --git a/Engine/Source/Runtime/Online/OnlineSubsystemAmazon/Private/OnlineIdentityAmazon.h b/Engine/Source/Runtime/Online/OnlineSubsystemAmazon/Private/OnlineIdentityAmazon.h
--- a/Engine/Source/Runtime/Online/OnlineSubsystemAmazon/Private/OnlineIdentityAmazon.h
+++ b/Engine/Source/Runtime/Online/OnlineSubsystemAmazon/Private/OnlineIdentityAmazon.h
@@ -18,7 +18,7 @@ public:
 	
 	virtual TSharedRef<const FUniqueNetId> GetUserId() const override;
 	virtual FString GetRealName() const override;
-	virtual FString GetDisplayName() const override;
+	virtual FString GetDisplayName(const FString& Platform = FString()) const override;
 	virtual bool GetUserAttribute(const FString& AttrName, FString& OutAttrValue) const override;
 	virtual bool SetUserAttribute(const FString& AttrName, const FString& AttrValue) override;
 
diff --git a/Engine/Source/Runtime/Online/OnlineSubsystemFacebook/Private/IOS/OnlineFriendsFacebook.cpp b/Engine/Source/Runtime/Online/OnlineSubsystemFacebook/Private/IOS/OnlineFriendsFacebook.cpp
--- a/Engine/Source/Runtime/Online/OnlineSubsystemFacebook/Private/IOS/OnlineFriendsFacebook.cpp
+++ b/Engine/Source/Runtime/Online/OnlineSubsystemFacebook/Private/IOS/OnlineFriendsFacebook.cpp
@@ -22,7 +22,7 @@ FString FOnlineFriendFacebook::GetRealName() const
 	return Result;
 }
 
-FString FOnlineFriendFacebook::GetDisplayName() const
+FString FOnlineFriendFacebook::GetDisplayName(const FString& Platform) const
 {
 	FString Result;
 	GetAccountData(TEXT("username"), Result);
diff --git a/Engine/Source/Runtime/Online/OnlineSubsystemFacebook/Private/IOS/OnlineIdentityFacebook.cpp b/Engine/Source/Runtime/Online/OnlineSubsystemFacebook/Private/IOS/OnlineIdentityFacebook.cpp
--- a/Engine/Source/Runtime/Online/OnlineSubsystemFacebook/Private/IOS/OnlineIdentityFacebook.cpp
+++ b/Engine/Source/Runtime/Online/OnlineSubsystemFacebook/Private/IOS/OnlineIdentityFacebook.cpp
@@ -34,7 +34,7 @@ FString FUserOnlineAccountFacebook::GetRealName() const
 	return FString();
 }
 
-FString FUserOnlineAccountFacebook::GetDisplayName() const
+FString FUserOnlineAccountFacebook::GetDisplayName(const FString& Platform) const
 {
 	//@todo samz - implement
 	return FString();
diff --git a/Engine/Source/Runtime/Online/OnlineSubsystemFacebook/Private/IOS/OnlineUserFacebook.cpp b/Engine/Source/Runtime/Online/OnlineSubsystemFacebook/Private/IOS/OnlineUserFacebook.cpp
--- a/Engine/Source/Runtime/Online/OnlineSubsystemFacebook/Private/IOS/OnlineUserFacebook.cpp
+++ b/Engine/Source/Runtime/Online/OnlineSubsystemFacebook/Private/IOS/OnlineUserFacebook.cpp
@@ -21,7 +21,7 @@ FString FOnlineUserInfoFacebook::GetRealName() const
 	return Result;
 }
 
-FString FOnlineUserInfoFacebook::GetDisplayName() const
+FString FOnlineUserInfoFacebook::GetDisplayName(const FString& Platform) const
 {
 	FString Result;
 	GetAccountData(TEXT("username"), Result);
diff --git a/Engine/Source/Runtime/Online/OnlineSubsystemFacebook/Private/Windows/OnlineFriendsFacebook.cpp b/Engine/Source/Runtime/Online/OnlineSubsystemFacebook/Private/Windows/OnlineFriendsFacebook.cpp
--- a/Engine/Source/Runtime/Online/OnlineSubsystemFacebook/Private/Windows/OnlineFriendsFacebook.cpp
+++ b/Engine/Source/Runtime/Online/OnlineSubsystemFacebook/Private/Windows/OnlineFriendsFacebook.cpp
@@ -18,7 +18,7 @@ FString FOnlineFriendFacebook::GetRealName() const
 	return Result;
 }
 
-FString FOnlineFriendFacebook::GetDisplayName() const
+FString FOnlineFriendFacebook::GetDisplayName(const FString& Platform) const
 {
 	FString Result;
 	GetAccountData(TEXT("username"), Result);
diff --git a/Engine/Source/Runtime/Online/OnlineSubsystemFacebook/Private/Windows/OnlineIdentityFacebook.cpp b/Engine/Source/Runtime/Online/OnlineSubsystemFacebook/Private/Windows/OnlineIdentityFacebook.cpp
--- a/Engine/Source/Runtime/Online/OnlineSubsystemFacebook/Private/Windows/OnlineIdentityFacebook.cpp
+++ b/Engine/Source/Runtime/Online/OnlineSubsystemFacebook/Private/Windows/OnlineIdentityFacebook.cpp
@@ -16,7 +16,7 @@ FString FUserOnlineAccountFacebook::GetRealName() const
 	return FString();
 }
 
-FString FUserOnlineAccountFacebook::GetDisplayName() const
+FString FUserOnlineAccountFacebook::GetDisplayName(const FString& Platform) const
 {
 	//@todo samz - implement
 	return FString();
diff --git a/Engine/Source/Runtime/Online/OnlineSubsystemNull/Private/OnlineIdentityNull.h b/Engine/Source/Runtime/Online/OnlineSubsystemNull/Private/OnlineIdentityNull.h
--- a/Engine/Source/Runtime/Online/OnlineSubsystemNull/Private/OnlineIdentityNull.h
+++ b/Engine/Source/Runtime/Online/OnlineSubsystemNull/Private/OnlineIdentityNull.h
@@ -17,7 +17,7 @@ public:
 	
 	virtual TSharedRef<const FUniqueNetId> GetUserId() const override { return UserIdPtr; }
 	virtual FString GetRealName() const override { return TEXT("DummyRealName"); }
-	virtual FString GetDisplayName() const override  { return TEXT("DummyDisplayName"); }
+	virtual FString GetDisplayName(const FString& Platform = FString()) const override  { return TEXT("DummyDisplayName"); }
 	virtual bool GetUserAttribute(const FString& AttrName, FString& OutAttrValue) const override;
 	virtual bool SetUserAttribute(const FString& AttrName, const FString& AttrValue) override;
 
diff --git a/Engine/Source/Runtime/Online/OnlineSubsystemSteam/Private/OnlineFriendsInterfaceSteam.cpp b/Engine/Source/Runtime/Online/OnlineSubsystemSteam/Private/OnlineFriendsInterfaceSteam.cpp
--- a/Engine/Source/Runtime/Online/OnlineSubsystemSteam/Private/OnlineFriendsInterfaceSteam.cpp
+++ b/Engine/Source/Runtime/Online/OnlineSubsystemSteam/Private/OnlineFriendsInterfaceSteam.cpp
@@ -25,7 +25,7 @@ FString FOnlineFriendSteam::GetRealName() const
 	return Result;
 }
 
-FString FOnlineFriendSteam::GetDisplayName() const
+FString FOnlineFriendSteam::GetDisplayName(const FString& Platform) const
 {
 	FString Result;
 	GetAccountData(TEXT("nickname"),Result);
diff --git a/Engine/Source/Runtime/Online/OnlineSubsystemUtils/Private/PartyBeaconClient.cpp b/Engine/Source/Runtime/Online/OnlineSubsystemUtils/Private/PartyBeaconClient.cpp
--- a/Engine/Source/Runtime/Online/OnlineSubsystemUtils/Private/PartyBeaconClient.cpp
+++ b/Engine/Source/Runtime/Online/OnlineSubsystemUtils/Private/PartyBeaconClient.cpp
@@ -5,6 +5,46 @@
 #include "PartyBeaconHost.h"
 #include "OnlineSessionInterface.h"
 
+#if !UE_BUILD_SHIPPING
+namespace BeaconConsoleVariables
+{
+	/** Time to delay delegates firing a reservation request response */
+	TAutoConsoleVariable<float> CVarDelayReservationResponse(
+		TEXT("beacon.DelayReservationResponse"),
+		0.0f,
+		TEXT("Delay time between received response and notification\n")
+		TEXT("Time in secs"),
+		ECVF_Default);
+
+	/** Time to delay delegates firing a cancel reservation request response */
+	TAutoConsoleVariable<float> CVarDelayCancellationResponse(
+		TEXT("beacon.DelayCancellationResponse"),
+		0.0f,
+		TEXT("Delay time between received cancel response and notification\n")
+		TEXT("Time in secs"),
+		ECVF_Default);
+
+	/** Time to delay delegates firing a reservation update response */
+	TAutoConsoleVariable<float> CVarDelayUpdateResponse(
+		TEXT("beacon.DelayUpdateResponse"),
+		0.0f,
+		TEXT("Delay time between received update response and notification\n")
+		TEXT("Time in secs"),
+		ECVF_Default);
+
+	/** Time to delay delegates firing a reservation full response */
+	TAutoConsoleVariable<float> CVarDelayFullResponse(
+		TEXT("beacon.DelayFullResponse"),
+		0.0f,
+		TEXT("Delay time between received full response and notification\n")
+		TEXT("Time in secs"),
+		ECVF_Default);
+}
+#endif
+
+/** Max time to wait for a response from the server for CancelReservation */
+#define CANCEL_FAILSAFE 5.0f
+
 APartyBeaconClient::APartyBeaconClient(const FObjectInitializer& ObjectInitializer) :
 	Super(ObjectInitializer),
 	RequestType(EClientRequestType::NonePending),
@@ -13,6 +53,57 @@ APartyBeaconClient::APartyBeaconClient(const FObjectInitializer& ObjectInitializ
 {
 }
 
+void APartyBeaconClient::BeginDestroy()
+{
+	ClearTimers();
+	Super::BeginDestroy();
+}
+
+void APartyBeaconClient::ClearTimers()
+{
+	UWorld* World = GetWorld();
+	if (World)
+	{
+		if (PendingResponseTimerHandle.IsValid())
+		{
+			UE_LOG(LogBeacon, Verbose, TEXT("ClearTimers: Pending reservation response cleared."));
+		}
+
+		if (PendingCancelResponseTimerHandle.IsValid())
+		{
+			UE_LOG(LogBeacon, Verbose, TEXT("ClearTimers: Pending cancel response cleared."));
+		}
+
+		if (PendingReservationUpdateTimerHandle.IsValid())
+		{
+			UE_LOG(LogBeacon, Verbose, TEXT("ClearTimers: Pending reservation update cleared."));
+		}
+
+		if (PendingReservationFullTimerHandle.IsValid())
+		{
+			UE_LOG(LogBeacon, Verbose, TEXT("ClearTimers: Pending reservation full cleared."));
+		}
+
+		if (CancelRPCFailsafe.IsValid())
+		{
+			UE_LOG(LogBeacon, Verbose, TEXT("ClearTimers: Cancel failsafe cleared."));
+		}
+
+		FTimerManager& TM = World->GetTimerManager();
+		TM.ClearTimer(PendingResponseTimerHandle);
+		TM.ClearTimer(PendingCancelResponseTimerHandle);
+		TM.ClearTimer(PendingReservationUpdateTimerHandle);
+		TM.ClearTimer(PendingReservationFullTimerHandle);
+		TM.ClearTimer(CancelRPCFailsafe);
+
+		PendingResponseTimerHandle.Invalidate();
+		PendingCancelResponseTimerHandle.Invalidate();
+		PendingReservationUpdateTimerHandle.Invalidate();
+		PendingReservationFullTimerHandle.Invalidate();
+		CancelRPCFailsafe.Invalidate();
+	}
+}
+
 bool APartyBeaconClient::RequestReservation(const FString& ConnectInfoStr, const FString& InSessionId, const FUniqueNetIdRepl& RequestingPartyLeader, const TArray<FPlayerReservation>& PartyMembers)
 {
 	bool bSuccess = false;
@@ -29,7 +120,7 @@ bool APartyBeaconClient::RequestReservation(const FString& ConnectInfoStr, const
 	}
 	else
 	{
-		UE_LOG_ONLINE(Warning, TEXT("RequestReservation: Failure to init client beacon with %s."), *ConnectURL.ToString());
+		UE_LOG(LogBeacon, Warning, TEXT("RequestReservation: Failure to init client beacon with %s."), *ConnectURL.ToString());
 		RequestType = EClientRequestType::NonePending;
 	}
 
@@ -139,26 +230,37 @@ bool APartyBeaconClient::RequestReservationUpdate(const FOnlineSessionSearchResu
 
 void APartyBeaconClient::CancelReservation()
 {
-	if (PendingReservation.PartyLeader.IsValid())
+	if (ensure(PendingReservation.PartyLeader.IsValid()))
 	{
 		bCancelReservation = true;
+
+		// Clear out any pending response handling, only the cancel matters
+		ClearTimers();
+
 		if (bPendingReservationSent)
 		{
 			UE_LOG(LogBeacon, Verbose, TEXT("Sending cancel reservation request."));
 			ServerCancelReservationRequest(PendingReservation.PartyLeader);
+
+			// In case the server is loading or unresponsive (ie no host beacon)
+			FTimerDelegate TimerDelegate;
+			TimerDelegate.BindUObject(this, &ThisClass::OnCancelledFailsafe);
+			
+			UWorld* World = GetWorld();
+			check(World);
+			World->GetTimerManager().SetTimer(CancelRPCFailsafe, TimerDelegate, CANCEL_FAILSAFE, false);
 		}
 		else
 		{
 			UE_LOG(LogBeacon, Verbose, TEXT("Reservation request never sent, no need to send cancelation request."));
-			ReservationRequestComplete.ExecuteIfBound( EPartyReservationResult::ReservationRequestCanceled );
+			OnCancelledComplete();
 		}
 	}
 	else
 	{
 		UE_LOG(LogBeacon, Verbose, TEXT("Unable to cancel reservation request with invalid party leader."));
+		OnCancelledComplete();
 	}
-
-	RequestType = EClientRequestType::NonePending;
 }
 
 void APartyBeaconClient::OnConnected()
@@ -185,14 +287,26 @@ void APartyBeaconClient::OnConnected()
 	}
 	else
 	{
-		UE_LOG(LogBeacon, Verbose, TEXT("Reservation request previously canceled, aborting reservation request."));
-		ReservationRequestComplete.ExecuteIfBound(EPartyReservationResult::ReservationRequestCanceled);
-		RequestType = EClientRequestType::NonePending;
+		UE_LOG(LogBeacon, Verbose, TEXT("Reservation request previously canceled, nothing sent."));
+		OnCancelledComplete();
 	}
 }
 
+void APartyBeaconClient::OnCancelledFailsafe()
+{
+	ClientCancelReservationResponse_Implementation(EPartyReservationResult::ReservationRequestCanceled);
+}
+ 
+void APartyBeaconClient::OnCancelledComplete()
+{
+	ReservationRequestComplete.ExecuteIfBound(EPartyReservationResult::ReservationRequestCanceled);
+	RequestType = EClientRequestType::NonePending;
+	bCancelReservation = false;
+}
+
 void APartyBeaconClient::OnFailure()
 {
+	ClearTimers();
 	RequestType = EClientRequestType::NonePending;
 	Super::OnFailure();
 }
@@ -246,19 +360,165 @@ void APartyBeaconClient::ServerCancelReservationRequest_Implementation(const FUn
 
 void APartyBeaconClient::ClientReservationResponse_Implementation(EPartyReservationResult::Type ReservationResponse)
 {
-	UE_LOG(LogBeacon, Verbose, TEXT("Party beacon response received %s"), EPartyReservationResult::ToString(ReservationResponse));
-	ReservationRequestComplete.ExecuteIfBound(ReservationResponse);
-	RequestType = EClientRequestType::NonePending;
+	if (!bCancelReservation)
+	{
+#if !UE_BUILD_SHIPPING
+		const float Rate = BeaconConsoleVariables::CVarDelayReservationResponse.GetValueOnGameThread();
+#else
+		const float Rate = 0.0f;
+#endif
+		if (Rate > 0.0f)
+		{
+			UE_LOG(LogBeacon, Verbose, TEXT("Party beacon response received %s, waiting %fs to notify"), EPartyReservationResult::ToString(ReservationResponse), Rate);
+
+			FTimerDelegate TimerDelegate;
+			TimerDelegate.BindLambda([this, ReservationResponse]()
+			{
+				ProcessReservationResponse(ReservationResponse);
+			});
+
+			PendingResponseTimerHandle = DelayResponse(TimerDelegate, Rate);
+		}
+		else
+		{
+			ProcessReservationResponse(ReservationResponse);
+		}
+	}
+	else
+	{
+		UE_LOG(LogBeacon, Verbose, TEXT("Party beacon response received %s, ignored due to cancel in progress"), EPartyReservationResult::ToString(ReservationResponse));
+		// Cancel RPC or failsafe timer will trigger the cancel
+	}
+}
+
+void APartyBeaconClient::ProcessReservationResponse(EPartyReservationResult::Type ReservationResponse)
+{
+	if (!bCancelReservation)
+	{
+		UE_LOG(LogBeacon, Verbose, TEXT("Party beacon response received %s"), EPartyReservationResult::ToString(ReservationResponse));
+		ReservationRequestComplete.ExecuteIfBound(ReservationResponse);
+		RequestType = EClientRequestType::NonePending;
+	}
+	else
+	{
+		UE_LOG(LogBeacon, Verbose, TEXT("Party beacon response received %s, ignored due to cancel in progress"), EPartyReservationResult::ToString(ReservationResponse));
+		// Cancel RPC or failsafe timer will trigger the cancel
+	}
+}
+
+void APartyBeaconClient::ClientCancelReservationResponse_Implementation(EPartyReservationResult::Type ReservationResponse)
+{
+	ensure(bCancelReservation);
+
+	// Clear out any pending response handling (including failsafe timer)
+	ClearTimers();
+#if !UE_BUILD_SHIPPING	
+	const float Rate = BeaconConsoleVariables::CVarDelayCancellationResponse.GetValueOnGameThread();
+#else
+	const float Rate = 0.0f;
+#endif
+	if (Rate > 0.0f)
+	{
+		UE_LOG(LogBeacon, Verbose, TEXT("Party beacon cancellation response received %s, waiting %fs to notify"), EPartyReservationResult::ToString(ReservationResponse), Rate);
+
+		FTimerDelegate TimerDelegate;
+		TimerDelegate.BindLambda([this, ReservationResponse]()
+		{
+			ProcessCancelReservationResponse(ReservationResponse);
+		});
+
+		PendingCancelResponseTimerHandle = DelayResponse(TimerDelegate, Rate);
+	}
+	else
+	{
+		ProcessCancelReservationResponse(ReservationResponse);
+	}
+}
+
+void APartyBeaconClient::ProcessCancelReservationResponse(EPartyReservationResult::Type ReservationResponse)
+{
+	ensure(ReservationResponse == EPartyReservationResult::ReservationRequestCanceled || ReservationResponse == EPartyReservationResult::ReservationNotFound);
+	ensure(bCancelReservation);
+	OnCancelledComplete();
 }
 
 void APartyBeaconClient::ClientSendReservationUpdates_Implementation(int32 NumRemainingReservations)
+{
+	if (!bCancelReservation)
+	{
+#if !UE_BUILD_SHIPPING
+		const float Rate = BeaconConsoleVariables::CVarDelayUpdateResponse.GetValueOnGameThread();
+#else
+		const float Rate = 0.0f;
+#endif
+		if (Rate > 0.0f)
+		{
+			UE_LOG(LogBeacon, Verbose, TEXT("Party beacon reservations remaining %d, waiting %fs to notify"), NumRemainingReservations, Rate);
+
+			FTimerDelegate TimerDelegate;
+			TimerDelegate.BindLambda([this, NumRemainingReservations]()
+			{
+				ProcessReservationUpdate(NumRemainingReservations);
+			});
+
+			PendingReservationUpdateTimerHandle = DelayResponse(TimerDelegate, Rate);
+		}
+		else
+		{
+			ProcessReservationUpdate(NumRemainingReservations);
+		}
+	}
+}
+
+void APartyBeaconClient::ProcessReservationUpdate(int32 NumRemainingReservations)
 {
 	UE_LOG(LogBeacon, Verbose, TEXT("Party beacon reservations remaining %d"), NumRemainingReservations);
 	ReservationCountUpdate.ExecuteIfBound(NumRemainingReservations);
 }
 
 void APartyBeaconClient::ClientSendReservationFull_Implementation()
+{
+	if (!bCancelReservation)
+	{
+#if !UE_BUILD_SHIPPING
+		const float Rate = BeaconConsoleVariables::CVarDelayFullResponse.GetValueOnGameThread();
+#else
+		const float Rate = 0.0f;
+#endif
+		if (Rate > 0.0f)
+		{
+			UE_LOG(LogBeacon, Verbose, TEXT("Party beacon reservations full, waiting %fs to notify"), Rate);
+
+			FTimerDelegate TimerDelegate;
+			TimerDelegate.BindLambda([this]()
+			{
+				ProcessReservationFull();
+			});
+
+			PendingReservationFullTimerHandle = DelayResponse(TimerDelegate, Rate);
+		}
+		else
+		{
+			ProcessReservationFull();
+		}
+	}
+}
+
+void APartyBeaconClient::ProcessReservationFull()
 {
 	UE_LOG(LogBeacon, Verbose, TEXT("Party beacon reservations full"));
 	ReservationFull.ExecuteIfBound();
 }
+
+FTimerHandle APartyBeaconClient::DelayResponse(FTimerDelegate& Delegate, float Delay)
+{
+	FTimerHandle TimerHandle;
+
+	UWorld* World = GetWorld();
+	if (ensure(World != nullptr))
+	{
+		World->GetTimerManager().SetTimer(TimerHandle, Delegate, Delay, false);
+	}
+
+	return TimerHandle;
+}
\ No newline at end of file
diff --git a/Engine/Source/Runtime/Online/OnlineSubsystemUtils/Private/PartyBeaconHost.cpp b/Engine/Source/Runtime/Online/OnlineSubsystemUtils/Private/PartyBeaconHost.cpp
--- a/Engine/Source/Runtime/Online/OnlineSubsystemUtils/Private/PartyBeaconHost.cpp
+++ b/Engine/Source/Runtime/Online/OnlineSubsystemUtils/Private/PartyBeaconHost.cpp
@@ -807,7 +807,7 @@ void APartyBeaconHost::ProcessCancelReservationRequest(APartyBeaconClient* Clien
 			DumpReservations();
 		}
 
-		Client->ClientReservationResponse(Result);
+		Client->ClientCancelReservationResponse(Result);
 	}
 }
 
diff --git a/Engine/Source/Runtime/Online/OnlineSubsystemUtils/Public/PartyBeaconClient.h b/Engine/Source/Runtime/Online/OnlineSubsystemUtils/Public/PartyBeaconClient.h
--- a/Engine/Source/Runtime/Online/OnlineSubsystemUtils/Public/PartyBeaconClient.h
+++ b/Engine/Source/Runtime/Online/OnlineSubsystemUtils/Public/PartyBeaconClient.h
@@ -82,6 +82,10 @@ class ONLINESUBSYSTEMUTILS_API APartyBeaconClient : public AOnlineBeaconClient
 {
 	GENERATED_UCLASS_BODY()
 
+	//~ Begin UObject Interface
+	virtual void BeginDestroy() override;
+	//~ End UObject Interface
+
 	//~ Begin AOnlineBeaconClient Interface
 	virtual void OnConnected() override;
 	virtual void OnFailure() override;
@@ -164,6 +168,14 @@ class ONLINESUBSYSTEMUTILS_API APartyBeaconClient : public AOnlineBeaconClient
  	UFUNCTION(client, reliable)
 	virtual void ClientReservationResponse(EPartyReservationResult::Type ReservationResponse);
 
+	/**
+	 * Response from the host session after making a cancellation request
+	 *
+	 * @param ReservationResponse response from server
+	 */
+	UFUNCTION(client, reliable)
+	virtual void ClientCancelReservationResponse(EPartyReservationResult::Type ReservationResponse);
+	
 	/**
 	 * Response from the host session that the reservation count has changed
 	 *
@@ -226,6 +238,31 @@ protected:
 	UPROPERTY()
 	bool bCancelReservation;
 
+	/** Timer to trigger a cancel reservation request if the server doesn't respond in time */
+	FTimerHandle CancelRPCFailsafe;
+
+	/** Timers for delaying various responses (debug) */
+	FTimerHandle PendingResponseTimerHandle;
+	FTimerHandle PendingCancelResponseTimerHandle;
+	FTimerHandle PendingReservationUpdateTimerHandle;
+	FTimerHandle PendingReservationFullTimerHandle;
+
+	/** Clear out all the timer handles listed above */
+	void ClearTimers();
+	/** Delegate triggered if the client doesn't hear from the server in time */
+	void OnCancelledFailsafe();
+	/** Delegate triggered when a cancel reservation request is complete */
+	void OnCancelledComplete();
+
+	/** Process a response to our RequestReservation request to the server */
+	void ProcessReservationResponse(EPartyReservationResult::Type ReservationResponse);
+	/** Process a response to our CancelReservation request to the server */
+	void ProcessCancelReservationResponse(EPartyReservationResult::Type ReservationResponse);
+	/** Process a response from the server with an update to the number of consumed reservations */
+	void ProcessReservationUpdate(int32 NumRemainingReservations);
+	/** Process a response from the server that the reservation beacon is full */
+	void ProcessReservationFull();
+
 	/**
 	 * Tell the server about the reservation request being made
 	 *
@@ -251,4 +288,14 @@ protected:
 	 */
 	UFUNCTION(server, reliable, WithValidation)
 	virtual void ServerCancelReservationRequest(const FUniqueNetIdRepl& PartyLeader);
+
+	/**
+	 * Trigger the given delegate at a later time
+	 *
+	 * @param Delegate function to call at a later date
+	 * @param Delay time to wait before calling function
+	 *
+	 * @return handle in the timer system for this entry
+	 */
+	FTimerHandle DelayResponse(FTimerDelegate& Delegate, float Delay);
 };
diff --git a/Engine/Source/Runtime/PacketHandlers/PacketHandler/Private/PacketHandler.cpp b/Engine/Source/Runtime/PacketHandlers/PacketHandler/Private/PacketHandler.cpp
--- a/Engine/Source/Runtime/PacketHandlers/PacketHandler/Private/PacketHandler.cpp
+++ b/Engine/Source/Runtime/PacketHandlers/PacketHandler/Private/PacketHandler.cpp
@@ -11,30 +11,36 @@ DEFINE_LOG_CATEGORY(PacketHandlerLog);
 // @todo #JohnB: For the moment, disable the reliability handler, while it is causing network trouble
 #define DISABLE_RELIABILITY_HANDLER 1
 
-// BUFFERED PACKET
+/**
+ * BufferedPacket
+ */
+
 BufferedPacket::BufferedPacket()
-: Data(nullptr)
-, BytesCount(0)
-, SendTime(0.f)
-, Id(0)
+	: Data(nullptr)
+	, BytesCount(0)
+	, SendTime(0.f)
+	, Id(0)
 {
 }
 
 BufferedPacket::~BufferedPacket()
 {
-	if(Data != nullptr)
+	if (Data != nullptr)
 	{
 		delete Data;
 	}
 }
 
-// PACKET HANDLER
+/**
+ * PacketHandler
+ */
 PacketHandler::PacketHandler()
-: bEnabled(false)
-, Time(0.f)
-, State(Handler::State::Uninitialized)
-, ReliabilityComponent(NULL)
+	: Time(0.f)
+	, State(Handler::State::Uninitialized)
+	, ReliabilityComponent(nullptr)
 {
+	GConfig->GetBool(TEXT("PacketHandlerComponents"), TEXT("bEnabled"), bEnabled, GEngineIni);
+
 	OutgoingPacket.SetAllowResize(true);
 	OutgoingPacket.AllowAppend(true);
 }
@@ -55,7 +61,7 @@ void PacketHandler::Tick(float DeltaTime)
 		}
 	}
 
-	for (int32 i = 0; i < HandlerComponents.Num(); ++i)
+	for (int32 i=0; i<HandlerComponents.Num(); ++i)
 	{
 		HandlerComponents[i]->Tick(DeltaTime);
 	}
@@ -63,8 +69,6 @@ void PacketHandler::Tick(float DeltaTime)
 
 void PacketHandler::Initialize(Handler::Mode InMode)
 {
-	GConfig->GetBool(TEXT("PacketHandlerComponents"), TEXT("bEnabled"), bEnabled, GEngineIni);
-
 	if (!bEnabled && FParse::Param(FCommandLine::Get(), TEXT("PacketHandler")))
 	{
 		UE_LOG(PacketHandlerLog, Log, TEXT("Force-enabling packet handler from commandline."));
@@ -161,11 +165,16 @@ void PacketHandler::Initialize(Handler::Mode InMode)
 				UE_LOG(PacketHandlerLog, Warning, TEXT("Unable to Load Module: %s"), *ComponentsArray[i].ComponentName);
 			}
 		}
+
+
+		// Called early, to ensure that all handlers report a valid packet overhead (triggers an assert if not)
+		GetTotalPacketOverheadBits();
 	}
 }
 
 void PacketHandler::Add(TSharedPtr<HandlerComponent> NewHandler)
 {
+	// This is never valid. Can end up silently changing maximum allow packet size, which could cause failure to send packets.
 	if (State != Handler::State::Uninitialized)
 	{
 		LowLevelFatalError(TEXT("Handler added during runtime."));
@@ -243,13 +252,13 @@ const ProcessedPacket PacketHandler::Outgoing(uint8* Packet, int32 Count)
 	}
 
 	// Queue packet for resending before handling
-	if(Count > 0 && ReliabilityComponent.IsValid())
+	if (Count > 0 && ReliabilityComponent.IsValid())
 	{
 		ReliabilityComponent->QueuePacketForResending(Packet, Count);
 	}
 
 	// Handle
-	for (int32 i = 0; i < HandlerComponents.Num(); ++i)
+	for (int32 i=0; i<HandlerComponents.Num(); ++i)
 	{
 		if (HandlerComponents[i]->IsActive())
 		{
@@ -300,7 +309,7 @@ const ProcessedPacket PacketHandler::Incoming(uint8* Packet, int32 Count)
 	}
 
 	// Handle
-	for (int32 i = HandlerComponents.Num() - 1; i >= 0; --i)
+	for (int32 i=HandlerComponents.Num() - 1; i>=0; --i)
 	{
 		if (HandlerComponents[i]->IsActive() && ProcessedPacketReader.GetNumBytes() > 0)
 		{
@@ -328,7 +337,7 @@ void PacketHandler::SetState(Handler::State InState)
 void PacketHandler::HandlerInitialized()
 {
 	// If any buffered packets, add to queue
-	for (int32 i = 0; i < BufferedPackets.Num(); ++i)
+	for (int32 i=0; i<BufferedPackets.Num(); ++i)
 	{
 		QueuedPackets.Enqueue(BufferedPackets[i]);
 		BufferedPackets[i] = nullptr;
@@ -342,7 +351,7 @@ void PacketHandler::HandlerComponentInitialized()
 	// Check if all handlers are initialized
 	if (State != Handler::State::Initialized)
 	{
-		for (int32 i = 0; i < HandlerComponents.Num(); ++i)
+		for (int32 i=0; i<HandlerComponents.Num(); ++i)
 		{
 			if (HandlerComponents[i]->IsInitialized() == false)
 			{
@@ -357,7 +366,9 @@ void PacketHandler::HandlerComponentInitialized()
 BufferedPacket* PacketHandler::GetQueuedPacket()
 {
 	BufferedPacket* QueuedPacket = nullptr;
+
 	QueuedPackets.Dequeue(QueuedPacket);
+
 	return QueuedPacket;
 }
 
@@ -366,13 +377,37 @@ void PacketHandler::QueuePacketForSending(BufferedPacket* PacketToQueue)
 
 }
 
-// HANDLER COMPONENT
-HandlerComponent::HandlerComponent()
-: State(Handler::Component::State::UnInitialized)
-, bActive(false)
-, bInitialized(false)
+int32 PacketHandler::GetTotalPacketOverheadBits()
 {
+	int32 ReturnVal = 0;
+
+	for (int32 i=0; i<HandlerComponents.Num(); i++)
+	{
+		int32 CurOverhead = HandlerComponents[i]->GetPacketOverheadBits();
+
+		// Specifying the packet overhead is mandatory, even if no overhead (as accidentally forgetting, leads to hard to trace issues).
+		if (CurOverhead == -1)
+		{
+			LowLevelFatalError(TEXT("Handler returned invalid 'GetPacketOverhead' value."));
+			continue;
+		}
+
+		ReturnVal += CurOverhead;
+	}
 
+	return ReturnVal;
+}
+
+
+/**
+ * HandlerComponent
+ */
+
+HandlerComponent::HandlerComponent()
+	: State(Handler::Component::State::UnInitialized)
+	, bActive(false)
+	, bInitialized(false)
+{
 }
 
 bool HandlerComponent::IsActive() const
@@ -401,9 +436,13 @@ bool HandlerComponent::IsInitialized() const
 	return bInitialized;
 }
 
-// PROCESSED PACKET
+
+/**
+ * ProcessedPacket
+ */
+
 ProcessedPacket::ProcessedPacket(uint8* InData, int32 InCount)
 {
 	Data = InData;
 	Count = InCount;
-}
\ No newline at end of file
+}
diff --git a/Engine/Source/Runtime/PacketHandlers/PacketHandler/Public/PacketHandler.h b/Engine/Source/Runtime/PacketHandlers/PacketHandler/Public/PacketHandler.h
--- a/Engine/Source/Runtime/PacketHandlers/PacketHandler/Public/PacketHandler.h
+++ b/Engine/Source/Runtime/PacketHandlers/PacketHandler/Public/PacketHandler.h
@@ -98,9 +98,16 @@ public:
 	/* Queue's a packet to be sent when the handler is ticked */
 	void QueuePacketForSending(BufferedPacket* PacketToQueue);
 
-	/* Gets a packet from the the buffered packet queue for sending */
+	/* Gets a packet from the buffered packet queue for sending */
 	BufferedPacket* GetQueuedPacket();
 
+	/**
+	 * Gets the combined packet/protocol overhead from all handlers, for reserving space in the parent connections packets
+	 *
+	 * @return	The combined packet/protocol overhead
+	 */
+	int32 GetTotalPacketOverheadBits();
+
 
 	/** Whether or not the PacketHandler system is enabled */
 	bool bEnabled;
@@ -182,12 +189,19 @@ public:
 	/* Tick functionality should be placed here */
 	virtual void Tick(float DeltaTime){};
 
-	/* The manager of the handler, set in initialization */
-	PacketHandler* Handler; 
-
 	/* Whether this handler is active */
 	virtual void SetActive(bool Active);
 
+	/**
+	 * Returns the amount of packet/protocol overhead expected from this handler.
+	 *
+	 * IMPORTANT: This MUST be accurate, and should represent the worst-case overhead expected from the handler.
+	 *				If this is inaccurate, packets will randomly fail to send, in rare cases which are extremely hard to trace.
+	 *
+	 * @return	The worst-case packet overhead for the handler
+	 */
+	virtual int32 GetPacketOverheadBits() PURE_VIRTUAL(Handler::Component::GetPacketOverheadBits, return -1;);
+
 protected:
 
 	/* Sets the state of the handler */
@@ -196,6 +210,12 @@ protected:
 	/* Should be called when the handler is fully initialized on both remote and local */
 	void Initialized();
 
+
+public:
+	/* The manager of the handler, set in initialization */
+	PacketHandler* Handler; 
+
+protected:
 	/* The state of this handler */
 	Handler::Component::State State;
 
