CIS fixes

diff --git a/Engine/Source/Runtime/Online/HTTP/Private/HTML5/HTML5HTTP.h b/Engine/Source/Runtime/Online/HTTP/Private/HTML5/HTML5HTTP.h
--- a/Engine/Source/Runtime/Online/HTTP/Private/HTML5/HTML5HTTP.h
+++ b/Engine/Source/Runtime/Online/HTTP/Private/HTML5/HTML5HTTP.h
@@ -53,6 +53,7 @@ public:
 	virtual void SetURL(const FString& InURL) override;
 	virtual void SetContent(const TArray<uint8>& ContentPayload) override;
 	virtual void SetContentAsString(const FString& ContentString) override;
+	virtual void SetHeader(const FString& HeaderName, const FString& HeaderValue) override;
 	virtual void AppendToHeader(const FString& HeaderName, const FString& AdditionalHeaderValue) override;
 	virtual bool ProcessRequest() override;
 	virtual FHttpRequestCompleteDelegate& OnProcessRequestComplete() override;
