Copying //UE4/Ocean-Staging to Dev-Main (//UE4/Dev-Main)

========================
MAJOR FEATURES + CHANGES
========================

Change 2889737 on 2016/03/02 by [REDACTED_PERSON]

	#ocn, #mcp, #ogs: Fixed several TEAM REST API calls - they required a json body in the POST; was getting 415 response, and the docs say the json body is required.  The calls are successful after these changes.

Change 2902154 on 2016/03/09 by [REDACTED_PERSON]

	#ocn, #mcp, #ogs: Added implementation for retrieving list of a user's pending group application list (list of groups to which the user has an outstanding membership application).

Change 2902973 on 2016/03/10 by [REDACTED_PERSON]

	#ocn, #mcp: Added "incoming applications" and renamed previous "user pending applications" to "outgoing applications".

Change 2903263 on 2016/03/10 by [REDACTED_PERSON]

	#ocn, #mcp: Added incoming/outgoing invitation queries.

Change 2907476 on 2016/03/13 by Jurre.deBaare

	Relocating SSF files

Change 2911725 on 2016/03/16 by [REDACTED_PERSON]

	#ocn, #mcp, #ogs: Fixed bad invitation queries.

diff --git a/Engine/Source/Runtime/Online/OnlineSubsystem/Public/Interfaces/OnlineGroupsInterface.h b/Engine/Source/Runtime/Online/OnlineSubsystem/Public/Interfaces/OnlineGroupsInterface.h
--- a/Engine/Source/Runtime/Online/OnlineSubsystem/Public/Interfaces/OnlineGroupsInterface.h
+++ b/Engine/Source/Runtime/Online/OnlineSubsystem/Public/Interfaces/OnlineGroupsInterface.h
@@ -128,6 +128,68 @@ struct FUserMembershipEntry
 	bool bIsAdmin;
 };
 
+/**
+ * An entry in the list of user pending membership applications.
+ */
+struct FApplicationEntry
+{
+	/** AccountId of the user who applied for group membership. */
+	TSharedPtr<const FUniqueNetId> Id;
+
+	/** Timestamp of when the application was sent. */
+	FDateTime SentAt;
+
+	/** Message associated with the membership application. */
+	FText Message;
+
+	/** JobId of the application. */
+	TSharedPtr<const FUniqueNetId> JobId;
+
+	/** Job summary. */
+	FString JobSummary;
+
+	/** GroupId of the group for which membership was applied. */
+	TSharedPtr<const FUniqueNetId> GroupId;
+
+	/** Namespace in which the application exists. */
+	FString Namespace;
+
+	/** Name of the group. */
+	FText Name;
+};
+
+/**
+* An entry in the list of user pending membership invitations.
+*/
+struct FInvitationEntry
+{
+	/** AccountId of the user who applied for group membership. */
+	TSharedPtr<const FUniqueNetId> Id;
+
+	/** The name of the account that is associated with the other end of the transaction. */
+	FText Host;
+
+	/** Timestamp of when the application was sent. */
+	FDateTime SentAt;
+
+	/** Message associated with the membership application. */
+	FText Message;
+
+	/** JobId of the application. */
+	TSharedPtr<const FUniqueNetId> JobId;
+
+	/** Job summary. */
+	FString JobSummary;
+
+	/** GroupId of the group for which membership was applied. */
+	TSharedPtr<const FUniqueNetId> GroupId;
+
+	/** Namespace in which the application exists. */
+	FString Namespace;
+
+	/** Name of the group. */
+	FText Name;
+};
 
 template <typename EntryType>
 struct IGroupUserCollection
@@ -163,6 +225,16 @@ typedef IGroupUserCollection<FGroupBlacklistEntry> IGroupBlacklist;
  */
 typedef IGroupUserCollection<FUserMembershipEntry> IUserMembership;
 
+/**
+ * A list of pending membership applications for a given user.
+ */
+typedef IGroupUserCollection<FApplicationEntry> IApplications;
+
+/**
+* A list of pending membership invitations for a given user.
+*/
+typedef IGroupUserCollection<FInvitationEntry> IInvitations;
+
 /**
  * Group search options
  */
@@ -398,6 +470,50 @@ public: // callable by all users
 	 */
 	virtual TSharedPtr<const IUserMembership> GetCachedUserMembership(const FUniqueNetId& ContextUserId, const FUniqueNetId& UserId) = 0;
 
+	/**
+	* Queries the server for a list of groups to which the user has applied for membership.
+	* If the callback reports success, use GetCachedApplications to retrieve details.
+	*
+	* @param ContextUserId The ID of the user whose credentials are being used to make this call
+	* @param UserId The user to query for pending membership application information.
+	* @param OnCompleted This callback is invoked after contacting the server. It is guaranteed to occur
+	*        (regardless of success/fail) and will not be called before this function returns.
+	*/
+	virtual void QueryOutgoingApplications(const FUniqueNetId& ContextUserId, const FUniqueNetId& UserId, const FOnGroupsRequestCompleted& OnCompleted) = 0;
+
+	/**
+	* Get cached pending application information (if it exists). This retrieves a list of membership
+	* applications that the user has created (outgoing) or to which the user can respond as admin (incoming).
+	*
+	* @param ContextUserId The ID of the user whose credentials are being used to make this call
+	* @param UserId The ID of the user to query.
+	* @param OnCompleted This callback is invoked after contacting the server. It is guaranteed to occur
+	*        (regardless of success/fail) and will not be called before this function returns.
+	*/
+	virtual TSharedPtr<const IApplications> GetCachedApplications(const FUniqueNetId& ContextUserId, const FUniqueNetId& UserId) = 0;
+
+	/**
+	* Queries the server for a list of groups to which the user has been invited.
+	* If the callback reports success, use GetCachedApplications to retrieve details.
+	*
+	* @param ContextUserId The ID of the user whose credentials are being used to make this call
+	* @param UserId The user to query for pending membership invitation information.
+	* @param OnCompleted This callback is invoked after contacting the server. It is guaranteed to occur
+	*        (regardless of success/fail) and will not be called before this function returns.
+	*/
+	virtual void QueryIncomingInvitations(const FUniqueNetId& ContextUserId, const FUniqueNetId& UserId, const FOnGroupsRequestCompleted& OnCompleted) = 0;
+
+	/**
+	* Get cached pending invitation information (if it exists). This retrieves a list of membership
+	* invitations that the user's admin clans have created (outgoing) or to which the user can respond (incoming).
+	*
+	* @param ContextUserId The ID of the user whose credentials are being used to make this call
+	* @param UserId The ID of the user to query.
+	* @param OnCompleted This callback is invoked after contacting the server. It is guaranteed to occur
+	*        (regardless of success/fail) and will not be called before this function returns.
+	*/
+	virtual TSharedPtr<const IInvitations> GetCachedInvitations(const FUniqueNetId& ContextUserId, const FUniqueNetId& UserId) = 0;
+
 public: // can be called by group admins
 
 	/**
@@ -596,6 +712,28 @@ public: // can be called by group admins
 	 */
 	virtual TSharedPtr<const IGroupBlacklist> GetCachedGroupBlacklist(const FUniqueNetId& ContextUserId, const FUniqueNetId& GroupId) = 0;
 
+	/**
+	* Queries the server for a list of membership applications that UserId can process (accept, reject, or block) as group admin.
+	* If the callback reports success, use GetCachedApplications to retrieve details.
+	*
+	* @param ContextUserId The ID of the user whose credentials are being used to make this call
+	* @param UserId The user to query for pending group membership application information.
+	* @param OnCompleted This callback is invoked after contacting the server. It is guaranteed to occur
+	*        (regardless of success/fail) and will not be called before this function returns.
+	*/
+	virtual void QueryIncomingApplications(const FUniqueNetId& ContextUserId, const FUniqueNetId& UserId, const FOnGroupsRequestCompleted& OnCompleted) = 0;
+
+	/**
+	* Queries the server for a list of invitations for other users to join groups to which UserId is an admin.
+	* If the callback reports success, use GetCachedInvitations to retrieve details.
+	*
+	* @param ContextUserId The ID of the user whose credentials are being used to make this call
+	* @param UserId The user to query for pending membership invitation information.
+	* @param OnCompleted This callback is invoked after contacting the server. It is guaranteed to occur
+	*        (regardless of success/fail) and will not be called before this function returns.
+	*/
+	virtual void QueryOutgoingInvitations(const FUniqueNetId& ContextUserId, const FUniqueNetId& UserId, const FOnGroupsRequestCompleted& OnCompleted) = 0;
+
 public: // can be called by group owner only
 
 	/**
