fix for tvOS build warning

diff --git a/Engine/Source/Runtime/Online/HTTP/Private/Apple/AppleHTTP.cpp b/Engine/Source/Runtime/Online/HTTP/Private/Apple/AppleHTTP.cpp
--- a/Engine/Source/Runtime/Online/HTTP/Private/Apple/AppleHTTP.cpp
+++ b/Engine/Source/Runtime/Online/HTTP/Private/Apple/AppleHTTP.cpp
@@ -267,9 +267,12 @@ bool FAppleHttpRequest::StartRequest()
 
 	Response = MakeShareable( new FAppleHttpResponse( *this ) );
 
+#pragma clang diagnostic push
+#pragma clang diagnostic ignored "-Wdeprecated-declarations"
 	// Create the connection, tell it to run in the main run loop, and kick it off.
 	Connection = [[NSURLConnection alloc] initWithRequest:Request delegate:Response->ResponseWrapper startImmediately:NO];
-	if( Connection != nullptr && Response->ResponseWrapper != nullptr )
+#pragma clang diagnostic pop
+	if (Connection != nullptr && Response->ResponseWrapper != nullptr)
 	{
 		CompletionStatus = EHttpRequestStatus::Processing;
 		[Connection scheduleInRunLoop:[NSRunLoop mainRunLoop] forMode:NSDefaultRunLoopMode];
