Removed error messages on curl hitching

diff --git a/Engine/Source/Runtime/Online/HTTP/Private/Curl/CurlHttpManager.cpp b/Engine/Source/Runtime/Online/HTTP/Private/Curl/CurlHttpManager.cpp
--- a/Engine/Source/Runtime/Online/HTTP/Private/Curl/CurlHttpManager.cpp
+++ b/Engine/Source/Runtime/Online/HTTP/Private/Curl/CurlHttpManager.cpp
@@ -315,12 +315,7 @@ void FCurlHttpManager::HttpThreadTick(float DeltaSeconds)
 	QUICK_SCOPE_CYCLE_COUNTER(STAT_FCurlHttpManager_Tick);
 	check(MultiHandle);
 
-	double CurlTick = 0.0;
-	const int NumRequestsToTick = Requests.Num();
-
 	{
-		FSimpleScopeSecondsCounter CurlTickTimer(CurlTick);
-
 		if (RunningThreadedRequests.Num() > 0)
 		{
 			int RunningRequests = -1;
@@ -393,16 +388,7 @@ void FCurlHttpManager::HttpThreadTick(float DeltaSeconds)
 		}
 	}
 
-	double ParentTick = 0.0;
-	{
-		FSimpleScopeSecondsCounter TickTimer(ParentTick);
-		FHttpManager::HttpThreadTick(DeltaSeconds);
-	}
-
-	if (ParentTick + CurlTick > 0.02)
-	{
-		UE_LOG(LogHttp, Warning, TEXT("HITCHHUNTER: Hitch in CurlHttp (CurlTick: %.1f ms, HttpManagerTick: %.1f) has been detected this frame, NumRequestsToTick = %d"), CurlTick * 1000.0, ParentTick * 1000.0, NumRequestsToTick);
-	}
+	FHttpManager::HttpThreadTick(DeltaSeconds);
 }
 
 bool FCurlHttpManager::StartThreadedRequest(IHttpThreadedRequest* Request)
