Copying //UE4/Orion-Staging to //UE4/Main (Source: //Orion/Dev-General @ 2997507)

diff --git a/Engine/Source/Runtime/GameplayAbilities/Private/AbilitySystemComponent.cpp b/Engine/Source/Runtime/GameplayAbilities/Private/AbilitySystemComponent.cpp
--- a/Engine/Source/Runtime/GameplayAbilities/Private/AbilitySystemComponent.cpp
+++ b/Engine/Source/Runtime/GameplayAbilities/Private/AbilitySystemComponent.cpp
@@ -632,7 +632,9 @@ FActiveGameplayEffectHandle UAbilitySystemComponent::ApplyGameplayEffectSpecToSe
 	bool bTreatAsInfiniteDuration = GetOwnerRole() != ROLE_Authority && PredictionKey.IsLocalClientKey() && Spec.Def->DurationPolicy == EGameplayEffectDurationType::Instant;
 
 	// Make sure we create our copy of the spec in the right place
-	FActiveGameplayEffectHandle	MyHandle;
+	// We initialize the FActiveGameplayEffectHandle here with INDEX_NONE to handle the case of instant GE
+	// Initializing it like this will set the bPassedFiltersAndWasExecuted on the FActiveGameplayEffectHandle to true so we can know that we applied a GE
+	FActiveGameplayEffectHandle	MyHandle(INDEX_NONE);
 	bool bInvokeGameplayCueApplied = Spec.Def->DurationPolicy != EGameplayEffectDurationType::Instant; // Cache this now before possibly modifying predictive instant effect to infinite duration effect.
 
 	FActiveGameplayEffect* AppliedEffect = nullptr;
@@ -1024,18 +1026,20 @@ void UAbilitySystemComponent::ExecuteGameplayCue(const FGameplayTag GameplayCueT
 
 void UAbilitySystemComponent::AddGameplayCue_Internal(const FGameplayTag GameplayCueTag, const FGameplayEffectContextHandle& EffectContext, FActiveGameplayCueContainer& GameplayCueContainer)
 {
+	FGameplayCueParameters Parameters(EffectContext);
+
 	if (IsOwnerActorAuthoritative())
 	{
 		bool bWasInList = HasMatchingGameplayTag(GameplayCueTag);
 
 		ForceReplication();
-		GameplayCueContainer.AddCue(GameplayCueTag, ScopedPredictionKey);
-		NetMulticast_InvokeGameplayCueAdded(GameplayCueTag, ScopedPredictionKey, EffectContext);
+		GameplayCueContainer.AddCue(GameplayCueTag, ScopedPredictionKey, Parameters);
+		NetMulticast_InvokeGameplayCueAdded_WithParams(GameplayCueTag, ScopedPredictionKey, Parameters);
 
 		if (!bWasInList)
 		{
 			// Call on server here, clients get it from repnotify
-			InvokeGameplayCueEvent(GameplayCueTag, EGameplayCueEvent::WhileActive);
+			InvokeGameplayCueEvent(GameplayCueTag, EGameplayCueEvent::WhileActive, Parameters);
 		}
 	}
 	else if (ScopedPredictionKey.IsLocalClientKey())
@@ -1043,8 +1047,8 @@ void UAbilitySystemComponent::AddGameplayCue_Internal(const FGameplayTag Gamepla
 		GameplayCueContainer.PredictiveAdd(GameplayCueTag, ScopedPredictionKey);
 
 		// Allow for predictive gameplaycue events? Needs more thought
-		InvokeGameplayCueEvent(GameplayCueTag, EGameplayCueEvent::OnActive, EffectContext);
-		InvokeGameplayCueEvent(GameplayCueTag, EGameplayCueEvent::WhileActive);
+		InvokeGameplayCueEvent(GameplayCueTag, EGameplayCueEvent::OnActive, Parameters);
+		InvokeGameplayCueEvent(GameplayCueTag, EGameplayCueEvent::WhileActive, Parameters);
 	}
 }
 
@@ -1137,6 +1141,14 @@ void UAbilitySystemComponent::NetMulticast_InvokeGameplayCueAdded_Implementation
 	}
 }
 
+void UAbilitySystemComponent::NetMulticast_InvokeGameplayCueAdded_WithParams_Implementation(const FGameplayTag GameplayCueTag, FPredictionKey PredictionKey, FGameplayCueParameters Parameters)
+{
+	if (IsOwnerActorAuthoritative() || PredictionKey.IsLocalClientKey() == false)
+	{
+		InvokeGameplayCueEvent(GameplayCueTag, EGameplayCueEvent::OnActive, Parameters);
+	}
+}
+
 // -----------
 
 void UAbilitySystemComponent::NetMulticast_InvokeGameplayCueAddedAndWhileActive_FromSpec_Implementation(const FGameplayEffectSpecForRPC& Spec, FPredictionKey PredictionKey)
@@ -1214,54 +1226,62 @@ TArray<FActiveGameplayEffectHandle> UAbilitySystemComponent::GetActiveEffects(co
 	return ActiveGameplayEffects.GetActiveEffects(Query);
 }
 
-void UAbilitySystemComponent::RemoveActiveEffectsWithTags(const FGameplayTagContainer Tags)
+int32 UAbilitySystemComponent::RemoveActiveEffectsWithTags(const FGameplayTagContainer Tags)
 {
 	if (IsOwnerActorAuthoritative())
 	{
-		RemoveActiveEffects(FGameplayEffectQuery::MakeQuery_MatchAnyEffectTags(Tags));
+		return RemoveActiveEffects(FGameplayEffectQuery::MakeQuery_MatchAnyEffectTags(Tags));
 	}
+	return 0;
 }
 
-void UAbilitySystemComponent::RemoveActiveEffectsWithSourceTags(FGameplayTagContainer Tags)
+int32 UAbilitySystemComponent::RemoveActiveEffectsWithSourceTags(FGameplayTagContainer Tags)
 {
 	if (IsOwnerActorAuthoritative())
 	{
-		RemoveActiveEffects(FGameplayEffectQuery::MakeQuery_MatchAnySourceTags(Tags));
+		return RemoveActiveEffects(FGameplayEffectQuery::MakeQuery_MatchAnySourceTags(Tags));
 	}
+	return 0;
 }
 
-void UAbilitySystemComponent::RemoveActiveEffectsWithAppliedTags(FGameplayTagContainer Tags)
+int32 UAbilitySystemComponent::RemoveActiveEffectsWithAppliedTags(FGameplayTagContainer Tags)
 {
 	if (IsOwnerActorAuthoritative())
 	{
-		RemoveActiveEffects(FGameplayEffectQuery::MakeQuery_MatchAnyOwningTags(Tags));
+		return RemoveActiveEffects(FGameplayEffectQuery::MakeQuery_MatchAnyOwningTags(Tags));
 	}
+	return 0;
 }
 
-void UAbilitySystemComponent::RemoveActiveEffectsWithGrantedTags(const FGameplayTagContainer Tags)
+int32 UAbilitySystemComponent::RemoveActiveEffectsWithGrantedTags(const FGameplayTagContainer Tags)
 {
 	if (IsOwnerActorAuthoritative())
 	{
-		RemoveActiveEffects(FGameplayEffectQuery::MakeQuery_MatchAnyOwningTags(Tags));
+		return RemoveActiveEffects(FGameplayEffectQuery::MakeQuery_MatchAnyOwningTags(Tags));
 	}
+	return 0;
 }
 
 // #deprecated, use FGameplayEffectQuery version
-void UAbilitySystemComponent::RemoveActiveEffects(const FActiveGameplayEffectQuery Query, int32 StacksToRemove)
+int32 UAbilitySystemComponent::RemoveActiveEffects(const FActiveGameplayEffectQuery Query, int32 StacksToRemove)
 {
 	if (IsOwnerActorAuthoritative())
 	{
 		PRAGMA_DISABLE_DEPRECATION_WARNINGS
-		ActiveGameplayEffects.RemoveActiveEffects(Query, StacksToRemove);
+		return ActiveGameplayEffects.RemoveActiveEffects(Query, StacksToRemove);
 		PRAGMA_ENABLE_DEPRECATION_WARNINGS
 	}
+
+	return 0;
 }
-void UAbilitySystemComponent::RemoveActiveEffects(const FGameplayEffectQuery& Query, int32 StacksToRemove)
+int32 UAbilitySystemComponent::RemoveActiveEffects(const FGameplayEffectQuery& Query, int32 StacksToRemove)
 {
 	if (IsOwnerActorAuthoritative())
 	{
-		ActiveGameplayEffects.RemoveActiveEffects(Query, StacksToRemove);
+		return ActiveGameplayEffects.RemoveActiveEffects(Query, StacksToRemove);
 	}
+
+	return 0;
 }
 
 // ---------------------------------------------------------------------------------------
@@ -1364,6 +1384,18 @@ void UAbilitySystemComponent::SetReplicationMode(EReplicationMode NewReplication
 	ReplicationMode = NewReplicationMode;
 }
 
+void UAbilitySystemComponent::OnPredictiveGameplayCueCatchup(FGameplayTag Tag)
+{
+	// Remove it
+	RemoveOneTagCount_NoReturn(Tag);
+
+	if (HasMatchingGameplayTag(Tag) == 0)
+	{
+		// Invoke Removed event if we no longer have this tag (probably a mispredict)
+		InvokeGameplayCueEvent(Tag, EGameplayCueEvent::Removed);
+	}
+}
+
 // ---------------------------------------------------------------------------------------
 
 void UAbilitySystemComponent::PrintAllGameplayEffects() const
diff --git a/Engine/Source/Runtime/GameplayAbilities/Private/GameplayCueInterface.cpp b/Engine/Source/Runtime/GameplayAbilities/Private/GameplayCueInterface.cpp
--- a/Engine/Source/Runtime/GameplayAbilities/Private/GameplayCueInterface.cpp
+++ b/Engine/Source/Runtime/GameplayAbilities/Private/GameplayCueInterface.cpp
@@ -160,7 +160,7 @@ void FActiveGameplayCue::PreReplicatedRemove(const struct FActiveGameplayCueCont
 	{
 		// If predicted ignore the add/remove
 		InArray.Owner->UpdateTagMap(GameplayCueTag, -1);
-		InArray.Owner->InvokeGameplayCueEvent(GameplayCueTag, EGameplayCueEvent::Removed);
+		InArray.Owner->InvokeGameplayCueEvent(GameplayCueTag, EGameplayCueEvent::Removed, Parameters);
 	}
 }
 
@@ -171,21 +171,21 @@ void FActiveGameplayCue::PostReplicatedAdd(const struct FActiveGameplayCueContai
 	if (PredictionKey.IsLocalClientKey() == false)
 	{
 		// If predicted ignore the add/remove
-		InArray.Owner->InvokeGameplayCueEvent(GameplayCueTag, EGameplayCueEvent::WhileActive);
+		InArray.Owner->InvokeGameplayCueEvent(GameplayCueTag, EGameplayCueEvent::WhileActive, Parameters);
 	}
 }
 
-void FActiveGameplayCueContainer::AddCue(const FGameplayTag& Tag, const FPredictionKey& PredictionKey)
+void FActiveGameplayCueContainer::AddCue(const FGameplayTag& Tag, const FPredictionKey& PredictionKey, const FGameplayCueParameters& Parameters)
 {
 	UWorld* World = Owner->GetWorld();
 
 	// Store the prediction key so the client can investigate it
-	FActiveGameplayCue	NewCue;
+	FActiveGameplayCue&	NewCue = GameplayCues[GameplayCues.AddDefaulted()];
 	NewCue.GameplayCueTag = Tag;
 	NewCue.PredictionKey = PredictionKey;
+	NewCue.Parameters = Parameters;
 	MarkItemDirty(NewCue);
-
-	GameplayCues.Add(NewCue);
+	
 	Owner->UpdateTagMap(Tag, 1);
 }
 
@@ -216,7 +216,7 @@ void FActiveGameplayCueContainer::PredictiveRemove(const FGameplayTag& Tag)
 			// DONT remove from the replicated array.
 			Cue.bPredictivelyRemoved = true;
 			Owner->UpdateTagMap(Tag, -1);
-			Owner->InvokeGameplayCueEvent(Tag, EGameplayCueEvent::Removed);	
+			Owner->InvokeGameplayCueEvent(Tag, EGameplayCueEvent::Removed, Cue.Parameters);	
 			return;
 		}
 	}
@@ -225,7 +225,7 @@ void FActiveGameplayCueContainer::PredictiveRemove(const FGameplayTag& Tag)
 void FActiveGameplayCueContainer::PredictiveAdd(const FGameplayTag& Tag, FPredictionKey& PredictionKey)
 {
 	Owner->UpdateTagMap(Tag, 1);	
-	PredictionKey.NewRejectOrCaughtUpDelegate(FPredictionKeyEvent::CreateUObject(Owner, &UAbilitySystemComponent::RemoveOneTagCount_NoReturn, Tag));
+	PredictionKey.NewRejectOrCaughtUpDelegate(FPredictionKeyEvent::CreateUObject(Owner, &UAbilitySystemComponent::OnPredictiveGameplayCueCatchup, Tag));
 }
 
 bool FActiveGameplayCueContainer::HasCue(const FGameplayTag& Tag) const
diff --git a/Engine/Source/Runtime/GameplayAbilities/Private/GameplayCueManager.cpp b/Engine/Source/Runtime/GameplayAbilities/Private/GameplayCueManager.cpp
--- a/Engine/Source/Runtime/GameplayAbilities/Private/GameplayCueManager.cpp
+++ b/Engine/Source/Runtime/GameplayAbilities/Private/GameplayCueManager.cpp
@@ -438,19 +438,41 @@ void UGameplayCueManager::InitObjectLibraries(TArray<FString> Paths, UObjectLibr
 	StaticObjectLibrary->GetAssetDataList(StaticAssetDatas);
 
 	TArray<FGameplayCueReferencePair> CuesToAdd;
-	BuildCuesToAddToGlobalSet(ActorAssetDatas, GET_MEMBER_NAME_CHECKED(AGameplayCueNotify_Actor, GameplayCueName), bAsyncLoadAtStartup, CuesToAdd, OnLoadDelegate, ShouldLoadDelegate);
-	BuildCuesToAddToGlobalSet(StaticAssetDatas, GET_MEMBER_NAME_CHECKED(UGameplayCueNotify_Static, GameplayCueName), bAsyncLoadAtStartup, CuesToAdd, OnLoadDelegate, ShouldLoadDelegate);
+	TArray<FStringAssetReference> AssetsToLoad;
+
+	// Build Cue lists for loading. Determines what from the obj library needs to be loaded
+	BuildCuesToAddToGlobalSet(ActorAssetDatas, GET_MEMBER_NAME_CHECKED(AGameplayCueNotify_Actor, GameplayCueName), CuesToAdd, AssetsToLoad, ShouldLoadDelegate);
+	BuildCuesToAddToGlobalSet(StaticAssetDatas, GET_MEMBER_NAME_CHECKED(UGameplayCueNotify_Static, GameplayCueName), CuesToAdd, AssetsToLoad, ShouldLoadDelegate);
 
+	// Add these cues to the global set
 	check(GlobalCueSet);
 	GlobalCueSet->AddCues(CuesToAdd);
+
+	// Start loading them if necessary
+	if (bAsyncLoadAtStartup)
+	{
+		auto ForwardLambda = [](TArray<FStringAssetReference> AssetList, FOnGameplayCueNotifySetLoaded OnLoadedDelegate)
+		{
+			OnLoadedDelegate.ExecuteIfBound(AssetList);
+		};
+
+		if (AssetsToLoad.Num() > 0)
+		{			
+			StreamableManager.RequestAsyncLoad(AssetsToLoad, FStreamableDelegate::CreateStatic( ForwardLambda, AssetsToLoad, OnLoadDelegate));
+		}
+		else
+		{
+			// Still fire the delegate even if nothing was found to load
+			OnLoadDelegate.ExecuteIfBound(AssetsToLoad);
+		}
+	}
 }
 
-void UGameplayCueManager::BuildCuesToAddToGlobalSet(const TArray<FAssetData>& AssetDataList, FName TagPropertyName, bool bAsyncLoadAfterAdd, TArray<FGameplayCueReferencePair>& OutCuesToAdd, FOnGameplayCueNotifySetLoaded OnLoaded, FShouldLoadGCNotifyDelegate ShouldLoad)
+void UGameplayCueManager::BuildCuesToAddToGlobalSet(const TArray<FAssetData>& AssetDataList, FName TagPropertyName, TArray<FGameplayCueReferencePair>& OutCuesToAdd, TArray<FStringAssetReference>& OutAssetsToLoad, FShouldLoadGCNotifyDelegate ShouldLoad)
 {
 	IGameplayTagsModule& GameplayTagsModule = IGameplayTagsModule::Get();
 
-	TArray<FStringAssetReference> AssetsToLoad;
-	AssetsToLoad.Reserve(AssetDataList.Num());
+	OutAssetsToLoad.Reserve(OutAssetsToLoad.Num() + AssetDataList.Num());
 
 	for (FAssetData Data: AssetDataList)
 	{
@@ -481,7 +503,7 @@ void UGameplayCueManager::BuildCuesToAddToGlobalSet(const TArray<FAssetData>& As
 
 				OutCuesToAdd.Add(FGameplayCueReferencePair(GameplayCueTag, StringRef));
 
-				AssetsToLoad.Add(StringRef);
+				OutAssetsToLoad.Add(StringRef);
 			}
 			else
 			{
@@ -489,24 +511,6 @@ void UGameplayCueManager::BuildCuesToAddToGlobalSet(const TArray<FAssetData>& As
 			}
 		}
 	}
-
-	if (bAsyncLoadAfterAdd)
-	{
-		auto ForwardLambda = [](TArray<FStringAssetReference> AssetList, FOnGameplayCueNotifySetLoaded OnLoadedDelegate)
-		{
-			OnLoadedDelegate.ExecuteIfBound(AssetList);
-		};
-
-		if (AssetsToLoad.Num() > 0)
-		{			
-			StreamableManager.RequestAsyncLoad(AssetsToLoad, FStreamableDelegate::CreateStatic( ForwardLambda, AssetsToLoad, OnLoaded));
-		}
-		else
-		{
-			// Still fire the delegate even if nothing was found to load
-			OnLoaded.ExecuteIfBound(AssetsToLoad);
-		}
-	}
 }
 
 int32 GameplayCueCheckForTooManyRPCs = 1;
diff --git a/Engine/Source/Runtime/GameplayAbilities/Private/GameplayEffect.cpp b/Engine/Source/Runtime/GameplayAbilities/Private/GameplayEffect.cpp
--- a/Engine/Source/Runtime/GameplayAbilities/Private/GameplayEffect.cpp
+++ b/Engine/Source/Runtime/GameplayAbilities/Private/GameplayEffect.cpp
@@ -3111,10 +3111,11 @@ bool FActiveGameplayEffectsContainer::NetDeltaSerialize(FNetDeltaSerializeInfo&
 	{
 		QUICK_SCOPE_CYCLE_COUNTER(STAT_ActiveGameplayEffectsContainer_NetDeltaSerialize_CheckRepGameplayCues);
 
-		for (const FActiveGameplayEffect& Effect : this)
+		if (!DeltaParms.bOutHasMoreUnmapped) // Do not invoke GCs when we have missing information (like AActor*s in EffectContext)
 		{
-			if (!DeltaParms.bOutHasMoreUnmapped) // Do not invoke GCs when we have missing information (like AActor*s in EffectContext)
+			for (const FActiveGameplayEffect& Effect : this)
 			{
+			
 				if (Effect.bIsInhibited == false)
 				{
 					if (Effect.bPendingRepOnActiveGC)
@@ -3527,10 +3528,11 @@ void FActiveGameplayEffectsContainer::ModifyActiveEffectStartTime(FActiveGamepla
 }
 
 // #deprecated, use FGameplayEffectQuery version
-void FActiveGameplayEffectsContainer::RemoveActiveEffects(const FActiveGameplayEffectQuery Query, int32 StacksToRemove)
+int32 FActiveGameplayEffectsContainer::RemoveActiveEffects(const FActiveGameplayEffectQuery Query, int32 StacksToRemove)
 {
 	// Force a lock because the removals could cause other removals earlier in the array, so iterating backwards is not safe all by itself
 	GAMEPLAYEFFECT_SCOPE_LOCK();
+	int32 NumRemoved = 0;
 
 	// Manually iterating through in reverse because this is a removal operation
 	for (int32 idx=GetNumGameplayEffects()-1; idx >= 0; --idx)
@@ -3538,14 +3540,18 @@ void FActiveGameplayEffectsContainer::RemoveActiveEffects(const FActiveGameplayE
 		const FActiveGameplayEffect& Effect = *GetActiveGameplayEffect(idx);
 		if (Effect.IsPendingRemove==false && Query.Matches(Effect))
 		{
-			InternalRemoveActiveGameplayEffect(idx, StacksToRemove, true);			
+			InternalRemoveActiveGameplayEffect(idx, StacksToRemove, true);		
+			++NumRemoved;
 		}
 	}
+	return NumRemoved;
 }
-void FActiveGameplayEffectsContainer::RemoveActiveEffects(const FGameplayEffectQuery& Query, int32 StacksToRemove)
+
+int32 FActiveGameplayEffectsContainer::RemoveActiveEffects(const FGameplayEffectQuery& Query, int32 StacksToRemove)
 {
 	// Force a lock because the removals could cause other removals earlier in the array, so iterating backwards is not safe all by itself
 	GAMEPLAYEFFECT_SCOPE_LOCK();
+	int32 NumRemoved = 0;
 
 	// Manually iterating through in reverse because this is a removal operation
 	for (int32 idx = GetNumGameplayEffects() - 1; idx >= 0; --idx)
@@ -3554,11 +3560,12 @@ void FActiveGameplayEffectsContainer::RemoveActiveEffects(const FGameplayEffectQ
 		if (Effect.IsPendingRemove == false && Query.Matches(Effect))
 		{
 			InternalRemoveActiveGameplayEffect(idx, StacksToRemove, true);
+			++NumRemoved;
 		}
 	}
+	return NumRemoved;
 }
 
-
 int32 FActiveGameplayEffectsContainer::GetActiveEffectCount(const FActiveGameplayEffectQuery Query, bool bEnforceOnGoingCheck) const
 {
 	int32 Count = 0;
diff --git a/Engine/Source/Runtime/GameplayAbilities/Private/GameplayEffectAggregator.cpp b/Engine/Source/Runtime/GameplayAbilities/Private/GameplayEffectAggregator.cpp
--- a/Engine/Source/Runtime/GameplayAbilities/Private/GameplayEffectAggregator.cpp
+++ b/Engine/Source/Runtime/GameplayAbilities/Private/GameplayEffectAggregator.cpp
@@ -36,6 +36,12 @@ bool FAggregatorMod::Qualifies(const FAggregatorEvaluateParameters& Parameters)
 	return bSourceMet && bTargetMet && bSourceFilterMet && bTargetFilterMet;
 }
 
+FAggregator::~FAggregator()
+{
+	int32 NumRemoved = FScopedAggregatorOnDirtyBatch::DirtyAggregators.Remove(this);
+	ensure(NumRemoved == 0);
+}
+
 float FAggregator::Evaluate(const FAggregatorEvaluateParameters& Parameters) const
 {
 	return EvaluateWithBase(BaseValue, Parameters);
diff --git a/Engine/Source/Runtime/GameplayAbilities/Private/GameplayEffectTypes.cpp b/Engine/Source/Runtime/GameplayAbilities/Private/GameplayEffectTypes.cpp
--- a/Engine/Source/Runtime/GameplayAbilities/Private/GameplayEffectTypes.cpp
+++ b/Engine/Source/Runtime/GameplayAbilities/Private/GameplayEffectTypes.cpp
@@ -567,6 +567,7 @@ bool FGameplayCueParameters::NetSerialize(FArchive& Ar, class UPackageMap* Map,
 		REP_Instigator,
 		REP_EffectCauser,
 		REP_SourceObject,
+		REP_TargetAttachComponent,
 		REP_PhysMaterial,
 		REP_GELevel,
 		REP_AbilityLevel,
@@ -609,6 +610,10 @@ bool FGameplayCueParameters::NetSerialize(FArchive& Ar, class UPackageMap* Map,
 		{
 			RepBits |= (1 << REP_SourceObject);
 		}
+		if (TargetAttachComponent.IsValid())
+		{
+			RepBits |= (1 << REP_TargetAttachComponent);
+		}
 		if (PhysicalMaterial.IsValid())
 		{
 			RepBits |= (1 << REP_PhysMaterial);
@@ -661,7 +666,11 @@ bool FGameplayCueParameters::NetSerialize(FArchive& Ar, class UPackageMap* Map,
 	{
 		Ar << SourceObject;
 	}
-	if (RepBits & (1 << REP_SourceObject))
+	if (RepBits & (1 << REP_TargetAttachComponent))
+	{
+		Ar << TargetAttachComponent;
+	}
+	if (RepBits & (1 << REP_PhysMaterial))
 	{
 		Ar << PhysicalMaterial;
 	}
diff --git a/Engine/Source/Runtime/GameplayAbilities/Public/Abilities/Tasks/AbilityTask_ApplyRootMotionRadialForce.h b/Engine/Source/Runtime/GameplayAbilities/Public/Abilities/Tasks/AbilityTask_ApplyRootMotionRadialForce.h
--- a/Engine/Source/Runtime/GameplayAbilities/Public/Abilities/Tasks/AbilityTask_ApplyRootMotionRadialForce.h
+++ b/Engine/Source/Runtime/GameplayAbilities/Public/Abilities/Tasks/AbilityTask_ApplyRootMotionRadialForce.h
@@ -22,7 +22,7 @@ class UAbilityTask_ApplyRootMotionRadialForce : public UAbilityTask
 
 	/** Apply force to character's movement */
 	UFUNCTION(BlueprintCallable, Category = "Ability|Tasks", meta = (HidePin = "WorldContextObject", DefaultToSelf = "WorldContextObject", BlueprintInternalUseOnly = "TRUE"))
-	static UAbilityTask_ApplyRootMotionRadialForce* ApplyRootMotionRadialForce(UObject* WorldContextObject, FName TaskInstanceName, FVector Location, AActor* LocationActor, float Strength, float Duration, float Radius, bool bIsPush, bool bIsAdditive, bool bNoZForce, UCurveFloat* StrengthDistanceFalloff, UCurveFloat* StrengthOverTime);
+	static UAbilityTask_ApplyRootMotionRadialForce* ApplyRootMotionRadialForce(UObject* WorldContextObject, FName TaskInstanceName, FVector Location, AActor* LocationActor, float Strength, float Duration, float Radius, bool bIsPush, bool bIsAdditive, bool bNoZForce, UCurveFloat* StrengthDistanceFalloff, UCurveFloat* StrengthOverTime, bool bUseFixedWorldDirection, FRotator FixedWorldDirection);
 
 	virtual void Activate() override;
 
@@ -78,6 +78,12 @@ protected:
 	UPROPERTY(Replicated)
 	UCurveFloat* StrengthOverTime;
 
+	UPROPERTY(Replicated)
+	bool bUseFixedWorldDirection;
+
+	UPROPERTY(Replicated)
+	FRotator FixedWorldDirection;
+
 	uint16 RootMotionSourceID;
 	bool bIsFinished;
 	float StartTime;
diff --git a/Engine/Source/Runtime/GameplayAbilities/Public/AbilitySystemBlueprintLibrary.h b/Engine/Source/Runtime/GameplayAbilities/Public/AbilitySystemBlueprintLibrary.h
--- a/Engine/Source/Runtime/GameplayAbilities/Public/AbilitySystemBlueprintLibrary.h
+++ b/Engine/Source/Runtime/GameplayAbilities/Public/AbilitySystemBlueprintLibrary.h
@@ -134,6 +134,9 @@ class GAMEPLAYABILITIES_API UAbilitySystemBlueprintLibrary : public UBlueprintFu
 	UFUNCTION(BlueprintPure, Category = "Ability|EffectContext", Meta = (DisplayName = "HasHitResult"))
 	static bool EffectContextHasHitResult(FGameplayEffectContextHandle EffectContext);
 
+	UFUNCTION(BlueprintCallable, Category = "Ability|EffectContext", Meta = (DisplayName = "AddHitResult"))
+	static void EffectContextAddHitResult(FGameplayEffectContextHandle EffectContext, FHitResult HitResult, bool bReset);
+
 	/** Gets the location the effect originated from */
 	UFUNCTION(BlueprintPure, Category = "Ability|EffectContext", Meta = (DisplayName = "GetOrigin"))
 	static FVector EffectContextGetOrigin(FGameplayEffectContextHandle EffectContext);
diff --git a/Engine/Source/Runtime/GameplayAbilities/Public/AbilitySystemComponent.h b/Engine/Source/Runtime/GameplayAbilities/Public/AbilitySystemComponent.h
--- a/Engine/Source/Runtime/GameplayAbilities/Public/AbilitySystemComponent.h
+++ b/Engine/Source/Runtime/GameplayAbilities/Public/AbilitySystemComponent.h
@@ -350,6 +350,9 @@ class GAMEPLAYABILITIES_API UAbilitySystemComponent : public UGameplayTasksCompo
 		UpdateTagMap(Tag, -1);
 	}
 
+	/** Called for predictively added gameplay cue. Needs to remove tag count and possible invoke OnRemove event if misprediction */
+	void OnPredictiveGameplayCueCatchup(FGameplayTag Tag);
+
 	float GetGameplayEffectDuration(FActiveGameplayEffectHandle Handle) const;
 
 	void GetGameplayEffectStartTimeAndDuration(FActiveGameplayEffectHandle Handle, float& StartEffectTime, float& Duration) const;
@@ -652,24 +655,24 @@ class GAMEPLAYABILITIES_API UAbilitySystemComponent : public UGameplayTasksCompo
 
 	/** Removes all active effects that contain any of the tags in Tags */
 	UFUNCTION(BlueprintCallable, Category = GameplayEffects)
-	void RemoveActiveEffectsWithTags(FGameplayTagContainer Tags);
+	int32 RemoveActiveEffectsWithTags(FGameplayTagContainer Tags);
 
 	/** Removes all active effects with captured source tags that contain any of the tags in Tags */
 	UFUNCTION(BlueprintCallable, Category = GameplayEffects)
-	void RemoveActiveEffectsWithSourceTags(FGameplayTagContainer Tags);
+	int32 RemoveActiveEffectsWithSourceTags(FGameplayTagContainer Tags);
 
 	/** Removes all active effects that apply any of the tags in Tags */
 	UFUNCTION(BlueprintCallable, Category = GameplayEffects)
-	void RemoveActiveEffectsWithAppliedTags(FGameplayTagContainer Tags);
+	int32 RemoveActiveEffectsWithAppliedTags(FGameplayTagContainer Tags);
 
 	/** Removes all active effects that grant any of the tags in Tags */
 	UFUNCTION(BlueprintCallable, Category = GameplayEffects)
-	void RemoveActiveEffectsWithGrantedTags(FGameplayTagContainer Tags);
+	int32 RemoveActiveEffectsWithGrantedTags(FGameplayTagContainer Tags);
 
 	/** Removes all active effects that match given query. StacksToRemove=-1 will remove all stacks. */
 	DEPRECATED(4.9, "FActiveGameplayEffectQuery is deprecated, use version that takes FGameplayEffectQuery")
-	void RemoveActiveEffects(const FActiveGameplayEffectQuery Query, int32 StacksToRemove=-1);
-	void RemoveActiveEffects(const FGameplayEffectQuery& Query, int32 StacksToRemove = -1);
+	int32 RemoveActiveEffects(const FActiveGameplayEffectQuery Query, int32 StacksToRemove=-1);
+	int32 RemoveActiveEffects(const FGameplayEffectQuery& Query, int32 StacksToRemove = -1);
 
 	void OnRestackGameplayEffects();	
 	
@@ -707,6 +710,9 @@ class GAMEPLAYABILITIES_API UAbilitySystemComponent : public UGameplayTasksCompo
 
 	UFUNCTION(NetMulticast, unreliable)
 	void NetMulticast_InvokeGameplayCueAdded(const FGameplayTag GameplayCueTag, FPredictionKey PredictionKey, FGameplayEffectContextHandle EffectContext);
+
+	UFUNCTION(NetMulticast, unreliable)
+	void NetMulticast_InvokeGameplayCueAdded_WithParams(const FGameplayTag GameplayCueTag, FPredictionKey PredictionKey, FGameplayCueParameters Parameters);
 	
 	UFUNCTION(NetMulticast, unreliable)
 	void NetMulticast_InvokeGameplayCueAddedAndWhileActive_FromSpec(const FGameplayEffectSpecForRPC& Spec, FPredictionKey PredictionKey);
diff --git a/Engine/Source/Runtime/GameplayAbilities/Public/GameplayCueInterface.h b/Engine/Source/Runtime/GameplayAbilities/Public/GameplayCueInterface.h
--- a/Engine/Source/Runtime/GameplayAbilities/Public/GameplayCueInterface.h
+++ b/Engine/Source/Runtime/GameplayAbilities/Public/GameplayCueInterface.h
@@ -79,6 +79,9 @@ struct FActiveGameplayCue : public FFastArraySerializerItem
 	UPROPERTY()
 	FPredictionKey PredictionKey;
 
+	UPROPERTY()
+	FGameplayCueParameters Parameters;
+
 	/** Has this been predictively removed on the client? */
 	UPROPERTY(NotReplicated)
 	bool bPredictivelyRemoved;
@@ -102,7 +105,7 @@ struct FActiveGameplayCueContainer : public FFastArraySerializer
 	/** Should this container only rpelicate in minimal replication mode */
 	bool bMinimalReplication;
 
-	void AddCue(const FGameplayTag& Tag, const FPredictionKey& PredictionKey);
+	void AddCue(const FGameplayTag& Tag, const FPredictionKey& PredictionKey, const FGameplayCueParameters& Parameters);
 	void RemoveCue(const FGameplayTag& Tag);
 
 	/** Marks as predictively removed so that we dont invoke remove event twice due to onrep */
diff --git a/Engine/Source/Runtime/GameplayAbilities/Public/GameplayCueManager.h b/Engine/Source/Runtime/GameplayAbilities/Public/GameplayCueManager.h
--- a/Engine/Source/Runtime/GameplayAbilities/Public/GameplayCueManager.h
+++ b/Engine/Source/Runtime/GameplayAbilities/Public/GameplayCueManager.h
@@ -195,7 +195,7 @@ protected:
 
 	virtual bool ShouldAsyncLoadAtStartup() const { return true; }
 
-	void BuildCuesToAddToGlobalSet(const TArray<FAssetData>& AssetDataList, FName TagPropertyName, bool bAsyncLoadAfterAdd, TArray<struct FGameplayCueReferencePair>& OutCuesToAdd, FOnGameplayCueNotifySetLoaded OnLoaded = FOnGameplayCueNotifySetLoaded(), FShouldLoadGCNotifyDelegate = FShouldLoadGCNotifyDelegate());
+	void BuildCuesToAddToGlobalSet(const TArray<FAssetData>& AssetDataList, FName TagPropertyName, TArray<struct FGameplayCueReferencePair>& OutCuesToAdd, TArray<FStringAssetReference>& OutAssetsToLoad, FShouldLoadGCNotifyDelegate = FShouldLoadGCNotifyDelegate());
 
 	/** The cue manager has a tendency to produce a lot of RPCs. This logs out when we are attempting to fire more RPCs than will actually go off */
 	void CheckForTooManyRPCs(FName FuncName, const FGameplayCuePendingExecute& PendingCue, const FString& CueID, const FGameplayEffectContext* EffectContext);
diff --git a/Engine/Source/Runtime/GameplayAbilities/Public/GameplayEffect.h b/Engine/Source/Runtime/GameplayAbilities/Public/GameplayEffect.h
--- a/Engine/Source/Runtime/GameplayAbilities/Public/GameplayEffect.h
+++ b/Engine/Source/Runtime/GameplayAbilities/Public/GameplayEffect.h
@@ -1691,8 +1691,8 @@ struct GAMEPLAYABILITIES_API FActiveGameplayEffectsContainer : public FFastArray
 	void ModifyActiveEffectStartTime(FActiveGameplayEffectHandle Handle, float StartTimeDiff);
 
 	DEPRECATED(4.9, "FActiveGameplayEffectQuery is deprecated, use FGameplayEffectQuery instead")
-	void RemoveActiveEffects(const FActiveGameplayEffectQuery Query, int32 StacksToRemove);
-	void RemoveActiveEffects(const FGameplayEffectQuery& Query, int32 StacksToRemove);
+	int32 RemoveActiveEffects(const FActiveGameplayEffectQuery Query, int32 StacksToRemove);
+	int32 RemoveActiveEffects(const FGameplayEffectQuery& Query, int32 StacksToRemove);
 
 	/**
 	 * Get the count of the effects matching the specified query (including stack count)
diff --git a/Engine/Source/Runtime/GameplayAbilities/Public/GameplayEffectAggregator.h b/Engine/Source/Runtime/GameplayAbilities/Public/GameplayEffectAggregator.h
--- a/Engine/Source/Runtime/GameplayAbilities/Public/GameplayEffectAggregator.h
+++ b/Engine/Source/Runtime/GameplayAbilities/Public/GameplayEffectAggregator.h
@@ -40,6 +40,7 @@ struct GAMEPLAYABILITIES_API FAggregator : public TSharedFromThis<FAggregator>
 	DECLARE_MULTICAST_DELEGATE_OneParam(FOnAggregatorDirty, FAggregator*);
 
 	FAggregator(float InBaseValue=0.f) : NetUpdateID(0), BaseValue(InBaseValue), IsBroadcastingDirty(false), NumPredictiveMods(0) { }
+	~FAggregator();
 
 	/** Simple accessor to base value */
 	float GetBaseValue() const;
diff --git a/Engine/Source/Runtime/GameplayAbilities/Public/GameplayEffectTypes.h b/Engine/Source/Runtime/GameplayAbilities/Public/GameplayEffectTypes.h
--- a/Engine/Source/Runtime/GameplayAbilities/Public/GameplayEffectTypes.h
+++ b/Engine/Source/Runtime/GameplayAbilities/Public/GameplayEffectTypes.h
@@ -109,13 +109,15 @@ struct GAMEPLAYABILITIES_API FActiveGameplayEffectHandle
 	GENERATED_USTRUCT_BODY()
 
 	FActiveGameplayEffectHandle()
-		: Handle(INDEX_NONE)
+		: Handle(INDEX_NONE),
+		bPassedFiltersAndWasExecuted(false)
 	{
 
 	}
 
 	FActiveGameplayEffectHandle(int32 InHandle)
-		: Handle(InHandle)
+		: Handle(InHandle),
+		bPassedFiltersAndWasExecuted(true)
 	{
 
 	}
@@ -125,6 +127,11 @@ struct GAMEPLAYABILITIES_API FActiveGameplayEffectHandle
 		return Handle != INDEX_NONE;
 	}
 
+	bool WasSuccessfullyApplied() const
+	{
+		return bPassedFiltersAndWasExecuted;
+	}
+
 	static FActiveGameplayEffectHandle GenerateNewHandle(UAbilitySystemComponent* OwningComponent);
 
 	static void ResetGlobalHandleMap();
@@ -163,6 +170,9 @@ private:
 
 	UPROPERTY()
 	int32 Handle;
+
+	UPROPERTY()
+	bool bPassedFiltersAndWasExecuted;
 };
 
 USTRUCT()
@@ -833,6 +843,10 @@ struct GAMEPLAYABILITIES_API FGameplayCueParameters
 	UPROPERTY(BlueprintReadWrite, Category = GameplayCue)
 	int32 AbilityLevel;
 
+	/** Could be used to say "attach FX to this component always" */
+	UPROPERTY(BlueprintReadWrite, Category = GameplayCue)
+	TWeakObjectPtr<USceneComponent> TargetAttachComponent;
+
 	bool NetSerialize(FArchive& Ar, class UPackageMap* Map, bool& bOutSuccess);
 
 	bool IsInstigatorLocallyControlled() const;
diff --git a/Engine/Source/Runtime/GameplayTags/Classes/GameplayTagContainer.h b/Engine/Source/Runtime/GameplayTags/Classes/GameplayTagContainer.h
--- a/Engine/Source/Runtime/GameplayTags/Classes/GameplayTagContainer.h
+++ b/Engine/Source/Runtime/GameplayTags/Classes/GameplayTagContainer.h
@@ -267,6 +267,20 @@ struct GAMEPLAYTAGS_API FGameplayTagContainer
 	 */
 	bool MatchesQuery(const struct FGameplayTagQuery& Query) const;
 
+
+	/**
+	 * Determine if the container has the specified tag. This forces an explicit match. 
+	 * This function exists for convenience and brevity. We do not wish to use default values for ::HasTag match type parameters, to avoid confusion on what the default behavior is. (E.g., we want people to think and use the right match type).
+	 * 
+	 * @param TagToCheck			Tag to check if it is present in the container
+	 * 
+	 * @return True if the tag is in the container, false if it is not
+	 */
+	FORCEINLINE_DEBUGGABLE bool HasTagExplicit(FGameplayTag const& TagToCheck) const
+	{
+		return HasTag(TagToCheck, EGameplayTagMatchType::Explicit, EGameplayTagMatchType::Explicit);
+	}
+
 	/**
 	 * Determine if the container has the specified tag
 	 * 
@@ -276,7 +290,7 @@ struct GAMEPLAYTAGS_API FGameplayTagContainer
 	 * 
 	 * @return True if the tag is in the container, false if it is not
 	 */
-	FORCEINLINE_DEBUGGABLE bool HasTag(FGameplayTag const& TagToCheck, TEnumAsByte<EGameplayTagMatchType::Type> TagMatchType, TEnumAsByte<EGameplayTagMatchType::Type> TagToCheckMatchType) const
+	FORCEINLINE_DEBUGGABLE bool HasTag(FGameplayTag const& TagToCheck, TEnumAsByte<EGameplayTagMatchType::Type> TagMatchType=EGameplayTagMatchType::Explicit, TEnumAsByte<EGameplayTagMatchType::Type> TagToCheckMatchType=EGameplayTagMatchType::Explicit) const
 	{
 		SCOPE_CYCLE_COUNTER(STAT_FGameplayTagContainer_HasTag);
 		bool bResult;
diff --git a/Engine/Source/Runtime/Landscape/Private/Landscape.cpp b/Engine/Source/Runtime/Landscape/Private/Landscape.cpp
--- a/Engine/Source/Runtime/Landscape/Private/Landscape.cpp
+++ b/Engine/Source/Runtime/Landscape/Private/Landscape.cpp
@@ -1038,9 +1038,9 @@ void ALandscape::CheckForErrors()
 }
 #endif
 
-void ALandscapeProxy::PreSave()
+void ALandscapeProxy::PreSave(const class ITargetPlatform* TargetPlatform)
 {
-	Super::PreSave();
+	Super::PreSave(TargetPlatform);
 
 #if WITH_EDITOR
 	// Work out whether we have grass or not for the next game run
@@ -1656,9 +1656,9 @@ UMaterialInterface* ALandscape::GetLandscapeHoleMaterial() const
 	return nullptr;
 }
 
-void ALandscape::PreSave()
+void ALandscape::PreSave(const class ITargetPlatform* TargetPlatform)
 {
-	Super::PreSave();
+	Super::PreSave(TargetPlatform);
 	//ULandscapeInfo* Info = GetLandscapeInfo();
 	//if (GIsEditor && Info && !IsRunningCommandlet())
 	//{
diff --git a/Engine/Source/Runtime/Landscape/Private/LandscapeCollision.cpp b/Engine/Source/Runtime/Landscape/Private/LandscapeCollision.cpp
--- a/Engine/Source/Runtime/Landscape/Private/LandscapeCollision.cpp
+++ b/Engine/Source/Runtime/Landscape/Private/LandscapeCollision.cpp
@@ -1606,9 +1606,9 @@ void ULandscapeHeightfieldCollisionComponent::PostLoad()
 #endif//WITH_EDITOR
 }
 
-void ULandscapeHeightfieldCollisionComponent::PreSave()
+void ULandscapeHeightfieldCollisionComponent::PreSave(const class ITargetPlatform* TargetPlatform)
 {
-	Super::PreSave();
+	Super::PreSave(TargetPlatform);
 
 	if (!IsRunningCommandlet())
 	{
diff --git a/Engine/Source/Runtime/Landscape/Private/LandscapeRender.cpp b/Engine/Source/Runtime/Landscape/Private/LandscapeRender.cpp
--- a/Engine/Source/Runtime/Landscape/Private/LandscapeRender.cpp
+++ b/Engine/Source/Runtime/Landscape/Private/LandscapeRender.cpp
@@ -30,6 +30,14 @@ IMPLEMENT_UNIFORM_BUFFER_STRUCT(FLandscapeUniformShaderParameters, TEXT("Landsca
 #define LANDSCAPE_MAX_COMPONENT_SIZE 255
 #define LANDSCAPE_LOD_SQUARE_ROOT_FACTOR 1.5f
 
+int32 GLandscapeMeshLODBias = 0;
+FAutoConsoleVariableRef CVarLandscapeMeshLODBias(
+	TEXT("r.LandscapeLODBias"),
+	GLandscapeMeshLODBias,
+	TEXT("LOD bias for landscape/terrain meshes."),
+	ECVF_Scalability
+	);
+
 /*------------------------------------------------------------------------------
 	Forsyth algorithm for cache optimizing index buffers.
 ------------------------------------------------------------------------------*/
@@ -533,8 +541,10 @@ TMap<uint32, FLandscapeSharedBuffers*>FLandscapeComponentSceneProxy::SharedBuffe
 TMap<uint32, FLandscapeSharedAdjacencyIndexBuffer*>FLandscapeComponentSceneProxy::SharedAdjacencyIndexBufferMap;
 TMap<FLandscapeNeighborInfo::FLandscapeKey, TMap<FIntPoint, const FLandscapeNeighborInfo*> > FLandscapeNeighborInfo::SharedSceneProxyMap;
 
+const static FName NAME_LandscapeResourceNameForDebugging(TEXT("Landscape"));
+
 FLandscapeComponentSceneProxy::FLandscapeComponentSceneProxy(ULandscapeComponent* InComponent, FLandscapeEditToolRenderData* InEditToolRenderData)
-	: FPrimitiveSceneProxy(InComponent)
+	: FPrimitiveSceneProxy(InComponent, NAME_LandscapeResourceNameForDebugging)
 	, FLandscapeNeighborInfo(InComponent->GetWorld(), InComponent->GetLandscapeProxy()->GetLandscapeGuid(), InComponent->GetSectionBase() / InComponent->ComponentSizeQuads, InComponent->HeightmapTexture, InComponent->ForcedLOD, InComponent->LODBias)
 	, MaxLOD(FMath::CeilLogTwo(InComponent->SubsectionSizeQuads + 1) - 1)
 	, FirstLOD(0)
@@ -1300,6 +1310,8 @@ float FLandscapeComponentSceneProxy::CalcDesiredLOD(const class FSceneView& View
 		SubsectionLODBias   = Neighbors[3] ? Neighbors[3]->LODBias : 0;
 	}
 
+	SubsectionLODBias = FMath::Clamp<int8>(SubsectionLODBias + GLandscapeMeshLODBias, FirstLOD, LastLOD);
+
 	const int32 MinStreamedLOD = SubsectionHeightmapTexture ? FMath::Min<int32>(((FTexture2DResource*)SubsectionHeightmapTexture->Resource)->GetCurrentFirstMip(), FMath::CeilLogTwo(SubsectionSizeVerts) - 1) : 0;
 
 	float fLOD = FLT_MAX;
diff --git a/Engine/Source/Runtime/Launch/Resources/VersionLocked.h b/Engine/Source/Runtime/Launch/Resources/VersionLocked.h
--- a/Engine/Source/Runtime/Launch/Resources/VersionLocked.h
+++ b/Engine/Source/Runtime/Launch/Resources/VersionLocked.h
@@ -7,7 +7,7 @@
 
 #define LOCKED_BUILD_VERSION 0
 #define LOCKED_NETWORK_VERSION 0
-#define LOCKED_REPLAY_VERSION 0
+#define LOCKED_REPLAY_VERSION 2984481
 
 #if LOCKED_BUILD_VERSION
 
diff --git a/Engine/Source/Runtime/MovieScene/Private/MovieScene.cpp b/Engine/Source/Runtime/MovieScene/Private/MovieScene.cpp
--- a/Engine/Source/Runtime/MovieScene/Private/MovieScene.cpp
+++ b/Engine/Source/Runtime/MovieScene/Private/MovieScene.cpp
@@ -573,9 +573,9 @@ void UMovieScene::PostLoad()
 }
 
 
-void UMovieScene::PreSave()
+void UMovieScene::PreSave(const class ITargetPlatform* TargetPlatform)
 {
-	Super::PreSave();
+	Super::PreSave(TargetPlatform);
 
 #if WITH_EDITORONLY_DATA
 	// compress meta data mappings prior to saving
diff --git a/Engine/Source/Runtime/MovieSceneTracks/Private/Sections/MovieSceneActorReferenceSection.cpp b/Engine/Source/Runtime/MovieSceneTracks/Private/Sections/MovieSceneActorReferenceSection.cpp
--- a/Engine/Source/Runtime/MovieSceneTracks/Private/Sections/MovieSceneActorReferenceSection.cpp
+++ b/Engine/Source/Runtime/MovieSceneTracks/Private/Sections/MovieSceneActorReferenceSection.cpp
@@ -96,14 +96,14 @@ void UMovieSceneActorReferenceSection::SetDefault( const FGuid& Value )
 }
 
 
-void UMovieSceneActorReferenceSection::PreSave()
+void UMovieSceneActorReferenceSection::PreSave(const class ITargetPlatform* TargetPlatform)
 {
 	ActorGuidStrings.Empty();
 	for ( const FGuid& ActorGuid : ActorGuids )
 	{
 		ActorGuidStrings.Add( ActorGuid.ToString() );
 	}
-	Super::PreSave();
+	Super::PreSave(TargetPlatform);
 }
 
 
diff --git a/Engine/Source/Runtime/NetworkReplayStreaming/HttpNetworkReplayStreaming/Private/HttpNetworkReplayStreaming.cpp b/Engine/Source/Runtime/NetworkReplayStreaming/HttpNetworkReplayStreaming/Private/HttpNetworkReplayStreaming.cpp
--- a/Engine/Source/Runtime/NetworkReplayStreaming/HttpNetworkReplayStreaming/Private/HttpNetworkReplayStreaming.cpp
+++ b/Engine/Source/Runtime/NetworkReplayStreaming/HttpNetworkReplayStreaming/Private/HttpNetworkReplayStreaming.cpp
@@ -1,6 +1,7 @@
 // Copyright 1998-2016 Epic Games, Inc. All Rights Reserved.
 
 #include "HttpNetworkReplayStreaming.h"
+#include "ScopedTimers.h"
 
 DEFINE_LOG_CATEGORY_STATIC( LogHttpReplay, Log, All );
 
@@ -822,43 +823,52 @@ void FHttpNetworkReplayStreamer::FlushCheckpointInternal( uint32 TimeInMS )
 		return;
 	}
 
-	// Upload any new streamed data to the http server
-	UE_LOG( LogHttpReplay, Log, TEXT( "FHttpNetworkReplayStreamer::FlushCheckpointInternal. Size: %i, StreamChunkIndex: %i" ), CheckpointArchive.Buffer.Num(), StreamChunkIndex );
+	double Duration = 0;
 
-	// Create the Http request and add to pending request list
-	TSharedRef<IHttpRequest> HttpRequest = FHttpModule::Get().CreateRequest();
+	{
+		FScopedDurationTimer Timer(Duration);
 
-	HttpRequest->OnProcessRequestComplete().BindRaw( this, &FHttpNetworkReplayStreamer::HttpUploadCheckpointFinished );
+		// Create the Http request and add to pending request list
+		TSharedRef<IHttpRequest> HttpRequest = FHttpModule::Get().CreateRequest();
 
-	HttpRequest->SetURL( FString::Printf( TEXT( "%sreplay/%s/event?group=checkpoint&time1=%i&time2=%i&meta=%i&incrementSize=false" ), *ServerURL, *SessionName, TimeInMS, TimeInMS, StreamChunkIndex ) );
-	HttpRequest->SetVerb( TEXT( "POST" ) );
-	HttpRequest->SetHeader( TEXT( "Content-Type" ), TEXT( "application/octet-stream" ) );
+		HttpRequest->OnProcessRequestComplete().BindRaw(this, &FHttpNetworkReplayStreamer::HttpUploadCheckpointFinished);
 
-	if ( SupportsCompression() )
-	{
-		SCOPE_CYCLE_COUNTER( STAT_HttpReplay_CompressTime );
+		HttpRequest->SetURL(FString::Printf(TEXT("%sreplay/%s/event?group=checkpoint&time1=%i&time2=%i&meta=%i&incrementSize=false"), *ServerURL, *SessionName, TimeInMS, TimeInMS, StreamChunkIndex));
+		HttpRequest->SetVerb(TEXT("POST"));
+		HttpRequest->SetHeader(TEXT("Content-Type"), TEXT("application/octet-stream"));
 
-		const double StartTime = FPlatformTime::Seconds();
-		FHttpStreamFArchive Compressed;
-		if ( !CompressBuffer( CheckpointArchive.Buffer, Compressed ) )
+		if (SupportsCompression())
 		{
-			SetLastError( ENetworkReplayError::ServiceUnavailable );
-			return;
+			SCOPE_CYCLE_COUNTER(STAT_HttpReplay_CompressTime);
+
+			const double StartTime = FPlatformTime::Seconds();
+			FHttpStreamFArchive Compressed;
+			if (!CompressBuffer(CheckpointArchive.Buffer, Compressed))
+			{
+				SetLastError(ENetworkReplayError::ServiceUnavailable);
+				return;
+			}
+			const double EndTime = FPlatformTime::Seconds();
+			HttpRequest->SetContent(Compressed.Buffer);
+
+			UE_LOG(LogHttpReplay, VeryVerbose, TEXT("Compressed checkpoint. Original: %i, Compressed: %i, Time: %2.2f MS"), CheckpointArchive.Buffer.Num(), Compressed.Buffer.Num(), (EndTime - StartTime) * 1000.0f);
+		}
+		else
+		{
+			HttpRequest->SetContent(CheckpointArchive.Buffer);
 		}
-		const double EndTime = FPlatformTime::Seconds();
-		HttpRequest->SetContent( Compressed.Buffer );
 
-		UE_LOG( LogHttpReplay, VeryVerbose, TEXT( "Compressed checkpoint. Original: %i, Compressed: %i, Time: %2.2f MS" ), CheckpointArchive.Buffer.Num(), Compressed.Buffer.Num(), ( EndTime - StartTime ) * 1000.0f );
-	}
-	else
-	{
-		HttpRequest->SetContent( CheckpointArchive.Buffer );
+		AddRequestToQueue(EQueuedHttpRequestType::UploadingCheckpoint, HttpRequest, 2, 2.0f);
 	}
 
-	CheckpointArchive.Buffer.Empty();
-	CheckpointArchive.Pos = 0;
+	// Upload any new streamed data to the http server
+	UE_LOG(LogHttpReplay, Log, TEXT("FHttpNetworkReplayStreamer::FlushCheckpointInternal. Size: %i, StreamChunkIndex: %i, Time: %2.2f MS"), 
+		CheckpointArchive.Buffer.Num(), 
+		StreamChunkIndex,
+		Duration * 1000.0f);
 
-	AddRequestToQueue( EQueuedHttpRequestType::UploadingCheckpoint, HttpRequest, 2, 2.0f );
+	CheckpointArchive.Buffer.Empty();
+	CheckpointArchive.Pos = 0;	
 }
 
 FQueuedHttpRequestAddEvent::FQueuedHttpRequestAddEvent( const FString& InName, const uint32 InTimeInMS, const FString& InGroup, const FString& InMeta, const TArray<uint8>& InData, TSharedRef< class IHttpRequest > InHttpRequest ) : FQueuedHttpRequest( EQueuedHttpRequestType::UploadingCustomEvent, InHttpRequest )
@@ -1717,9 +1727,26 @@ void FHttpNetworkReplayStreamer::HttpDownloadFinished( FHttpRequestPtr HttpReque
 			return;
 		}
 
-		NumTotalStreamChunks = FCString::Atoi( *HttpResponse->GetHeader( TEXT( "NumChunks" ) ) );
-		TotalDemoTimeInMS = FCString::Atoi( *HttpResponse->GetHeader( TEXT( "Time" ) ) );
-		bStreamIsLive = HttpResponse->GetHeader( TEXT( "State" ) ) == TEXT( "Live" );
+		const int32 NewNumTotalStreamChunks = FCString::Atoi( *HttpResponse->GetHeader( TEXT( "NumChunks" ) ) );
+		const uint32 NewTotalDemoTimeInMS = FCString::Atoi( *HttpResponse->GetHeader( TEXT( "Time" ) ) );
+
+		// Cached http calls can make time go backwards, so protect against that
+		if ( NewNumTotalStreamChunks > NumTotalStreamChunks )
+		{
+			NumTotalStreamChunks = NewNumTotalStreamChunks;
+		}
+
+		if ( NewTotalDemoTimeInMS > TotalDemoTimeInMS )
+		{
+			TotalDemoTimeInMS = NewTotalDemoTimeInMS;
+		}
+
+		const bool bNewStreamIsLive = HttpResponse->GetHeader( TEXT( "State" ) ) == TEXT( "Live" );
+
+		if ( bStreamIsLive && !bNewStreamIsLive )
+		{
+			bStreamIsLive = bNewStreamIsLive;
+		}
 
 		if ( StreamArchive.Buffer.Num() == 0 )
 		{
diff --git a/Engine/Source/Runtime/NullDrv/Private/NullRHI.cpp b/Engine/Source/Runtime/NullDrv/Private/NullRHI.cpp
--- a/Engine/Source/Runtime/NullDrv/Private/NullRHI.cpp
+++ b/Engine/Source/Runtime/NullDrv/Private/NullRHI.cpp
@@ -62,7 +62,7 @@ void FNullDynamicRHI::Shutdown()
 void* FNullDynamicRHI::GetStaticBuffer()
 {
 #if !WITH_EDITOR
-	checkf(!IsRunningDedicatedServer(), TEXT("NullRHI should never allocate memory on the server. Change the caller to avoid doing allocs in when FApp::ShouldUseNullRHI() is true."));
+    checkf(WITH_EDITOR || !IsRunningDedicatedServer(), TEXT("NullRHI should never allocate memory on the server. Change the caller to avoid doing allocs in when FApp::ShouldUseNullRHI() is true."));
 #endif
 
 	static void* Buffer = nullptr;
diff --git a/Engine/Source/Runtime/Online/HTTP/Private/Android/AndroidHttp.cpp b/Engine/Source/Runtime/Online/HTTP/Private/Android/AndroidHttp.cpp
--- a/Engine/Source/Runtime/Online/HTTP/Private/Android/AndroidHttp.cpp
+++ b/Engine/Source/Runtime/Online/HTTP/Private/Android/AndroidHttp.cpp
@@ -21,6 +21,6 @@ void FAndroidPlatformHttp::Shutdown()
 
 IHttpRequest* FAndroidPlatformHttp::ConstructRequest()
 {
-	return new FCurlHttpRequest(FCurlHttpManager::GShareHandle);
+	return new FCurlHttpRequest();
 }
 
diff --git a/Engine/Source/Runtime/Online/HTTP/Private/Curl/CurlHttp.cpp b/Engine/Source/Runtime/Online/HTTP/Private/Curl/CurlHttp.cpp
--- a/Engine/Source/Runtime/Online/HTTP/Private/Curl/CurlHttp.cpp
+++ b/Engine/Source/Runtime/Online/HTTP/Private/Curl/CurlHttp.cpp
@@ -9,7 +9,7 @@
 
 // FCurlHttpRequest
 
-FCurlHttpRequest::FCurlHttpRequest(CURLSH* InShareHandle)
+FCurlHttpRequest::FCurlHttpRequest()
 	:	EasyHandle(NULL)
 	,	HeaderList(NULL)
 	,	bCanceled(false)
@@ -34,7 +34,7 @@ FCurlHttpRequest::FCurlHttpRequest(CURLSH* InShareHandle)
 
 #endif // !UE_BUILD_SHIPPING && !UE_BUILD_TEST
 
-	curl_easy_setopt(EasyHandle, CURLOPT_SHARE, InShareHandle);
+	curl_easy_setopt(EasyHandle, CURLOPT_SHARE, FCurlHttpManager::GShareHandle);
 
 	// set certificate verification (disable to allow self-signed certificates)
 	if (FCurlHttpManager::CurlRequestOptions.bVerifyPeer)
@@ -328,21 +328,12 @@ size_t FCurlHttpRequest::ReceiveResponseBodyCallback(void* Ptr, size_t SizeInBlo
 	{
 		TimeSinceLastResponse = 0.0f;
 
-		int32 OldTotalBytesRead = 0;
 		uint32 SizeToDownload = SizeInBlocks * BlockSizeInBytes;
-		{
-			FScopeLock ScopeLock(&ProgressLock);
-			// Note that we are incrementing the number of total bytes read before we actually perform the memcpy, and this progress could be
-			// reported via the progress delegate but we are expecting the progress delegate to not expose access to the raw bytes so this
-			// would be a benign race condition
-			OldTotalBytesRead = Response->TotalBytesRead;
-			Response->TotalBytesRead += SizeToDownload;
-		}
 
 		UE_LOG(LogHttp, Verbose, TEXT("%p: ReceiveResponseBodyCallback: %d bytes out of %d received. (SizeInBlocks=%d, BlockSizeInBytes=%d, Response->TotalBytesRead=%d, Response->GetContentLength()=%d, SizeToDownload=%d (<-this will get returned from the callback))"),
 			this,
-			static_cast<int32>(OldTotalBytesRead + SizeToDownload), Response->GetContentLength(),
-			static_cast<int32>(SizeInBlocks), static_cast<int32>(BlockSizeInBytes), OldTotalBytesRead, Response->GetContentLength(), static_cast<int32>(SizeToDownload)
+			static_cast<int32>(Response->TotalBytesRead.GetValue() + SizeToDownload), Response->GetContentLength(),
+			static_cast<int32>(SizeInBlocks), static_cast<int32>(BlockSizeInBytes), Response->TotalBytesRead.GetValue(), Response->GetContentLength(), static_cast<int32>(SizeToDownload)
 			);
 
 		// note that we can be passed 0 bytes if file transmitted has 0 length
@@ -351,7 +342,8 @@ size_t FCurlHttpRequest::ReceiveResponseBodyCallback(void* Ptr, size_t SizeInBlo
 			Response->Payload.AddUninitialized(SizeToDownload);
 
 			// save
-			FMemory::Memcpy(static_cast<uint8*>(Response->Payload.GetData()) + OldTotalBytesRead, Ptr, SizeToDownload);
+			FMemory::Memcpy(static_cast< uint8* >(Response->Payload.GetData()) + Response->TotalBytesRead.GetValue(), Ptr, SizeToDownload);
+			Response->TotalBytesRead.Add(SizeToDownload);
 
 			return SizeToDownload;
 		}
@@ -368,38 +360,25 @@ size_t FCurlHttpRequest::UploadCallback(void* Ptr, size_t SizeInBlocks, size_t B
 {
 	TimeSinceLastResponse = 0.0f;
 
-	int32 OldBytesSent = 0;
-	size_t SizeToSend = 0;
+	size_t SizeToSend = RequestPayload.Num() - BytesSent.GetValue();
 	size_t SizeToSendThisTime = 0;
 
+	if (SizeToSend != 0)
 	{
-		FScopeLock ScopeLock(&ProgressLock);
-		OldBytesSent = BytesSent;
-		SizeToSend = RequestPayload.Num() - BytesSent;
-		if (SizeToSend != 0)
+		SizeToSendThisTime = FMath::Min(SizeToSend, SizeInBlocks * BlockSizeInBytes);
+		if (SizeToSendThisTime != 0)
 		{
-			SizeToSendThisTime = FMath::Min(SizeToSend, SizeInBlocks * BlockSizeInBytes);
-			if (SizeToSendThisTime != 0)
-			{
-				// Note that we are incrementing the number of total bytes sent before we actually perform the memcpy, and this progress could be
-				// reported via the progress delegate but we are expecting the progress delegate to not expose access to the raw bytes so this
-				// would be a benign race condition
-				BytesSent += SizeToSendThisTime;
-			}
+			// static cast just ensures that this is uint8* in fact
+			FMemory::Memcpy(Ptr, static_cast< uint8* >(RequestPayload.GetData()) + BytesSent.GetValue(), SizeToSendThisTime);
+			BytesSent.Add(SizeToSendThisTime);
 		}
 	}
 
-	if (SizeToSendThisTime != 0)
-	{
-		// static cast just ensures that this is uint8* in fact
-		FMemory::Memcpy(Ptr, static_cast<uint8*>(RequestPayload.GetData()) + OldBytesSent, SizeToSendThisTime);
-	}
-
 	UE_LOG(LogHttp, Verbose, TEXT("%p: UploadCallback: %d bytes out of %d sent. (SizeInBlocks=%d, BlockSizeInBytes=%d, RequestPayload.Num()=%d, BytesSent=%d, SizeToSend=%d, SizeToSendThisTime=%d (<-this will get returned from the callback))"),
 		this,
-		static_cast<int32>(OldBytesSent + SizeToSendThisTime), RequestPayload.Num(),
-		static_cast<int32>(SizeInBlocks), static_cast<int32>(BlockSizeInBytes), RequestPayload.Num(), static_cast<int32>(OldBytesSent + SizeToSendThisTime), static_cast<int32>(SizeToSend),
-		static_cast<int32>(SizeToSendThisTime)
+		static_cast< int32 >(BytesSent.GetValue()), RequestPayload.Num(),
+		static_cast< int32 >(SizeInBlocks), static_cast< int32 >(BlockSizeInBytes), RequestPayload.Num(), static_cast< int32 >(BytesSent.GetValue()), static_cast< int32 >(SizeToSend),
+		static_cast< int32 >(SizeToSendThisTime)
 		);
 
 	return SizeToSendThisTime;
@@ -554,7 +533,7 @@ bool FCurlHttpRequest::SetupRequest()
 		curl_easy_setopt(EasyHandle, CURLOPT_INFILESIZE, RequestPayload.Num());
 
 		// reset the counter
-		BytesSent = 0;
+		BytesSent.Reset();
 	}
 	else if (Verb == TEXT("GET"))
 	{
@@ -727,19 +706,11 @@ void FCurlHttpRequest::Tick(float DeltaSeconds)
 void FCurlHttpRequest::CheckProgressDelegate()
 {
 	QUICK_SCOPE_CYCLE_COUNTER(STAT_FCurlHttpRequest_CheckProgressDelegate);
-	int32 CurrentBytesRead = 0;
-	int32 CurrentBytesSent = 0;
-	{
-		FScopeLock ScopeLock(&ProgressLock);
-		if (Response.IsValid())
-		{
-			CurrentBytesRead = Response->TotalBytesRead;
-		}
-		CurrentBytesSent = BytesSent;
-	}
+	int32 CurrentBytesRead = Response.IsValid() ? Response->TotalBytesRead.GetValue() : 0;
+	int32 CurrentBytesSent = BytesSent.GetValue();
 
-	bool bProcessing = CompletionStatus == EHttpRequestStatus::Processing;
-	bool bProgessChanged = (CurrentBytesSent != LastReportedBytesSent) ||
+	const bool bProcessing = CompletionStatus == EHttpRequestStatus::Processing;
+	const bool bProgessChanged = (CurrentBytesSent != LastReportedBytesSent) ||
 		(Response.IsValid() && CurrentBytesRead != LastReportedBytesRead);
 	if (bProcessing && bProgessChanged)
 	{
diff --git a/Engine/Source/Runtime/Online/HTTP/Private/Curl/CurlHttp.h b/Engine/Source/Runtime/Online/HTTP/Private/Curl/CurlHttp.h
--- a/Engine/Source/Runtime/Online/HTTP/Private/Curl/CurlHttp.h
+++ b/Engine/Source/Runtime/Online/HTTP/Private/Curl/CurlHttp.h
@@ -176,7 +176,7 @@ public:
 	/**
 	 * Constructor
 	 */
-	FCurlHttpRequest(CURLSH* InShareHandle);
+	FCurlHttpRequest();
 
 	/**
 	 * Destructor. Clean up any connection/request handles
@@ -334,10 +334,8 @@ private:
 	/** Elapsed time since the last received HTTP response. */
 	float TimeSinceLastResponse;
 
-	/** Used to lock access to bytes read/sent */
-	FCriticalSection ProgressLock;
 	/** Number of bytes sent already */
-	uint32 BytesSent;
+	FThreadSafeCounter BytesSent;
 	/** Last bytes read reported to progress delegate */
 	int32 LastReportedBytesRead;
 	/** Last bytes sent reported to progress delegate */
@@ -399,7 +397,7 @@ private:
 	/** BYTE array to fill in as the response is read via didReceiveData */
 	TArray<uint8> Payload;
 	/** Caches how many bytes of the response we've read so far */
-	int32 TotalBytesRead;
+	FThreadSafeCounter TotalBytesRead;
 	/** Cached key/value header pairs. Parsed once request completes */
 	TMap<FString, FString> Headers;
 	/** Cached code from completed response */
diff --git a/Engine/Source/Runtime/Online/HTTP/Private/Curl/CurlHttpManager.cpp b/Engine/Source/Runtime/Online/HTTP/Private/Curl/CurlHttpManager.cpp
--- a/Engine/Source/Runtime/Online/HTTP/Private/Curl/CurlHttpManager.cpp
+++ b/Engine/Source/Runtime/Online/HTTP/Private/Curl/CurlHttpManager.cpp
@@ -2,6 +2,7 @@
 
 #include "HttpPrivatePCH.h"
 #include "CurlHttpManager.h"
+#include "CurlHttpThread.h"
 #include "CurlHttp.h"
 
 #if WITH_LIBCURL
@@ -268,155 +269,9 @@ void FCurlHttpManager::ShutdownCurl()
 	curl_global_cleanup();
 }
 
-FCurlHttpManager::FCurlHttpManager()
-	:	FHttpManager()
-	,	MultiHandle(GMultiHandle)
-	,	MaxSimultaneousRequests(0)
-	,	MaxRequestsAddedPerFrame(0)
-	,	NumRequestsAddedToMulti(0)
+FHttpThread* FCurlHttpManager::CreateHttpThread()
 {
-	check(MultiHandle);
-	if (GConfig)
-	{
-		GConfig->GetInt(TEXT("HTTP"), TEXT("CurlMaxSimultaneousRequests"), MaxSimultaneousRequests, GEngineIni);
-		GConfig->GetInt(TEXT("HTTP"), TEXT("CurlMaxRequestsAddedPerFrame"), MaxRequestsAddedPerFrame, GEngineIni);
-	}
-}
-
-void FCurlHttpManager::AddRequest(const TSharedRef<IHttpRequest>& Request)
-{
-	checkf(false, TEXT("Should not be called for curl http anymore, should be using AddThreadedRequest."));
-}
-
-void FCurlHttpManager::RemoveRequest(const TSharedRef<IHttpRequest>& Request)
-{
-	checkf(false, TEXT("Should not be called for curl http anymore."));
-}
-
-bool FCurlHttpManager::FindNextEasyHandle(CURL** OutEasyHandle) const
-{
-	FDateTime LowestDateTime = FDateTime::MaxValue();
-	for (auto& It : HandlesToRequests)
-	{
-		if (!It.Value.bProcessingStarted)
-		{
-			if (It.Value.DateTime < LowestDateTime)
-			{
-				LowestDateTime = It.Value.DateTime;
-				*OutEasyHandle = It.Key;
-			}
-		}
-	}
-	return LowestDateTime != FDateTime::MaxValue();
-}
-
-void FCurlHttpManager::HttpThreadTick(float DeltaSeconds)
-{
-	QUICK_SCOPE_CYCLE_COUNTER(STAT_FCurlHttpManager_Tick);
-	check(MultiHandle);
-
-	{
-		if (RunningThreadedRequests.Num() > 0)
-		{
-			int RunningRequests = -1;
-			curl_multi_perform(MultiHandle, &RunningRequests);
-
-			// read more info if number of requests changed or if there's zero running
-			// (note that some requests might have never be "running" from libcurl's point of view)
-			if (RunningRequests == 0 || RunningRequests != NumRequestsAddedToMulti)
-			{
-				for(;;)
-				{
-					int MsgsStillInQueue = 0;	// may use that to impose some upper limit we may spend in that loop
-					CURLMsg * Message = curl_multi_info_read(MultiHandle, &MsgsStillInQueue);
-
-					if (Message == NULL)
-					{
-						break;
-					}
-
-					// find out which requests have completed
-					if (Message->msg == CURLMSG_DONE)
-					{
-						CURL* CompletedHandle = Message->easy_handle;
-						curl_multi_remove_handle(GMultiHandle, CompletedHandle);
-						--NumRequestsAddedToMulti;
-
-						CurlEasyRequestData* RequestData = HandlesToRequests.Find(CompletedHandle);
-						if (RequestData)
-						{
-							RequestData->bAddedToMulti = false;
-
-							FCurlHttpRequest* CurlRequest = static_cast< FCurlHttpRequest* >(RequestData->Request);
-							CurlRequest->MarkAsCompleted(Message->data.result);
-
-							UE_LOG(LogHttp, Verbose, TEXT("Request %p (easy handle:%p) has completed (code:%d) and has been marked as such"), CurlRequest, CompletedHandle, (int32)Message->data.result);
-						}
-						else
-						{
-							UE_LOG(LogHttp, Warning, TEXT("Could not find mapping for completed request (easy handle: %p)"), CompletedHandle);
-						}
-					}
-				}
-			}
-
-			int NumAdded = 0;
-			CURL* EasyHandle;
-			while ((MaxSimultaneousRequests == 0 || NumRequestsAddedToMulti < MaxSimultaneousRequests) &&
-				(MaxRequestsAddedPerFrame == 0 || NumAdded < MaxRequestsAddedPerFrame) &&
-				FindNextEasyHandle(&EasyHandle))
-			{
-				CurlEasyRequestData& RequestData = HandlesToRequests.FindChecked(EasyHandle);
-
-				CURLMcode AddResult = curl_multi_add_handle(GMultiHandle, EasyHandle);
-				RequestData.bProcessingStarted = true;
-
-				if (AddResult == CURLM_OK)
-				{
-					++NumAdded;
-					++NumRequestsAddedToMulti;
-					RequestData.bAddedToMulti = true;
-				}
-				else
-				{
-					UE_LOG(LogHttp, Warning, TEXT("Failed to add easy handle %p to multi handle with code %d"), EasyHandle, (int)AddResult);
-				}
-
-				FCurlHttpRequest* CurlRequest = static_cast<FCurlHttpRequest*>(RequestData.Request);
-				CurlRequest->SetAddToCurlMultiResult(AddResult);
-			}
-		}
-	}
-
-	FHttpManager::HttpThreadTick(DeltaSeconds);
-}
-
-bool FCurlHttpManager::StartThreadedRequest(IHttpThreadedRequest* Request)
-{
-	FCurlHttpRequest* CurlRequest = static_cast<FCurlHttpRequest*>(Request);
-	CURL* EasyHandle = CurlRequest->GetEasyHandle();
-	ensure(!HandlesToRequests.Contains(EasyHandle));
-	HandlesToRequests.Add(EasyHandle, CurlEasyRequestData(Request));
-
-	return true;
-}
-
-void FCurlHttpManager::CompleteThreadedRequest(IHttpThreadedRequest* Request)
-{
-	FCurlHttpRequest* CurlRequest = static_cast<FCurlHttpRequest*>(Request);
-	CURL* EasyHandle = CurlRequest->GetEasyHandle();
-
-	CurlEasyRequestData* RequestData = HandlesToRequests.Find(EasyHandle);
-	ensure(RequestData);
-	if (RequestData)
-	{
-		if (RequestData->bAddedToMulti)
-		{
-			curl_multi_remove_handle(GMultiHandle, EasyHandle);
-			--NumRequestsAddedToMulti;
-		}
-		HandlesToRequests.Remove(EasyHandle);
-	}
+	return new FCurlHttpThread();
 }
 
 #endif //WITH_LIBCURL
diff --git a/Engine/Source/Runtime/Online/HTTP/Private/Curl/CurlHttpManager.h b/Engine/Source/Runtime/Online/HTTP/Private/Curl/CurlHttpManager.h
--- a/Engine/Source/Runtime/Online/HTTP/Private/Curl/CurlHttpManager.h
+++ b/Engine/Source/Runtime/Online/HTTP/Private/Curl/CurlHttpManager.h
@@ -14,70 +14,11 @@
 
 class FCurlHttpManager : public FHttpManager
 {
-protected:
-	/** 
-	 * Struture containing data associated with a CURL easy handle
-	 */
-	struct CurlEasyRequestData
-	{
-		CurlEasyRequestData(IHttpRequest* InRequest)
-			: Request(InRequest)
-			, DateTime(FDateTime::Now())
-			, bProcessingStarted(false)
-			, bAddedToMulti(false)
-		{
-		}
-
-		/** Pointer to the Http Request */
-		IHttpRequest* Request;
-		
-		/** Date time this was added to ensure first in first out */
-		FDateTime DateTime;
-
-		/** Have we started processing the request by adding it to curl's multi handle */
-		bool bProcessingStarted;
-
-		/** Is the request currently added to curl's multi handle */
-		bool bAddedToMulti;
-	};
-
-
-	/** multi handle that groups all the requests - not owned by this class */
-	CURLM * MultiHandle;
-
-	/** Maximum number of requests that can be added to MultiHandle simultaneously (0 = no limit) */
-	int32 MaxSimultaneousRequests;
-	
-	/** Maximum number of requests that can be added to MultiHandle in one tick (0 = no limit) */
-	int32 MaxRequestsAddedPerFrame;
-
-	/** Number of requests that have been added to curl multi */
-	int NumRequestsAddedToMulti;
-
-	/** Mapping of libcurl easy handles to HTTP requests */
-	TMap<CURL*, CurlEasyRequestData> HandlesToRequests;
-
-	/** Find the next easy handle to be processed */
-	bool FindNextEasyHandle(CURL** OutEasyHandle) const;
-
-	/** Global multi handle */
-	static CURLM * GMultiHandle;
-
 public:
-
-	//~ Begin HttpManager Interface
-	virtual void AddRequest(const TSharedRef<class IHttpRequest>& Request) override;
-	virtual void RemoveRequest(const TSharedRef<class IHttpRequest>& Request) override;
-	virtual void HttpThreadTick(float DeltaSeconds) override;
-	virtual bool StartThreadedRequest(IHttpThreadedRequest* Request) override;
-	virtual void CompleteThreadedRequest(IHttpThreadedRequest* Request) override;
-	//~ End HttpManager Interface
-
-	FCurlHttpManager();
-
 	static void InitCurl();
 	static void ShutdownCurl();
 	static CURLSH* GShareHandle;
+	static CURLM * GMultiHandle;
 
 	static struct FCurlRequestOptions
 	{
@@ -107,6 +48,11 @@ public:
 		const char * CertBundlePath;
 	}
 	CurlRequestOptions;
+
+protected:
+	//~ Begin HttpManager Interface
+	virtual FHttpThread* CreateHttpThread() override;
+	//~ End HttpManager Interface
 };
 
 #endif //WITH_LIBCURL
diff --git a/Engine/Source/Runtime/Online/HTTP/Private/Curl/CurlHttpThread.cpp b/Engine/Source/Runtime/Online/HTTP/Private/Curl/CurlHttpThread.cpp
new file mode 100644
--- /dev/null
+++ b/Engine/Source/Runtime/Online/HTTP/Private/Curl/CurlHttpThread.cpp
@@ -0,0 +1,106 @@
+// Copyright 1998-2016 Epic Games, Inc. All Rights Reserved.
+
+#include "HttpPrivatePCH.h"
+#include "CurlHttpThread.h"
+#include "CurlHttp.h"
+#include "CurlHttpManager.h"
+
+#if WITH_LIBCURL
+
+FCurlHttpThread::FCurlHttpThread()
+{
+}
+
+void FCurlHttpThread::HttpThreadTick(float DeltaSeconds)
+{
+	QUICK_SCOPE_CYCLE_COUNTER(STAT_FCurlHttpThread_HttpThreadTick);
+	check(FCurlHttpManager::GMultiHandle);
+
+	const int NumRequestsToTick = RunningThreadedRequests.Num();
+
+	{
+		if (RunningThreadedRequests.Num() > 0)
+		{
+			int RunningRequests = -1;
+			{
+				QUICK_SCOPE_CYCLE_COUNTER(STAT_FCurlHttpThread_HttpThreadTick_Perform);
+				curl_multi_perform(FCurlHttpManager::GMultiHandle, &RunningRequests);
+			}
+
+			// read more info if number of requests changed or if there's zero running
+			// (note that some requests might have never be "running" from libcurl's point of view)
+			if (RunningRequests == 0 || RunningRequests != RunningThreadedRequests.Num())
+			{
+				for (;;)
+				{
+					QUICK_SCOPE_CYCLE_COUNTER(STAT_FCurlHttpThread_HttpThreadTick_Loop);
+					int MsgsStillInQueue = 0;	// may use that to impose some upper limit we may spend in that loop
+					CURLMsg * Message = curl_multi_info_read(FCurlHttpManager::GMultiHandle, &MsgsStillInQueue);
+
+					if (Message == NULL)
+					{
+						break;
+					}
+
+					// find out which requests have completed
+					if (Message->msg == CURLMSG_DONE)
+					{
+						CURL* CompletedHandle = Message->easy_handle;
+						curl_multi_remove_handle(FCurlHttpManager::GMultiHandle, CompletedHandle);
+
+						IHttpThreadedRequest** Request = HandlesToRequests.Find(CompletedHandle);
+						if (Request)
+						{
+							FCurlHttpRequest* CurlRequest = static_cast< FCurlHttpRequest* >(*Request);
+							CurlRequest->MarkAsCompleted(Message->data.result);
+
+							UE_LOG(LogHttp, Verbose, TEXT("Request %p (easy handle:%p) has completed (code:%d) and has been marked as such"), CurlRequest, CompletedHandle, (int32)Message->data.result);
+
+							HandlesToRequests.Remove(CompletedHandle);
+						}
+						else
+						{
+							UE_LOG(LogHttp, Warning, TEXT("Could not find mapping for completed request (easy handle: %p)"), CompletedHandle);
+						}
+					}
+				}
+			}
+		}
+	}
+
+	FHttpThread::HttpThreadTick(DeltaSeconds);
+}
+
+bool FCurlHttpThread::StartThreadedRequest(IHttpThreadedRequest* Request)
+{
+	FCurlHttpRequest* CurlRequest = static_cast<FCurlHttpRequest*>(Request);
+	CURL* EasyHandle = CurlRequest->GetEasyHandle();
+	ensure(!HandlesToRequests.Contains(EasyHandle));
+
+	CURLMcode AddResult = curl_multi_add_handle(FCurlHttpManager::GMultiHandle, EasyHandle);
+	CurlRequest->SetAddToCurlMultiResult(AddResult);
+
+	if (AddResult != CURLM_OK)
+	{
+		UE_LOG(LogHttp, Warning, TEXT("Failed to add easy handle %p to multi handle with code %d"), EasyHandle, (int)AddResult);
+		return false;
+	}
+
+	HandlesToRequests.Add(EasyHandle, Request);
+
+	return FHttpThread::StartThreadedRequest(Request);
+}
+
+void FCurlHttpThread::CompleteThreadedRequest(IHttpThreadedRequest* Request)
+{
+	FCurlHttpRequest* CurlRequest = static_cast<FCurlHttpRequest*>(Request);
+	CURL* EasyHandle = CurlRequest->GetEasyHandle();
+
+	if (HandlesToRequests.Find(EasyHandle))
+	{
+		curl_multi_remove_handle(FCurlHttpManager::GMultiHandle, EasyHandle);
+		HandlesToRequests.Remove(EasyHandle);
+	}
+}
+
+#endif
\ No newline at end of file
diff --git a/Engine/Source/Runtime/Online/HTTP/Private/Curl/CurlHttpThread.h b/Engine/Source/Runtime/Online/HTTP/Private/Curl/CurlHttpThread.h
new file mode 100644
--- /dev/null
+++ b/Engine/Source/Runtime/Online/HTTP/Private/Curl/CurlHttpThread.h
@@ -0,0 +1,38 @@
+// Copyright 1998-2016 Epic Games, Inc. All Rights Reserved.
+
+#pragma once
+
+#if WITH_LIBCURL
+
+#include "HttpThread.h"
+
+#if PLATFORM_WINDOWS
+#include "AllowWindowsPlatformTypes.h"
+#endif
+	#include "curl/curl.h"
+#if PLATFORM_WINDOWS
+#include "HideWindowsPlatformTypes.h"
+#endif
+
+
+class FCurlHttpThread
+	: public FHttpThread
+{
+public:
+	
+	FCurlHttpThread();
+
+protected:
+	//~ Begin FHttpThread Interface
+	virtual void HttpThreadTick(float DeltaSeconds) override;
+	virtual bool StartThreadedRequest(IHttpThreadedRequest* Request) override;
+	virtual void CompleteThreadedRequest(IHttpThreadedRequest* Request) override;
+	//~ End FHttpThread Interface
+protected:
+
+	/** Mapping of libcurl easy handles to HTTP requests */
+	TMap<CURL*, IHttpThreadedRequest*> HandlesToRequests;
+};
+
+
+#endif //WITH_LIBCURL
diff --git a/Engine/Source/Runtime/Online/HTTP/Private/HttpManager.cpp b/Engine/Source/Runtime/Online/HTTP/Private/HttpManager.cpp
--- a/Engine/Source/Runtime/Online/HTTP/Private/HttpManager.cpp
+++ b/Engine/Source/Runtime/Online/HTTP/Private/HttpManager.cpp
@@ -3,81 +3,71 @@
 #include "HttpPrivatePCH.h"
 #include "HttpManager.h"
 
+#include "HttpThread.h"
+
 // FHttpManager
 
 FCriticalSection FHttpManager::RequestLock;
 
 FHttpManager::FHttpManager()
 	:	FTickerObjectBase(0.0f)
-	,	HttpThreadTickBudget(0.033)
 	,	DeferredDestroyDelay(10.0f)
 {
-	double TargetTickRate = FHttpModule::Get().GetHttpThreadTickRate();
-	if (TargetTickRate > 0.0)
-	{
-		HttpThreadTickBudget = 1.0 / TargetTickRate;
-	}
-
-	UE_LOG(LogHttp, Log, TEXT("HTTP thread tick budget is %.1f ms (tick rate is %f)"), HttpThreadTickBudget, TargetTickRate);
-
-	Thread = FRunnableThread::Create(this, TEXT("HttpManagerThread"), 1024 * 1024, TPri_Normal);
 }
 
 FHttpManager::~FHttpManager()
 {
-	if (Thread != nullptr)
+	if (Thread)
 	{
-		Thread->Kill(true);
+		Thread->StopThread();
 		delete Thread;
 	}
 }
 
+void FHttpManager::Initialize()
+{
+	Thread = CreateHttpThread();
+	Thread->StartThread();
+}
+
+FHttpThread* FHttpManager::CreateHttpThread()
+{
+	return new FHttpThread();
+}
+
 void FHttpManager::Flush(bool bShutdown)
 {
-	bool bHasRequests = false;
-	
+	FScopeLock ScopeLock(&RequestLock);
+
+	if (bShutdown)
 	{
-		FScopeLock ScopeLock(&RequestLock);
-		if (bShutdown)
+		if (Requests.Num())
 		{
-			if (Requests.Num())
-			{
-				UE_LOG(LogHttp, Display, TEXT("Http module shutting down, but needs to wait on %d outstanding Http requests:"), Requests.Num());
-			}
-			// Clear delegates since they may point to deleted instances
-			for (TArray<TSharedRef<IHttpRequest>>::TIterator It(Requests); It; ++It)
-			{
-				TSharedRef<IHttpRequest> Request = *It;
-				Request->OnProcessRequestComplete().Unbind();
-				Request->OnRequestProgress().Unbind();
-				UE_LOG(LogHttp, Display, TEXT("	verb=[%s] url=[%s] status=%s"), *Request->GetVerb(), *Request->GetURL(), EHttpRequestStatus::ToString(Request->GetStatus()));
-			}
+			UE_LOG(LogHttp, Display, TEXT("Http module shutting down, but needs to wait on %d outstanding Http requests:"), Requests.Num());
+		}
+		// Clear delegates since they may point to deleted instances
+		for (TArray<TSharedRef<IHttpRequest>>::TIterator It(Requests); It; ++It)
+		{
+			TSharedRef<IHttpRequest> Request = *It;
+			Request->OnProcessRequestComplete().Unbind();
+			Request->OnRequestProgress().Unbind();
+			UE_LOG(LogHttp, Display, TEXT("	verb=[%s] url=[%s] status=%s"), *Request->GetVerb(), *Request->GetURL(), EHttpRequestStatus::ToString(Request->GetStatus()));
 		}
-		bHasRequests = Requests.Num() > 0;
 	}
 
 	// block until all active requests have completed
 	double LastTime = FPlatformTime::Seconds();
-	while (bHasRequests)
-	{	
-		{
-			FScopeLock ScopeLock(&RequestLock);
-			const double AppTime = FPlatformTime::Seconds();
-			Tick(AppTime - LastTime);
-			LastTime = AppTime;
-			bHasRequests = Requests.Num() > 0;
-		}
-		if (bHasRequests)
+	while (Requests.Num() > 0)
+	{
+		const double AppTime = FPlatformTime::Seconds();
+		Tick(AppTime - LastTime);
+		LastTime = AppTime;
+		if (Requests.Num() > 0)
 		{
 			UE_LOG(LogHttp, Display, TEXT("Sleeping 0.5s to wait for %d outstanding Http requests."), Requests.Num());
 			FPlatformProcess::Sleep(0.5f);
 		}
 	}
-	if (bShutdown)
-	{
-		// Stop the http thread
-		Stop();
-	}
 }
 
 bool FHttpManager::Tick(float DeltaSeconds)
@@ -100,18 +90,19 @@ bool FHttpManager::Tick(float DeltaSeconds)
 			PendingDestroyRequests.RemoveAt(Idx--);
 		}		
 	}
+
+	TArray<IHttpThreadedRequest*> CompletedThreadedRequests;
+	Thread->GetCompletedRequests(CompletedThreadedRequests);
+
 	// Finish and remove any completed requests
 	for (IHttpThreadedRequest* CompletedRequest : CompletedThreadedRequests)
 	{
-		TSharedRef<IHttpRequest> CompletedRequestRef(CompletedRequest->AsShared());
-		Requests.Remove(CompletedRequestRef);
-
 		// Keep track of requests that have been removed to be destroyed later
-		PendingDestroyRequests.AddUnique(FRequestPendingDestroy(DeferredDestroyDelay,CompletedRequestRef));
+		PendingDestroyRequests.AddUnique(FRequestPendingDestroy(DeferredDestroyDelay,CompletedRequest->AsShared()));
 
 		CompletedRequest->FinishRequest();
+		Requests.Remove(CompletedRequest->AsShared());
 	}
-	CompletedThreadedRequests.Empty();
 	// keep ticking
 	return true;
 }
@@ -135,20 +126,16 @@ void FHttpManager::RemoveRequest(const TSharedRef<IHttpRequest>& Request)
 
 void FHttpManager::AddThreadedRequest(const TSharedRef<IHttpThreadedRequest>& Request)
 {
-	FScopeLock ScopeLock(&RequestLock);
-
-	Requests.Add(Request);
-
-	PendingThreadedRequests.Add(&Request.Get());
+	{
+		FScopeLock ScopeLock(&RequestLock);
+		Requests.Add(Request);
+	}
+	Thread->AddRequest(&Request.Get());
 }
 
 void FHttpManager::CancelThreadedRequest(const TSharedRef<IHttpThreadedRequest>& Request)
 {
-	FScopeLock ScopeLock(&RequestLock);
-	if (Requests.Contains(Request))
-	{
-		CancelledThreadedRequests.Add(&Request.Get());
-	}
+	Thread->CancelRequest(&Request.Get());
 }
 
 bool FHttpManager::IsValidRequest(const IHttpRequest* RequestPtr) const
@@ -156,7 +143,7 @@ bool FHttpManager::IsValidRequest(const IHttpRequest* RequestPtr) const
 	FScopeLock ScopeLock(&RequestLock);
 
 	bool bResult = false;
-	for (auto& Request : Requests)
+	for (const TSharedRef<IHttpRequest>& Request : Requests)
 	{
 		if (&Request.Get() == RequestPtr)
 		{
@@ -164,10 +151,6 @@ bool FHttpManager::IsValidRequest(const IHttpRequest* RequestPtr) const
 			break;
 		}
 	}
-	if (!bResult)
-	{
-		bResult = CompletedThreadedRequests.Contains(RequestPtr);
-	}
 
 	return bResult;
 }
@@ -177,7 +160,7 @@ void FHttpManager::DumpRequests(FOutputDevice& Ar) const
 	FScopeLock ScopeLock(&RequestLock);
 	
 	Ar.Logf(TEXT("------- (%d) Http Requests"), Requests.Num());
-	for (auto& Request : Requests)
+	for (const TSharedRef<IHttpRequest>& Request : Requests)
 	{
 		Ar.Logf(TEXT("	verb=[%s] url=[%s] status=%s"),
 			*Request->GetVerb(), *Request->GetURL(), EHttpRequestStatus::ToString(Request->GetStatus()));
@@ -200,115 +183,3 @@ static FString GetUrlDomain(const FString& Url)
 	}
 	return Domain;
 }
-
-bool FHttpManager::Init()
-{
-	return true;
-}
-
-void FHttpManager::HttpThreadTick(float DeltaSeconds)
-{
-	// empty
-}
-
-bool FHttpManager::StartThreadedRequest(IHttpThreadedRequest* Request)
-{
-	return Request->StartThreadedRequest();
-}
-
-void FHttpManager::CompleteThreadedRequest(IHttpThreadedRequest* Request)
-{
-	// empty
-}
-
-uint32 FHttpManager::Run()
-{
-	double LastTime = FPlatformTime::Seconds();
-	TArray<IHttpThreadedRequest*> RequestsToCancel;
-	TArray<IHttpThreadedRequest*> RequestsToStart;
-	TArray<IHttpThreadedRequest*> RequestsToComplete;
-	while (!ExitRequest.GetValue())
-	{
-		double TickBegin = FPlatformTime::Seconds();
-
-		{
-			FScopeLock ScopeLock(&RequestLock);
-
-			RequestsToCancel = CancelledThreadedRequests;
-			CancelledThreadedRequests.Empty();
-
-			RequestsToStart = PendingThreadedRequests;
-			PendingThreadedRequests.Empty();
-		}
-
-		// Cancel any pending cancel requests
-		for (auto& Request : RequestsToCancel)
-		{
-			if (RunningThreadedRequests.Remove(Request) > 0)
-			{
-				RequestsToComplete.Add(Request);
-			}
-		}
-		// Start any pending requests
-		for (auto& Request : RequestsToStart)
-		{
-			if (StartThreadedRequest(Request))
-			{
-				RunningThreadedRequests.Add(Request);
-			}
-			else
-			{
-				RequestsToComplete.Add(Request);
-			}
-		}
-		const double AppTime = FPlatformTime::Seconds();
-		// Tick any running requests
-		for (int Index = 0; Index < RunningThreadedRequests.Num(); ++Index)
-		{
-			auto& Request = RunningThreadedRequests[Index];
-			Request->TickThreadedRequest(AppTime - LastTime);
-		}
-		HttpThreadTick(AppTime - LastTime);
-		// Move any completed requests
-		for (int Index = 0; Index < RunningThreadedRequests.Num(); ++Index)
-		{
-			auto& Request = RunningThreadedRequests[Index];
-			if (Request->IsThreadedRequestComplete())
-			{
-				RequestsToComplete.Add(Request);
-				RunningThreadedRequests.RemoveAtSwap(Index);
-				--Index;
-			}
-		}
-		LastTime = AppTime;
-
-		if (RequestsToComplete.Num() > 0)
-		{
-			for (auto& Request : RequestsToComplete)
-			{
-				CompleteThreadedRequest(Request);
-			}
-
-			{
-				FScopeLock ScopeLock(&RequestLock);
-				CompletedThreadedRequests.Append(RequestsToComplete);
-			}
-			RequestsToComplete.Empty();
-		}
-
-		double TickDuration = (FPlatformTime::Seconds() - TickBegin);
-		double WaitTime = FMath::Max(0.0, HttpThreadTickBudget - TickDuration);
-		FPlatformProcess::SleepNoStats(WaitTime);
-	}
-	return 0;
-}
-
-void FHttpManager::Stop()
-{
-	ExitRequest.Set(true);
-}
-	
-void FHttpManager::Exit()
-{
-	// empty
-}
diff --git a/Engine/Source/Runtime/Online/HTTP/Private/HttpModule.cpp b/Engine/Source/Runtime/Online/HTTP/Private/HttpModule.cpp
--- a/Engine/Source/Runtime/Online/HTTP/Private/HttpModule.cpp
+++ b/Engine/Source/Runtime/Online/HTTP/Private/HttpModule.cpp
@@ -52,6 +52,7 @@ void FHttpModule::StartupModule()
 		// platform does not provide specific HTTP manager, use generic one
 		HttpManager = new FHttpManager();
 	}
+	HttpManager->Initialize();
 }
 
 void FHttpModule::ShutdownModule()
diff --git a/Engine/Source/Runtime/Online/HTTP/Private/HttpThread.cpp b/Engine/Source/Runtime/Online/HTTP/Private/HttpThread.cpp
new file mode 100644
--- /dev/null
+++ b/Engine/Source/Runtime/Online/HTTP/Private/HttpThread.cpp
@@ -0,0 +1,175 @@
+// Copyright 1998-2016 Epic Games, Inc. All Rights Reserved.
+
+#include "HttpPrivatePCH.h"
+#include "HttpThread.h"
+
+// FHttpThread
+
+FHttpThread::FHttpThread()
+	:	HttpThreadTickBudget(0.033)
+	,	Thread(nullptr)
+{
+	double TargetTickRate = FHttpModule::Get().GetHttpThreadTickRate();
+	if (TargetTickRate > 0.0)
+	{
+		HttpThreadTickBudget = 1.0 / TargetTickRate;
+	}
+
+	UE_LOG(LogHttp, Log, TEXT("HTTP thread tick budget is %.1f ms (tick rate is %f)"), HttpThreadTickBudget, TargetTickRate);
+}
+
+FHttpThread::~FHttpThread()
+{
+	StopThread();
+}
+
+void FHttpThread::StartThread()
+{
+	Thread = FRunnableThread::Create(this, TEXT("HttpManagerThread"), 128 * 1024, TPri_Normal);
+}
+
+void FHttpThread::StopThread()
+{
+	if (Thread != nullptr)
+	{
+		Thread->Kill(true);
+		delete Thread;
+		Thread = nullptr;
+	}
+}
+
+void FHttpThread::AddRequest(IHttpThreadedRequest* Request)
+{
+	FScopeLock ScopeLock(&RequestArraysLock);
+	PendingThreadedRequests.Add(Request);
+}
+
+void FHttpThread::CancelRequest(IHttpThreadedRequest* Request)
+{
+	FScopeLock ScopeLock(&RequestArraysLock);
+	CancelledThreadedRequests.Add(Request);
+}
+
+void FHttpThread::GetCompletedRequests(TArray<IHttpThreadedRequest*>& OutCompletedRequests)
+{
+	FScopeLock ScopeLock(&RequestArraysLock);
+	OutCompletedRequests = CompletedThreadedRequests;
+	CompletedThreadedRequests.Reset();
+}
+
+bool FHttpThread::Init()
+{
+	return true;
+}
+
+void FHttpThread::HttpThreadTick(float DeltaSeconds)
+{
+	// empty
+}
+
+bool FHttpThread::StartThreadedRequest(IHttpThreadedRequest* Request)
+{
+	return Request->StartThreadedRequest();
+}
+
+void FHttpThread::CompleteThreadedRequest(IHttpThreadedRequest* Request)
+{
+	// empty
+}
+
+uint32 FHttpThread::Run()
+{
+	double LastTime = FPlatformTime::Seconds();
+	TArray<IHttpThreadedRequest*> RequestsToCancel;
+	TArray<IHttpThreadedRequest*> RequestsToStart;
+	TArray<IHttpThreadedRequest*> RequestsToComplete;
+	while (!ExitRequest.GetValue())
+	{
+		double TickBegin = FPlatformTime::Seconds();
+
+		{
+			FScopeLock ScopeLock(&RequestArraysLock);
+
+			RequestsToCancel = CancelledThreadedRequests;
+			CancelledThreadedRequests.Reset();
+
+			RequestsToStart = PendingThreadedRequests;
+			PendingThreadedRequests.Reset();
+		}
+
+		// Cancel any pending cancel requests
+		for (IHttpThreadedRequest* Request : RequestsToCancel)
+		{
+			if (RunningThreadedRequests.Remove(Request) > 0)
+			{
+				RequestsToComplete.Add(Request);
+			}
+		}
+		// Start any pending requests
+		for (IHttpThreadedRequest* Request : RequestsToStart)
+		{
+			if (StartThreadedRequest(Request))
+			{
+				RunningThreadedRequests.Add(Request);
+			}
+			else
+			{
+				RequestsToComplete.Add(Request);
+			}
+		}
+
+		const double AppTime = FPlatformTime::Seconds();
+		const double ElapsedTime = AppTime - LastTime;
+		LastTime = AppTime;
+
+		// Tick any running requests
+		for (int32 Index = 0; Index < RunningThreadedRequests.Num(); ++Index)
+		{
+			IHttpThreadedRequest* Request = RunningThreadedRequests[Index];
+			Request->TickThreadedRequest(ElapsedTime);
+		}
+
+		HttpThreadTick(ElapsedTime);
+
+		// Move any completed requests
+		for (int32 Index = 0; Index < RunningThreadedRequests.Num(); ++Index)
+		{
+			IHttpThreadedRequest* Request = RunningThreadedRequests[Index];
+			if (Request->IsThreadedRequestComplete())
+			{
+				RequestsToComplete.Add(Request);
+				RunningThreadedRequests.RemoveAtSwap(Index);
+				--Index;
+			}
+		}
+
+		if (RequestsToComplete.Num() > 0)
+		{
+			for (IHttpThreadedRequest* Request : RequestsToComplete)
+			{
+				CompleteThreadedRequest(Request);
+			}
+
+			{
+				FScopeLock ScopeLock(&RequestArraysLock);
+				CompletedThreadedRequests.Append(RequestsToComplete);
+			}
+			RequestsToComplete.Reset();
+		}
+
+		double TickDuration = (FPlatformTime::Seconds() - TickBegin);
+		double WaitTime = FMath::Max(0.0, HttpThreadTickBudget - TickDuration);
+		FPlatformProcess::SleepNoStats(WaitTime);
+	}
+	return 0;
+}
+
+void FHttpThread::Stop()
+{
+	ExitRequest.Set(true);
+}
+	
+void FHttpThread::Exit()
+{
+	// empty
+}
diff --git a/Engine/Source/Runtime/Online/HTTP/Private/HttpThread.h b/Engine/Source/Runtime/Online/HTTP/Private/HttpThread.h
new file mode 100644
--- /dev/null
+++ b/Engine/Source/Runtime/Online/HTTP/Private/HttpThread.h
@@ -0,0 +1,116 @@
+// Copyright 1998-2016 Epic Games, Inc. All Rights Reserved.
+
+#pragma once
+
+#include "HttpPackage.h"
+
+class IHttpThreadedRequest;
+
+/**
+ * Manages Http thread
+ * Assumes any requests entering the system will remain valid (not deleted) until they exit the system
+ */
+class FHttpThread
+	: FRunnable
+{
+public:
+
+	FHttpThread();
+	virtual ~FHttpThread();
+
+	/** 
+	 * Start the HTTP thread.
+	 */
+	void StartThread();
+
+	/** 
+	 * Stop the HTTP thread.  Blocks until thread has stopped.
+	 */
+	void StopThread();
+
+	/** 
+	 * Add a request to begin processing on HTTP thread.
+	 *
+	 * @param Request the request to be processed on the HTTP thread
+	 */
+	void AddRequest(IHttpThreadedRequest* Request);
+
+	/** 
+	 * Mark a request as cancelled.    Called on non-HTTP thread.
+	 *
+	 * @param Request the request to be processed on the HTTP thread
+	 */
+	void CancelRequest(IHttpThreadedRequest* Request);
+
+	/** 
+	 * Get completed requests.  Clears internal arrays.  Called on non-HTTP thread.
+	 *
+	 * @param OutCompletedRequests array of requests that have been completed
+	 */
+	void GetCompletedRequests(TArray<IHttpThreadedRequest*>& OutCompletedRequests);
+
+protected:
+
+	/**
+	 * Tick on http thread
+	 */
+	virtual void HttpThreadTick(float DeltaSeconds);
+	
+	/** 
+	 * Start processing a request on the http thread
+	 */
+	virtual bool StartThreadedRequest(IHttpThreadedRequest* Request);
+
+	/** 
+	 * Complete a request on the http thread
+	 */
+	virtual void CompleteThreadedRequest(IHttpThreadedRequest* Request);
+
+
+protected:
+	// Threading functions
+
+	//~ Begin FRunnable Interface
+	virtual bool Init() override;
+	virtual uint32 Run() override;
+	virtual void Stop() override;
+	virtual void Exit() override;
+	//~ End FRunnable Interface
+
+	/** signal request to stop and exit thread */
+	FThreadSafeCounter ExitRequest;
+
+	/** Target frame duration of a http thread tick, in seconds. */
+	double	HttpThreadTickBudget;
+
+protected:
+	/** Critical section to lock access to PendingThreadedRequests, CancelledThreadedRequests, and CompletedThreadedRequests */
+	FCriticalSection RequestArraysLock;
+
+	/** 
+	 * Threaded requests that are waiting to be processed on the http thread.
+	 * Added to on non-HTTP thread, processed then cleared on HTTP thread.
+	 */
+	TArray<IHttpThreadedRequest*> PendingThreadedRequests;
+
+	/**
+	 * Threaded requests that are waiting to be cancelled on the http thread.
+	 * Added to on non-HTTP thread, processed then cleared on HTTP thread.
+	 */
+	TArray<IHttpThreadedRequest*> CancelledThreadedRequests;
+
+	/**
+	 * Currently running threaded requests (not in any of the other arrays).
+	 * Only accessed on the HTTP thread.
+	 */
+	TArray<IHttpThreadedRequest*> RunningThreadedRequests;
+
+	/**
+	 * Threaded requests that have completed and are waiting for the game thread to process.
+	 * Added to on HTTP thread, processed then cleared on game thread.
+	 */
+	TArray<IHttpThreadedRequest*> CompletedThreadedRequests;
+
+	/** Pointer to Runnable Thread */
+	FRunnableThread* Thread;
+};
diff --git a/Engine/Source/Runtime/Online/HTTP/Private/Linux/LinuxPlatformHttp.cpp b/Engine/Source/Runtime/Online/HTTP/Private/Linux/LinuxPlatformHttp.cpp
--- a/Engine/Source/Runtime/Online/HTTP/Private/Linux/LinuxPlatformHttp.cpp
+++ b/Engine/Source/Runtime/Online/HTTP/Private/Linux/LinuxPlatformHttp.cpp
@@ -21,6 +21,6 @@ void FLinuxPlatformHttp::Shutdown()
 
 IHttpRequest* FLinuxPlatformHttp::ConstructRequest()
 {
-	return new FCurlHttpRequest(FCurlHttpManager::GShareHandle);
+	return new FCurlHttpRequest();
 }
 
diff --git a/Engine/Source/Runtime/Online/HTTP/Public/HttpManager.h b/Engine/Source/Runtime/Online/HTTP/Public/HttpManager.h
--- a/Engine/Source/Runtime/Online/HTTP/Public/HttpManager.h
+++ b/Engine/Source/Runtime/Online/HTTP/Public/HttpManager.h
@@ -7,13 +7,13 @@
 
 class IHttpRequest;
 class IHttpThreadedRequest;
+class FHttpThread;
 
 /**
  * Manages Http request that are currently being processed
  */
 class FHttpManager
 	: public FTickerObjectBase
-	, FRunnable
 {
 public:
 
@@ -27,7 +27,12 @@ public:
 	/**
 	 * Destructor
 	 */
-	virtual ~FHttpManager();
+	~FHttpManager();
+
+	/**
+	 * Initialize
+	 */
+	void Initialize();
 
 	/**
 	 * Adds an Http request instance to the manager for tracking/ticking
@@ -35,7 +40,7 @@ public:
 	 *
 	 * @param Request - the request object to add
 	 */
-	virtual void AddRequest(const TSharedRef<IHttpRequest>& Request);
+	void AddRequest(const TSharedRef<IHttpRequest>& Request);
 
 	/**
 	 * Removes an Http request instance from the manager
@@ -43,7 +48,7 @@ public:
 	 *
 	 * @param Request - the request object to remove
 	 */
-	virtual void RemoveRequest(const TSharedRef<IHttpRequest>& Request);
+	void RemoveRequest(const TSharedRef<IHttpRequest>& Request);
 
 	/**
 	* Find an Http request in the lists of current valid requests
@@ -52,14 +57,14 @@ public:
 	*
 	* @return true if the request is being tracked, false if not
 	*/
-	virtual bool IsValidRequest(const IHttpRequest* RequestPtr) const;
+	bool IsValidRequest(const IHttpRequest* RequestPtr) const;
 
 	/**
 	 * Block until all pending requests are finished processing
 	 *
 	 * @param bShutdown true if final flush during shutdown
 	 */
-	virtual void Flush(bool bShutdown);
+	void Flush(bool bShutdown);
 
 	/**
 	 * FTicker callback
@@ -68,56 +73,39 @@ public:
 	 *
 	 * @return false if no longer needs ticking
 	 */
-	virtual bool Tick(float DeltaSeconds) override;
+	bool Tick(float DeltaSeconds) override;
 
 	/** 
 	 * Add a http request to be executed on the http thread
 	 *
 	 * @param Request - the request object to add
 	 */
-	virtual void AddThreadedRequest(const TSharedRef<IHttpThreadedRequest>& Request);
+	void AddThreadedRequest(const TSharedRef<IHttpThreadedRequest>& Request);
 
 	/**
 	 * Mark a threaded http request as cancelled to be removed from the http thread
 	 *
 	 * @param Request - the request object to cancel
 	 */
-	virtual void CancelThreadedRequest(const TSharedRef<IHttpThreadedRequest>& Request);
+	void CancelThreadedRequest(const TSharedRef<IHttpThreadedRequest>& Request);
 
 	/**
 	 * List all of the Http requests currently being processed
 	 *
 	 * @param Ar - output device to log with
 	 */
-	virtual void DumpRequests(FOutputDevice& Ar) const;
+	void DumpRequests(FOutputDevice& Ar) const;
 
 protected:
-	// FRunnable
-	virtual bool Init() override;
-	virtual uint32 Run() override;
-	virtual void Stop() override;
-	virtual void Exit() override;
-
-	/** signal request to stop and exit thread */
-	FThreadSafeCounter ExitRequest;
-
-	/**
-	 * Tick on http thread
-	 */
-	virtual void HttpThreadTick(float DeltaSeconds);
-	
 	/** 
-	 * Start processing a request on the http thread
+	 * Create HTTP thread object
+	 *
+	 * @return the HTTP thread object
 	 */
-	virtual bool StartThreadedRequest(IHttpThreadedRequest* Request);
+	virtual FHttpThread* CreateHttpThread();
 
-	/** 
-	 * Complete a request on the http thread
-	 */
-	virtual void CompleteThreadedRequest(IHttpThreadedRequest* Request);
 
 protected:
-	
 	/** List of Http requests that are actively being processed */
 	TArray<TSharedRef<IHttpRequest>> Requests;
 
@@ -142,22 +130,7 @@ protected:
 	/** Dead requests that need to be destroyed */
 	TArray<FRequestPendingDestroy> PendingDestroyRequests;
 
-	/** Threaded requests that have been added on the game thread and are waiting to be processed on the http thread */
-	TArray<IHttpThreadedRequest*> PendingThreadedRequests;
-
-	/** Threaded requests that have been cancelled on the game thread and are waiting to be cancelled on the http thread */
-	TArray<IHttpThreadedRequest*> CancelledThreadedRequests;
-
-	/** Currently running threaded requests */
-	TArray<IHttpThreadedRequest*> RunningThreadedRequests;
-
-	/** Threaded requests completed on the http thread, awaiting processing on the game thread */
-	TArray<IHttpThreadedRequest*> CompletedThreadedRequests;
-
-	/** Target frame duration of a http thread tick, in seconds. */
-	double	HttpThreadTickBudget;
-
-	FRunnableThread* Thread;
+	FHttpThread* Thread;
 
 PACKAGE_SCOPE:
 
diff --git a/Engine/Source/Runtime/Online/OnlineFramework/Party/Private/Party.cpp b/Engine/Source/Runtime/Online/OnlineFramework/Party/Private/Party.cpp
--- a/Engine/Source/Runtime/Online/OnlineFramework/Party/Private/Party.cpp
+++ b/Engine/Source/Runtime/Online/OnlineFramework/Party/Private/Party.cpp
@@ -1027,17 +1027,21 @@ void UParty::LeavePersistentParty(const FUniqueNetId& InUserId, const UPartyDele
 {
 	if (PersistentPartyId.IsValid())
 	{
-		UPartyDelegates::FOnLeaveUPartyComplete CompletionDelegate;
-		CompletionDelegate.BindUObject(this, &ThisClass::OnLeavePersistentPartyComplete, InCompletionDelegate);
-
 		if (!bLeavingPersistentParty)
 		{
+			UE_LOG(LogParty, Verbose, TEXT("LeavePersistentParty"));
+
+			UPartyDelegates::FOnLeaveUPartyComplete CompletionDelegate;
+			CompletionDelegate.BindUObject(this, &ThisClass::OnLeavePersistentPartyComplete, InCompletionDelegate);
+
 			bLeavingPersistentParty = true;
 			LeavePartyInternal(InUserId, IOnlinePartySystem::GetPrimaryPartyTypeId(), CompletionDelegate);
 		}
 		else
 		{
 			LeavePartyCompleteDelegates.Add(InCompletionDelegate);
+
+			UE_LOG(LogParty, Verbose, TEXT("LeavePersistentParty %d extra delegates"), LeavePartyCompleteDelegates.Num());
 		}
 	}
 	else
diff --git a/Engine/Source/Runtime/Online/OnlineFramework/Qos/Private/QosRegionManager.cpp b/Engine/Source/Runtime/Online/OnlineFramework/Qos/Private/QosRegionManager.cpp
--- a/Engine/Source/Runtime/Online/OnlineFramework/Qos/Private/QosRegionManager.cpp
+++ b/Engine/Source/Runtime/Online/OnlineFramework/Qos/Private/QosRegionManager.cpp
@@ -288,6 +288,8 @@ void UQosRegionManager::TrySetDefaultRegion()
 		if (!SetSelectedRegion(BestRegionId))
 		{
 			UE_LOG(LogQos, Warning, TEXT("Unable to set a good region!"));
+			UE_LOG(LogQos, Warning, TEXT("Wanted to set %s, failed to fall back to %s"), *GetRegionId(), *BestRegionId);
+			DumpRegionStats();
 		}
 	}
 }
diff --git a/Engine/Source/Runtime/PakFile/Private/SignedArchiveReader.cpp b/Engine/Source/Runtime/PakFile/Private/SignedArchiveReader.cpp
--- a/Engine/Source/Runtime/PakFile/Private/SignedArchiveReader.cpp
+++ b/Engine/Source/Runtime/PakFile/Private/SignedArchiveReader.cpp
@@ -39,6 +39,9 @@ FChunkCacheWorker::FChunkCacheWorker(FArchive* InReader, const TCHAR* Filename)
 	*SigFileReader << EncryptedSignatures;
 	delete SigFileReader;
 
+	EncryptedSignaturesCRC = FCrc::MemCrc32(&EncryptedSignatures[0], EncryptedSignatures.Num() * sizeof(FEncryptedSignature));
+	DecryptedSignaturesCRC = 0;
+
 	DecryptedSignatures.AddDefaulted(EncryptedSignatures.Num());
 
 	SetupDecryptionKey();
@@ -102,6 +105,7 @@ void FChunkCacheWorker::DecryptSignatures(int32 NextIndexToDecrypt)
 uint32 FChunkCacheWorker::Run()
 {
 	int32 NextIndexToDecrypt = 0;
+	LastDecryptedSignatureIndex = -1;
 	StartTime = FPlatformTime::Seconds();
 	
 	while (StopTaskCounter.GetValue() == 0)
@@ -112,19 +116,19 @@ uint32 FChunkCacheWorker::Run()
 			uint32 WaitTime = 500;
 
 			// Try and decrypt some signatures if there are any left
-			if (EncryptedSignatures.Num() > 0)
+			if (NextIndexToDecrypt < EncryptedSignatures.Num())
 			{
 				check(NextIndexToDecrypt < EncryptedSignatures.Num());
 
 				DecryptSignatures(NextIndexToDecrypt);
+				LastDecryptedSignatureIndex = NextIndexToDecrypt;
 
 				if (++NextIndexToDecrypt == EncryptedSignatures.Num())
 				{
-					// Dump the encrypted signatures as we don't need them anymore
-					EncryptedSignatures.Empty();
-					NextIndexToDecrypt = 0;
 					double Time = FPlatformTime::Seconds() - StartTime;
 					UE_LOG(LogPakFile, Log, TEXT("PakFile signature decryption complete in %.2fs"), Time);
+
+					DecryptedSignaturesCRC = FCrc::MemCrc32(&DecryptedSignatures[0], DecryptedSignatures.Num() * sizeof(FDecryptedSignature));
 				}
 				else
 				{
@@ -257,35 +261,64 @@ bool FChunkCacheWorker::CheckSignature(const FChunkRequest& ChunkInfo)
 	SCOPE_SECONDS_ACCUMULATOR(STAT_FChunkCacheWorker_CheckSignature);
 
 	FDecryptedSignature SourceSignature;
+	const int32 MaxNumRetries = 3;
+	int32 RetriesRemaining = MaxNumRetries;
+	bool bPakChunkSignaturesMatched = false;
 
+	while (!bPakChunkSignaturesMatched && RetriesRemaining >= 0)
 	{
-		SCOPE_SECONDS_ACCUMULATOR(STAT_FChunkCacheWorker_Serialize);
-		Reader->Seek(ChunkInfo.Offset);
-		Reader->Serialize(ChunkInfo.Buffer->Data, ChunkInfo.Size);
-	}
-	{
-		SCOPE_SECONDS_ACCUMULATOR(STAT_FChunkCacheWorker_HashBuffer);
-		SourceSignature.Data = FCrc::MemCrc32(ChunkInfo.Buffer->Data, ChunkInfo.Size);
-	}
+		{
+			SCOPE_SECONDS_ACCUMULATOR(STAT_FChunkCacheWorker_Serialize);
+			Reader->Seek(ChunkInfo.Offset);
+			Reader->Serialize(ChunkInfo.Buffer->Data, ChunkInfo.Size);
+		}
+		{
+			SCOPE_SECONDS_ACCUMULATOR(STAT_FChunkCacheWorker_HashBuffer);
+			SourceSignature.Data = FCrc::MemCrc32(ChunkInfo.Buffer->Data, ChunkInfo.Size);
+		}
 
-	// Decrypt serialized hash
-	FDecryptedSignature& DecryptedSignature = DecryptedSignatures[ChunkInfo.Index];
-	if (!DecryptedSignature.IsValid())
-	{
-		INC_DWORD_STAT(STAT_FChunkCacheWorker_SignaturesDecryptedDuringVerify);
-		SCOPE_SECONDS_ACCUMULATOR(STAT_FChunkCacheWorker_DecryptDuringVerify);
-		FEncryption::DecryptSignature(EncryptedSignatures[ChunkInfo.Index], DecryptedSignature, DecryptionKey);
-		check(DecryptedSignature.IsValid());
+		// Decrypt serialized hash
+		FDecryptedSignature& DecryptedSignature = DecryptedSignatures[ChunkInfo.Index];
+
+		if (RetriesRemaining < MaxNumRetries || !DecryptedSignature.IsValid())
+		{
+			INC_DWORD_STAT(STAT_FChunkCacheWorker_SignaturesDecryptedDuringVerify);
+			SCOPE_SECONDS_ACCUMULATOR(STAT_FChunkCacheWorker_DecryptDuringVerify);
+			FEncryptedSignature EncryptedSignature = EncryptedSignatures[ChunkInfo.Index];
+			check(EncryptedSignature.IsValid());
+			FEncryption::DecryptSignature(EncryptedSignature, DecryptedSignature, DecryptionKey);
+			check(DecryptedSignature.IsValid());
+		}
+
+		// Compare hashes
+		bPakChunkSignaturesMatched = DecryptedSignature == SourceSignature;
+		--RetriesRemaining;
 	}
 
-	// Compare hashes
-	if (!ensure(DecryptedSignature == SourceSignature))
+	if (!ensure(bPakChunkSignaturesMatched))
 	{
+		// Check that the encrypted signatures data is the same as it was when we started up
+		int32 CurrentEncryptedSignaturesCRC = FCrc::MemCrc32(&EncryptedSignatures[0], EncryptedSignatures.Num() * sizeof(FEncryptedSignature));
+		ensure(CurrentEncryptedSignaturesCRC == EncryptedSignaturesCRC);
+
+		// If we finished decrypting all the signatures, check that they are still the same as they were when decryption completed
+		if (DecryptedSignaturesCRC != 0)
+		{
+			int32 CurrentDecryptedSignaturesCRC = FCrc::MemCrc32(&DecryptedSignatures[0], DecryptedSignatures.Num() * sizeof(FDecryptedSignature));
+			ensure(CurrentDecryptedSignaturesCRC == DecryptedSignaturesCRC);
+		}
+
+		UE_LOG(LogPakFile, Warning, TEXT("Pak chunk signature verification failed!"));
+		UE_LOG(LogPakFile, Warning, TEXT("  Chunk Index: %d"), ChunkInfo.Index);
+		UE_LOG(LogPakFile, Warning, TEXT("  Chunk Offset: %d"), ChunkInfo.Offset);
+		UE_LOG(LogPakFile, Warning, TEXT("  Chunk Size: %d"), ChunkInfo.Size);
+		UE_LOG(LogPakFile, Warning, TEXT("  Background decrypt: %d\\%d"), LastDecryptedSignatureIndex, DecryptedSignatures.Num());
+
 		FPlatformMisc::MessageBoxExt(EAppMsgType::Ok, TEXT("Corrupt Installation Detected. Please use \"Verify\" in the Epic Games Launcher"), TEXT("Pakfile Error"));
 		FPlatformMisc::RequestExit(1);
 	}
-
-	return true;
+	
+	return bPakChunkSignaturesMatched;
 }
 
 FChunkRequest& FChunkCacheWorker::RequestChunk(int32 ChunkIndex, int64 StartOffset, int64 ChunkSize)
diff --git a/Engine/Source/Runtime/PakFile/Private/SignedArchiveReader.h b/Engine/Source/Runtime/PakFile/Private/SignedArchiveReader.h
--- a/Engine/Source/Runtime/PakFile/Private/SignedArchiveReader.h
+++ b/Engine/Source/Runtime/PakFile/Private/SignedArchiveReader.h
@@ -124,6 +124,12 @@ class FChunkCacheWorker : public FRunnable
 	FEncryptionKey DecryptionKey;
 	/** The time at which the chunk cache worker started, for debugging purposes */
 	double StartTime;
+	/** Track the index of the last signature decrypted in the background, for debugging purposes */
+	int32 LastDecryptedSignatureIndex;
+	/** Stored CRC for the encrypted signature data after loading */
+	uint32 EncryptedSignaturesCRC;
+	/** Stored CRC for the fully dencrypted signature data after loading. Only computed when all signatures were decrypted */
+	uint32 DecryptedSignaturesCRC;
 
 	/** 
 	 * Process requested chunks 
diff --git a/Engine/Source/Runtime/PakFile/Public/IPlatformFilePak.h b/Engine/Source/Runtime/PakFile/Public/IPlatformFilePak.h
--- a/Engine/Source/Runtime/PakFile/Public/IPlatformFilePak.h
+++ b/Engine/Source/Runtime/PakFile/Public/IPlatformFilePak.h
@@ -1040,7 +1040,7 @@ public:
 			return FileEntry->CompressionMethod != COMPRESS_None ? FileEntry->UncompressedSize : FileEntry->Size;
 		}
 		// First look for the file in the user dir.
-		int64 Result = 0;
+		int64 Result = INDEX_NONE;
 		if (IsNonPakFilenameAllowed(Filename))
 		{
 			Result = LowerLevel->FileSize(Filename);
