Fixing a bug in BuildPatchTool which could cause the loss of chunk information, if a new chunk generated had a hash collision with a previously existing chunk.
#change Use correct method of merging two TMaps which have TSet value types.
#remove Overly verbose temporary debugging code.

diff --git a/Engine/Source/Runtime/Online/BuildPatchServices/Private/BuildPatchGeneration.cpp b/Engine/Source/Runtime/Online/BuildPatchServices/Private/BuildPatchGeneration.cpp
--- a/Engine/Source/Runtime/Online/BuildPatchServices/Private/BuildPatchGeneration.cpp
+++ b/Engine/Source/Runtime/Online/BuildPatchServices/Private/BuildPatchGeneration.cpp
@@ -554,7 +554,11 @@ bool FBuildDataGenerator::GenerateChunksManifestFromDirectory(const FBuildPatchS
 	TMap<FGuid, FChunkInfoData> ChunkInfoMap;
 	TMap<FGuid, int64> ChunkFileSizes = CloudEnumeration->GetChunkFileSizes();
 	ChunkWriter.GetChunkFilesizes(ChunkFileSizes);
-	ChunkInventory.Append(CloudEnumeration->GetChunkInventory());
+	for (const TPair<uint64, TSet<FGuid>>& ChunkInventoryPair : CloudEnumeration->GetChunkInventory())
+	{
+		TSet<FGuid>& ChunkSet = ChunkInventory.FindOrAdd(ChunkInventoryPair.Key);
+		ChunkSet = ChunkSet.Union(ChunkInventoryPair.Value);
+	}
 	ChunkShaHashes.Append(CloudEnumeration->GetChunkShaHashes());
 	bool bFoundAllChunkInfo = true;
 	for (const TPair<uint64, TSet<FGuid>>& ChunkInventoryPair : ChunkInventory)
diff --git a/Engine/Source/Runtime/Online/BuildPatchServices/Private/Generation/ManifestBuilder.cpp b/Engine/Source/Runtime/Online/BuildPatchServices/Private/Generation/ManifestBuilder.cpp
--- a/Engine/Source/Runtime/Online/BuildPatchServices/Private/Generation/ManifestBuilder.cpp
+++ b/Engine/Source/Runtime/Online/BuildPatchServices/Private/Generation/ManifestBuilder.cpp
@@ -112,53 +112,33 @@ namespace BuildPatchServices
 			FileManifest.Init();
 			check(FileManifest.GetFileSize() == FileSpan.Size);
 		}
-		UE_LOG(LogManifestBuilder, Log, TEXT("Manifest references %d chunks."), ReferencedChunks.Num());
+		UE_LOG(LogManifestBuilder, Verbose, TEXT("Manifest references %d chunks."), ReferencedChunks.Num());
 
 		// Setup chunk list, removing all that were not referenced.
 		Manifest->Data->ChunkList = MoveTemp(ChunkInfo);
 		int32 TotalChunkListNum = Manifest->Data->ChunkList.Num();
 		Manifest->Data->ChunkList.RemoveAll([&](FChunkInfoData& Candidate){ return ReferencedChunks.Contains(Candidate.Guid) == false; });
-		UE_LOG(LogManifestBuilder, Log, TEXT("Chunk info list trimmed from %d to %d."), TotalChunkListNum, Manifest->Data->ChunkList.Num());
+		UE_LOG(LogManifestBuilder, Verbose, TEXT("Chunk info list trimmed from %d to %d."), TotalChunkListNum, Manifest->Data->ChunkList.Num());
 
-		// Sanity check that all referenced chunks exist in the ChunkList data, in two different ways - we are trying to track down an issue.
-		// Start by making a set of chunks in ChunkList to ensure the checks are speedy enough.
-		TSet<FGuid> ChunkListReferenceSet;
-		for (const FChunkInfoData& ChunkListElement : Manifest->Data->ChunkList)
-		{
-			ChunkListReferenceSet.Add(ChunkListElement.Guid);
-		}
-		// Check all values in ReferencedChunks have entries in ChunkList.
-		bool bReferencedChunksInChunkList = true;
-		for (const FGuid& ReferencedChunkGuid : ReferencedChunks)
-		{
-			if (ChunkListReferenceSet.Contains(ReferencedChunkGuid) == false)
-			{
-				UE_LOG(LogManifestBuilder, Error, TEXT("Missing chunk info for referenced chunk %s."), *ReferencedChunkGuid.ToString());
-				bReferencedChunksInChunkList = false;
-			}
-		}
-		// Check all parts of file manifests have their chunk info in ChunkList.
-		bool bFilePartsInChunkList = true;
-		for (const FFileManifestData& FileManifest : Manifest->Data->FileManifestList)
+		// Init the manifest, and we are done.
+		Manifest->InitLookups();
+
+		// Sanity check all chunk info was provided
+		bool bHasAllInfo = true;
+		for (const FGuid& ReferencedChunk : ReferencedChunks)
 		{
-			for (const FChunkPartData& FileChunkPart : FileManifest.FileChunkParts)
+			uint64 ChunkHash;
+			if(Manifest->GetChunkHash(ReferencedChunk, ChunkHash) == false)
 			{
-				if (ChunkListReferenceSet.Contains(FileChunkPart.Guid) == false)
-				{
-					UE_LOG(LogManifestBuilder, Error, TEXT("Missing chunk info for file part chunk %s."), *FileChunkPart.Guid.ToString());
-					bFilePartsInChunkList = false;
-				}
+				UE_LOG(LogManifestBuilder, Error, TEXT("Generated manifest is missing ChunkInfo for chunk %s."), *ReferencedChunk.ToString());
+				bHasAllInfo = false;
 			}
 		}
-		// Fail if any missing chunk info.
-		if (!(bReferencedChunksInChunkList && bFilePartsInChunkList))
+		if (bHasAllInfo == false)
 		{
 			return false;
 		}
 
-		// Init the manifest, and we are done.
-		Manifest->InitLookups();
-
 		// Some sanity checks for build integrity.
 		if (BuildStructureAdded.GetHead() == nullptr || BuildStructureAdded.GetHead()->GetNext() != nullptr)
 		{
@@ -179,7 +159,7 @@ namespace BuildPatchServices
 	{
 		// Previous validation from FinaliseData, but this time we assert if fail as the error should have been picked up.
 		checkf(BuildStructureAdded.GetHead() != nullptr && BuildStructureAdded.GetHead()->GetNext() == nullptr, TEXT("Build integrity check failed. No structure was added."));
-		checkf(BuildStructureAdded.GetHead()->GetSize() == Manifest->GetBuildSize(), TEXT("Build integrity check failed. Structure added is not the same size and manifest data setup, did you call FinalizeData?"));
+		checkf(BuildStructureAdded.GetHead()->GetSize() == Manifest->GetBuildSize(), TEXT("Build integrity check failed. Structure added is not the same size as the manifest data setup; did you call FinalizeData?"));
 
 		// Currently we only save out in JSON format
 		Manifest->Data->ManifestFileVersion = EBuildPatchAppManifestVersion::GetLatestJsonVersion();
