Fixes for merge from //UE4/Portal-Staging.

diff --git a/Engine/Source/Runtime/Online/WebSockets/Private/Lws/LwsWebSocket.cpp b/Engine/Source/Runtime/Online/WebSockets/Private/Lws/LwsWebSocket.cpp
--- a/Engine/Source/Runtime/Online/WebSockets/Private/Lws/LwsWebSocket.cpp
+++ b/Engine/Source/Runtime/Online/WebSockets/Private/Lws/LwsWebSocket.cpp
@@ -1,3 +1,7 @@
+// Copyright 1998-2016 Epic Games, Inc. All Rights Reserved.
+
+#if WITH_WEBSOCKETS
+
 #include "WebSocketsPrivatePCH.h"
 #include "LwsWebSocket.h"
 #include "Ssl.h"
@@ -380,5 +384,4 @@ int FLwsWebSocket::CallbackWrapper(struct lws *Instance, enum lws_callback_reaso
 	return 0;
 }
 
-
-
+#endif // #if WITH_WEBSOCKETS
\ No newline at end of file
diff --git a/Engine/Source/Runtime/Online/WebSockets/Private/WebSocketsModule.cpp b/Engine/Source/Runtime/Online/WebSockets/Private/WebSocketsModule.cpp
--- a/Engine/Source/Runtime/Online/WebSockets/Private/WebSocketsModule.cpp
+++ b/Engine/Source/Runtime/Online/WebSockets/Private/WebSocketsModule.cpp
@@ -2,7 +2,10 @@
 
 #include "WebSocketsPrivatePCH.h"
 #include "WebSocketsModule.h"
+
+#if WITH_WEBSOCKETS
 #include "PlatformWebSocket.h"
+#endif // #if WITH_WEBSOCKETS
 
 DEFINE_LOG_CATEGORY(LogWebSockets);
 
@@ -27,6 +30,7 @@ FWebSocketsModule& FWebSocketsModule::Get()
 	return *Singleton;
 }
 
+#if WITH_WEBSOCKETS
 TSharedRef<IWebSocket> FWebSocketsModule::CreateWebSocket(const FString& Url, const TArray<FString>& Protocols)
 {
 	return MakeShareable(new FPlatformWebSocket(Url, Protocols));
@@ -38,3 +42,4 @@ TSharedRef<IWebSocket> FWebSocketsModule::CreateWebSocket(const FString& Url, co
 	Protocols.Add(Protocol);
 	return MakeShareable(new FPlatformWebSocket(Url, Protocols));
 }
+#endif // #if WITH_WEBSOCKETS
diff --git a/Engine/Source/Runtime/Online/WebSockets/Private/WebSocketsPrivatePCH.h b/Engine/Source/Runtime/Online/WebSockets/Private/WebSocketsPrivatePCH.h
--- a/Engine/Source/Runtime/Online/WebSockets/Private/WebSocketsPrivatePCH.h
+++ b/Engine/Source/Runtime/Online/WebSockets/Private/WebSocketsPrivatePCH.h
@@ -2,8 +2,12 @@
 #pragma once
 
 #include "Core.h"
+
 #include "WebSocketsModule.h"
+
+#if WITH_WEBSOCKETS
 #include "IWebSocket.h"
+#endif // #if WITH_WEBSOCKETS
 
 DECLARE_LOG_CATEGORY_EXTERN(LogWebSockets, Warning, All);
 
diff --git a/Engine/Source/Runtime/Online/WebSockets/Public/WebSocketsModule.h b/Engine/Source/Runtime/Online/WebSockets/Public/WebSocketsModule.h
--- a/Engine/Source/Runtime/Online/WebSockets/Public/WebSocketsModule.h
+++ b/Engine/Source/Runtime/Online/WebSockets/Public/WebSocketsModule.h
@@ -25,6 +25,7 @@ public:
 	 */
 	static FWebSocketsModule& Get();
 
+#if WITH_WEBSOCKETS
 	/**
 	 * Instantiates a new web socket for the current platform
 	 *
@@ -43,6 +44,7 @@ public:
 	 * @return new IWebSocket instance
 	 */
 	virtual TSharedRef<IWebSocket> CreateWebSocket(const FString& Url, const FString& Protocol = FString());
+#endif // #if WITH_WEBSOCKETS
 
 private:
 
diff --git a/Engine/Source/Runtime/Online/WebSockets/WebSockets.Build.cs b/Engine/Source/Runtime/Online/WebSockets/WebSockets.Build.cs
--- a/Engine/Source/Runtime/Online/WebSockets/WebSockets.Build.cs
+++ b/Engine/Source/Runtime/Online/WebSockets/WebSockets.Build.cs
@@ -8,25 +8,32 @@ public class WebSockets : ModuleRules
     {
         Definitions.Add("WEBSOCKETS_PACKAGE=1");
 
-        PrivateIncludePaths.AddRange(
-            new string[] 
-			{
-				"Runtime/Online/WebSockets/Private",
-			}
-        );
-
-        PrivateDependencyModuleNames.AddRange(
-			new string[] { 
+		PrivateDependencyModuleNames.AddRange(
+			new string[] {
 				"Core",
 			}
 		);
 
-        if (Target.Platform == UnrealTargetPlatform.Win32 ||
-            Target.Platform == UnrealTargetPlatform.Win64 ||
-            Target.Platform == UnrealTargetPlatform.Mac
+		bool bShouldUseModule = false;
+		if (Target.Platform == UnrealTargetPlatform.Win32 ||
+			Target.Platform == UnrealTargetPlatform.Win64 ||
+			Target.Platform == UnrealTargetPlatform.Mac
 			)
-        {
-            AddEngineThirdPartyPrivateStaticDependencies(Target, "OpenSSL", "libWebSockets", "zlib");
+		{
+			bShouldUseModule = true;
+		}
+
+		if (bShouldUseModule)
+		{
+			Definitions.Add("WITH_WEBSOCKETS=1");
+
+			PrivateIncludePaths.AddRange(
+				new string[] {
+					"Runtime/Online/WebSockets/Private",
+				}
+			);
+
+			AddEngineThirdPartyPrivateStaticDependencies(Target, "OpenSSL", "libWebSockets", "zlib");
 			PrivateDependencyModuleNames.Add("SSL");
 		}
 	}
