Fix Mac non-unity failures due to #pragma once in cpp file.

diff --git a/Engine/Source/Runtime/Online/BuildPatchServices/Private/BuildPatchState.cpp b/Engine/Source/Runtime/Online/BuildPatchServices/Private/BuildPatchState.cpp
--- a/Engine/Source/Runtime/Online/BuildPatchServices/Private/BuildPatchState.cpp
+++ b/Engine/Source/Runtime/Online/BuildPatchServices/Private/BuildPatchState.cpp
@@ -1,5 +1,4 @@
 // Copyright 1998-2016 Epic Games, Inc. All Rights Reserved.
-#pragma once
 
 #include "BuildPatchServicesPrivatePCH.h"
 #include "BuildPatchState.h"
