Copying //UE4/Portal-Staging to //UE4/Dev-Main (Source: //Portal/Main @ 3216504)

diff --git a/Engine/Source/Runtime/Online/BuildPatchServices/Private/BuildPatchManifest.cpp b/Engine/Source/Runtime/Online/BuildPatchServices/Private/BuildPatchManifest.cpp
--- a/Engine/Source/Runtime/Online/BuildPatchServices/Private/BuildPatchManifest.cpp
+++ b/Engine/Source/Runtime/Online/BuildPatchServices/Private/BuildPatchManifest.cpp
@@ -546,49 +546,69 @@ namespace BuildPatchAppManifest
 *****************************************************************************/
 
 FBuildPatchAppManifest::FBuildPatchAppManifest()
-	: Data(nullptr) // This MUST be in the initializer list, due to FGCObject inheritance.
+	: ManifestFileVersion(EBuildPatchAppManifestVersion::Invalid)
+	, bIsFileData(false)
+	, AppID(INDEX_NONE)
+	, AppName(TEXT(""))
+	, BuildVersion(TEXT(""))
+	, LaunchExe(TEXT(""))
+	, LaunchCommand(TEXT(""))
+	, PrereqName(TEXT(""))
+	, PrereqPath(TEXT(""))
+	, PrereqArgs(TEXT(""))
+	, FileManifestList()
+	, ChunkList()
+	, CustomFields()
 	, TotalBuildSize(INDEX_NONE)
 	, TotalDownloadSize(INDEX_NONE)
 	, bNeedsResaving(false)
 {
-	Data = NewObject<UBuildPatchManifest>();
 }
 
-FBuildPatchAppManifest::FBuildPatchAppManifest(const uint32& InAppID, const FString& AppName)
-	: Data(nullptr) // This MUST be in the initializer list, due to FGCObject inheritance.
+FBuildPatchAppManifest::FBuildPatchAppManifest(const uint32& InAppID, const FString& InAppName)
+	: ManifestFileVersion(EBuildPatchAppManifestVersion::Invalid)
+	, bIsFileData(false)
+	, AppID(InAppID)
+	, AppName(InAppName)
+	, BuildVersion(TEXT(""))
+	, LaunchExe(TEXT(""))
+	, LaunchCommand(TEXT(""))
+	, PrereqName(TEXT(""))
+	, PrereqPath(TEXT(""))
+	, PrereqArgs(TEXT(""))
+	, FileManifestList()
+	, ChunkList()
+	, CustomFields()
 	, TotalBuildSize(INDEX_NONE)
 	, TotalDownloadSize(INDEX_NONE)
 	, bNeedsResaving(false)
 {
-	Data = NewObject<UBuildPatchManifest>();
-	Data->AppID = InAppID;
-	Data->AppName = AppName;
 }
 
 FBuildPatchAppManifest::FBuildPatchAppManifest(const FBuildPatchAppManifest& Other)
-	: Data(nullptr) // This MUST be in the initializer list, due to FGCObject inheritance.
-{
-	Data = NewObject<UBuildPatchManifest>();
-	Data->ManifestFileVersion = Other.Data->ManifestFileVersion;
-	Data->bIsFileData = Other.Data->bIsFileData;
-	Data->AppID = Other.Data->AppID;
-	Data->AppName = Other.Data->AppName;
-	Data->BuildVersion = Other.Data->BuildVersion;
-	Data->LaunchExe = Other.Data->LaunchExe;
-	Data->LaunchCommand = Other.Data->LaunchCommand;
-	Data->PrereqName = Other.Data->PrereqName;
-	Data->PrereqPath = Other.Data->PrereqPath;
-	Data->PrereqArgs = Other.Data->PrereqArgs;
-	Data->FileManifestList = Other.Data->FileManifestList;
-	Data->ChunkList = Other.Data->ChunkList;
-	Data->CustomFields = Other.Data->CustomFields;
+	: ManifestFileVersion(Other.ManifestFileVersion)
+	, bIsFileData(Other.bIsFileData)
+	, AppID(Other.AppID)
+	, AppName(Other.AppName)
+	, BuildVersion(Other.BuildVersion)
+	, LaunchExe(Other.LaunchExe)
+	, LaunchCommand(Other.LaunchCommand)
+	, PrereqName(Other.PrereqName)
+	, PrereqPath(Other.PrereqPath)
+	, PrereqArgs(Other.PrereqArgs)
+	, FileManifestList(Other.FileManifestList)
+	, ChunkList(Other.ChunkList)
+	, CustomFields(Other.CustomFields)
+	, TotalBuildSize(INDEX_NONE)
+	, TotalDownloadSize(INDEX_NONE)
+	, bNeedsResaving(Other.bNeedsResaving)
+{
 	InitLookups();
-	bNeedsResaving = Other.bNeedsResaving;
 }
 
 FBuildPatchAppManifest::~FBuildPatchAppManifest()
 {
-	Data = nullptr;
+	DestroyData();
 }
 
 bool FBuildPatchAppManifest::SaveToFile(const FString& Filename, bool bUseBinary)
@@ -720,6 +740,9 @@ bool FBuildPatchAppManifest::DeserializeFromData(const TArray<uint8>& DataInput)
 
 bool FBuildPatchAppManifest::Serialize(FArchive& Ar)
 {
+	UBuildPatchManifest* Data = NewObject<UBuildPatchManifest>();
+	Data->AddToRoot();
+
 	// Make sure we use the correct serialization version, this is now fixed and must never use a newer version,
 	// because the property tag has changed in structure meaning older clients would not read correctly.
 	Ar.SetUE4Ver(VER_UE4_STRUCT_GUID_IN_PROPERTY_TAG - 1);
@@ -728,29 +751,73 @@ bool FBuildPatchAppManifest::Serialize(FArchive& Ar)
 	{
 		DestroyData();
 	}
+	else
+	{
+		Data->ManifestFileVersion = ManifestFileVersion;
+		Data->bIsFileData = bIsFileData;
+		Data->AppID = AppID;
+		Data->AppName = AppName;
+		Data->BuildVersion = BuildVersion;
+		Data->LaunchExe = LaunchExe;
+		Data->LaunchCommand = LaunchCommand;
+		Data->PrereqName = PrereqName;
+		Data->PrereqPath = PrereqPath;
+		Data->PrereqArgs = PrereqArgs;
+		Data->FileManifestList = FileManifestList;
+		Data->ChunkList = ChunkList;
+		Data->CustomFields = CustomFields;
+	}
 
 	Data->Serialize(Ar);
 
 	if (Ar.IsLoading())
 	{
+		ManifestFileVersion = MoveTemp(Data->ManifestFileVersion);
+		bIsFileData = MoveTemp(Data->bIsFileData);
+		AppID = MoveTemp(Data->AppID);
+		AppName = MoveTemp(Data->AppName);
+		BuildVersion = MoveTemp(Data->BuildVersion);
+		LaunchExe = MoveTemp(Data->LaunchExe);
+		LaunchCommand = MoveTemp(Data->LaunchCommand);
+		PrereqName = MoveTemp(Data->PrereqName);
+		PrereqPath = MoveTemp(Data->PrereqPath);
+		PrereqArgs = MoveTemp(Data->PrereqArgs);
+		FileManifestList = MoveTemp(Data->FileManifestList);
+		ChunkList = MoveTemp(Data->ChunkList);
+		CustomFields = MoveTemp(Data->CustomFields);
+
 		// If we didn't load the version number, we know it was skipped when saving therefore must be
 		// the first UObject version
-		if (Data->ManifestFileVersion == static_cast<uint8>(EBuildPatchAppManifestVersion::Invalid))
+		if (ManifestFileVersion == static_cast<uint8>(EBuildPatchAppManifestVersion::Invalid))
 		{
-			Data->ManifestFileVersion = EBuildPatchAppManifestVersion::StoredAsCompressedUClass;
+			ManifestFileVersion = EBuildPatchAppManifestVersion::StoredAsCompressedUClass;
 		}
 
 		// Setup internal lookups
 		InitLookups();
 	}
 
+	Data->Clear();
+	Data->RemoveFromRoot();
 	return !Ar.IsError();
 }
 
 void FBuildPatchAppManifest::DestroyData()
 {
 	// Clear Manifest data
-	Data->Clear();
+	ManifestFileVersion = EBuildPatchAppManifestVersion::Invalid;
+	bIsFileData = false;
+	AppID = INDEX_NONE;
+	AppName.Empty();
+	BuildVersion.Empty();
+	LaunchExe.Empty();
+	LaunchCommand.Empty();
+	PrereqName.Empty();
+	PrereqPath.Empty();
+	PrereqArgs.Empty();
+	FileManifestList.Empty();
+	ChunkList.Empty();
+	CustomFields.Empty();
 	FileNameLookup.Empty();
 	FileManifestLookup.Empty();
 	ChunkInfoLookup.Empty();
@@ -763,19 +830,19 @@ void FBuildPatchAppManifest::DestroyData()
 void FBuildPatchAppManifest::InitLookups()
 {
 	// Make sure file list is sorted
-	Data->FileManifestList.Sort();
+	FileManifestList.Sort();
 
 	// Setup internals
 	TotalBuildSize = 0;
-	FileManifestLookup.Empty(Data->FileManifestList.Num());
+	FileManifestLookup.Empty(FileManifestList.Num());
 	TaggedFilesLookup.Empty();
-	FileNameLookup.Empty(Data->bIsFileData ? Data->FileManifestList.Num() : 0);
-	for (auto& File : Data->FileManifestList)
+	FileNameLookup.Empty(bIsFileData ? FileManifestList.Num() : 0);
+	for (FFileManifestData& File : FileManifestList)
 	{
 		File.Init();
 		TotalBuildSize += File.GetFileSize();
 		FileManifestLookup.Add(File.Filename, &File);
-		if (Data->bIsFileData)
+		if (bIsFileData)
 		{
 			// File data chunk parts should have been checked already
 			FileNameLookup.Add(File.FileChunkParts[0].Guid, &File.Filename);
@@ -788,21 +855,21 @@ void FBuildPatchAppManifest::InitLookups()
 		else
 		{
 			// Fill out lookup for optional files
-			for (auto& FileTag : File.InstallTags)
+			for (const FString& FileTag : File.InstallTags)
 			{
 				TaggedFilesLookup.FindOrAdd(FileTag).Add(&File);
 			}
 		}
 	}
 	TotalDownloadSize = 0;
-	ChunkInfoLookup.Empty(Data->ChunkList.Num());
-	for (auto& Chunk : Data->ChunkList)
+	ChunkInfoLookup.Empty(ChunkList.Num());
+	for (FChunkInfoData& Chunk : ChunkList)
 	{
 		ChunkInfoLookup.Add(Chunk.Guid, &Chunk);
 		TotalDownloadSize += Chunk.FileSize;
 	}
-	CustomFieldLookup.Empty(Data->CustomFields.Num());
-	for (auto& CustomField : Data->CustomFields)
+	CustomFieldLookup.Empty(CustomFields.Num());
+	for (FCustomFieldData& CustomField : CustomFields)
 	{
 		CustomFieldLookup.Add(CustomField.Key, &CustomField);
 	}
@@ -819,19 +886,19 @@ void FBuildPatchAppManifest::SerializeToJSON(FString& JSONOutput)
 	Writer->WriteObjectStart();
 	{
 		// Write general data
-		Writer->WriteValue(TEXT("ManifestFileVersion"), ToStringBlob(static_cast<int32>(Data->ManifestFileVersion)));
-		Writer->WriteValue(TEXT("bIsFileData"), Data->bIsFileData);
-		Writer->WriteValue(TEXT("AppID"), ToStringBlob(Data->AppID));
-		Writer->WriteValue(TEXT("AppNameString"), Data->AppName);
-		Writer->WriteValue(TEXT("BuildVersionString"), Data->BuildVersion);
-		Writer->WriteValue(TEXT("LaunchExeString"), Data->LaunchExe);
-		Writer->WriteValue(TEXT("LaunchCommand"), Data->LaunchCommand);
-		Writer->WriteValue(TEXT("PrereqName"), Data->PrereqName);
-		Writer->WriteValue(TEXT("PrereqPath"), Data->PrereqPath);
-		Writer->WriteValue(TEXT("PrereqArgs"), Data->PrereqArgs);
+		Writer->WriteValue(TEXT("ManifestFileVersion"), ToStringBlob(static_cast<int32>(ManifestFileVersion)));
+		Writer->WriteValue(TEXT("bIsFileData"), bIsFileData);
+		Writer->WriteValue(TEXT("AppID"), ToStringBlob(AppID));
+		Writer->WriteValue(TEXT("AppNameString"), AppName);
+		Writer->WriteValue(TEXT("BuildVersionString"), BuildVersion);
+		Writer->WriteValue(TEXT("LaunchExeString"), LaunchExe);
+		Writer->WriteValue(TEXT("LaunchCommand"), LaunchCommand);
+		Writer->WriteValue(TEXT("PrereqName"), PrereqName);
+		Writer->WriteValue(TEXT("PrereqPath"), PrereqPath);
+		Writer->WriteValue(TEXT("PrereqArgs"), PrereqArgs);
 		// Write file manifest data
 		Writer->WriteArrayStart(TEXT("FileManifestList"));
-		for (const auto& FileManifest : Data->FileManifestList)
+		for (const FFileManifestData& FileManifest : FileManifestList)
 		{
 			Writer->WriteObjectStart();
 			{
@@ -858,7 +925,7 @@ void FBuildPatchAppManifest::SerializeToJSON(FString& JSONOutput)
 				{
 					Writer->WriteArrayStart(TEXT("FileChunkParts"));
 					{
-						for (const auto& FileChunkPart : FileManifest.FileChunkParts)
+						for (const FChunkPartData& FileChunkPart : FileManifest.FileChunkParts)
 						{
 							Writer->WriteObjectStart();
 							{
@@ -875,7 +942,7 @@ void FBuildPatchAppManifest::SerializeToJSON(FString& JSONOutput)
 				{
 					Writer->WriteArrayStart(TEXT("InstallTags"));
 					{
-						for (const auto& InstallTag : FileManifest.InstallTags)
+						for (const FString& InstallTag : FileManifest.InstallTags)
 						{
 							Writer->WriteValue(InstallTag);
 						}
@@ -888,7 +955,7 @@ void FBuildPatchAppManifest::SerializeToJSON(FString& JSONOutput)
 		Writer->WriteArrayEnd();
 		// Write chunk hash list
 		Writer->WriteObjectStart(TEXT("ChunkHashList"));
-		for (const auto& ChunkInfo : Data->ChunkList)
+		for (const FChunkInfoData& ChunkInfo : ChunkList)
 		{
 			const FGuid& ChunkGuid = ChunkInfo.Guid;
 			const uint64& ChunkHash = ChunkInfo.Hash;
@@ -897,7 +964,7 @@ void FBuildPatchAppManifest::SerializeToJSON(FString& JSONOutput)
 		Writer->WriteObjectEnd();
 		// Write chunk sha list
 		Writer->WriteObjectStart(TEXT("ChunkShaList"));
-		for (const auto& ChunkInfo : Data->ChunkList)
+		for (const FChunkInfoData& ChunkInfo : ChunkList)
 		{
 			const FGuid& ChunkGuid = ChunkInfo.Guid;
 			const FSHAHashData& ChunkSha = ChunkInfo.ShaHash;
@@ -906,7 +973,7 @@ void FBuildPatchAppManifest::SerializeToJSON(FString& JSONOutput)
 		Writer->WriteObjectEnd();
 		// Write data group list
 		Writer->WriteObjectStart(TEXT("DataGroupList"));
-		for (const auto& ChunkInfo : Data->ChunkList)
+		for (const FChunkInfoData& ChunkInfo : ChunkList)
 		{
 			const FGuid& DataGuid = ChunkInfo.Guid;
 			const uint8& DataGroup = ChunkInfo.GroupNumber;
@@ -915,7 +982,7 @@ void FBuildPatchAppManifest::SerializeToJSON(FString& JSONOutput)
 		Writer->WriteObjectEnd();
 		// Write chunk size list
 		Writer->WriteObjectStart(TEXT("ChunkFilesizeList"));
-		for (const auto& ChunkInfo : Data->ChunkList)
+		for (const FChunkInfoData& ChunkInfo : ChunkList)
 		{
 			const FGuid& ChunkGuid = ChunkInfo.Guid;
 			const int64& ChunkSize = ChunkInfo.FileSize;
@@ -924,7 +991,7 @@ void FBuildPatchAppManifest::SerializeToJSON(FString& JSONOutput)
 		Writer->WriteObjectEnd();
 		// Write custom fields
 		Writer->WriteObjectStart(TEXT("CustomFields"));
-		for (const auto& CustomField : Data->CustomFields)
+		for (const FCustomFieldData& CustomField : CustomFields)
 		{
 			Writer->WriteValue(CustomField.Key, CustomField.Value);
 		}
@@ -962,12 +1029,12 @@ bool FBuildPatchAppManifest::DeserializeFromJSON( const FString& JSONInput )
 	TSharedPtr<FJsonValue> JsonManifestFileVersion = JsonValueMap.FindRef(TEXT("ManifestFileVersion"));
 	if (JsonManifestFileVersion.IsValid() && FromStringBlob(JsonManifestFileVersion->AsString(), ManifestFileVersionInt))
 	{
-		Data->ManifestFileVersion = static_cast<EBuildPatchAppManifestVersion::Type>(ManifestFileVersionInt);
+		ManifestFileVersion = static_cast<EBuildPatchAppManifestVersion::Type>(ManifestFileVersionInt);
 	}
 	else
 	{
 		// Then we presume version just before we started outputting the version
-		Data->ManifestFileVersion = static_cast<EBuildPatchAppManifestVersion::Type>(EBuildPatchAppManifestVersion::StartStoringVersion - 1);
+		ManifestFileVersion = static_cast<EBuildPatchAppManifestVersion::Type>(EBuildPatchAppManifestVersion::StartStoringVersion - 1);
 	}
 
 	// Get the app and version strings
@@ -982,33 +1049,33 @@ bool FBuildPatchAppManifest::DeserializeFromJSON( const FString& JSONInput )
 	bSuccess = bSuccess && JsonAppID.IsValid();
 	if( bSuccess )
 	{
-		bSuccess = bSuccess && FromStringBlob( JsonAppID->AsString(), Data->AppID );
+		bSuccess = bSuccess && FromStringBlob( JsonAppID->AsString(), AppID );
 	}
 	bSuccess = bSuccess && JsonAppNameString.IsValid();
 	if( bSuccess )
 	{
-		Data->AppName = JsonAppNameString->AsString();
+		AppName = JsonAppNameString->AsString();
 	}
 	bSuccess = bSuccess && JsonBuildVersionString.IsValid();
 	if( bSuccess )
 	{
-		Data->BuildVersion = JsonBuildVersionString->AsString();
+		BuildVersion = JsonBuildVersionString->AsString();
 	}
 	bSuccess = bSuccess && JsonLaunchExe.IsValid();
 	if( bSuccess )
 	{
-		Data->LaunchExe = JsonLaunchExe->AsString();
+		LaunchExe = JsonLaunchExe->AsString();
 	}
 	bSuccess = bSuccess && JsonLaunchCommand.IsValid();
 	if( bSuccess )
 	{
-		Data->LaunchCommand = JsonLaunchCommand->AsString();
+		LaunchCommand = JsonLaunchCommand->AsString();
 	}
 
 	// Get the prerequisites installer info.  These are optional entries.
-	Data->PrereqName = JsonPrereqName.IsValid() ? JsonPrereqName->AsString() : FString();
-	Data->PrereqPath = JsonPrereqPath.IsValid() ? JsonPrereqPath->AsString() : FString();
-	Data->PrereqArgs = JsonPrereqArgs.IsValid() ? JsonPrereqArgs->AsString() : FString();
+	PrereqName = JsonPrereqName.IsValid() ? JsonPrereqName->AsString() : FString();
+	PrereqPath = JsonPrereqPath.IsValid() ? JsonPrereqPath->AsString() : FString();
+	PrereqArgs = JsonPrereqArgs.IsValid() ? JsonPrereqArgs->AsString() : FString();
 
 	// Get the FileManifestList
 	TSharedPtr<FJsonValue> JsonFileManifestList = JsonValueMap.FindRef(TEXT("FileManifestList"));
@@ -1020,8 +1087,8 @@ bool FBuildPatchAppManifest::DeserializeFromJSON( const FString& JSONInput )
 		{
 			TSharedPtr<FJsonObject> JsonFileManifest = (*JsonFileManifestIt)->AsObject();
 
-			const int32 FileIndex = Data->FileManifestList.Add(FFileManifestData());
-			FFileManifestData& FileManifest = Data->FileManifestList[FileIndex];
+			const int32 FileIndex = FileManifestList.Add(FFileManifestData());
+			FFileManifestData& FileManifest = FileManifestList[FileIndex];
 			FileManifest.Filename = JsonFileManifest->GetStringField(TEXT("Filename"));
 			bSuccess = bSuccess && FString::ToBlob(JsonFileManifest->GetStringField(TEXT("FileHash")), FileManifest.FileHash.Hash, FSHA1::DigestSize);
 			TArray<TSharedPtr<FJsonValue>> JsonChunkPartArray = JsonFileManifest->GetArrayField(TEXT("FileChunkParts"));
@@ -1050,21 +1117,21 @@ bool FBuildPatchAppManifest::DeserializeFromJSON( const FString& JSONInput )
 			FileManifest.Init();
 		}
 	}
-	Data->FileManifestList.Sort();
-	for (auto& FileManifest : Data->FileManifestList)
+	FileManifestList.Sort();
+	for (FFileManifestData& FileManifest : FileManifestList)
 	{
 		FileManifestLookup.Add(FileManifest.Filename, &FileManifest);
 	}
 
 	// For each chunk setup it's info
-	for (const auto& DataGuid : AllDataGuids)
+	for (const FGuid& DataGuid : AllDataGuids)
 	{
-		int32 ChunkIndex = Data->ChunkList.Add(FChunkInfoData());
-		Data->ChunkList[ChunkIndex].Guid = DataGuid;
+		int32 ChunkIndex = ChunkList.Add(FChunkInfoData());
+		ChunkList[ChunkIndex].Guid = DataGuid;
 	}
 
 	// Setup chunk info lookup
-	for (auto& ChunkInfo : Data->ChunkList)
+	for (FChunkInfoData& ChunkInfo : ChunkList)
 	{
 		ChunkInfoLookup.Add(ChunkInfo.Guid, &ChunkInfo);
 	}
@@ -1133,7 +1200,7 @@ bool FBuildPatchAppManifest::DeserializeFromJSON( const FString& JSONInput )
 	{
 		// If the list did not exist in the manifest then the grouping is the deprecated crc functionality, as long
 		// as there are no previous parsing errors we can build the group list from the Guids.
-		for (auto& ChunkInfo : Data->ChunkList)
+		for (FChunkInfoData& ChunkInfo : ChunkList)
 		{
 			ChunkInfo.GroupNumber = FCrc::MemCrc_DEPRECATED(&ChunkInfo.Guid, sizeof(FGuid)) % 100;
 		}
@@ -1164,7 +1231,7 @@ bool FBuildPatchAppManifest::DeserializeFromJSON( const FString& JSONInput )
 	if (bHasChunkFilesizeList == false)
 	{
 		// Missing chunk list, version before we saved them compressed.. Assume chunk size
-		for (FChunkInfoData& ChunkInfo : Data->ChunkList)
+		for (FChunkInfoData& ChunkInfo : ChunkList)
 		{
 			ChunkInfo.FileSize = FBuildPatchData::ChunkDataSize;
 		}
@@ -1175,11 +1242,11 @@ bool FBuildPatchAppManifest::DeserializeFromJSON( const FString& JSONInput )
 	TSharedPtr<FJsonValue> JsonIsFileData = JsonValueMap.FindRef(TEXT("bIsFileData"));
 	if (JsonIsFileData.IsValid() && JsonIsFileData->Type == EJson::Boolean)
 	{
-		Data->bIsFileData = JsonIsFileData->AsBool();
+		bIsFileData = JsonIsFileData->AsBool();
 	}
 	else
 	{
-		Data->bIsFileData = !bHasChunkHashList;
+		bIsFileData = !bHasChunkHashList;
 	}
 
 	// Get the custom fields. This is optional, and should not fail if it does not exist
@@ -1189,23 +1256,23 @@ bool FBuildPatchAppManifest::DeserializeFromJSON( const FString& JSONInput )
 		TSharedPtr< FJsonObject > JsonCustomFieldsObj = JsonCustomFields->AsObject();
 		for( auto CustomFieldIt = JsonCustomFieldsObj->Values.CreateConstIterator(); CustomFieldIt && bSuccess; ++CustomFieldIt )
 		{
-			Data->CustomFields.Add(FCustomFieldData(CustomFieldIt.Key(), CustomFieldIt.Value()->AsString()));
+			CustomFields.Add(FCustomFieldData(CustomFieldIt.Key(), CustomFieldIt.Value()->AsString()));
 		}
 	}
-	CustomFieldLookup.Empty(Data->CustomFields.Num());
-	for (auto& CustomField : Data->CustomFields)
+	CustomFieldLookup.Empty(CustomFields.Num());
+	for (FCustomFieldData& CustomField : CustomFields)
 	{
 		CustomFieldLookup.Add(CustomField.Key, &CustomField);
 	}
 
 	// If this is file data, fill out the guid to filename lookup, and chunk file size
-	if (Data->bIsFileData)
+	if (bIsFileData)
 	{
-		for (auto& FileManifest : Data->FileManifestList)
+		for (FFileManifestData& FileManifest : FileManifestList)
 		{
 			if (FileManifest.FileChunkParts.Num() == 1)
 			{
-				FGuid& Guid = FileManifest.FileChunkParts[0].Guid;
+				const FGuid& Guid = FileManifest.FileChunkParts[0].Guid;
 				FileNameLookup.Add(Guid, &FileManifest.Filename);
 				if (ChunkInfoLookup.Contains(Guid))
 				{
@@ -1225,11 +1292,11 @@ bool FBuildPatchAppManifest::DeserializeFromJSON( const FString& JSONInput )
 	TotalDownloadSize = 0;
 	if (bSuccess)
 	{
-		for (auto& FileManifest : Data->FileManifestList)
+		for (const FFileManifestData& FileManifest : FileManifestList)
 		{
 			TotalBuildSize += FileManifest.GetFileSize();
 		}
-		for (auto& Chunk : Data->ChunkList)
+		for (const FChunkInfoData& Chunk : ChunkList)
 		{
 			TotalDownloadSize += Chunk.FileSize;
 		}
@@ -1252,7 +1319,7 @@ bool FBuildPatchAppManifest::DeserializeFromJSON( const FString& JSONInput )
 
 EBuildPatchAppManifestVersion::Type FBuildPatchAppManifest::GetManifestVersion() const
 {
-	return static_cast<EBuildPatchAppManifestVersion::Type>(Data->ManifestFileVersion);
+	return static_cast<EBuildPatchAppManifestVersion::Type>(ManifestFileVersion);
 }
 
 void FBuildPatchAppManifest::GetChunksRequiredForFiles(const TArray<FString>& FileList, TArray<FGuid>& RequiredChunks, bool bAddUnique) const
@@ -1395,7 +1462,7 @@ int64 FBuildPatchAppManifest::GetDataSize(const FGuid& DataGuid) const
 		// Chunk file sizes are stored in the info
 		return ChunkInfoLookup[DataGuid]->FileSize;
 	}
-	else if (Data->bIsFileData)
+	else if (bIsFileData)
 	{
 		// For file data, the file must exist in the list
 		check(FileNameLookup.Contains(DataGuid));
@@ -1411,7 +1478,7 @@ int64 FBuildPatchAppManifest::GetDataSize(const FGuid& DataGuid) const
 int64 FBuildPatchAppManifest::GetDataSize(const TArray<FGuid>& DataGuids) const
 {
 	int64 TotalSize = 0;
-	for (const auto& DataGuid : DataGuids)
+	for (const FGuid& DataGuid : DataGuids)
 	{
 		TotalSize += GetDataSize(DataGuid);
 	}
@@ -1420,7 +1487,7 @@ int64 FBuildPatchAppManifest::GetDataSize(const TArray<FGuid>& DataGuids) const
 
 uint32 FBuildPatchAppManifest::GetNumFiles() const
 {
-	return Data->FileManifestList.Num();
+	return FileManifestList.Num();
 }
 
 void FBuildPatchAppManifest::GetFileList(TArray<FString>& Filenames) const
@@ -1437,12 +1504,12 @@ void FBuildPatchAppManifest::GetFileTagList(TSet<FString>& Tags) const
 
 void FBuildPatchAppManifest::GetTaggedFileList(const TSet<FString>& Tags, TSet<FString>& TaggedFiles) const
 {
-	for (auto& Tag : Tags)
+	for (const FString& Tag : Tags)
 	{
-		auto* Files = TaggedFilesLookup.Find(Tag);
+		const TArray<FFileManifestData*>* Files = TaggedFilesLookup.Find(Tag);
 		if (Files != nullptr)
 		{
-			for (auto& File : *Files)
+			for (const FFileManifestData* File : *Files)
 			{
 				TaggedFiles.Add(File->Filename);
 			}
@@ -1463,7 +1530,7 @@ const FFileManifestData* FBuildPatchAppManifest::GetFileManifest(const FString&
 
 bool FBuildPatchAppManifest::IsFileDataManifest() const
 {
-	return Data->bIsFileData;
+	return bIsFileData;
 }
 
 bool FBuildPatchAppManifest::GetChunkHash(const FGuid& ChunkGuid, uint64& OutHash) const
@@ -1522,42 +1589,42 @@ bool FBuildPatchAppManifest::GetFilePartHash(const FGuid& FilePartGuid, uint64&
 
 uint32 FBuildPatchAppManifest::GetAppID() const
 {
-	return Data->AppID;
+	return AppID;
 }
 
 const FString& FBuildPatchAppManifest::GetAppName() const
 {
-	return Data->AppName;
+	return AppName;
 }
 
 const FString& FBuildPatchAppManifest::GetVersionString() const
 {
-	return Data->BuildVersion;
+	return BuildVersion;
 }
 
 const FString& FBuildPatchAppManifest::GetLaunchExe() const
 {
-	return Data->LaunchExe;
+	return LaunchExe;
 }
 
 const FString& FBuildPatchAppManifest::GetLaunchCommand() const
 {
-	return Data->LaunchCommand;
+	return LaunchCommand;
 }
 
 const FString& FBuildPatchAppManifest::GetPrereqName() const
 {
-	return Data->PrereqName;
+	return PrereqName;
 }
 
 const FString& FBuildPatchAppManifest::GetPrereqPath() const
 {
-	return Data->PrereqPath;
+	return PrereqPath;
 }
 
 const FString& FBuildPatchAppManifest::GetPrereqArgs() const
 {
-	return Data->PrereqArgs;
+	return PrereqArgs;
 }
 
 IBuildManifestRef FBuildPatchAppManifest::Duplicate() const
@@ -1572,7 +1639,7 @@ void FBuildPatchAppManifest::CopyCustomFields(const IBuildManifestRef& InOther,
 
 	// Use lookup to overwrite existing and list additional fields
 	TArray<FCustomFieldData> Extras;
-	for (const auto& CustomField : Other->Data->CustomFields)
+	for (const FCustomFieldData& CustomField : Other->CustomFields)
 	{
 		if (CustomFieldLookup.Contains(CustomField.Key))
 		{
@@ -1588,11 +1655,11 @@ void FBuildPatchAppManifest::CopyCustomFields(const IBuildManifestRef& InOther,
 	}
 
 	// Add the extra fields
-	Data->CustomFields.Append(Extras);
+	CustomFields.Append(Extras);
 
 	// Reset the loookup
-	CustomFieldLookup.Empty(Data->CustomFields.Num());
-	for (auto& CustomField : Data->CustomFields)
+	CustomFieldLookup.Empty(CustomFields.Num());
+	for (FCustomFieldData& CustomField : CustomFields)
 	{
 		CustomFieldLookup.Add(CustomField.Key, &CustomField);
 	}
@@ -1616,9 +1683,9 @@ const IManifestFieldPtr FBuildPatchAppManifest::SetCustomField(const FString& Fi
 	}
 	else
 	{
-		Data->CustomFields.Add(FCustomFieldData(FieldName, Value));
-		CustomFieldLookup.Empty(Data->CustomFields.Num());
-		for (auto& CustomField : Data->CustomFields)
+		CustomFields.Add(FCustomFieldData(FieldName, Value));
+		CustomFieldLookup.Empty(CustomFields.Num());
+		for (FCustomFieldData& CustomField : CustomFields)
 		{
 			CustomFieldLookup.Add(CustomField.Key, &CustomField);
 		}
@@ -1638,9 +1705,9 @@ const IManifestFieldPtr FBuildPatchAppManifest::SetCustomField(const FString& Fi
 
 void FBuildPatchAppManifest::RemoveCustomField(const FString& FieldName)
 {
-	Data->CustomFields.RemoveAll([&](const FCustomFieldData& Entry){ return Entry.Key == FieldName; });
-	CustomFieldLookup.Empty(Data->CustomFields.Num());
-	for (auto& CustomField : Data->CustomFields)
+	CustomFields.RemoveAll([&](const FCustomFieldData& Entry){ return Entry.Key == FieldName; });
+	CustomFieldLookup.Empty(CustomFields.Num());
+	for (FCustomFieldData& CustomField : CustomFields)
 	{
 		CustomFieldLookup.Add(CustomField.Key, &CustomField);
 	}
@@ -1726,7 +1793,7 @@ void FBuildPatchAppManifest::EnumerateChunkPartInventory(const TArray<FGuid>& Ch
 	// Use a set to optimize
 	TSet<FGuid> ChunksReqSet(ChunksRequired);
 	// For each file in the manifest, check what chunks it is made out of, and grab details for the ones in ChunksRequired
-	for (auto FileManifestIt = Data->FileManifestList.CreateConstIterator(); FileManifestIt && !FBuildPatchInstallError::HasFatalError(); ++FileManifestIt)
+	for (auto FileManifestIt = FileManifestList.CreateConstIterator(); FileManifestIt && !FBuildPatchInstallError::HasFatalError(); ++FileManifestIt)
 	{
 		const FFileManifestData& FileManifest = *FileManifestIt;
 		uint64 FileOffset = 0;
@@ -1749,7 +1816,7 @@ void FBuildPatchAppManifest::EnumerateChunkPartInventory(const TArray<FGuid>& Ch
 
 bool FBuildPatchAppManifest::HasFileAttributes() const
 {
-	for (const auto& FileManifestPair : FileManifestLookup)
+	for (const TPair<FString, FFileManifestData*>& FileManifestPair : FileManifestLookup)
 	{
 		const FFileManifestData* FileManifest = FileManifestPair.Value;
 		if (FileManifest->bIsReadOnly || FileManifest->bIsUnixExecutable || FileManifest->bIsCompressed)
@@ -1771,7 +1838,7 @@ void FBuildPatchAppManifest::GetRemovableFiles(const IBuildManifestRef& InOldMan
 	const FBuildPatchAppManifestRef OldManifest = StaticCastSharedRef< FBuildPatchAppManifest >(InOldManifest);
 	// Simply put, any files that exist in the OldManifest file list, but do not in this manifest's file list, are assumed
 	// to be files no longer required by the build
-	for (const auto& OldFile : OldManifest->Data->FileManifestList)
+	for (const FFileManifestData& OldFile : OldManifest->FileManifestList)
 	{
 		if (!FileManifestLookup.Contains(OldFile.Filename))
 		{
@@ -1825,7 +1892,7 @@ void FBuildPatchAppManifest::GetOutdatedFiles(const FBuildPatchAppManifestPtr& O
 	{
 		// Enumerate files in the this file list, that do not exist, or have different hashes in the OldManifest
 		// to be files no longer required by the build
-		for (const auto& NewFile : Data->FileManifestList)
+		for (const FFileManifestData& NewFile : FileManifestList)
 		{
 			const int64 ExistingFileSize = IFileManager::Get().FileSize(*(InstallDirectory / NewFile.Filename));
 			// Check changed
@@ -1869,9 +1936,9 @@ uint32 FBuildPatchAppManifest::GetNumberOfChunkReferences(const FGuid& ChunkGuid
 {
 	uint32 RefCount = 0;
 	// For each file in the manifest, check if it has references to this chunk
-	for (const auto& FileManifest : Data->FileManifestList)
+	for (const FFileManifestData& FileManifest : FileManifestList)
 	{
-		for (const auto& ChunkPart : FileManifest.FileChunkParts)
+		for (const FChunkPartData& ChunkPart : FileManifest.FileChunkParts)
 		{
 			if (ChunkPart.Guid == ChunkGuid)
 			{
@@ -1882,13 +1949,4 @@ uint32 FBuildPatchAppManifest::GetNumberOfChunkReferences(const FGuid& ChunkGuid
 	return RefCount;
 }
 
-void FBuildPatchAppManifest::AddReferencedObjects( FReferenceCollector& Collector )
-{
-	if (Data != nullptr)
-	{
-		// Ensure Data is not garbage collected as long as this object is valid.
-		Collector.AddReferencedObject(Data);
-	}
-}
-
 #undef LOCTEXT_NAMESPACE
diff --git a/Engine/Source/Runtime/Online/BuildPatchServices/Private/BuildPatchManifest.h b/Engine/Source/Runtime/Online/BuildPatchServices/Private/BuildPatchManifest.h
--- a/Engine/Source/Runtime/Online/BuildPatchServices/Private/BuildPatchManifest.h
+++ b/Engine/Source/Runtime/Online/BuildPatchServices/Private/BuildPatchManifest.h
@@ -361,7 +361,7 @@ namespace BuildPatchServices
  * Declare the FBuildPatchAppManifest object class. This holds the UObject data, and the implemented build manifest functionality
  */
 class FBuildPatchAppManifest
-	: public IBuildManifest, FGCObject
+	: public IBuildManifest
 {
 	// Allow access to build processor classes
 	friend class FBuildDataGenerator;
@@ -621,10 +621,6 @@ public:
 
 	/** @return True if this manifest is for the same build, i.e. same ID, Name, and Version */
 	bool IsSameAs(FBuildPatchAppManifestRef InstallManifest) const;
-public:
-
-	// FGCObject API
-	virtual void AddReferencedObjects( FReferenceCollector& Collector ) override;
 
 private:
 
@@ -642,8 +638,23 @@ private:
 
 private:
 
-	/** Holds the actual manifest data. Some other variables point to the memory held by this object */
-	UBuildPatchManifest* Data;
+	/** Holds the actual manifest data. Some other variables point to the memory held by these objects */
+	uint8 ManifestFileVersion;
+	bool bIsFileData;
+	uint32 AppID;
+	FString AppName;
+	FString BuildVersion;
+	FString LaunchExe;
+	FString LaunchCommand;
+	FString PrereqName;
+	FString PrereqPath;
+	FString PrereqArgs;
+	TArray<FFileManifestData> FileManifestList;
+	TArray<FChunkInfoData> ChunkList;
+	TArray<FCustomFieldData> CustomFields;
+
+	/** Holds the handle to our PreExit delegate */
+	FDelegateHandle OnPreExitHandle;
 
 	/** Some lookups to optimize data access */
 	TMap<FGuid, FString*> FileNameLookup;
diff --git a/Engine/Source/Runtime/Online/BuildPatchServices/Private/BuildPatchMergeManifests.cpp b/Engine/Source/Runtime/Online/BuildPatchServices/Private/BuildPatchMergeManifests.cpp
--- a/Engine/Source/Runtime/Online/BuildPatchServices/Private/BuildPatchMergeManifests.cpp
+++ b/Engine/Source/Runtime/Online/BuildPatchServices/Private/BuildPatchMergeManifests.cpp
@@ -27,22 +27,21 @@ namespace MergeHelpers
 		return FBuildPatchAppManifestPtr();
 	}
 
-	bool CopyFileDataFromManifestToObject(const TSet<FString>& Filenames, const FBuildPatchAppManifestPtr& Manifest, UBuildPatchManifest* Object, const TCHAR* ManifestName)
+	bool CopyFileDataFromManifestToArray(const TSet<FString>& Filenames, const FBuildPatchAppManifestPtr& Source, TArray<FFileManifestData>& DestArray)
 	{
 		bool bSuccess = true;
 		for (const FString& Filename : Filenames)
 		{
-			check(Manifest.IsValid());
-			check(Object != nullptr);
-			const FFileManifestData* FileManifest = Manifest->GetFileManifest(Filename);
+			check(Source.IsValid());
+			const FFileManifestData* FileManifest = Source->GetFileManifest(Filename);
 			if (FileManifest == nullptr)
 			{
-				UE_LOG(LogMergeManifests, Error, TEXT("Could not find file in %s: %s"), ManifestName, *Filename);
+				UE_LOG(LogMergeManifests, Error, TEXT("Could not find file in %s %s: %s"), *Source->GetAppName(), *Source->GetVersionString(), *Filename);
 				bSuccess = false;
 			}
 			else
 			{
-				Object->FileManifestList.Add(*FileManifest);
+				DestArray.Add(*FileManifest);
 			}
 		}
 		return bSuccess;
@@ -62,12 +61,12 @@ bool FBuildMergeManifests::MergeManifests(const FString& ManifestFilePathA, cons
 	{
 		return MergeHelpers::LoadManifestFile(ManifestFilePathB, &UObjectAllocationLock);
 	};
-	typedef TPair<TSet<FString>,TSet<FString>> TStringSetPair;
+	typedef TPair<TSet<FString>,TSet<FString>> FStringSetPair;
 	FThreadSafeBool bSelectionDetailSuccess = false;
-	TFunction<TStringSetPair()> TaskSelectionInfo = [&SelectionDetailFilePath, &bSelectionDetailSuccess]()
+	TFunction<FStringSetPair()> TaskSelectionInfo = [&SelectionDetailFilePath, &bSelectionDetailSuccess]()
 	{
 		bSelectionDetailSuccess = true;
-		TStringSetPair StringSetPair;
+		FStringSetPair StringSetPair;
 		if(SelectionDetailFilePath.IsEmpty() == false)
 		{
 			FString SelectionDetailFileData;
@@ -108,11 +107,11 @@ bool FBuildMergeManifests::MergeManifests(const FString& ManifestFilePathA, cons
 
 	TFuture<FBuildPatchAppManifestPtr> FutureManifestA = Async(EAsyncExecution::ThreadPool, MoveTemp(TaskManifestA));
 	TFuture<FBuildPatchAppManifestPtr> FutureManifestB = Async(EAsyncExecution::ThreadPool, MoveTemp(TaskManifestB));
-	TFuture<TStringSetPair> FutureSelectionDetail = Async(EAsyncExecution::ThreadPool, MoveTemp(TaskSelectionInfo));
+	TFuture<FStringSetPair> FutureSelectionDetail = Async(EAsyncExecution::ThreadPool, MoveTemp(TaskSelectionInfo));
 
 	FBuildPatchAppManifestPtr ManifestA = FutureManifestA.Get();
 	FBuildPatchAppManifestPtr ManifestB = FutureManifestB.Get();
-	TStringSetPair SelectionDetail = FutureSelectionDetail.Get();
+	FStringSetPair SelectionDetail = FutureSelectionDetail.Get();
 
 	// Flush any logs collected by tasks
 	GLog->FlushThreadedLogs();
@@ -146,38 +145,33 @@ bool FBuildMergeManifests::MergeManifests(const FString& ManifestFilePathA, cons
 	// Create the new manifest
 	FBuildPatchAppManifest MergedManifest;
 
-	// Helpful pointers
-	UBuildPatchManifest* DataA = ManifestA->Data;
-	UBuildPatchManifest* DataB = ManifestB->Data;
-	UBuildPatchManifest* DataC = MergedManifest.Data;
-
 	// Copy basic info from B
-	DataC->ManifestFileVersion = DataB->ManifestFileVersion;
-	DataC->bIsFileData = DataB->bIsFileData;
-	DataC->AppID = DataB->AppID;
-	DataC->AppName = DataB->AppName;
-	DataC->LaunchExe = DataB->LaunchExe;
-	DataC->LaunchCommand = DataB->LaunchCommand;
-	DataC->PrereqName = DataB->PrereqName;
-	DataC->PrereqPath = DataB->PrereqPath;
-	DataC->PrereqArgs = DataB->PrereqArgs;
-	DataC->CustomFields = DataB->CustomFields;
+	MergedManifest.ManifestFileVersion = ManifestB->ManifestFileVersion;
+	MergedManifest.bIsFileData = ManifestB->bIsFileData;
+	MergedManifest.AppID = ManifestB->AppID;
+	MergedManifest.AppName = ManifestB->AppName;
+	MergedManifest.LaunchExe = ManifestB->LaunchExe;
+	MergedManifest.LaunchCommand = ManifestB->LaunchCommand;
+	MergedManifest.PrereqName = ManifestB->PrereqName;
+	MergedManifest.PrereqPath = ManifestB->PrereqPath;
+	MergedManifest.PrereqArgs = ManifestB->PrereqArgs;
+	MergedManifest.CustomFields = ManifestB->CustomFields;
 
 	// Set the new version string
-	DataC->BuildVersion = NewVersionString;
+	MergedManifest.BuildVersion = NewVersionString;
 
 	// Copy the file manifests required from A
-	bSuccess = MergeHelpers::CopyFileDataFromManifestToObject(SelectionDetail.Key, ManifestA, DataC, TEXT("ManifestA")) && bSuccess;
+	bSuccess = MergeHelpers::CopyFileDataFromManifestToArray(SelectionDetail.Key, ManifestA, MergedManifest.FileManifestList) && bSuccess;
 
 	// Copy the file manifests required from B
-	bSuccess = MergeHelpers::CopyFileDataFromManifestToObject(SelectionDetail.Value, ManifestB, DataC, TEXT("ManifestB")) && bSuccess;
+	bSuccess = MergeHelpers::CopyFileDataFromManifestToArray(SelectionDetail.Value, ManifestB, MergedManifest.FileManifestList) && bSuccess;
 
 	// Sort the file manifests before entering chunk info
-	DataC->FileManifestList.Sort();
+	MergedManifest.FileManifestList.Sort();
 
 	// Fill out the chunk list in order of reference
 	TSet<FGuid> ReferencedChunks;
-	for (const FFileManifestData& FileManifest: DataC->FileManifestList)
+	for (const FFileManifestData& FileManifest: MergedManifest.FileManifestList)
 	{
 		for (const FChunkPartData& FileChunkPart : FileManifest.FileChunkParts)
 		{
@@ -198,7 +192,7 @@ bool FBuildMergeManifests::MergeManifests(const FString& ManifestFilePathA, cons
 				}
 				else
 				{
-					DataC->ChunkList.Add(**ChunkInfo);
+					MergedManifest.ChunkList.Add(**ChunkInfo);
 				}
 			}
 		}
diff --git a/Engine/Source/Runtime/Online/BuildPatchServices/Private/Generation/DataScanner.cpp b/Engine/Source/Runtime/Online/BuildPatchServices/Private/Generation/DataScanner.cpp
--- a/Engine/Source/Runtime/Online/BuildPatchServices/Private/Generation/DataScanner.cpp
+++ b/Engine/Source/Runtime/Online/BuildPatchServices/Private/Generation/DataScanner.cpp
@@ -38,10 +38,12 @@ namespace BuildPatchServices
 		TMap<FGuid, FSHAHash> ChunkShaHashes;
 		volatile FStatsCollector::FAtomicValue* StatCreatedScanners;
 		volatile FStatsCollector::FAtomicValue* StatRunningScanners;
+		volatile FStatsCollector::FAtomicValue* StatCompleteScanners;
 		volatile FStatsCollector::FAtomicValue* StatCpuTime;
 		volatile FStatsCollector::FAtomicValue* StatRealTime;
 		volatile FStatsCollector::FAtomicValue* StatHashCollisions;
-		volatile FStatsCollector::FAtomicValue* StatProcessedData;
+		volatile FStatsCollector::FAtomicValue* StatTotalData;
+		volatile FStatsCollector::FAtomicValue* StatSkippedData;
 		volatile FStatsCollector::FAtomicValue* StatProcessingSpeed;
 
 	public:
@@ -59,10 +61,12 @@ namespace BuildPatchServices
 		// Create statistics.
 		StatCreatedScanners = StatsCollector->CreateStat(TEXT("Scanner: Created Scanners"), EStatFormat::Value);
 		StatRunningScanners = StatsCollector->CreateStat(TEXT("Scanner: Running Scanners"), EStatFormat::Value);
+		StatCompleteScanners = StatsCollector->CreateStat(TEXT("Scanner: Complete Scanners"), EStatFormat::Value);
 		StatCpuTime = StatsCollector->CreateStat(TEXT("Scanner: CPU Time"), EStatFormat::Timer);
 		StatRealTime = StatsCollector->CreateStat(TEXT("Scanner: Real Time"), EStatFormat::Timer);
 		StatHashCollisions = StatsCollector->CreateStat(TEXT("Scanner: Hash Collisions"), EStatFormat::Value);
-		StatProcessedData = StatsCollector->CreateStat(TEXT("Scanner: Processed Data"), EStatFormat::DataSize);
+		StatTotalData = StatsCollector->CreateStat(TEXT("Scanner: Total Data"), EStatFormat::DataSize);
+		StatSkippedData = StatsCollector->CreateStat(TEXT("Scanner: Skipped Data"), EStatFormat::DataSize);
 		StatProcessingSpeed = StatsCollector->CreateStat(TEXT("Scanner: Processing Speed"), EStatFormat::DataSpeed);
 		FStatsCollector::Accumulate(StatCreatedScanners, 1);
 
@@ -72,6 +76,7 @@ namespace BuildPatchServices
 		{
 			TArray<FChunkMatch> Result = ScanData();
 			NumIncompleteScanners.Decrement();
+			FStatsCollector::Accumulate(StatCompleteScanners, 1);
 			return MoveTemp(Result);
 		};
 		FutureResult = Async(EAsyncExecution::ThreadPool, MoveTemp(Task));
@@ -152,6 +157,9 @@ namespace BuildPatchServices
 		FSHAHash ChunkSha;
 		uint64 CpuTimer;
 
+		// Track last match so we know if we can start skipping data. This will also cover us for the overlap with previous scanner.
+		uint64 LastMatch = 0;
+
 		// Loop over and process all data.
 		uint32 NextByte = ConsumeData(&Data[0], Data.Num());
 		bool bScanningData = true;
@@ -160,12 +168,34 @@ namespace BuildPatchServices
 			FStatsParallelScopeTimer ParallelScopeTimer(&TempTimerValue, StatRealTime, StatRunningScanners);
 			while (bScanningData && !bShouldAbort)
 			{
+				const uint32 DataStart = NextByte - WindowSize;
+				const bool bChunkOverlap = DataStart < (LastMatch + WindowSize);
 				// Check for a chunk match at this offset.
-				if (FindChunkDataMatch(ChunkMatch, ChunkSha))
+				const bool bFoundChunkMatch = FindChunkDataMatch(ChunkMatch, ChunkSha);
+				if (bFoundChunkMatch)
 				{
-					DataScanResult.Emplace(NextByte - WindowSize, ChunkMatch);
+					LastMatch = DataStart;
+					DataScanResult.Emplace(DataStart, ChunkMatch);
 				}
-
+				// We can start skipping over the chunk that we matched if we have no overlap potential, i.e. we know this match will not be rejected.
+				if (bFoundChunkMatch && !bChunkOverlap)
+				{
+					RollingHash.Clear();
+					const bool bHasEnoughData = (NextByte + WindowSize - 1) < static_cast<uint32>(Data.Num());
+					if (bHasEnoughData)
+				{
+						const uint32 Consumed = ConsumeData(&Data[NextByte], Data.Num() - NextByte);
+						FStatsCollector::Accumulate(StatSkippedData, Consumed);
+						NextByte += Consumed;
+				}
+					else
+					{
+						bScanningData = false;
+					}
+				}
+				// Otherwise we only move forwards by one byte.
+				else
+				{
 				const bool bHasMoreData = NextByte < static_cast<uint32>(Data.Num());
 				if (bHasMoreData)
 				{
@@ -177,9 +207,10 @@ namespace BuildPatchServices
 					bScanningData = false;
 				}
 			}
+			}
 			FStatsCollector::AccumulateTimeEnd(StatCpuTime, CpuTimer);
-			FStatsCollector::Accumulate(StatProcessedData, Data.Num());
-			FStatsCollector::Set(StatProcessingSpeed, *StatProcessedData / FStatsCollector::CyclesToSeconds(ParallelScopeTimer.GetCurrentTime()));
+			FStatsCollector::Accumulate(StatTotalData, Data.Num());
+			FStatsCollector::Set(StatProcessingSpeed, *StatTotalData / FStatsCollector::CyclesToSeconds(ParallelScopeTimer.GetCurrentTime()));
 		}
 
 		// Count running scanners.
diff --git a/Engine/Source/Runtime/Online/BuildPatchServices/Private/Generation/ManifestBuilder.cpp b/Engine/Source/Runtime/Online/BuildPatchServices/Private/Generation/ManifestBuilder.cpp
--- a/Engine/Source/Runtime/Online/BuildPatchServices/Private/Generation/ManifestBuilder.cpp
+++ b/Engine/Source/Runtime/Online/BuildPatchServices/Private/Generation/ManifestBuilder.cpp
@@ -49,15 +49,15 @@ namespace BuildPatchServices
 		: Manifest(MakeShareable(new FBuildPatchAppManifest()))
 		, FileAttributesMap(InDetails.FileAttributesMap)
 	{
-		Manifest->Data->bIsFileData = false;
-		Manifest->Data->AppID = InDetails.AppId;
-		Manifest->Data->AppName = InDetails.AppName;
-		Manifest->Data->BuildVersion = InDetails.BuildVersion;
-		Manifest->Data->LaunchExe = InDetails.LaunchExe;
-		Manifest->Data->LaunchCommand = InDetails.LaunchCommand;
-		Manifest->Data->PrereqName = InDetails.PrereqName;
-		Manifest->Data->PrereqPath = InDetails.PrereqPath;
-		Manifest->Data->PrereqArgs = InDetails.PrereqArgs;
+		Manifest->bIsFileData = false;
+		Manifest->AppID = InDetails.AppId;
+		Manifest->AppName = InDetails.AppName;
+		Manifest->BuildVersion = InDetails.BuildVersion;
+		Manifest->LaunchExe = InDetails.LaunchExe;
+		Manifest->LaunchCommand = InDetails.LaunchCommand;
+		Manifest->PrereqName = InDetails.PrereqName;
+		Manifest->PrereqPath = InDetails.PrereqPath;
+		Manifest->PrereqArgs = InDetails.PrereqArgs;
 		for (const auto& CustomField : InDetails.CustomFields)
 		{
 			int32 VarType = CustomField.Value.GetType();
@@ -100,8 +100,8 @@ namespace BuildPatchServices
 		for (const FFileSpan& FileSpan : FileSpans)
 		{
 			FFileAttributes FileAttributes = FileAttributesMap.FindRef(FileSpan.Filename);
-			Manifest->Data->FileManifestList.AddDefaulted();
-			FFileManifestData& FileManifest = Manifest->Data->FileManifestList.Last();
+			Manifest->FileManifestList.AddDefaulted();
+			FFileManifestData& FileManifest = Manifest->FileManifestList.Last();
 			FileManifest.Filename = FileSpan.Filename;
 			FMemory::Memcpy(FileManifest.FileHash.Hash, FileSpan.SHAHash.Hash, FSHA1::DigestSize);
 			FileManifest.InstallTags = FileAttributes.InstallTags.Array();
@@ -116,10 +116,10 @@ namespace BuildPatchServices
 		UE_LOG(LogManifestBuilder, Verbose, TEXT("Manifest references %d chunks."), ReferencedChunks.Num());
 
 		// Setup chunk list, removing all that were not referenced.
-		Manifest->Data->ChunkList = MoveTemp(ChunkInfo);
-		int32 TotalChunkListNum = Manifest->Data->ChunkList.Num();
-		Manifest->Data->ChunkList.RemoveAll([&](FChunkInfoData& Candidate){ return ReferencedChunks.Contains(Candidate.Guid) == false; });
-		UE_LOG(LogManifestBuilder, Verbose, TEXT("Chunk info list trimmed from %d to %d."), TotalChunkListNum, Manifest->Data->ChunkList.Num());
+		Manifest->ChunkList = MoveTemp(ChunkInfo);
+		int32 TotalChunkListNum = Manifest->ChunkList.Num();
+		Manifest->ChunkList.RemoveAll([&](FChunkInfoData& Candidate){ return ReferencedChunks.Contains(Candidate.Guid) == false; });
+		UE_LOG(LogManifestBuilder, Verbose, TEXT("Chunk info list trimmed from %d to %d."), TotalChunkListNum, Manifest->ChunkList.Num());
 
 		// Init the manifest, and we are done.
 		Manifest->InitLookups();
@@ -163,7 +163,7 @@ namespace BuildPatchServices
 		checkf(BuildStructureAdded.GetHead()->GetSize() == Manifest->GetBuildSize(), TEXT("Build integrity check failed. Structure added is not the same size as the manifest data setup; did you call FinalizeData?"));
 
 		// Currently we only save out in JSON format
-		Manifest->Data->ManifestFileVersion = EBuildPatchAppManifestVersion::GetLatestJsonVersion();
+		Manifest->ManifestFileVersion = EBuildPatchAppManifestVersion::GetLatestJsonVersion();
 		return Manifest->SaveToFile(Filename, false);
 	}
 
diff --git a/Engine/Source/Runtime/Online/HTTP/Private/HttpModule.cpp b/Engine/Source/Runtime/Online/HTTP/Private/HttpModule.cpp
--- a/Engine/Source/Runtime/Online/HTTP/Private/HttpModule.cpp
+++ b/Engine/Source/Runtime/Online/HTTP/Private/HttpModule.cpp
@@ -47,14 +47,17 @@ void FHttpModule::StartupModule()
 	HttpDelayTime = 0;
 	GConfig->GetFloat(TEXT("HTTP"), TEXT("HttpDelayTime"), HttpDelayTime, GEngineIni);
 
-	HttpThreadProcessingClampInSeconds = 0.005f; // 5ms
-	GConfig->GetFloat(TEXT("HTTP"), TEXT("HttpThreadProcessingClampInSeconds"), HttpThreadProcessingClampInSeconds, GEngineIni);
+	HttpThreadActiveFrameTimeInSeconds = 1.0f / 200.0f; // 200Hz
+	GConfig->GetFloat(TEXT("HTTP"), TEXT("HttpThreadActiveFrameTimeInSeconds"), HttpThreadActiveFrameTimeInSeconds, GEngineIni);
 
-	HttpThreadFrameTimeInSeconds = 1.0f / 30.0f; // 30Hz
-	GConfig->GetFloat(TEXT("HTTP"), TEXT("HttpThreadFrameTimeInSeconds"), HttpThreadFrameTimeInSeconds, GEngineIni);
+	HttpThreadActiveMinimumSleepTimeInSeconds = 0.0f;
+	GConfig->GetFloat(TEXT("HTTP"), TEXT("HttpThreadActiveMinimumSleepTimeInSeconds"), HttpThreadActiveMinimumSleepTimeInSeconds, GEngineIni);
 
-	HttpThreadMinimumSleepTimeInSeconds = 0.0f;
-	GConfig->GetFloat(TEXT("HTTP"), TEXT("HttpThreadMinimumSleepTimeInSeconds"), HttpThreadMinimumSleepTimeInSeconds, GEngineIni);
+	HttpThreadIdleFrameTimeInSeconds = 1.0f / 30.0f; // 30Hz
+	GConfig->GetFloat(TEXT("HTTP"), TEXT("HttpThreadIdleFrameTimeInSeconds"), HttpThreadIdleFrameTimeInSeconds, GEngineIni);
+
+	HttpThreadIdleMinimumSleepTimeInSeconds = 0.0f;
+	GConfig->GetFloat(TEXT("HTTP"), TEXT("HttpThreadIdleMinimumSleepTimeInSeconds"), HttpThreadIdleMinimumSleepTimeInSeconds, GEngineIni);
 
 	HttpManager = FPlatformHttp::CreatePlatformHttpManager();
 	if (NULL == HttpManager)
diff --git a/Engine/Source/Runtime/Online/HTTP/Private/HttpThread.cpp b/Engine/Source/Runtime/Online/HTTP/Private/HttpThread.cpp
--- a/Engine/Source/Runtime/Online/HTTP/Private/HttpThread.cpp
+++ b/Engine/Source/Runtime/Online/HTTP/Private/HttpThread.cpp
@@ -14,11 +14,12 @@
 FHttpThread::FHttpThread()
 	:	Thread(nullptr)
 {
-	HttpThreadProcessingClampInSeconds = FHttpModule::Get().GetHttpThreadProcessingClampInSeconds();
-	HttpThreadFrameTimeInSeconds = FHttpModule::Get().GetHttpThreadFrameTimeInSeconds();
-	HttpThreadMinimumSleepTimeInSeconds = FHttpModule::Get().GetHttpThreadMinimumSleepTimeInSeconds();
+	HttpThreadActiveFrameTimeInSeconds = FHttpModule::Get().GetHttpThreadActiveFrameTimeInSeconds();
+	HttpThreadActiveMinimumSleepTimeInSeconds = FHttpModule::Get().GetHttpThreadActiveMinimumSleepTimeInSeconds();
+	HttpThreadIdleFrameTimeInSeconds = FHttpModule::Get().GetHttpThreadIdleFrameTimeInSeconds();
+	HttpThreadIdleMinimumSleepTimeInSeconds = FHttpModule::Get().GetHttpThreadIdleMinimumSleepTimeInSeconds();
 
-	UE_LOG(LogHttp, Log, TEXT("HTTP thread processing clamped at %.1f ms. HTTP thread frame time clamped to %.1f ms. Minimum sleep time is %.1f ms"), HttpThreadProcessingClampInSeconds * 1000.0, HttpThreadFrameTimeInSeconds * 1000.0, HttpThreadMinimumSleepTimeInSeconds * 1000.0);
+	UE_LOG(LogHttp, Log, TEXT("HTTP thread active frame time %.1f ms. Minimum active sleep time is %.1f ms. HTTP thread idle frame time %.1f ms. Minimum idle sleep time is %.1f ms."), HttpThreadActiveFrameTimeInSeconds * 1000.0, HttpThreadActiveMinimumSleepTimeInSeconds * 1000.0, HttpThreadIdleFrameTimeInSeconds * 1000.0, HttpThreadIdleMinimumSleepTimeInSeconds * 1000.0);
 }
 
 FHttpThread::~FHttpThread()
@@ -88,12 +89,12 @@ uint32 FHttpThread::Run()
 	TArray<IHttpThreadedRequest*> RequestsToComplete;
 	while (!ExitRequest.GetValue())
 	{
-		double TickBegin = FPlatformTime::Seconds();
-		double TickEnd = 0.0;
+		double OuterLoopBegin = FPlatformTime::Seconds();
+		double OuterLoopEnd = 0.0;
 		bool bKeepProcessing = true;
 		while (bKeepProcessing)
 		{
-			double InterationBegin = FPlatformTime::Seconds();
+			double InnerLoopBegin = FPlatformTime::Seconds();
 
 			// cache all cancelled and pending requests
 			{
@@ -167,37 +168,27 @@ uint32 FHttpThread::Run()
 				RequestsToComplete.Reset();
 			}
 
-			double TimeNow = FPlatformTime::Seconds();
-			double TimeWithoutSleeping = TimeNow - TickBegin;
-			bool bEndit = false;
 			if (RunningThreadedRequests.Num() == 0)
 			{
-				bEndit = true;
+				bKeepProcessing = false;
 			}
 
-			if (HttpThreadProcessingClampInSeconds > 0.0 && TimeWithoutSleeping >= HttpThreadProcessingClampInSeconds)
+			double InnerLoopEnd = FPlatformTime::Seconds();
+			if (bKeepProcessing)
 			{
-				bEndit = true;
+				double InnerLoopTime = InnerLoopEnd - InnerLoopBegin;
+				double InnerSleep = FMath::Max(HttpThreadActiveFrameTimeInSeconds - InnerLoopTime, HttpThreadActiveMinimumSleepTimeInSeconds);
+				FPlatformProcess::SleepNoStats(InnerSleep);
 			}
-
-			if (bEndit)
+			else
 			{
-				bKeepProcessing = false;
-				TickEnd = TimeNow;
+				OuterLoopEnd = InnerLoopEnd;
 			}
 		}
 
-		double TickTime = TickEnd - TickBegin;
-		double TimeToSleep = HttpThreadMinimumSleepTimeInSeconds;
-		if (HttpThreadFrameTimeInSeconds > 0.0 && HttpThreadFrameTimeInSeconds >= TickTime)
-		{
-			TimeToSleep = FMath::Max(HttpThreadFrameTimeInSeconds - TickTime, TimeToSleep);
-		}
-
-		if (TimeToSleep > 0.0)
-		{
-			FPlatformProcess::SleepNoStats(TimeToSleep);
-		}
+		double OuterLoopTime = OuterLoopEnd - OuterLoopBegin;
+		double OuterSleep = FMath::Max(HttpThreadIdleFrameTimeInSeconds - OuterLoopTime, HttpThreadIdleMinimumSleepTimeInSeconds);
+		FPlatformProcess::SleepNoStats(OuterSleep);
 	}
 	return 0;
 }
diff --git a/Engine/Source/Runtime/Online/HTTP/Private/HttpThread.h b/Engine/Source/Runtime/Online/HTTP/Private/HttpThread.h
--- a/Engine/Source/Runtime/Online/HTTP/Private/HttpThread.h
+++ b/Engine/Source/Runtime/Online/HTTP/Private/HttpThread.h
@@ -83,12 +83,14 @@ protected:
 	/** signal request to stop and exit thread */
 	FThreadSafeCounter ExitRequest;
 
-	/** Time in seconds to clamp the processing of HTTP requests on the HTTP thread to. 0 means unclamped */
-	double HttpThreadProcessingClampInSeconds;
-	/** Time in seconds to use as a frame time. 0 means no frame time. */
-	double HttpThreadFrameTimeInSeconds;
-	/** Time in seconds to sleep minimally */
-	double HttpThreadMinimumSleepTimeInSeconds;
+	/** Time in seconds to use as frame time when actively processing requests. 0 means no frame time. */
+	double HttpThreadActiveFrameTimeInSeconds;
+	/** Time in seconds to sleep minimally when actively processing requests. */
+	double HttpThreadActiveMinimumSleepTimeInSeconds;
+	/** Time in seconds to use as frame time when idle, waiting for requests. 0 means no frame time. */
+	double HttpThreadIdleFrameTimeInSeconds;
+	/** Time in seconds to sleep minimally when idle, waiting for requests. */
+	double HttpThreadIdleMinimumSleepTimeInSeconds;
 
 protected:
 	/** Critical section to lock access to PendingThreadedRequests, CancelledThreadedRequests, and CompletedThreadedRequests */
diff --git a/Engine/Source/Runtime/Online/HTTP/Public/HttpModule.h b/Engine/Source/Runtime/Online/HTTP/Public/HttpModule.h
--- a/Engine/Source/Runtime/Online/HTTP/Public/HttpModule.h
+++ b/Engine/Source/Runtime/Online/HTTP/Public/HttpModule.h
@@ -171,51 +171,67 @@ public:
 	}
 
 	/**
-	 * @return Target tick rate of an http thread
+	 * @return Target tick rate of an active http thread
 	 */
-	inline float GetHttpThreadProcessingClampInSeconds() const
+	inline float GetHttpThreadActiveFrameTimeInSeconds() const
 	{
-		return HttpThreadProcessingClampInSeconds;
+		return HttpThreadActiveFrameTimeInSeconds;
 	}
 
 	/**
-	 * Set the target tick rate of an http thread
+	 * Set the target tick rate of an active http thread
 	 */
-	inline void SetHttpThreadProcessingClampInSeconds(float InHttpThreadProcessingClampInSeconds)
+	inline void SetHttpThreadActiveFrameTimeInSeconds(float InHttpThreadActiveFrameTimeInSeconds)
 	{
-		HttpThreadProcessingClampInSeconds = InHttpThreadProcessingClampInSeconds;
+		HttpThreadActiveFrameTimeInSeconds = InHttpThreadActiveFrameTimeInSeconds;
 	}
 
 	/**
-	* @return Target tick rate of an http thread
-	*/
-	inline float GetHttpThreadFrameTimeInSeconds() const
+	 * @return Minimum sleep time of an active http thread
+	 */
+	inline float GetHttpThreadActiveMinimumSleepTimeInSeconds() const
+	{
+		return HttpThreadActiveMinimumSleepTimeInSeconds;
+	}
+
+	/**
+	 * Set the minimum sleep time of an active http thread
+	 */
+	inline void SetHttpThreadActiveMinimumSleepTimeInSeconds(float InHttpThreadActiveMinimumSleepTimeInSeconds)
+	{
+		HttpThreadActiveMinimumSleepTimeInSeconds = InHttpThreadActiveMinimumSleepTimeInSeconds;
+	}
+
+	/**
+	 * @return Target tick rate of an idle http thread
+	 */
+	inline float GetHttpThreadIdleFrameTimeInSeconds() const
 	{
-		return HttpThreadFrameTimeInSeconds;
+		return HttpThreadIdleFrameTimeInSeconds;
 	}
 
 	/**
-	* Set the target tick rate of an http thread
-	*/
-	inline void SetHttpThreadFrameTimeInSeconds(float InHttpThreadFrameTimeInSeconds)
+	 * Set the target tick rate of an idle http thread
+	 */
+	inline void SetHttpThreadIdleFrameTimeInSeconds(float InHttpThreadIdleFrameTimeInSeconds)
 	{
-		HttpThreadFrameTimeInSeconds = InHttpThreadFrameTimeInSeconds;
+		HttpThreadIdleFrameTimeInSeconds = InHttpThreadIdleFrameTimeInSeconds;
 	}
 
 	/**
-	* @return Target tick rate of an http thread
-	*/
-	inline float GetHttpThreadMinimumSleepTimeInSeconds() const
+	 * @return Minimum sleep time when idle, waiting for requests
+	 */
+	inline float GetHttpThreadIdleMinimumSleepTimeInSeconds() const
 	{
-		return HttpThreadMinimumSleepTimeInSeconds;
+		return HttpThreadIdleMinimumSleepTimeInSeconds;
 	}
 
 	/**
-	* Set the target tick rate of an http thread
-	*/
-	inline void SetHttpThreadMinimumSleepTimeInSeconds(float InHttpThreadMinimumSleepTimeInSeconds)
+	 * Set the minimum sleep time when idle, waiting for requests
+	 */
+	inline void SetHttpThreadIdleMinimumSleepTimeInSeconds(float InHttpThreadIdleMinimumSleepTimeInSeconds)
 	{
-		HttpThreadMinimumSleepTimeInSeconds = InHttpThreadMinimumSleepTimeInSeconds;
+		HttpThreadIdleMinimumSleepTimeInSeconds = InHttpThreadIdleMinimumSleepTimeInSeconds;
 	}
 
 private:
@@ -258,12 +274,14 @@ private:
 	float HttpSendTimeout;
 	/** total time to delay the request */
 	float HttpDelayTime;
-	/** Time in seconds to clamp the processing of HTTP requests on the HTTP thread to. 0 means unclamped */
-	float HttpThreadProcessingClampInSeconds;
-	/** Time in seconds to use as a frame time. 0 means no frame time. */
-	float HttpThreadFrameTimeInSeconds;
-	/** Time in seconds to sleep minimally */
-	float HttpThreadMinimumSleepTimeInSeconds;
+	/** Time in seconds to use as frame time when actively processing requests. 0 means no frame time. */
+	float HttpThreadActiveFrameTimeInSeconds;
+	/** Time in seconds to sleep minimally when actively processing requests. */
+	float HttpThreadActiveMinimumSleepTimeInSeconds;
+	/** Time in seconds to use as frame time when idle, waiting for requests. 0 means no frame time. */
+	float HttpThreadIdleFrameTimeInSeconds;
+	/** Time in seconds to sleep minimally when idle, waiting for requests. */
+	float HttpThreadIdleMinimumSleepTimeInSeconds;
 	/** Max number of simultaneous connections to a specific server */
 	int32 HttpMaxConnectionsPerServer;
 	/** Max buffer size for individual http reads */
diff --git a/Engine/Source/Runtime/Slate/Private/Widgets/Layout/SExpandableArea.cpp b/Engine/Source/Runtime/Slate/Private/Widgets/Layout/SExpandableArea.cpp
--- a/Engine/Source/Runtime/Slate/Private/Widgets/Layout/SExpandableArea.cpp
+++ b/Engine/Source/Runtime/Slate/Private/Widgets/Layout/SExpandableArea.cpp
@@ -16,6 +16,7 @@ void SExpandableArea::Construct( const FArguments& InArgs )
 	check(InArgs._Style);
 
 	bAreaCollapsed = InArgs._InitiallyCollapsed;
+	MinWidth = InArgs._MinWidth;
 	MaxHeight = InArgs._MaxHeight;
 	OnAreaExpansionChanged = InArgs._OnAreaExpansionChanged;
 	CollapsedImage = &InArgs._Style->CollapsedImage;
@@ -39,24 +40,41 @@ void SExpandableArea::Construct( const FArguments& InArgs )
 			.ShadowOffset(FVector2D(1.0f, 1.0f));
 	}
 
+	// If the user wants the body of the expanded section to be different from the title area,
+	// then we have to do two separate borders incase the body has any transparency.
+	// Furthermore, we still need to fallback to just using one border if they do want them the
+	// same, otherwise we could introduce curved edges between the upper and lower sections.
+	const bool bBodyDiffers = InArgs._BodyBorderImage != nullptr || InArgs._BodyBorderBackgroundColor.IsSet();
+	const FSlateBrush* FullBorderImage = bBodyDiffers ? FStyleDefaults::GetNoBrush() : InArgs._BorderImage;
+	const TAttribute<FSlateColor> FullBorderBackgroundColor = bBodyDiffers ? FLinearColor::Transparent : InArgs._BorderBackgroundColor;
+	const FSlateBrush* TitleBorderImage = !bBodyDiffers ? FStyleDefaults::GetNoBrush() : InArgs._BorderImage;
+	const TAttribute<FSlateColor> TitleBorderBackgroundColor = !bBodyDiffers ? FLinearColor::Transparent : InArgs._BorderBackgroundColor;
+
 	ChildSlot
 	[
 		SNew( SBorder )
-		.BorderImage( InArgs._BorderImage )
-		.BorderBackgroundColor( InArgs._BorderBackgroundColor )
+		.BorderImage( FullBorderImage )
+		.BorderBackgroundColor( FullBorderBackgroundColor )
 		.Padding(0.0f)
 		[
 			SNew(SVerticalBox)
 			+ SVerticalBox::Slot()
 			.AutoHeight()
 			[
-				SNew( SButton )
-				.ButtonStyle(FCoreStyle::Get(), "NoBorder")
-				.ContentPadding(InArgs._HeaderPadding)
-				.ForegroundColor(FSlateColor::UseForeground())
-				.OnClicked( this, &SExpandableArea::OnHeaderClicked )
+				SNew( SBorder )
+				.BorderImage(TitleBorderImage)
+				.BorderBackgroundColor(TitleBorderBackgroundColor)
+				.Padding(0.0f)
 				[
-					ConstructHeaderWidget( InArgs, HeaderContent )
+					SNew( SButton )
+					.Cursor(InArgs._HeaderCursor.IsSet() ? InArgs._HeaderCursor : Cursor)
+					.ButtonStyle(FCoreStyle::Get(), "NoBorder")
+					.ContentPadding(InArgs._HeaderPadding)
+					.ForegroundColor(FSlateColor::UseForeground())
+					.OnClicked( this, &SExpandableArea::OnHeaderClicked )
+					[
+						ConstructHeaderWidget( InArgs, HeaderContent )
+					]
 				]
 			]
 			+ SVerticalBox::Slot()
@@ -64,7 +82,8 @@ void SExpandableArea::Construct( const FArguments& InArgs )
 			[
 				SNew(SBorder)
 				.Visibility(this, &SExpandableArea::OnGetContentVisibility)
-				.BorderImage(FStyleDefaults::GetNoBrush())
+				.BorderImage(InArgs._BodyBorderImage)
+				.BorderBackgroundColor(InArgs._BodyBorderBackgroundColor)
 				.Padding(InArgs._Padding)
 				.DesiredSizeScale(this, &SExpandableArea::GetSectionScale)
 				[
@@ -102,7 +121,7 @@ TSharedRef<SWidget> SExpandableArea::ConstructHeaderWidget( const FArguments& In
 		SNew(SHorizontalBox)
 		+ SHorizontalBox::Slot()
 		.AutoWidth()
-		.Padding(0, 0, 3, 0)
+		.Padding(InArgs._AreaTitlePadding)
 		.VAlign(VAlign_Center)
 		[
 			SNew(SImage)
@@ -182,11 +201,18 @@ FVector2D SExpandableArea::ComputeDesiredSize( float ) const
 	{
 		FVector2D SlotWidgetDesiredSize = ChildSlot.GetWidget()->GetDesiredSize() + ChildSlot.SlotPadding.Get().GetDesiredSize();
 		
+		// Only clamp if the user specified a min width
+		if( MinWidth > 0.0f )
+		{
+			SlotWidgetDesiredSize.X = FMath::Max( MinWidth, SlotWidgetDesiredSize.X );
+		}
+
 		// Only clamp if the user specified a max height
 		if( MaxHeight > 0.0f )
 		{
 			SlotWidgetDesiredSize.Y = FMath::Min( MaxHeight, SlotWidgetDesiredSize.Y );
 		}
+
 		return SlotWidgetDesiredSize;
 	}
 
diff --git a/Engine/Source/Runtime/Slate/Private/Widgets/Layout/SScrollBar.cpp b/Engine/Source/Runtime/Slate/Private/Widgets/Layout/SScrollBar.cpp
--- a/Engine/Source/Runtime/Slate/Private/Widgets/Layout/SScrollBar.cpp
+++ b/Engine/Source/Runtime/Slate/Private/Widgets/Layout/SScrollBar.cpp
@@ -285,6 +285,11 @@ bool SScrollBar::IsScrolling() const
 	return bDraggingThumb;
 }
 
+EOrientation SScrollBar::GetOrientation() const
+{
+	return Orientation;
+}
+
 void SScrollBar::SetStyle(const FScrollBarStyle* InStyle)
 {
 	const FScrollBarStyle* Style = InStyle;
diff --git a/Engine/Source/Runtime/Slate/Public/Framework/Application/SlateApplication.h b/Engine/Source/Runtime/Slate/Public/Framework/Application/SlateApplication.h
--- a/Engine/Source/Runtime/Slate/Public/Framework/Application/SlateApplication.h
+++ b/Engine/Source/Runtime/Slate/Public/Framework/Application/SlateApplication.h
@@ -828,6 +828,8 @@ public:
 	 */
 	void EnableMenuAnimations( const bool bEnableAnimations );
 
+	void SetPlatformApplication(const TSharedRef<class GenericApplication>& InPlatformApplication);
+
 	/** Set the global application icon */
 	void SetAppIcon(const FSlateBrush* const InAppIcon);
 
diff --git a/Engine/Source/Runtime/Slate/Public/Widgets/Layout/SExpandableArea.h b/Engine/Source/Runtime/Slate/Public/Widgets/Layout/SExpandableArea.h
--- a/Engine/Source/Runtime/Slate/Public/Widgets/Layout/SExpandableArea.h
+++ b/Engine/Source/Runtime/Slate/Public/Widgets/Layout/SExpandableArea.h
@@ -30,12 +30,17 @@ public:
 		: _Style( &FCoreStyle::Get().GetWidgetStyle<FExpandableAreaStyle>("ExpandableArea") )
 		, _BorderBackgroundColor( FLinearColor::White )
 		, _BorderImage( FCoreStyle::Get().GetBrush("ExpandableArea.Border") )
+		, _BodyBorderBackgroundColor()
+		, _BodyBorderImage()
 		, _AreaTitle( )
 		, _InitiallyCollapsed( false )
+		, _MinWidth( 0.0f )
 		, _MaxHeight( 0.0f )
-		, _HeaderPadding(FMargin(4.0, 2.0))
+		, _AreaTitlePadding(FMargin(0.0f, 0.0f, 3.0f, 0.0f))
+		, _HeaderPadding(FMargin(4.0f, 2.0f))
 		, _Padding( 1.0f )
 		, _AreaTitleFont( FCoreStyle::Get().GetFontStyle("ExpandableArea.TitleFont") )
+		, _HeaderCursor( )
 		{}
 		
 		/** Style used to draw this area */
@@ -44,6 +49,10 @@ public:
 		SLATE_ATTRIBUTE( FSlateColor, BorderBackgroundColor )
 		/** Border to use around the area */
 		SLATE_ARGUMENT( const FSlateBrush*, BorderImage )
+		/** Background color to apply to the body's border image. Unspecified uses BorderBackgroundColor */
+		SLATE_ATTRIBUTE( FSlateColor, BodyBorderBackgroundColor )
+		/** Border to use around the body. Unspecified uses BorderImage */
+		SLATE_ARGUMENT( const FSlateBrush*, BodyBorderImage )
 		/** Content displayed next to the expansion arrow.  This is always visible */
 		SLATE_NAMED_SLOT( FArguments, HeaderContent )
 		/** Content displayed inside the area that is expanded */
@@ -52,8 +61,12 @@ public:
 		SLATE_ATTRIBUTE( FText, AreaTitle )
 		/** Whether or not the area is initially collapsed */
 		SLATE_ARGUMENT( bool, InitiallyCollapsed )
+		/** The minimum width of the area */
+		SLATE_ARGUMENT( float, MinWidth )
 		/** The maximum height of the area */
 		SLATE_ARGUMENT( float, MaxHeight )
+		/** The title padding */
+		SLATE_ATTRIBUTE(FMargin, AreaTitlePadding)
 		/** The header padding */
 		SLATE_ATTRIBUTE( FMargin, HeaderPadding )
 		/** The content padding */
@@ -62,6 +75,9 @@ public:
 		SLATE_EVENT( FOnBooleanValueChanged, OnAreaExpansionChanged )
 		/** Sets the font used to draw the title text */
 		SLATE_ATTRIBUTE( FSlateFontInfo, AreaTitleFont )
+		/** Override for Cursor, so you can specify a different cursor for the header */
+		SLATE_ATTRIBUTE( TOptional<EMouseCursor::Type>, HeaderCursor )
+			
 	SLATE_END_ARGS()
 
 
@@ -130,6 +146,9 @@ protected:
 	/** Delegate that fires when the area becomes expanded or collapsed the user */
 	FOnBooleanValueChanged OnAreaExpansionChanged;
 
+	/** The minimum width of this area */
+	float MinWidth;
+
 	/** The maximum height of this area */
 	float MaxHeight;
 
diff --git a/Engine/Source/Runtime/Slate/Public/Widgets/Layout/SScrollBar.h b/Engine/Source/Runtime/Slate/Public/Widgets/Layout/SScrollBar.h
--- a/Engine/Source/Runtime/Slate/Public/Widgets/Layout/SScrollBar.h
+++ b/Engine/Source/Runtime/Slate/Public/Widgets/Layout/SScrollBar.h
@@ -113,6 +113,9 @@ public:
 	/** @return True if the user is scrolling by dragging the scroll bar thumb. */
 	bool IsScrolling() const;
 
+	/** @return the orientation in which the scrollbar is scrolling. */
+	EOrientation GetOrientation() const; 
+
 	/** See argument Style */
 	void SetStyle(const FScrollBarStyle* InStyle);
 
diff --git a/Engine/Source/Runtime/SlateCore/Private/Widgets/SWidget.cpp b/Engine/Source/Runtime/SlateCore/Private/Widgets/SWidget.cpp
--- a/Engine/Source/Runtime/SlateCore/Private/Widgets/SWidget.cpp
+++ b/Engine/Source/Runtime/SlateCore/Private/Widgets/SWidget.cpp
@@ -260,11 +260,23 @@ FReply SWidget::OnMouseButtonDoubleClick(const FGeometry& MyGeometry, const FPoi
 void SWidget::OnMouseEnter( const FGeometry& MyGeometry, const FPointerEvent& MouseEvent )
 {
 	bIsHovered = true;
+
+	if (MouseEnterHandler.IsBound())
+	{
+		// A valid handler is assigned; let it handle the event.
+		MouseEnterHandler.Execute(MyGeometry, MouseEvent);
+	}
 }
 
 void SWidget::OnMouseLeave( const FPointerEvent& MouseEvent )
 {
 	bIsHovered = false;
+
+	if (MouseLeaveHandler.IsBound())
+	{
+		// A valid handler is assigned; let it handle the event.
+		MouseLeaveHandler.Execute(MouseEvent);
+	}
 }
 
 FReply SWidget::OnMouseWheel( const FGeometry& MyGeometry, const FPointerEvent& MouseEvent )
@@ -849,3 +861,13 @@ void SWidget::SetOnMouseDoubleClick(FPointerEventHandler EventHandler)
 {
 	MouseDoubleClickHandler = EventHandler;
 }
+
+void SWidget::SetOnMouseEnter(FNoReplyPointerEventHandler EventHandler)
+{
+	MouseEnterHandler = EventHandler;
+}
+
+void SWidget::SetOnMouseLeave(FSimpleNoReplyPointerEventHandler EventHandler)
+{
+	MouseLeaveHandler = EventHandler;
+}
diff --git a/Engine/Source/Runtime/SlateCore/Public/Application/SlateApplicationBase.h b/Engine/Source/Runtime/SlateCore/Public/Application/SlateApplicationBase.h
--- a/Engine/Source/Runtime/SlateCore/Public/Application/SlateApplicationBase.h
+++ b/Engine/Source/Runtime/SlateCore/Public/Application/SlateApplicationBase.h
@@ -518,6 +518,11 @@ public:
 		return PlatformApplication->Cursor;
 	}
 
+	TSharedPtr<class GenericApplication> GetPlatformApplication()
+	{
+		return PlatformApplication;
+	}
+
 protected:
 	/** multicast delegate to broadcast when a global invalidate is requested */
 	FOnGlobalInvalidate OnGlobalInvalidateEvent;
diff --git a/Engine/Source/Runtime/SlateCore/Public/Widgets/SWidget.h b/Engine/Source/Runtime/SlateCore/Public/Widgets/SWidget.h
--- a/Engine/Source/Runtime/SlateCore/Public/Widgets/SWidget.h
+++ b/Engine/Source/Runtime/SlateCore/Public/Widgets/SWidget.h
@@ -40,6 +40,18 @@ DECLARE_DELEGATE_RetVal_TwoParams(
 	/** The Mouse Event that we are processing */
 	const FPointerEvent&)
 
+DECLARE_DELEGATE_TwoParams(
+	FNoReplyPointerEventHandler,
+	/** The geometry of the widget*/
+	const FGeometry&,
+	/** The Mouse Event that we are processing */
+	const FPointerEvent&)
+
+DECLARE_DELEGATE_OneParam(
+	FSimpleNoReplyPointerEventHandler,
+	/** The Mouse Event that we are processing */
+	const FPointerEvent&)
+
 enum class EPopupMethod : uint8;
 
 class SLATECORE_API FSlateControlledConstruction
@@ -990,6 +1002,12 @@ public:
 	/** See OnMouseDoubleClick event */
 	void SetOnMouseDoubleClick(FPointerEventHandler EventHandler);
 
+	/** See OnMouseEnter event */
+	void SetOnMouseEnter(FNoReplyPointerEventHandler EventHandler);
+
+	/** See OnMouseLeave event */
+	void SetOnMouseLeave(FSimpleNoReplyPointerEventHandler EventHandler);
+
 public:
 
 	// Widget Inspector and debugging methods
@@ -1257,6 +1275,8 @@ private:
 	FPointerEventHandler MouseButtonUpHandler;
 	FPointerEventHandler MouseMoveHandler;
 	FPointerEventHandler MouseDoubleClickHandler;
+	FNoReplyPointerEventHandler MouseEnterHandler;
+	FSimpleNoReplyPointerEventHandler MouseLeaveHandler;
 };
 
 //=================================================================
diff --git a/Engine/Source/Runtime/WebBrowser/Private/Android/AndroidJSScripting.cpp b/Engine/Source/Runtime/WebBrowser/Private/Android/AndroidJSScripting.cpp
--- a/Engine/Source/Runtime/WebBrowser/Private/Android/AndroidJSScripting.cpp
+++ b/Engine/Source/Runtime/WebBrowser/Private/Android/AndroidJSScripting.cpp
@@ -329,9 +329,9 @@ FString FAndroidJSScripting::ConvertObject(UObject* Object)
 		{
 			first = false;
 		}
-		Result.Append(*Function->GetName());
+		Result.Append(*GetBindingName(Function));
 		Result.Append(TEXT(": function "));
-		Result.Append(*Function->GetName());
+		Result.Append(*GetBindingName(Function));
 		Result.Append(TEXT(" ("));
 
 		bool firstArg = true;
@@ -351,7 +351,7 @@ FString FAndroidJSScripting::ConvertObject(UObject* Object)
 					{
 						firstArg = false;
 					}
-					Result.Append(*Param->GetName());
+					Result.Append(*GetBindingName(Param));
 				}
 			}
 		}
@@ -510,7 +510,7 @@ bool FAndroidJSScripting::HandleExecuteUObjectMethodMessage(const TArray<FString
 			FString ResultJS = ReturnBackend.ToString();
 
 			ResultJS.Append(TEXT("['"));
-			ResultJS.Append(ReturnParam->GetName().ReplaceCharWithEscapedChar());
+			ResultJS.Append(GetBindingName(ReturnParam).ReplaceCharWithEscapedChar());
 			ResultJS.Append(TEXT("']"));
 
 			InvokeJSFunctionRaw(ResultCallbackId, ResultJS, false);
@@ -545,7 +545,7 @@ void FAndroidJSScripting::PageLoaded(TSharedRef<class FAndroidWebBrowserWindow>
 	InWindow->ExecuteJavascript(Script);
 }
 
-FAndroidJSScripting::FAndroidJSScripting()
-	: FWebJSScripting()
+FAndroidJSScripting::FAndroidJSScripting(bool bJSBindingToLoweringEnabled)
+	: FWebJSScripting(bJSBindingToLoweringEnabled)
 {
 }
diff --git a/Engine/Source/Runtime/WebBrowser/Private/Android/AndroidJSScripting.h b/Engine/Source/Runtime/WebBrowser/Private/Android/AndroidJSScripting.h
--- a/Engine/Source/Runtime/WebBrowser/Private/Android/AndroidJSScripting.h
+++ b/Engine/Source/Runtime/WebBrowser/Private/Android/AndroidJSScripting.h
@@ -18,7 +18,7 @@ class FAndroidJSScripting
 public:
 	static const FString JSMessageTag;
 
-	FAndroidJSScripting();
+	FAndroidJSScripting(bool bJSBindingToLoweringEnabled);
 
 	virtual void BindUObject(const FString& Name, UObject* Object, bool bIsPermanent = true) override;
 	virtual void UnbindUObject(const FString& Name, UObject* Object = nullptr, bool bIsPermanent = true) override;
diff --git a/Engine/Source/Runtime/WebBrowser/Private/Android/AndroidJSStructSerializerBackend.cpp b/Engine/Source/Runtime/WebBrowser/Private/Android/AndroidJSStructSerializerBackend.cpp
--- a/Engine/Source/Runtime/WebBrowser/Private/Android/AndroidJSStructSerializerBackend.cpp
+++ b/Engine/Source/Runtime/WebBrowser/Private/Android/AndroidJSStructSerializerBackend.cpp
@@ -36,7 +36,7 @@ void FAndroidJSStructSerializerBackend::WriteUObject(const FStructSerializerStat
 	}
 	else
 	{
-		GetWriter()->WriteRawJSONValue(State.ValueProperty->GetName(), RawValue);
+		GetWriter()->WriteRawJSONValue(Scripting->GetBindingName(State.ValueProperty), RawValue);
 	}
 }
 
diff --git a/Engine/Source/Runtime/WebBrowser/Private/Android/AndroidWebBrowserWindow.cpp b/Engine/Source/Runtime/WebBrowser/Private/Android/AndroidWebBrowserWindow.cpp
--- a/Engine/Source/Runtime/WebBrowser/Private/Android/AndroidWebBrowserWindow.cpp
+++ b/Engine/Source/Runtime/WebBrowser/Private/Android/AndroidWebBrowserWindow.cpp
@@ -19,12 +19,12 @@ namespace {
 
 }
 
-FAndroidWebBrowserWindow::FAndroidWebBrowserWindow(FString InUrl, TOptional<FString> InContentsToLoad, bool InShowErrorMessage, bool InThumbMouseButtonNavigation, bool InUseTransparency)
+FAndroidWebBrowserWindow::FAndroidWebBrowserWindow(FString InUrl, TOptional<FString> InContentsToLoad, bool InShowErrorMessage, bool InThumbMouseButtonNavigation, bool InUseTransparency, bool bInJSBindingToLoweringEnabled)
 	: CurrentUrl(MoveTemp(InUrl))
 	, ContentsToLoad(MoveTemp(InContentsToLoad))
 	, DocumentState(EWebBrowserDocumentState::NoDocument)
 	, ErrorCode(0)
-	, Scripting(new FAndroidJSScripting())
+	, Scripting(new FAndroidJSScripting(bInJSBindingToLoweringEnabled))
 {
 }
 
@@ -128,6 +128,10 @@ FReply FAndroidWebBrowserWindow::OnMouseMove(const FGeometry& MyGeometry, const
 	return FReply::Unhandled();
 }
 
+void FAndroidWebBrowserWindow::OnMouseLeave(const FPointerEvent& MouseEvent)
+{
+}
+
 FReply FAndroidWebBrowserWindow::OnMouseWheel(const FGeometry& MyGeometry, const FPointerEvent& MouseEvent, bool bIsPopup)
 {
 	return FReply::Unhandled();
diff --git a/Engine/Source/Runtime/WebBrowser/Private/Android/AndroidWebBrowserWindow.h b/Engine/Source/Runtime/WebBrowser/Private/Android/AndroidWebBrowserWindow.h
--- a/Engine/Source/Runtime/WebBrowser/Private/Android/AndroidWebBrowserWindow.h
+++ b/Engine/Source/Runtime/WebBrowser/Private/Android/AndroidWebBrowserWindow.h
@@ -24,13 +24,16 @@ class FAndroidWebBrowserWindow
 
 private:
 	/**
-	* Creates and initializes a new instance.
-	*
-	* @param InUrl The Initial URL that will be loaded.
-	* @param InContentsToLoad Optional string to load as a web page.
-	* @param InShowErrorMessage Whether to show an error message in case of loading errors.
-	*/
-	FAndroidWebBrowserWindow(FString InUrl, TOptional<FString> InContentsToLoad, bool ShowErrorMessage, bool bThumbMouseButtonNavigation, bool bUseTransparency);
+	 * Creates and initializes a new instance.
+	 *
+	 * @param Url The Initial URL that will be loaded.
+	 * @param ContentsToLoad Optional string to load as a web page.
+	 * @param bShowErrorMessage Whether to show an error message in case of loading errors.
+	 * @param bThumbMouseButtonNavigation Whether to allow forward and back navigation via the mouse thumb buttons.
+	 * @param bUseTransparency Whether to enable transparency.
+	 * @param bJSBindingToLoweringEnabled Whether we ToLower all JavaScript member names.
+	 */
+	FAndroidWebBrowserWindow(FString Url, TOptional<FString> ContentsToLoad, bool bShowErrorMessage, bool bThumbMouseButtonNavigation, bool bUseTransparency, bool bJSBindingToLoweringEnabled);
 
 	/**
 	 * Create the SWidget for this WebBrowserWindow
@@ -63,6 +66,7 @@ public:
 	virtual FReply OnMouseButtonUp(const FGeometry& MyGeometry, const FPointerEvent& MouseEvent, bool bIsPopup) override;
 	virtual FReply OnMouseButtonDoubleClick(const FGeometry& MyGeometry, const FPointerEvent& MouseEvent, bool bIsPopup) override;
 	virtual FReply OnMouseMove(const FGeometry& MyGeometry, const FPointerEvent& MouseEvent, bool bIsPopup) override;
+	virtual void OnMouseLeave(const FPointerEvent& MouseEvent) override;
 	virtual FReply OnMouseWheel(const FGeometry& MyGeometry, const FPointerEvent& MouseEvent, bool bIsPopup) override;
 	virtual void OnFocus(bool SetFocus, bool bIsPopup) override;
 	virtual void OnCaptureLost() override;
diff --git a/Engine/Source/Runtime/WebBrowser/Private/CEF/CEFJSScripting.cpp b/Engine/Source/Runtime/WebBrowser/Private/CEF/CEFJSScripting.cpp
--- a/Engine/Source/Runtime/WebBrowser/Private/CEF/CEFJSScripting.cpp
+++ b/Engine/Source/Runtime/WebBrowser/Private/CEF/CEFJSScripting.cpp
@@ -53,7 +53,7 @@ CefRefPtr<CefDictionaryValue> FCEFJSScripting::ConvertStruct(UStruct* TypeInfo,
 
 	CefRefPtr<CefDictionaryValue> Result = CefDictionaryValue::Create();
 	Result->SetString("$type", "struct");
-	Result->SetString("$ue4Type", *TypeInfo->GetName());
+	Result->SetString("$ue4Type", *GetBindingName(TypeInfo));
 	Result->SetDictionary("$value", Backend.GetResult());
 	return Result;
 }
@@ -69,7 +69,7 @@ CefRefPtr<CefDictionaryValue> FCEFJSScripting::ConvertObject(UObject* Object)
 	for (TFieldIterator<UFunction> FunctionIt(Class, EFieldIteratorFlags::IncludeSuper); FunctionIt; ++FunctionIt)
 	{
 		UFunction* Function = *FunctionIt;
-		MethodNames->SetString(MethodIndex++, *Function->GetName());
+		MethodNames->SetString(MethodIndex++, *GetBindingName(Function));
 	}
 
 	Result->SetString("$type", "uobject");
@@ -271,7 +271,7 @@ bool FCEFJSScripting::HandleExecuteUObjectMethodMessage(CefRefPtr<CefListValue>
 					}
 					else
 					{
-						CopyContainerValue(NamedArgs, CefArgs, CefString(*Param->GetName()), CurrentArg);
+						CopyContainerValue(NamedArgs, CefArgs, CefString(*GetBindingName(Param)), CurrentArg);
 						CurrentArg++;
 					}
 				}
@@ -311,7 +311,7 @@ bool FCEFJSScripting::HandleExecuteUObjectMethodMessage(CefRefPtr<CefListValue>
 			CefRefPtr<CefDictionaryValue> ResultDict = ReturnBackend.GetResult();
 
 			// Extract the single return value from the serialized dictionary to an array
-			CopyContainerValue(Results, ResultDict, 0, CefString(*ReturnParam->GetName()));
+			CopyContainerValue(Results, ResultDict, 0, *GetBindingName(ReturnParam));
 		}
 		InvokeJSFunction(ResultCallbackId, Results, false);
 	}
diff --git a/Engine/Source/Runtime/WebBrowser/Private/CEF/CEFJSScripting.h b/Engine/Source/Runtime/WebBrowser/Private/CEF/CEFJSScripting.h
--- a/Engine/Source/Runtime/WebBrowser/Private/CEF/CEFJSScripting.h
+++ b/Engine/Source/Runtime/WebBrowser/Private/CEF/CEFJSScripting.h
@@ -36,8 +36,8 @@ class FCEFJSScripting
 	, public TSharedFromThis<FCEFJSScripting>
 {
 public:
-	FCEFJSScripting(CefRefPtr<CefBrowser> Browser)
-		: FWebJSScripting()
+	FCEFJSScripting(CefRefPtr<CefBrowser> Browser, bool bJSBindingToLoweringEnabled)
+		: FWebJSScripting(bJSBindingToLoweringEnabled)
 		, InternalCefBrowser(Browser)
 	{}
 
diff --git a/Engine/Source/Runtime/WebBrowser/Private/CEF/CEFJSStructSerializerBackend.cpp b/Engine/Source/Runtime/WebBrowser/Private/CEF/CEFJSStructSerializerBackend.cpp
--- a/Engine/Source/Runtime/WebBrowser/Private/CEF/CEFJSStructSerializerBackend.cpp
+++ b/Engine/Source/Runtime/WebBrowser/Private/CEF/CEFJSStructSerializerBackend.cpp
@@ -16,7 +16,7 @@ void FCEFJSStructSerializerBackend::AddNull(const FStructSerializerState& State)
 	StackItem& Current = Stack.Top();
 	switch (Current.Kind) {
 		case StackItem::STYPE_DICTIONARY:
-			Current.DictionaryValue->SetNull(*State.ValueProperty->GetName());
+			Current.DictionaryValue->SetNull(*Scripting->GetBindingName(State.ValueProperty));
 			break;
 		case StackItem::STYPE_LIST:
 			Current.ListValue->SetNull(Current.ListValue->GetSize());
@@ -30,7 +30,7 @@ void FCEFJSStructSerializerBackend::Add(const FStructSerializerState& State, boo
 	StackItem& Current = Stack.Top();
 	switch (Current.Kind) {
 		case StackItem::STYPE_DICTIONARY:
-			Current.DictionaryValue->SetBool(*State.ValueProperty->GetName(), Value);
+			Current.DictionaryValue->SetBool(*Scripting->GetBindingName(State.ValueProperty), Value);
 			break;
 		case StackItem::STYPE_LIST:
 			Current.ListValue->SetBool(Current.ListValue->GetSize(), Value);
@@ -44,7 +44,7 @@ void FCEFJSStructSerializerBackend::Add(const FStructSerializerState& State, int
 	StackItem& Current = Stack.Top();
 	switch (Current.Kind) {
 		case StackItem::STYPE_DICTIONARY:
-			Current.DictionaryValue->SetInt(*State.ValueProperty->GetName(), Value);
+			Current.DictionaryValue->SetInt(*Scripting->GetBindingName(State.ValueProperty), Value);
 			break;
 		case StackItem::STYPE_LIST:
 			Current.ListValue->SetInt(Current.ListValue->GetSize(), Value);
@@ -58,7 +58,7 @@ void FCEFJSStructSerializerBackend::Add(const FStructSerializerState& State, dou
 	StackItem& Current = Stack.Top();
 	switch (Current.Kind) {
 		case StackItem::STYPE_DICTIONARY:
-			Current.DictionaryValue->SetDouble(*State.ValueProperty->GetName(), Value);
+			Current.DictionaryValue->SetDouble(*Scripting->GetBindingName(State.ValueProperty), Value);
 			break;
 		case StackItem::STYPE_LIST:
 			Current.ListValue->SetDouble(Current.ListValue->GetSize(), Value);
@@ -72,7 +72,7 @@ void FCEFJSStructSerializerBackend::Add(const FStructSerializerState& State, FSt
 	StackItem& Current = Stack.Top();
 	switch (Current.Kind) {
 		case StackItem::STYPE_DICTIONARY:
-			Current.DictionaryValue->SetString(*State.ValueProperty->GetName(), *Value);
+			Current.DictionaryValue->SetString(*Scripting->GetBindingName(State.ValueProperty), *Value);
 			break;
 		case StackItem::STYPE_LIST:
 			Current.ListValue->SetString(Current.ListValue->GetSize(), *Value);
@@ -86,7 +86,7 @@ void FCEFJSStructSerializerBackend::Add(const FStructSerializerState& State, UOb
 	StackItem& Current = Stack.Top();
 	switch (Current.Kind) {
 		case StackItem::STYPE_DICTIONARY:
-			Current.DictionaryValue->SetDictionary(*State.ValueProperty->GetName(), Scripting->ConvertObject(Value));
+			Current.DictionaryValue->SetDictionary(*Scripting->GetBindingName(State.ValueProperty), Scripting->ConvertObject(Value));
 			break;
 		case StackItem::STYPE_LIST:
 			Current.ListValue->SetDictionary(Current.ListValue->GetSize(), Scripting->ConvertObject(Value));
@@ -101,7 +101,7 @@ void FCEFJSStructSerializerBackend::Add(const FStructSerializerState& State, UOb
 void FCEFJSStructSerializerBackend::BeginArray(const FStructSerializerState& State)
 {
 	CefRefPtr<CefListValue> ListValue = CefListValue::Create();
-	Stack.Push(StackItem(State.ValueProperty->GetName(), ListValue));
+	Stack.Push(StackItem(Scripting->GetBindingName(State.ValueProperty), ListValue));
 }
 
 
@@ -118,7 +118,7 @@ void FCEFJSStructSerializerBackend::BeginStructure(const FStructSerializerState&
 	else if (State.ValueProperty != nullptr)
 	{
 		CefRefPtr<CefDictionaryValue> DictionaryValue = CefDictionaryValue::Create();
-		Stack.Push(StackItem(State.ValueProperty->GetName(), DictionaryValue));
+		Stack.Push(StackItem(Scripting->GetBindingName(State.ValueProperty), DictionaryValue));
 	}
 	else
 	{
diff --git a/Engine/Source/Runtime/WebBrowser/Private/CEF/CEFWebBrowserWindow.cpp b/Engine/Source/Runtime/WebBrowser/Private/CEF/CEFWebBrowserWindow.cpp
--- a/Engine/Source/Runtime/WebBrowser/Private/CEF/CEFWebBrowserWindow.cpp
+++ b/Engine/Source/Runtime/WebBrowser/Private/CEF/CEFWebBrowserWindow.cpp
@@ -364,7 +364,7 @@ private:
 
 
 
-FCEFWebBrowserWindow::FCEFWebBrowserWindow(CefRefPtr<CefBrowser> InBrowser, CefRefPtr<FCEFBrowserHandler> InHandler, FString InUrl, TOptional<FString> InContentsToLoad, bool InShowErrorMessage, bool InThumbMouseButtonNavigation, bool InUseTransparency)
+FCEFWebBrowserWindow::FCEFWebBrowserWindow(CefRefPtr<CefBrowser> InBrowser, CefRefPtr<FCEFBrowserHandler> InHandler, FString InUrl, TOptional<FString> InContentsToLoad, bool bInShowErrorMessage, bool bInThumbMouseButtonNavigation, bool bInUseTransparency, bool bInJSBindingToLoweringEnabled)
 	: DocumentState(EWebBrowserDocumentState::NoDocument)
 	, InternalCefBrowser(InBrowser)
 	, WebBrowserHandler(InHandler)
@@ -373,9 +373,9 @@ FCEFWebBrowserWindow::FCEFWebBrowserWindow(CefRefPtr<CefBrowser> InBrowser, CefR
 	, bIsClosing(false)
 	, bIsInitialized(false)
 	, ContentsToLoad(InContentsToLoad)
-	, ShowErrorMessage(InShowErrorMessage)
-	, bThumbMouseButtonNavigation(InThumbMouseButtonNavigation)
-	, bUseTransparency(InUseTransparency)
+	, bShowErrorMessage(bInShowErrorMessage)
+	, bThumbMouseButtonNavigation(bInThumbMouseButtonNavigation)
+	, bUseTransparency(bInUseTransparency)
 	, Cursor(EMouseCursor::Default)
 	, bIsDisabled(false)
 	, bIsHidden(false)
@@ -390,7 +390,7 @@ FCEFWebBrowserWindow::FCEFWebBrowserWindow(CefRefPtr<CefBrowser> InBrowser, CefR
 	, bPopupHasFocus(false)
 	, bRecoverFromRenderProcessCrash(false)
 	, ErrorCode(0)
-	, Scripting(new FCEFJSScripting(InBrowser))
+	, Scripting(new FCEFJSScripting(InBrowser, bInJSBindingToLoweringEnabled))
 {
 	check(InBrowser.get() != nullptr);
 
@@ -1225,6 +1225,12 @@ FReply FCEFWebBrowserWindow::OnMouseMove(const FGeometry& MyGeometry, const FPoi
 	return Reply;
 }
 
+void FCEFWebBrowserWindow::OnMouseLeave(const FPointerEvent& MouseEvent)
+{
+	// Ensure we clear any tooltips if the mouse leaves the window.
+	SetToolTip(CefString());
+}
+
 FReply FCEFWebBrowserWindow::OnMouseWheel(const FGeometry& MyGeometry, const FPointerEvent& MouseEvent, bool bIsPopup)
 {
 	FReply Reply = FReply::Unhandled();
@@ -1508,9 +1514,13 @@ void FCEFWebBrowserWindow::OnCursorChange(CefCursorHandle CefCursor, CefRenderHa
 		// Platform specific support for native cursor types
 		default:
 			{
-				FPlatformCursor* PlatformCursor = (FPlatformCursor*)FSlateApplication::Get().GetPlatformCursor().Get();
-				PlatformCursor->SetCustomShape(CefCursor);
-				Cursor = EMouseCursor::Custom;
+				TSharedPtr<ICursor> PlatformCursor = FSlateApplication::Get().GetPlatformCursor();
+
+				if (PlatformCursor.IsValid())
+				{
+					PlatformCursor->SetCustomShape((void*)CefCursor);
+					Cursor = EMouseCursor::Custom;
+				}
 			}
 			break;
 		#else
diff --git a/Engine/Source/Runtime/WebBrowser/Private/CEF/CEFWebBrowserWindow.h b/Engine/Source/Runtime/WebBrowser/Private/CEF/CEFWebBrowserWindow.h
--- a/Engine/Source/Runtime/WebBrowser/Private/CEF/CEFWebBrowserWindow.h
+++ b/Engine/Source/Runtime/WebBrowser/Private/CEF/CEFWebBrowserWindow.h
@@ -87,12 +87,16 @@ private:
 	/**
 	 * Creates and initializes a new instance.
 	 *
-	 * @param InBrowser The CefBrowser object representing this browser window.
-	 * @param InUrl The Initial URL that will be loaded.
-	 * @param InContentsToLoad Optional string to load as a web page.
-	 * @param InShowErrorMessage Whether to show an error message in case of loading errors.
+	 * @param Browser The CefBrowser object representing this browser window.
+	 * @param Handler Pointer to the CEF handler for this window.
+	 * @param Url The Initial URL that will be loaded.
+	 * @param ContentsToLoad Optional string to load as a web page.
+	 * @param bShowErrorMessage Whether to show an error message in case of loading errors.
+	 * @param bThumbMouseButtonNavigation Whether to allow forward and back navigation via the mouse thumb buttons.
+	 * @param bUseTransparency Whether to enable transparency.
+	 * @param bJSBindingToLoweringEnabled Whether we ToLower all JavaScript member names.
 	 */
-	FCEFWebBrowserWindow(CefRefPtr<CefBrowser> InBrowser, CefRefPtr<FCEFBrowserHandler> InHandler, FString InUrl, TOptional<FString> InContentsToLoad, bool ShowErrorMessage, bool bThumbMouseButtonNavigation, bool bUseTransparency);
+	FCEFWebBrowserWindow(CefRefPtr<CefBrowser> Browser, CefRefPtr<FCEFBrowserHandler> Handler, FString Url, TOptional<FString> ContentsToLoad, bool bShowErrorMessage, bool bThumbMouseButtonNavigation, bool bUseTransparency, bool bJSBindingToLoweringEnabled);
 
 	/**
 	 * Create the SWidget for this WebBrowserWindow
@@ -103,7 +107,7 @@ public:
 	/** Virtual Destructor. */
 	virtual ~FCEFWebBrowserWindow();
 
-	bool IsShowingErrorMessages() { return ShowErrorMessage; }
+	bool IsShowingErrorMessages() { return bShowErrorMessage; }
 	bool IsThumbMouseButtonNavigationEnabled() { return bThumbMouseButtonNavigation; }
 	bool UseTransparency() { return bUseTransparency; }
 
@@ -129,6 +133,8 @@ public:
 	virtual FReply OnMouseButtonUp(const FGeometry& MyGeometry, const FPointerEvent& MouseEvent, bool bIsPopup) override;
 	virtual FReply OnMouseButtonDoubleClick(const FGeometry& MyGeometry, const FPointerEvent& MouseEvent, bool bIsPopup) override;
 	virtual FReply OnMouseMove(const FGeometry& MyGeometry, const FPointerEvent& MouseEvent, bool bIsPopup) override;
+	virtual void OnMouseLeave(const FPointerEvent& MouseEvent) override;
+
 	virtual FReply OnMouseWheel(const FGeometry& MyGeometry, const FPointerEvent& MouseEvent, bool bIsPopup) override;
 	virtual void OnFocus(bool SetFocus, bool bIsPopup) override;
 	virtual void OnCaptureLost() override;
@@ -480,7 +486,7 @@ private:
 	FOnDocumentStateChanged DocumentStateChangedEvent;
 
 	/** Whether to show an error message in case of loading errors. */
-	bool ShowErrorMessage;
+	bool bShowErrorMessage;
 
 	/** Whether to allow forward and back navigation via the mouse thumb buttons. */
 	bool bThumbMouseButtonNavigation;
diff --git a/Engine/Source/Runtime/WebBrowser/Private/IOS/IOSPlatformWebBrowser.cpp b/Engine/Source/Runtime/WebBrowser/Private/IOS/IOSPlatformWebBrowser.cpp
--- a/Engine/Source/Runtime/WebBrowser/Private/IOS/IOSPlatformWebBrowser.cpp
+++ b/Engine/Source/Runtime/WebBrowser/Private/IOS/IOSPlatformWebBrowser.cpp
@@ -298,6 +298,10 @@ FReply FWebBrowserWindow::OnMouseMove(const FGeometry& MyGeometry, const FPointe
 	return FReply::Unhandled();
 }
 
+void FWebBrowserWindow::OnMouseLeave(const FPointerEvent& MouseEvent)
+{
+}
+
 FReply FWebBrowserWindow::OnMouseWheel(const FGeometry& MyGeometry, const FPointerEvent& MouseEvent, bool bIsPopup)
 {
 	return FReply::Unhandled();
diff --git a/Engine/Source/Runtime/WebBrowser/Private/IOS/IOSPlatformWebBrowser.h b/Engine/Source/Runtime/WebBrowser/Private/IOS/IOSPlatformWebBrowser.h
--- a/Engine/Source/Runtime/WebBrowser/Private/IOS/IOSPlatformWebBrowser.h
+++ b/Engine/Source/Runtime/WebBrowser/Private/IOS/IOSPlatformWebBrowser.h
@@ -93,6 +93,7 @@ public:
 	virtual FReply OnMouseButtonUp(const FGeometry& MyGeometry, const FPointerEvent& MouseEvent, bool bIsPopup) override;
 	virtual FReply OnMouseButtonDoubleClick(const FGeometry& MyGeometry, const FPointerEvent& MouseEvent, bool bIsPopup) override;
 	virtual FReply OnMouseMove(const FGeometry& MyGeometry, const FPointerEvent& MouseEvent, bool bIsPopup) override;
+	virtual void OnMouseLeave(const FPointerEvent& MouseEvent) override;
 	virtual FReply OnMouseWheel(const FGeometry& MyGeometry, const FPointerEvent& MouseEvent, bool bIsPopup) override;
 	virtual void OnFocus(bool SetFocus, bool bIsPopup) override;
 	virtual void OnCaptureLost() override;
diff --git a/Engine/Source/Runtime/WebBrowser/Private/SWebBrowserView.cpp b/Engine/Source/Runtime/WebBrowser/Private/SWebBrowserView.cpp
--- a/Engine/Source/Runtime/WebBrowser/Private/SWebBrowserView.cpp
+++ b/Engine/Source/Runtime/WebBrowser/Private/SWebBrowserView.cpp
@@ -205,6 +205,17 @@ void SWebBrowserView::HandleWindowDeactivated()
 	}
 }
 
+void SWebBrowserView::HandleWindowActivated()
+{
+	if (BrowserViewport.IsValid())
+	{
+		if (HasAnyUserFocusOrFocusedDescendants())
+		{
+			BrowserViewport->OnFocusReceived(FFocusEvent());
+		}	
+	}
+}
+
 void SWebBrowserView::LoadURL(FString NewURL)
 {
 	AddressBarUrl = FText::FromString(NewURL);
@@ -335,6 +346,7 @@ void SWebBrowserView::SetupParentWindowHandlers()
 	if (SlateParentWindowPtr.IsValid() && BrowserWindow.IsValid())
 	{
 		SlateParentWindowPtr.Pin()->GetOnWindowDeactivatedEvent().AddSP(this, &SWebBrowserView::HandleWindowDeactivated);
+		SlateParentWindowPtr.Pin()->GetOnWindowActivatedEvent().AddSP(this, &SWebBrowserView::HandleWindowActivated);
 	}
 }
 
diff --git a/Engine/Source/Runtime/WebBrowser/Private/WebBrowserSingleton.cpp b/Engine/Source/Runtime/WebBrowser/Private/WebBrowserSingleton.cpp
--- a/Engine/Source/Runtime/WebBrowser/Private/WebBrowserSingleton.cpp
+++ b/Engine/Source/Runtime/WebBrowser/Private/WebBrowserSingleton.cpp
@@ -208,6 +208,7 @@ FWebBrowserSingleton::FWebBrowserSingleton()
 	: WebBrowserWindowFactory(MakeShareable(new FNoWebBrowserWindowFactory()))
 #endif
 	, bDevToolsShortcutEnabled(UE_BUILD_DEBUG)
+	, bJSBindingsToLoweringEnabled(true)
 {
 #if WITH_CEF3
 	// The FWebBrowserSingleton must be initialized on the game thread
@@ -357,7 +358,7 @@ TSharedPtr<IWebBrowserWindow> FWebBrowserSingleton::CreateBrowserWindow(
 	bool bThumbMouseButtonNavigation = BrowserWindowParent->IsThumbMouseButtonNavigationEnabled();
 	bool bUseTransparency = BrowserWindowParent->UseTransparency();
 	FString InitialURL = BrowserWindowInfo->Browser->GetMainFrame()->GetURL().ToWString().c_str();
-	TSharedPtr<FCEFWebBrowserWindow> NewBrowserWindow(new FCEFWebBrowserWindow(BrowserWindowInfo->Browser, BrowserWindowInfo->Handler, InitialURL, ContentsToLoad, bShowErrorMessage, bThumbMouseButtonNavigation, bUseTransparency));
+	TSharedPtr<FCEFWebBrowserWindow> NewBrowserWindow(new FCEFWebBrowserWindow(BrowserWindowInfo->Browser, BrowserWindowInfo->Handler, InitialURL, ContentsToLoad, bShowErrorMessage, bThumbMouseButtonNavigation, bUseTransparency, bJSBindingsToLoweringEnabled));
 	BrowserWindowInfo->Handler->SetBrowserWindow(NewBrowserWindow);
 
 	WindowInterfaces.Add(NewBrowserWindow);
@@ -469,7 +470,8 @@ TSharedPtr<IWebBrowserWindow> FWebBrowserSingleton::CreateBrowserWindow(const FC
 				WindowSettings.ContentsToLoad,
 				WindowSettings.bShowErrorMessage,
 				WindowSettings.bThumbMouseButtonNavigation,
-				WindowSettings.bUseTransparency));
+				WindowSettings.bUseTransparency,
+				bJSBindingsToLoweringEnabled));
 			NewHandler->SetBrowserWindow(NewBrowserWindow);
 
 			WindowInterfaces.Add(NewBrowserWindow);
@@ -483,7 +485,8 @@ TSharedPtr<IWebBrowserWindow> FWebBrowserSingleton::CreateBrowserWindow(const FC
 		WindowSettings.ContentsToLoad,
 		WindowSettings.bShowErrorMessage,
 		WindowSettings.bThumbMouseButtonNavigation,
-		WindowSettings.bUseTransparency));
+		WindowSettings.bUseTransparency,
+		bJSBindingsToLoweringEnabled));
 
 	//WindowInterfaces.Add(NewBrowserWindow);
 	return NewBrowserWindow;
diff --git a/Engine/Source/Runtime/WebBrowser/Private/WebBrowserSingleton.h b/Engine/Source/Runtime/WebBrowser/Private/WebBrowserSingleton.h
--- a/Engine/Source/Runtime/WebBrowser/Private/WebBrowserSingleton.h
+++ b/Engine/Source/Runtime/WebBrowser/Private/WebBrowserSingleton.h
@@ -106,6 +106,11 @@ public:
 		bDevToolsShortcutEnabled = Value;
 	}
 
+	virtual void SetJSBindingToLoweringEnabled(bool bEnabled) override
+	{
+		bJSBindingsToLoweringEnabled = bEnabled;
+	}
+
 public:
 
 	// FTickerObjectBase Interface
@@ -131,6 +136,8 @@ private:
 
 	bool bDevToolsShortcutEnabled;
 
+	bool bJSBindingsToLoweringEnabled;
+
 };
 
 PRAGMA_ENABLE_DEPRECATION_WARNINGS
diff --git a/Engine/Source/Runtime/WebBrowser/Private/WebBrowserViewport.cpp b/Engine/Source/Runtime/WebBrowser/Private/WebBrowserViewport.cpp
--- a/Engine/Source/Runtime/WebBrowser/Private/WebBrowserViewport.cpp
+++ b/Engine/Source/Runtime/WebBrowser/Private/WebBrowserViewport.cpp
@@ -77,6 +77,7 @@ void FWebBrowserViewport::OnMouseEnter(const FGeometry& MyGeometry, const FPoint
 
 void FWebBrowserViewport::OnMouseLeave(const FPointerEvent& MouseEvent)
 {
+	WebBrowserWindow->OnMouseLeave(MouseEvent);
 }
 
 FReply FWebBrowserViewport::OnMouseMove(const FGeometry& MyGeometry, const FPointerEvent& MouseEvent)
diff --git a/Engine/Source/Runtime/WebBrowser/Private/WebJSScripting.h b/Engine/Source/Runtime/WebBrowser/Private/WebJSScripting.h
--- a/Engine/Source/Runtime/WebBrowser/Private/WebJSScripting.h
+++ b/Engine/Source/Runtime/WebBrowser/Private/WebJSScripting.h
@@ -16,8 +16,9 @@ class FWebJSScripting
 	: public FGCObject
 {
 public:
-	FWebJSScripting()
+	FWebJSScripting(bool bInJSBindingToLoweringEnabled)
 		: BaseGuid(FGuid::NewGuid())
+		, bJSBindingToLoweringEnabled(bInJSBindingToLoweringEnabled)
 	{}
 
 	virtual void BindUObject(const FString& Name, UObject* Object, bool bIsPermanent = true) =0;
@@ -27,6 +28,12 @@ public:
 	virtual void InvokeJSFunction(FGuid FunctionId, int32 ArgCount, FWebJSParam Arguments[], bool bIsError=false) =0;
 	virtual void InvokeJSErrorResult(FGuid FunctionId, const FString& Error) =0;
 
+
+	FString GetBindingName(const UField* Property) const
+	{
+		return bJSBindingToLoweringEnabled ? Property->GetName().ToLower() : Property->GetName();
+	}
+
 public:
 
 	// FGCObject API
@@ -129,4 +136,6 @@ protected:
 	/** Reverse lookup for permanent bindings */
 	TMap<FString, UObject*> PermanentUObjectsByName;
 
+	/** The to-lowering option enable for the binding names. */
+	const bool bJSBindingToLoweringEnabled;
 };
diff --git a/Engine/Source/Runtime/WebBrowser/Public/IWebBrowserSingleton.h b/Engine/Source/Runtime/WebBrowser/Public/IWebBrowserSingleton.h
--- a/Engine/Source/Runtime/WebBrowser/Public/IWebBrowserSingleton.h
+++ b/Engine/Source/Runtime/WebBrowser/Public/IWebBrowserSingleton.h
@@ -165,4 +165,18 @@ public:
 	 */
 	virtual bool IsDevToolsShortcutEnabled() = 0;
 
+
+	/**
+	 * Enable or disable to-lowering of JavaScript object member bindings.
+	 *
+	 * Due to how JavaScript to UObject bridges require the use of FNames for building up the JS API objects, it is possible for case-sensitivity issues
+	 * to develop if an FName has been previously created with differing case to your function or property names. To-lowering the member names allows
+	 * a guaranteed casing for the web page's JS to reference.
+	 *
+	 * Default behavior is enabled, so that all JS side objects have only lowercase members.
+	 *
+	 * @param bEnabled a boolean value to enable or disable the to-lowering.
+	 */
+	virtual void SetJSBindingToLoweringEnabled(bool bEnabled) = 0;
+
 };
diff --git a/Engine/Source/Runtime/WebBrowser/Public/IWebBrowserWindow.h b/Engine/Source/Runtime/WebBrowser/Public/IWebBrowserWindow.h
--- a/Engine/Source/Runtime/WebBrowser/Public/IWebBrowserWindow.h
+++ b/Engine/Source/Runtime/WebBrowser/Public/IWebBrowserWindow.h
@@ -168,6 +168,13 @@ public:
 	 */
 	virtual FReply OnMouseMove(const FGeometry& MyGeometry, const FPointerEvent& MouseEvent, bool bIsPopup) = 0;
 
+	/**
+	 * Notify the browser that a mouse has left the window
+	 *
+	 * @param MouseEvent Information about the input event
+	 */
+	virtual void OnMouseLeave(const FPointerEvent& MouseEvent) = 0;
+
 	/**
 	 * Called when the mouse wheel is spun
 	 *
diff --git a/Engine/Source/Runtime/WebBrowser/Public/SWebBrowserView.h b/Engine/Source/Runtime/WebBrowser/Public/SWebBrowserView.h
--- a/Engine/Source/Runtime/WebBrowser/Public/SWebBrowserView.h
+++ b/Engine/Source/Runtime/WebBrowser/Public/SWebBrowserView.h
@@ -295,6 +295,7 @@ private:
 	}
 
 	void HandleWindowDeactivated();
+	void HandleWindowActivated();
 
 private:
 
