Export GetOutdatedFiles with BUILDPATCHSERVICES_API to link properly with UnrealTournament non-monolithic build

Merging 3223132

//UE4/UT-Staging/...

to //UE4/Main/...

diff --git a/Engine/Source/Runtime/Online/BuildPatchServices/Private/BuildPatchManifest.h b/Engine/Source/Runtime/Online/BuildPatchServices/Private/BuildPatchManifest.h
--- a/Engine/Source/Runtime/Online/BuildPatchServices/Private/BuildPatchManifest.h
+++ b/Engine/Source/Runtime/Online/BuildPatchServices/Private/BuildPatchManifest.h
@@ -600,7 +600,7 @@ public:
 	 * @param InstallDirectory	IN		The Build installation directory, so that it can be checked for missing files.
 	 * @param OutDatedFiles		OUT		The files that changed hash, are new, are wrong size, or missing on disk.
 	 */
-	void GetOutdatedFiles(const FBuildPatchAppManifestPtr& OldManifest, const FString& InstallDirectory, TSet<FString>& OutDatedFiles) const;
+	BUILDPATCHSERVICES_API void GetOutdatedFiles(const FBuildPatchAppManifestPtr& OldManifest, const FString& InstallDirectory, TSet<FString>& OutDatedFiles) const;
 
 	/**
 	 * Check a single file to see if it will be effected by patching from a previous version.
