Fixing merge that readded superfluous comment.

diff --git a/Engine/Source/Runtime/Online/BuildPatchServices/Private/Installer/InstallerError.h b/Engine/Source/Runtime/Online/BuildPatchServices/Private/Installer/InstallerError.h
--- a/Engine/Source/Runtime/Online/BuildPatchServices/Private/Installer/InstallerError.h
+++ b/Engine/Source/Runtime/Online/BuildPatchServices/Private/Installer/InstallerError.h
@@ -1,9 +1,4 @@
 // Copyright 1998-2017 Epic Games, Inc. All Rights Reserved.
-
-/*=============================================================================
-	BuildPatchError.h: Declares classes involved setting and getting error information.
-=============================================================================*/
-
 #pragma once
 
 #include "CoreMinimal.h"
