Copying //UE4/Dev-Online to //UE4/Dev-Main (Source: //UE4/Dev-Online @ 3605423)

============================
  MAJOR FEATURES & CHANGES
============================

Change 3529773 by [REDACTED_PERSON]

	#UE4 - voice module improvements / refactors for Windows/Opus codec
	- added ability to change more parameters at runtime (samplerate, channels, device)
	- ability to change bitrate at runtime
	- ability to change complexity at runtime
	- ability to set VBR on/off
	- fixed some bugs related to stereo playback / diagnostics
	- exposed more settings / information in the various classes
	- added num channels as an input parameter to CreateVoiceAudioComponent helper function
	- improved test voice class and its capabilities
	- various exec commands to change audio capture values at runtime
	- queue audio data differently for the procedural sound wave class
	- more buffer sanity checks
	- fixed bug when nothing to encode (super low bitrate) causing the decoder to fail
	-- encode 0 size frame data and make decoder aware of it
	- added more stats under "stat voice"
	- added LogVoiceEngine to separate out OSS related voice calls from rest of audio capture code

	#tests local voice chat loopback, test voice harness with external audio input

Change 3532546 by [REDACTED_PERSON]

	#UE4 - voip streaming improvements
	- use the underflow feature of sound streaming to feed the audio buffer rather than enqueue the data directly
	- wait a small period of time to buffer period actually playing new voip streams
	- some general code cleanup

Change 3538715 by [REDACTED_PERSON]

	#UE4 - OGS-472, OGS-473 NBOSeralizer fixes
	- protect against possible overflow in the values

Change 3538723 by [REDACTED_PERSON]

	#UE4 - OGS-475 VoiceDecoder Opus security fixes
	- sanity check packet header calculations to prevent buffer overrun

Change 3587162 by [REDACTED_PERSON]

	#UE4 - fix potential null pointer in Oculus session cleanup

Change 3590508 by [REDACTED_PERSON]

	#UE4 - fixed some static analysis checks
	- null check on CreateEvent for stop notifications
	- null dereference in audio device enumeration
	- made sure to add the HMD device override to the list of devices
	- fixed success logic in CreateNotifications (possible to return true when it actually failed)

Change 3531444 by [REDACTED_PERSON]

	#UE4 - used intrinsic FName for Game/PartySession and Game/BeaconPort
	- left old variable around as a #define for a while

diff --git a/Engine/Plugins/Online/Android/OnlineSubsystemGameCircle/Source/OnlineSubsystemGameCircle/Private/OnlineIdentityInterfaceGameCircle.cpp b/Engine/Plugins/Online/Android/OnlineSubsystemGameCircle/Source/OnlineSubsystemGameCircle/Private/OnlineIdentityInterfaceGameCircle.cpp
--- a/Engine/Plugins/Online/Android/OnlineSubsystemGameCircle/Source/OnlineSubsystemGameCircle/Private/OnlineIdentityInterfaceGameCircle.cpp
+++ b/Engine/Plugins/Online/Android/OnlineSubsystemGameCircle/Source/OnlineSubsystemGameCircle/Private/OnlineIdentityInterfaceGameCircle.cpp
@@ -123,7 +123,7 @@ void FOnlineIdentityGameCircle::GetUserPrivilege(const FUniqueNetId& UserId, EUs
 	Delegate.ExecuteIfBound(UserId, Privilege, (uint32)EPrivilegeResults::NoFailures);
 }
 
-FPlatformUserId FOnlineIdentityGameCircle::GetPlatformUserIdFromUniqueNetId(const FUniqueNetId& NetId)
+FPlatformUserId FOnlineIdentityGameCircle::GetPlatformUserIdFromUniqueNetId(const FUniqueNetId& NetId) const
 {
 	for (int i = 0; i < MAX_LOCAL_PLAYERS; ++i)
 	{
diff --git a/Engine/Plugins/Online/Android/OnlineSubsystemGameCircle/Source/OnlineSubsystemGameCircle/Private/OnlineIdentityInterfaceGameCircle.h b/Engine/Plugins/Online/Android/OnlineSubsystemGameCircle/Source/OnlineSubsystemGameCircle/Private/OnlineIdentityInterfaceGameCircle.h
--- a/Engine/Plugins/Online/Android/OnlineSubsystemGameCircle/Source/OnlineSubsystemGameCircle/Private/OnlineIdentityInterfaceGameCircle.h
+++ b/Engine/Plugins/Online/Android/OnlineSubsystemGameCircle/Source/OnlineSubsystemGameCircle/Private/OnlineIdentityInterfaceGameCircle.h
@@ -29,7 +29,7 @@ public:
 	virtual FString GetAuthToken(int32 LocalUserNum) const override;
 	virtual void RevokeAuthToken(const FUniqueNetId& UserId, const FOnRevokeAuthTokenCompleteDelegate& Delegate) override;
 	virtual void GetUserPrivilege(const FUniqueNetId& UserId, EUserPrivileges::Type Privilege, const FOnGetUserPrivilegeCompleteDelegate& Delegate) override;
-	virtual FPlatformUserId GetPlatformUserIdFromUniqueNetId(const FUniqueNetId& NetId) override;
+	virtual FPlatformUserId GetPlatformUserIdFromUniqueNetId(const FUniqueNetId& NetId) const override;
 	virtual FString GetAuthType() const override;
 	void Tick(float DeltaTime);
 	//~ End IOnlineIdentity Interface
diff --git a/Engine/Plugins/Online/Android/OnlineSubsystemGooglePlay/Source/Private/OnlineIdentityInterfaceGooglePlay.cpp b/Engine/Plugins/Online/Android/OnlineSubsystemGooglePlay/Source/Private/OnlineIdentityInterfaceGooglePlay.cpp
--- a/Engine/Plugins/Online/Android/OnlineSubsystemGooglePlay/Source/Private/OnlineIdentityInterfaceGooglePlay.cpp
+++ b/Engine/Plugins/Online/Android/OnlineSubsystemGooglePlay/Source/Private/OnlineIdentityInterfaceGooglePlay.cpp
@@ -182,7 +182,7 @@ void FOnlineIdentityGooglePlay::GetUserPrivilege(const FUniqueNetId& UserId, EUs
 	Delegate.ExecuteIfBound(UserId, Privilege, (uint32)EPrivilegeResults::NoFailures);
 }
 
-FPlatformUserId FOnlineIdentityGooglePlay::GetPlatformUserIdFromUniqueNetId(const FUniqueNetId& NetId)
+FPlatformUserId FOnlineIdentityGooglePlay::GetPlatformUserIdFromUniqueNetId(const FUniqueNetId& NetId) const
 {
 	for (int i = 0; i < MAX_LOCAL_PLAYERS; ++i)
 	{
diff --git a/Engine/Plugins/Online/Android/OnlineSubsystemGooglePlay/Source/Private/OnlineIdentityInterfaceGooglePlay.h b/Engine/Plugins/Online/Android/OnlineSubsystemGooglePlay/Source/Private/OnlineIdentityInterfaceGooglePlay.h
--- a/Engine/Plugins/Online/Android/OnlineSubsystemGooglePlay/Source/Private/OnlineIdentityInterfaceGooglePlay.h
+++ b/Engine/Plugins/Online/Android/OnlineSubsystemGooglePlay/Source/Private/OnlineIdentityInterfaceGooglePlay.h
@@ -74,7 +74,7 @@ public:
 	virtual FString GetAuthToken(int32 LocalUserNum) const override;
 	virtual void RevokeAuthToken(const FUniqueNetId& UserId, const FOnRevokeAuthTokenCompleteDelegate& Delegate) override;
 	virtual void GetUserPrivilege(const FUniqueNetId& UserId, EUserPrivileges::Type Privilege, const FOnGetUserPrivilegeCompleteDelegate& Delegate) override;
-	virtual FPlatformUserId GetPlatformUserIdFromUniqueNetId(const FUniqueNetId& NetId) override;
+	virtual FPlatformUserId GetPlatformUserIdFromUniqueNetId(const FUniqueNetId& NetId) const override;
 	virtual FString GetAuthType() const override;
 	//~ End IOnlineIdentity interface
 
diff --git a/Engine/Plugins/Online/IOS/OnlineSubsystemIOS/Source/Private/OnlineExternalUIInterfaceIOS.h b/Engine/Plugins/Online/IOS/OnlineSubsystemIOS/Source/Private/OnlineExternalUIInterfaceIOS.h
--- a/Engine/Plugins/Online/IOS/OnlineSubsystemIOS/Source/Private/OnlineExternalUIInterfaceIOS.h
+++ b/Engine/Plugins/Online/IOS/OnlineSubsystemIOS/Source/Private/OnlineExternalUIInterfaceIOS.h
@@ -18,7 +18,7 @@ public:
 	//~ Begin IOnlineExternalUI Interface
 	virtual bool ShowLoginUI(const int ControllerIndex, bool bShowOnlineOnly, bool bShowSkipButton, const FOnLoginUIClosedDelegate& Delegate = FOnLoginUIClosedDelegate()) override;
 	virtual bool ShowFriendsUI(int32 LocalUserNum) override;
-	virtual bool ShowInviteUI(int32 LocalUserNum, FName SessionName = GameSessionName) override;
+	virtual bool ShowInviteUI(int32 LocalUserNum, FName SessionName = NAME_GameSession) override;
 	virtual bool ShowAchievementsUI(int32 LocalUserNum) override;
 	virtual bool ShowLeaderboardUI(const FString& LeaderboardName) override;
 	virtual bool ShowWebURL(const FString& Url, const FShowWebUrlParams& ShowParams, const FOnShowWebUrlClosedDelegate& Delegate = FOnShowWebUrlClosedDelegate()) override;
diff --git a/Engine/Plugins/Online/IOS/OnlineSubsystemIOS/Source/Private/OnlineIdentityInterfaceIOS.cpp b/Engine/Plugins/Online/IOS/OnlineSubsystemIOS/Source/Private/OnlineIdentityInterfaceIOS.cpp
--- a/Engine/Plugins/Online/IOS/OnlineSubsystemIOS/Source/Private/OnlineIdentityInterfaceIOS.cpp
+++ b/Engine/Plugins/Online/IOS/OnlineSubsystemIOS/Source/Private/OnlineIdentityInterfaceIOS.cpp
@@ -316,7 +316,7 @@ void FOnlineIdentityIOS::GetUserPrivilege(const FUniqueNetId& UserId, EUserPrivi
 	 }
 }
 
-FPlatformUserId FOnlineIdentityIOS::GetPlatformUserIdFromUniqueNetId(const FUniqueNetId& InUniqueNetId)
+FPlatformUserId FOnlineIdentityIOS::GetPlatformUserIdFromUniqueNetId(const FUniqueNetId& InUniqueNetId) const
 {
 	for (int i = 0; i < MAX_LOCAL_PLAYERS; ++i)
 	{
diff --git a/Engine/Plugins/Online/IOS/OnlineSubsystemIOS/Source/Private/OnlineIdentityInterfaceIOS.h b/Engine/Plugins/Online/IOS/OnlineSubsystemIOS/Source/Private/OnlineIdentityInterfaceIOS.h
--- a/Engine/Plugins/Online/IOS/OnlineSubsystemIOS/Source/Private/OnlineIdentityInterfaceIOS.h
+++ b/Engine/Plugins/Online/IOS/OnlineSubsystemIOS/Source/Private/OnlineIdentityInterfaceIOS.h
@@ -48,7 +48,7 @@ public:
 	virtual FString GetAuthToken(int32 LocalUserNum) const override;
 	virtual void RevokeAuthToken(const FUniqueNetId& UserId, const FOnRevokeAuthTokenCompleteDelegate& Delegate) override;
 	virtual void GetUserPrivilege(const FUniqueNetId& UserId, EUserPrivileges::Type Privilege, const FOnGetUserPrivilegeCompleteDelegate& Delegate) override;
-	virtual FPlatformUserId GetPlatformUserIdFromUniqueNetId(const FUniqueNetId& UniqueNetId) override;
+	virtual FPlatformUserId GetPlatformUserIdFromUniqueNetId(const FUniqueNetId& UniqueNetId) const override;
 	virtual FString GetAuthType() const override;
 	//~ End IOnlineIdentity Interface
 
diff --git a/Engine/Plugins/Online/OnlineFramework/Source/Lobby/Private/LobbyBeaconClient.cpp b/Engine/Plugins/Online/OnlineFramework/Source/Lobby/Private/LobbyBeaconClient.cpp
--- a/Engine/Plugins/Online/OnlineFramework/Source/Lobby/Private/LobbyBeaconClient.cpp
+++ b/Engine/Plugins/Online/OnlineFramework/Source/Lobby/Private/LobbyBeaconClient.cpp
@@ -48,7 +48,7 @@ void ALobbyBeaconClient::ConnectToLobby(const FOnlineSessionSearchResult& Desire
 		if (SessionInt.IsValid())
 		{
 			FString ConnectInfo;
-			if (SessionInt->GetResolvedConnectString(DesiredHost, BeaconPort, ConnectInfo))
+			if (SessionInt->GetResolvedConnectString(DesiredHost, NAME_BeaconPort, ConnectInfo))
 			{
 				FURL ConnectURL(NULL, *ConnectInfo, TRAVEL_Absolute);
 				if (InitClient(ConnectURL) && DesiredHost.Session.SessionInfo.IsValid())
@@ -306,7 +306,7 @@ void ALobbyBeaconClient::ClientPlayerJoined_Implementation(const FText& NewPlaye
 		if (SessionInt.IsValid() && InUniqueId.IsValid())
 		{
 			// Register the player as part of the session
-			SessionInt->RegisterPlayer(GameSessionName, *InUniqueId, false);
+			SessionInt->RegisterPlayer(NAME_GameSession, *InUniqueId, false);
 		}
 	}
 
@@ -323,7 +323,7 @@ void ALobbyBeaconClient::ClientPlayerLeft_Implementation(const FUniqueNetIdRepl&
 		if (SessionInt.IsValid() && InUniqueId.IsValid())
 		{
 			// Register the player as part of the session
-			SessionInt->UnregisterPlayer(GameSessionName, *InUniqueId);
+			SessionInt->UnregisterPlayer(NAME_GameSession, *InUniqueId);
 		}
 	}
 
diff --git a/Engine/Plugins/Online/OnlineFramework/Source/Lobby/Private/LobbyBeaconHost.cpp b/Engine/Plugins/Online/OnlineFramework/Source/Lobby/Private/LobbyBeaconHost.cpp
--- a/Engine/Plugins/Online/OnlineFramework/Source/Lobby/Private/LobbyBeaconHost.cpp
+++ b/Engine/Plugins/Online/OnlineFramework/Source/Lobby/Private/LobbyBeaconHost.cpp
@@ -103,7 +103,7 @@ ALobbyBeaconPlayerState* ALobbyBeaconHost::HandlePlayerLogin(ALobbyBeaconClient*
 		{
 			// Register the player as part of the session
 			bool bWasFromInvite = UGameplayStatics::HasOption(Options, TEXT("bIsFromInvite"));
-			SessionInt->RegisterPlayer(GameSessionName, *InUniqueId, bWasFromInvite);
+			SessionInt->RegisterPlayer(NAME_GameSession, *InUniqueId, bWasFromInvite);
 		}
 	}
 
@@ -298,7 +298,7 @@ void ALobbyBeaconHost::NotifyClientDisconnected(AOnlineBeaconClient* LeavingClie
 			check(GameMode->GameSession);
 
 			// Notify the session (updates reservation beacon, unregisters the player, etc)
-			GameMode->GameSession->NotifyLogout(GameSessionName, Player->UniqueId);
+			GameMode->GameSession->NotifyLogout(NAME_GameSession, Player->UniqueId);
 			HandlePlayerLogout(Player->UniqueId);
 		}
 	}
diff --git a/Engine/Plugins/Online/OnlineFramework/Source/Party/Private/Party.cpp b/Engine/Plugins/Online/OnlineFramework/Source/Party/Private/Party.cpp
--- a/Engine/Plugins/Online/OnlineFramework/Source/Party/Private/Party.cpp
+++ b/Engine/Plugins/Online/OnlineFramework/Source/Party/Private/Party.cpp
@@ -298,7 +298,7 @@ FName UParty::GetPlayerSessionName() const
 			return PlayerController->PlayerState->SessionName;
 		}
 	}
-	return GameSessionName;
+	return NAME_GameSession;
 }
 
 UPartyGameState* UParty::GetParty(const FOnlinePartyId& InPartyId) const
diff --git a/Engine/Plugins/Online/OnlineFramework/Source/Party/Private/PartyGameState.cpp b/Engine/Plugins/Online/OnlineFramework/Source/Party/Private/PartyGameState.cpp
--- a/Engine/Plugins/Online/OnlineFramework/Source/Party/Private/PartyGameState.cpp
+++ b/Engine/Plugins/Online/OnlineFramework/Source/Party/Private/PartyGameState.cpp
@@ -1052,7 +1052,7 @@ bool UPartyGameState::IsInJoinableGameState() const
 		bool bGameInviteOnly = false;
 		bool bGameAllowInvites = false;
 
-		FNamedOnlineSession* GameSession = SessionInt->GetNamedSession(GameSessionName);
+		FNamedOnlineSession* GameSession = SessionInt->GetNamedSession(NAME_GameSession);
 		if (GameSession != NULL &&
 			GameSession->GetJoinability(bGamePublicJoinable, bGameFriendJoinable, bGameInviteOnly, bGameAllowInvites))
 		{
@@ -1408,7 +1408,7 @@ void UPartyGameState::GetSessionInfo(FName SessionName, FString& URL, FString& S
 	IOnlineSessionPtr SessionInt = Online::GetSessionInterface(World);
 	if (ensure(SessionInt.IsValid()))
 	{
-		ensure(SessionInt->GetResolvedConnectString(SessionName, URL, BeaconPort));
+		ensure(SessionInt->GetResolvedConnectString(SessionName, URL, NAME_BeaconPort));
 
 		FNamedOnlineSession* Session = SessionInt->GetNamedSession(SessionName);
 		if (Session)
diff --git a/Engine/Plugins/Online/OnlineFramework/Source/Party/Public/Party.h b/Engine/Plugins/Online/OnlineFramework/Source/Party/Public/Party.h
--- a/Engine/Plugins/Online/OnlineFramework/Source/Party/Public/Party.h
+++ b/Engine/Plugins/Online/OnlineFramework/Source/Party/Public/Party.h
@@ -332,7 +332,7 @@ public:
 	bool HasPendingPartyJoin() const;
 
 	/**
-	 * Get the session name (if available) for the primary player, typically GameSessionName
+	 * Get the session name (if available) for the primary player, typically NAME_GameSession
 	 *
 	 * @return Session name
 	 */
diff --git a/Engine/Plugins/Online/OnlineFramework/Source/Qos/Private/QosBeaconClient.cpp b/Engine/Plugins/Online/OnlineFramework/Source/Qos/Private/QosBeaconClient.cpp
--- a/Engine/Plugins/Online/OnlineFramework/Source/Qos/Private/QosBeaconClient.cpp
+++ b/Engine/Plugins/Online/OnlineFramework/Source/Qos/Private/QosBeaconClient.cpp
@@ -31,7 +31,7 @@ void AQosBeaconClient::SendQosRequest(const FOnlineSessionSearchResult& DesiredH
 		if (SessionInt.IsValid())
 		{
 			FString ConnectInfo;
-			if (SessionInt->GetResolvedConnectString(DesiredHost, BeaconPort, ConnectInfo))
+			if (SessionInt->GetResolvedConnectString(DesiredHost, NAME_BeaconPort, ConnectInfo))
 			{
 				FURL ConnectURL(NULL, *ConnectInfo, TRAVEL_Absolute);
 				if (InitClient(ConnectURL) && DesiredHost.Session.SessionInfo.IsValid())
diff --git a/Engine/Plugins/Online/OnlineFramework/Source/Rejoin/Private/RejoinCheck.cpp b/Engine/Plugins/Online/OnlineFramework/Source/Rejoin/Private/RejoinCheck.cpp
--- a/Engine/Plugins/Online/OnlineFramework/Source/Rejoin/Private/RejoinCheck.cpp
+++ b/Engine/Plugins/Online/OnlineFramework/Source/Rejoin/Private/RejoinCheck.cpp
@@ -311,7 +311,7 @@ void URejoinCheck::TravelToSession()
 	ULocalPlayer* LP = GEngine->GetFirstGamePlayer(GetWorld());
 	if (ensure(LP))
 	{
-		bResult = GameInstance->ClientTravelToSession(LP->GetControllerId(), GameSessionName);
+		bResult = GameInstance->ClientTravelToSession(LP->GetControllerId(), NAME_GameSession);
 		if (bResult)
 		{
 			UE_LOG(LogOnline, Log, TEXT("URejoinCheck::TravelToSession: Performing ClientTravelToSession"));
diff --git a/Engine/Plugins/Online/OnlineSubsystem/Source/Private/OnlineSessionSettings.cpp b/Engine/Plugins/Online/OnlineSubsystem/Source/Private/OnlineSessionSettings.cpp
--- a/Engine/Plugins/Online/OnlineSubsystem/Source/Private/OnlineSessionSettings.cpp
+++ b/Engine/Plugins/Online/OnlineSubsystem/Source/Private/OnlineSessionSettings.cpp
@@ -2,11 +2,13 @@
 
 #include "OnlineSessionSettings.h"
 #include "OnlineSubsystem.h"
+#include "Logging/LogScopedVerbosityOverride.h"
 
 void DumpNamedSession(const FNamedOnlineSession* NamedSession)
 {
 	if (NamedSession != NULL)
 	{
+		LOG_SCOPE_VERBOSITY_OVERRIDE(LogOnline, ELogVerbosity::VeryVerbose);
 		UE_LOG(LogOnline, Verbose, TEXT("dumping NamedSession: "));
 		UE_LOG(LogOnline, Verbose, TEXT("	SessionName: %s"), *NamedSession->SessionName.ToString());	
 		UE_LOG(LogOnline, Verbose, TEXT("	HostingPlayerNum: %d"), NamedSession->HostingPlayerNum);
diff --git a/Engine/Plugins/Online/OnlineSubsystem/Source/Private/OnlineSubsystem.cpp b/Engine/Plugins/Online/OnlineSubsystem/Source/Private/OnlineSubsystem.cpp
--- a/Engine/Plugins/Online/OnlineSubsystem/Source/Private/OnlineSubsystem.cpp
+++ b/Engine/Plugins/Online/OnlineSubsystem/Source/Private/OnlineSubsystem.cpp
@@ -11,6 +11,7 @@
 DEFINE_LOG_CATEGORY(LogOnline);
 DEFINE_LOG_CATEGORY(LogOnlineGame);
 DEFINE_LOG_CATEGORY(LogOnlineChat);
+DEFINE_LOG_CATEGORY(LogVoiceEngine);
 
 #if STATS
 ONLINESUBSYSTEM_API DEFINE_STAT(STAT_Online_Async);
diff --git a/Engine/Plugins/Online/OnlineSubsystem/Source/Public/Interfaces/OnlineExternalUIInterface.h b/Engine/Plugins/Online/OnlineSubsystem/Source/Public/Interfaces/OnlineExternalUIInterface.h
--- a/Engine/Plugins/Online/OnlineSubsystem/Source/Public/Interfaces/OnlineExternalUIInterface.h
+++ b/Engine/Plugins/Online/OnlineSubsystem/Source/Public/Interfaces/OnlineExternalUIInterface.h
@@ -258,7 +258,7 @@ public:
 	 *
 	 * @return true if it was able to show the UI, false if it failed
 	 */
-	virtual bool ShowInviteUI(int32 LocalUserNum, FName SessionName = GameSessionName) = 0;
+	virtual bool ShowInviteUI(int32 LocalUserNum, FName SessionName = NAME_GameSession) = 0;
 
 	/**
 	 *	Displays the UI that shows a user's list of achievements
diff --git a/Engine/Plugins/Online/OnlineSubsystem/Source/Public/Interfaces/OnlineIdentityInterface.h b/Engine/Plugins/Online/OnlineSubsystem/Source/Public/Interfaces/OnlineIdentityInterface.h
--- a/Engine/Plugins/Online/OnlineSubsystem/Source/Public/Interfaces/OnlineIdentityInterface.h
+++ b/Engine/Plugins/Online/OnlineSubsystem/Source/Public/Interfaces/OnlineIdentityInterface.h
@@ -386,7 +386,7 @@ public:
 	 * @param UniqueNetId The unique id to look up
 	 * @return The corresponding id or PLATFORMID_NONE if not found
 	 */
-	virtual FPlatformUserId GetPlatformUserIdFromUniqueNetId(const FUniqueNetId& UniqueNetId) = 0;
+	virtual FPlatformUserId GetPlatformUserIdFromUniqueNetId(const FUniqueNetId& UniqueNetId) const = 0;
 
 	/**
 	 * Get the auth type associated with accounts for this platform
diff --git a/Engine/Plugins/Online/OnlineSubsystem/Source/Public/Interfaces/OnlineSessionInterface.h b/Engine/Plugins/Online/OnlineSubsystem/Source/Public/Interfaces/OnlineSessionInterface.h
--- a/Engine/Plugins/Online/OnlineSubsystem/Source/Public/Interfaces/OnlineSessionInterface.h
+++ b/Engine/Plugins/Online/OnlineSubsystem/Source/Public/Interfaces/OnlineSessionInterface.h
@@ -690,7 +690,7 @@ public:
 	 *
 	 * @return true if the call was successful, false otherwise
 	 */
-	virtual bool GetResolvedConnectString(FName SessionName, FString& ConnectInfo, FName PortType = GamePort) = 0;
+	virtual bool GetResolvedConnectString(FName SessionName, FString& ConnectInfo, FName PortType = NAME_GamePort) = 0;
 
 	/**
 	 * Returns the platform specific connection information for joining a search result.
diff --git a/Engine/Plugins/Online/OnlineSubsystem/Source/Public/Interfaces/VoiceInterface.h b/Engine/Plugins/Online/OnlineSubsystem/Source/Public/Interfaces/VoiceInterface.h
--- a/Engine/Plugins/Online/OnlineSubsystem/Source/Public/Interfaces/VoiceInterface.h
+++ b/Engine/Plugins/Online/OnlineSubsystem/Source/Public/Interfaces/VoiceInterface.h
@@ -9,6 +9,8 @@
 
 class FVoicePacket;
 
+ONLINESUBSYSTEM_API DECLARE_LOG_CATEGORY_EXTERN(LogVoiceEngine, Display, All);
+
 /** Enable to pipe local voice data back to this client as remote data */
 #define VOICE_LOOPBACK !UE_BUILD_SHIPPING
 
diff --git a/Engine/Plugins/Online/OnlineSubsystem/Source/Public/NboSerializer.h b/Engine/Plugins/Online/OnlineSubsystem/Source/Public/NboSerializer.h
--- a/Engine/Plugins/Online/OnlineSubsystem/Source/Public/NboSerializer.h
+++ b/Engine/Plugins/Online/OnlineSubsystem/Source/Public/NboSerializer.h
@@ -635,7 +635,9 @@ public:
 		int32 Len = 0;
 		Ar >> Len;
 
-		if (!Ar.HasOverflow() && Ar.CurrentOffset + Len <= Ar.NumBytes)
+		// Check this way to trust NumBytes and CurrentOffset to be more accurate than the packet Len value
+		const bool bSizeOk = (Len >= 0) && (Len <= (Ar.NumBytes - Ar.CurrentOffset));
+		if (!Ar.HasOverflow() && bSizeOk)
 		{
 			// Handle strings of zero length
 			if (Len > 0)
@@ -741,10 +743,12 @@ public:
 					int32 Length;
 					Ar >> Length;
 
-					if (!Ar.HasOverflow() && Ar.CurrentOffset + Length <= Ar.NumBytes)
+					// Check this way to trust NumBytes and CurrentOffset to be more accurate than the packet Len value
+					const bool bSizeOk = (Length >= 0) && (Length <= (Ar.NumBytes - Ar.CurrentOffset));
+					if (!Ar.HasOverflow() && bSizeOk)
 					{
 						// Now directly copy the blob data
-						KeyValuePair.SetValue(Length,&Ar.Data[Ar.CurrentOffset]);
+						KeyValuePair.SetValue(Length, &Ar.Data[Ar.CurrentOffset]);
 						Ar.CurrentOffset += Length;
 					}
 					else
diff --git a/Engine/Plugins/Online/OnlineSubsystemAmazon/Source/Private/OnlineIdentityAmazon.cpp b/Engine/Plugins/Online/OnlineSubsystemAmazon/Source/Private/OnlineIdentityAmazon.cpp
--- a/Engine/Plugins/Online/OnlineSubsystemAmazon/Source/Private/OnlineIdentityAmazon.cpp
+++ b/Engine/Plugins/Online/OnlineSubsystemAmazon/Source/Private/OnlineIdentityAmazon.cpp
@@ -360,7 +360,7 @@ void FOnlineIdentityAmazon::GetUserPrivilege(const FUniqueNetId& UserId, EUserPr
 	Delegate.ExecuteIfBound(UserId, Privilege, (uint32)EPrivilegeResults::NoFailures);
 }
 
-FPlatformUserId FOnlineIdentityAmazon::GetPlatformUserIdFromUniqueNetId(const FUniqueNetId& UniqueNetId)
+FPlatformUserId FOnlineIdentityAmazon::GetPlatformUserIdFromUniqueNetId(const FUniqueNetId& UniqueNetId) const
 {
 	for (int i = 0; i < MAX_LOCAL_PLAYERS; ++i)
 	{
diff --git a/Engine/Plugins/Online/OnlineSubsystemAmazon/Source/Private/OnlineIdentityAmazon.h b/Engine/Plugins/Online/OnlineSubsystemAmazon/Source/Private/OnlineIdentityAmazon.h
--- a/Engine/Plugins/Online/OnlineSubsystemAmazon/Source/Private/OnlineIdentityAmazon.h
+++ b/Engine/Plugins/Online/OnlineSubsystemAmazon/Source/Private/OnlineIdentityAmazon.h
@@ -127,7 +127,7 @@ public:
 	virtual FString GetAuthToken(int32 LocalUserNum) const override;
 	virtual void RevokeAuthToken(const FUniqueNetId& UserId, const FOnRevokeAuthTokenCompleteDelegate& Delegate) override;
 	virtual void GetUserPrivilege(const FUniqueNetId& UserId, EUserPrivileges::Type Privilege, const FOnGetUserPrivilegeCompleteDelegate& Delegate) override;
-	virtual FPlatformUserId GetPlatformUserIdFromUniqueNetId(const FUniqueNetId& UniqueNetId) override;
+	virtual FPlatformUserId GetPlatformUserIdFromUniqueNetId(const FUniqueNetId& UniqueNetId) const override;
 	virtual FString GetAuthType() const override;
 
 	// FOnlineIdentityAmazon
diff --git a/Engine/Plugins/Online/OnlineSubsystemFacebook/Source/OnlineSubsystemFacebook.Build.cs b/Engine/Plugins/Online/OnlineSubsystemFacebook/Source/OnlineSubsystemFacebook.Build.cs
--- a/Engine/Plugins/Online/OnlineSubsystemFacebook/Source/OnlineSubsystemFacebook.Build.cs
+++ b/Engine/Plugins/Online/OnlineSubsystemFacebook/Source/OnlineSubsystemFacebook.Build.cs
@@ -74,6 +74,14 @@ public class OnlineSubsystemFacebook : ModuleRules
 			Definitions.Add("WITH_FACEBOOK=1");
 			PrivateIncludePaths.Add("Private/Windows");
 		}
+		else if (Target.Platform == UnrealTargetPlatform.XboxOne)
+		{
+			PrivateIncludePaths.Add("Private/XboxOne");
+		}
+		else if (Target.Platform == UnrealTargetPlatform.PS4)
+		{
+			PrivateIncludePaths.Add("Private/PS4");
+		}
 		else
 		{
 			Definitions.Add("WITH_FACEBOOK=0");
diff --git a/Engine/Plugins/Online/OnlineSubsystemFacebook/Source/Private/OnlineExternalUIFacebookCommon.h b/Engine/Plugins/Online/OnlineSubsystemFacebook/Source/Private/OnlineExternalUIFacebookCommon.h
--- a/Engine/Plugins/Online/OnlineSubsystemFacebook/Source/Private/OnlineExternalUIFacebookCommon.h
+++ b/Engine/Plugins/Online/OnlineSubsystemFacebook/Source/Private/OnlineExternalUIFacebookCommon.h
@@ -41,7 +41,7 @@ public:
 	// IOnlineExternalUI
 	virtual bool ShowLoginUI(const int ControllerIndex, bool bShowOnlineOnly, bool bShowSkipButton, const FOnLoginUIClosedDelegate& Delegate = FOnLoginUIClosedDelegate()) override;
 	virtual bool ShowFriendsUI(int32 LocalUserNum) override;
-	virtual bool ShowInviteUI(int32 LocalUserNum, FName SessionName = GameSessionName) override;
+	virtual bool ShowInviteUI(int32 LocalUserNum, FName SessionName = NAME_GameSession) override;
 	virtual bool ShowAchievementsUI(int32 LocalUserNum) override;
 	virtual bool ShowLeaderboardUI(const FString& LeaderboardName) override;
 	virtual bool ShowWebURL(const FString& Url, const FShowWebUrlParams& ShowParams, const FOnShowWebUrlClosedDelegate& Delegate = FOnShowWebUrlClosedDelegate()) override;
diff --git a/Engine/Plugins/Online/OnlineSubsystemFacebook/Source/Private/OnlineIdentityFacebookCommon.cpp b/Engine/Plugins/Online/OnlineSubsystemFacebook/Source/Private/OnlineIdentityFacebookCommon.cpp
--- a/Engine/Plugins/Online/OnlineSubsystemFacebook/Source/Private/OnlineIdentityFacebookCommon.cpp
+++ b/Engine/Plugins/Online/OnlineSubsystemFacebook/Source/Private/OnlineIdentityFacebookCommon.cpp
@@ -292,7 +292,7 @@ void FOnlineIdentityFacebookCommon::GetUserPrivilege(const FUniqueNetId& UserId,
 	Delegate.ExecuteIfBound(UserId, Privilege, (uint32)EPrivilegeResults::NoFailures);
 }	
 
-FPlatformUserId FOnlineIdentityFacebookCommon::GetPlatformUserIdFromUniqueNetId(const FUniqueNetId& UniqueNetId)
+FPlatformUserId FOnlineIdentityFacebookCommon::GetPlatformUserIdFromUniqueNetId(const FUniqueNetId& UniqueNetId) const
 {
 	for (int i = 0; i < MAX_LOCAL_PLAYERS; ++i)
 	{
diff --git a/Engine/Plugins/Online/OnlineSubsystemFacebook/Source/Private/OnlineIdentityFacebookCommon.h b/Engine/Plugins/Online/OnlineSubsystemFacebook/Source/Private/OnlineIdentityFacebookCommon.h
--- a/Engine/Plugins/Online/OnlineSubsystemFacebook/Source/Private/OnlineIdentityFacebookCommon.h
+++ b/Engine/Plugins/Online/OnlineSubsystemFacebook/Source/Private/OnlineIdentityFacebookCommon.h
@@ -59,7 +59,7 @@ public:
 	virtual FString GetAuthToken(int32 LocalUserNum) const override;
 	virtual void RevokeAuthToken(const FUniqueNetId& UserId, const FOnRevokeAuthTokenCompleteDelegate& Delegate) override;
 	virtual void GetUserPrivilege(const FUniqueNetId& UserId, EUserPrivileges::Type Privilege, const FOnGetUserPrivilegeCompleteDelegate& Delegate) override;
-	virtual FPlatformUserId GetPlatformUserIdFromUniqueNetId(const FUniqueNetId& UniqueNetId) override;
+	virtual FPlatformUserId GetPlatformUserIdFromUniqueNetId(const FUniqueNetId& UniqueNetId) const override;
 	virtual FString GetAuthType() const override;
 
 	// FOnlineIdentityFacebook
diff --git a/Engine/Plugins/Online/OnlineSubsystemGoogle/Source/OnlineSubsystemGoogle.Build.cs b/Engine/Plugins/Online/OnlineSubsystemGoogle/Source/OnlineSubsystemGoogle.Build.cs
--- a/Engine/Plugins/Online/OnlineSubsystemGoogle/Source/OnlineSubsystemGoogle.Build.cs
+++ b/Engine/Plugins/Online/OnlineSubsystemGoogle/Source/OnlineSubsystemGoogle.Build.cs
@@ -85,6 +85,14 @@ public class OnlineSubsystemGoogle : ModuleRules
 		{
 			PrivateIncludePaths.Add("Private/Windows");
 		}
+		else if (Target.Platform == UnrealTargetPlatform.XboxOne)
+		{
+			PrivateIncludePaths.Add("Private/XboxOne");
+		}
+		else if (Target.Platform == UnrealTargetPlatform.PS4)
+		{
+			PrivateIncludePaths.Add("Private/PS4");
+		}
 		else
 		{
 			PrecompileForTargets = PrecompileTargetsType.None;
diff --git a/Engine/Plugins/Online/OnlineSubsystemGoogle/Source/Private/OnlineExternalUIGoogleCommon.h b/Engine/Plugins/Online/OnlineSubsystemGoogle/Source/Private/OnlineExternalUIGoogleCommon.h
--- a/Engine/Plugins/Online/OnlineSubsystemGoogle/Source/Private/OnlineExternalUIGoogleCommon.h
+++ b/Engine/Plugins/Online/OnlineSubsystemGoogle/Source/Private/OnlineExternalUIGoogleCommon.h
@@ -42,7 +42,7 @@ public:
 	// IOnlineExternalUI
 	virtual bool ShowLoginUI(const int ControllerIndex, bool bShowOnlineOnly, bool bShowSkipButton, const FOnLoginUIClosedDelegate& Delegate = FOnLoginUIClosedDelegate()) override;
 	virtual bool ShowFriendsUI(int32 LocalUserNum) override;
-	virtual bool ShowInviteUI(int32 LocalUserNum, FName SessionName = GameSessionName) override;
+	virtual bool ShowInviteUI(int32 LocalUserNum, FName SessionName = NAME_GameSession) override;
 	virtual bool ShowAchievementsUI(int32 LocalUserNum) override;
 	virtual bool ShowLeaderboardUI(const FString& LeaderboardName) override;
 	virtual bool ShowWebURL(const FString& Url, const FShowWebUrlParams& ShowParams, const FOnShowWebUrlClosedDelegate& Delegate = FOnShowWebUrlClosedDelegate()) override;
diff --git a/Engine/Plugins/Online/OnlineSubsystemGoogle/Source/Private/OnlineIdentityGoogleCommon.cpp b/Engine/Plugins/Online/OnlineSubsystemGoogle/Source/Private/OnlineIdentityGoogleCommon.cpp
--- a/Engine/Plugins/Online/OnlineSubsystemGoogle/Source/Private/OnlineIdentityGoogleCommon.cpp
+++ b/Engine/Plugins/Online/OnlineSubsystemGoogle/Source/Private/OnlineIdentityGoogleCommon.cpp
@@ -52,7 +52,7 @@ bool FJsonWebTokenGoogle::Parse(const FString& InJWTStr)
 							static const FString Issuer2 = TEXT("accounts.google.com");
 							if ((Payload.ISS == Issuer1) || (Payload.ISS == Issuer2))
 							{
-								// Verify that the value of aud in the ID token is equal to your app’s client ID.
+								// Verify that the value of aud in the ID token is equal to your app's client ID.
 								FOnlineSubsystemGoogle* GoogleSubsystem = static_cast<FOnlineSubsystemGoogle*>(IOnlineSubsystem::Get(GOOGLE_SUBSYSTEM));
 								if (Payload.Aud == GoogleSubsystem->GetAppId() ||
 									Payload.Aud == GoogleSubsystem->GetServerClientId())
@@ -493,7 +493,7 @@ void FOnlineIdentityGoogleCommon::GetUserPrivilege(const FUniqueNetId& UserId, E
 	Delegate.ExecuteIfBound(UserId, Privilege, (uint32)EPrivilegeResults::NoFailures);
 }	
 
-FPlatformUserId FOnlineIdentityGoogleCommon::GetPlatformUserIdFromUniqueNetId(const FUniqueNetId& UniqueNetId)
+FPlatformUserId FOnlineIdentityGoogleCommon::GetPlatformUserIdFromUniqueNetId(const FUniqueNetId& UniqueNetId) const
 {
 	for (int i = 0; i < MAX_LOCAL_PLAYERS; ++i)
 	{
diff --git a/Engine/Plugins/Online/OnlineSubsystemGoogle/Source/Private/OnlineIdentityGoogleCommon.h b/Engine/Plugins/Online/OnlineSubsystemGoogle/Source/Private/OnlineIdentityGoogleCommon.h
--- a/Engine/Plugins/Online/OnlineSubsystemGoogle/Source/Private/OnlineIdentityGoogleCommon.h
+++ b/Engine/Plugins/Online/OnlineSubsystemGoogle/Source/Private/OnlineIdentityGoogleCommon.h
@@ -70,7 +70,7 @@ public:
 	virtual FString GetPlayerNickname(const FUniqueNetId& UserId) const override;
 	virtual FString GetAuthToken(int32 LocalUserNum) const override;
 	virtual void GetUserPrivilege(const FUniqueNetId& UserId, EUserPrivileges::Type Privilege, const FOnGetUserPrivilegeCompleteDelegate& Delegate) override;
-	virtual FPlatformUserId GetPlatformUserIdFromUniqueNetId(const FUniqueNetId& UniqueNetId) override;
+	virtual FPlatformUserId GetPlatformUserIdFromUniqueNetId(const FUniqueNetId& UniqueNetId) const override;
 	virtual FString GetAuthType() const override;
 	virtual void RevokeAuthToken(const FUniqueNetId& UserId, const FOnRevokeAuthTokenCompleteDelegate& Delegate) override;
 
diff --git a/Engine/Plugins/Online/OnlineSubsystemNull/Source/Private/OnlineIdentityNull.cpp b/Engine/Plugins/Online/OnlineSubsystemNull/Source/Private/OnlineIdentityNull.cpp
--- a/Engine/Plugins/Online/OnlineSubsystemNull/Source/Private/OnlineIdentityNull.cpp
+++ b/Engine/Plugins/Online/OnlineSubsystemNull/Source/Private/OnlineIdentityNull.cpp
@@ -302,7 +302,7 @@ void FOnlineIdentityNull::GetUserPrivilege(const FUniqueNetId& UserId, EUserPriv
 	Delegate.ExecuteIfBound(UserId, Privilege, (uint32)EPrivilegeResults::NoFailures);
 }
 
-FPlatformUserId FOnlineIdentityNull::GetPlatformUserIdFromUniqueNetId(const FUniqueNetId& UniqueNetId)
+FPlatformUserId FOnlineIdentityNull::GetPlatformUserIdFromUniqueNetId(const FUniqueNetId& UniqueNetId) const
 {
 	for (int i = 0; i < MAX_LOCAL_PLAYERS; ++i)
 	{
diff --git a/Engine/Plugins/Online/OnlineSubsystemNull/Source/Private/OnlineIdentityNull.h b/Engine/Plugins/Online/OnlineSubsystemNull/Source/Private/OnlineIdentityNull.h
--- a/Engine/Plugins/Online/OnlineSubsystemNull/Source/Private/OnlineIdentityNull.h
+++ b/Engine/Plugins/Online/OnlineSubsystemNull/Source/Private/OnlineIdentityNull.h
@@ -74,7 +74,7 @@ public:
 	virtual FString GetAuthToken(int32 LocalUserNum) const override;
 	virtual void RevokeAuthToken(const FUniqueNetId& UserId, const FOnRevokeAuthTokenCompleteDelegate& Delegate) override;
 	virtual void GetUserPrivilege(const FUniqueNetId& UserId, EUserPrivileges::Type Privilege, const FOnGetUserPrivilegeCompleteDelegate& Delegate) override;
-	virtual FPlatformUserId GetPlatformUserIdFromUniqueNetId(const FUniqueNetId& UniqueNetId) override;
+	virtual FPlatformUserId GetPlatformUserIdFromUniqueNetId(const FUniqueNetId& UniqueNetId) const override;
 	virtual FString GetAuthType() const override;
 
 	// FOnlineIdentityNull
diff --git a/Engine/Plugins/Online/OnlineSubsystemNull/Source/Private/OnlineSessionInterfaceNull.cpp b/Engine/Plugins/Online/OnlineSubsystemNull/Source/Private/OnlineSessionInterfaceNull.cpp
--- a/Engine/Plugins/Online/OnlineSubsystemNull/Source/Private/OnlineSessionInterfaceNull.cpp
+++ b/Engine/Plugins/Online/OnlineSubsystemNull/Source/Private/OnlineSessionInterfaceNull.cpp
@@ -745,12 +745,12 @@ bool FOnlineSessionNull::GetResolvedConnectString(FName SessionName, FString& Co
 	if (Session != NULL)
 	{
 		TSharedPtr<FOnlineSessionInfoNull> SessionInfo = StaticCastSharedPtr<FOnlineSessionInfoNull>(Session->SessionInfo);
-		if (PortType == BeaconPort)
+		if (PortType == NAME_BeaconPort)
 		{
 			int32 BeaconListenPort = GetBeaconPortFromSessionSettings(Session->SessionSettings);
 			bSuccess = GetConnectStringFromSessionInfo(SessionInfo, ConnectInfo, BeaconListenPort);
 		}
-		else if (PortType == GamePort)
+		else if (PortType == NAME_GamePort)
 		{
 			bSuccess = GetConnectStringFromSessionInfo(SessionInfo, ConnectInfo);
 		}
@@ -777,13 +777,13 @@ bool FOnlineSessionNull::GetResolvedConnectString(const FOnlineSessionSearchResu
 	{
 		TSharedPtr<FOnlineSessionInfoNull> SessionInfo = StaticCastSharedPtr<FOnlineSessionInfoNull>(SearchResult.Session.SessionInfo);
 
-		if (PortType == BeaconPort)
+		if (PortType == NAME_BeaconPort)
 		{
 			int32 BeaconListenPort = GetBeaconPortFromSessionSettings(SearchResult.Session.SessionSettings);
 			bSuccess = GetConnectStringFromSessionInfo(SessionInfo, ConnectInfo, BeaconListenPort);
 
 		}
-		else if (PortType == GamePort)
+		else if (PortType == NAME_GamePort)
 		{
 			bSuccess = GetConnectStringFromSessionInfo(SessionInfo, ConnectInfo);
 		}
diff --git a/Engine/Plugins/Online/OnlineSubsystemOculus/Source/Private/OculusCreateSessionCallbackProxy.cpp b/Engine/Plugins/Online/OnlineSubsystemOculus/Source/Private/OculusCreateSessionCallbackProxy.cpp
--- a/Engine/Plugins/Online/OnlineSubsystemOculus/Source/Private/OculusCreateSessionCallbackProxy.cpp
+++ b/Engine/Plugins/Online/OnlineSubsystemOculus/Source/Private/OculusCreateSessionCallbackProxy.cpp
@@ -43,7 +43,7 @@ void UOculusCreateSessionCallbackProxy::Activate()
 			Settings.Set(SETTING_OCULUS_POOL, OculusPool, EOnlineDataAdvertisementType::ViaOnlineService);
 		}
 
-		OculusSessionInterface->CreateSession(0, GameSessionName, Settings);
+		OculusSessionInterface->CreateSession(0, NAME_GameSession, Settings);
 	}
 	else
 	{
@@ -63,7 +63,7 @@ void UOculusCreateSessionCallbackProxy::OnCreateCompleted(FName SessionName, boo
 		if (bWasSuccessful)
 		{
 			StartCompleteDelegateHandle = OculusSessionInterface->AddOnStartSessionCompleteDelegate_Handle(StartCompleteDelegate);
-			OculusSessionInterface->StartSession(GameSessionName);
+			OculusSessionInterface->StartSession(NAME_GameSession);
 
 			// OnStartCompleted will get called, nothing more to do now
 			return;
diff --git a/Engine/Plugins/Online/OnlineSubsystemOculus/Source/Private/OnlineIdentityOculus.cpp b/Engine/Plugins/Online/OnlineSubsystemOculus/Source/Private/OnlineIdentityOculus.cpp
--- a/Engine/Plugins/Online/OnlineSubsystemOculus/Source/Private/OnlineIdentityOculus.cpp
+++ b/Engine/Plugins/Online/OnlineSubsystemOculus/Source/Private/OnlineIdentityOculus.cpp
@@ -294,7 +294,7 @@ void FOnlineIdentityOculus::GetUserPrivilege(const FUniqueNetId& UserId, EUserPr
 		}));
 }
 
-FPlatformUserId FOnlineIdentityOculus::GetPlatformUserIdFromUniqueNetId(const FUniqueNetId& UniqueNetId)
+FPlatformUserId FOnlineIdentityOculus::GetPlatformUserIdFromUniqueNetId(const FUniqueNetId& UniqueNetId) const
 {
 	for (int i = 0; i < MAX_LOCAL_PLAYERS; ++i)
 	{
diff --git a/Engine/Plugins/Online/OnlineSubsystemOculus/Source/Private/OnlineIdentityOculus.h b/Engine/Plugins/Online/OnlineSubsystemOculus/Source/Private/OnlineIdentityOculus.h
--- a/Engine/Plugins/Online/OnlineSubsystemOculus/Source/Private/OnlineIdentityOculus.h
+++ b/Engine/Plugins/Online/OnlineSubsystemOculus/Source/Private/OnlineIdentityOculus.h
@@ -77,7 +77,7 @@ public:
 	virtual FString GetAuthToken(int32 LocalUserNum) const override;
 	virtual void RevokeAuthToken(const FUniqueNetId& UserId, const FOnRevokeAuthTokenCompleteDelegate& Delegate) override;
 	virtual void GetUserPrivilege(const FUniqueNetId& UserId, EUserPrivileges::Type Privilege, const FOnGetUserPrivilegeCompleteDelegate& Delegate) override;
-	virtual FPlatformUserId GetPlatformUserIdFromUniqueNetId(const FUniqueNetId& UniqueNetId) override;
+	virtual FPlatformUserId GetPlatformUserIdFromUniqueNetId(const FUniqueNetId& UniqueNetId) const override;
 	virtual FString GetAuthType() const override;
 
 	// FOnlineIdentityOculus
diff --git a/Engine/Plugins/Online/OnlineSubsystemOculus/Source/Private/OnlineSessionInterfaceOculus.cpp b/Engine/Plugins/Online/OnlineSubsystemOculus/Source/Private/OnlineSessionInterfaceOculus.cpp
--- a/Engine/Plugins/Online/OnlineSubsystemOculus/Source/Private/OnlineSessionInterfaceOculus.cpp
+++ b/Engine/Plugins/Online/OnlineSubsystemOculus/Source/Private/OnlineSessionInterfaceOculus.cpp
@@ -57,13 +57,25 @@ FOnlineSessionOculus::~FOnlineSessionOculus()
 	// Make sure the player leaves all the sessions they were in before destroying this
 	for (auto It = Sessions.CreateConstIterator(); It; ++It)
 	{
-		auto Session = It.Value();
-		auto RoomId = GetOvrIDFromSession(*Session);
-		if (RoomId != 0)
+		TSharedPtr<FNamedOnlineSession> Session = It.Value();
+		if (Session.IsValid())
 		{
-			ovr_Room_Leave(RoomId);
+			ovrID RoomId = GetOvrIDFromSession(*Session);
+			if (RoomId != 0)
+			{
+				ovr_Room_Leave(RoomId);
+			}
+
+			if (!Session.IsUnique())
+			{
+				UE_LOG_ONLINE(Warning, TEXT("Session pointer (room %d) not unique during cleanup!"), RoomId);
+			}
+			Session->SessionState = EOnlineSessionState::Destroying;
+		}
+		else
+		{
+			UE_LOG_ONLINE(Warning, TEXT("Invalid session during shutdown!"));
 		}
-		Session->SessionState = EOnlineSessionState::Destroying;
 	}
 	Sessions.Empty();
 };
@@ -971,8 +983,8 @@ bool FOnlineSessionOculus::FindFriendSession(int32 LocalUserNum, const FUniqueNe
 			TriggerOnFindFriendSessionCompleteDelegates(LocalUserNum, false, SearchResult);
 			return;
 		}
-    
-    // UE4 specific setting.  Sessions with different build unique ids shouldn't be able to see each other
+	
+	// UE4 specific setting.  Sessions with different build unique ids shouldn't be able to see each other
 		// because the builds are not compatible
 		int32 BuildUniqueId = GetBuildUniqueId();
 
diff --git a/Engine/Plugins/Online/OnlineSubsystemOculus/Source/Private/OnlineSessionInterfaceOculus.h b/Engine/Plugins/Online/OnlineSubsystemOculus/Source/Private/OnlineSessionInterfaceOculus.h
--- a/Engine/Plugins/Online/OnlineSubsystemOculus/Source/Private/OnlineSessionInterfaceOculus.h
+++ b/Engine/Plugins/Online/OnlineSubsystemOculus/Source/Private/OnlineSessionInterfaceOculus.h
@@ -104,7 +104,7 @@ public:
 	virtual bool SendSessionInviteToFriend(const FUniqueNetId& LocalUserId, FName SessionName, const FUniqueNetId& Friend) override;
 	virtual bool SendSessionInviteToFriends(int32 LocalUserNum, FName SessionName, const TArray< TSharedRef<const FUniqueNetId> >& Friends) override;
 	virtual bool SendSessionInviteToFriends(const FUniqueNetId& LocalUserId, FName SessionName, const TArray< TSharedRef<const FUniqueNetId> >& Friends) override;
-	virtual bool GetResolvedConnectString(FName SessionName, FString& ConnectInfo, FName PortType = GamePort) override;
+	virtual bool GetResolvedConnectString(FName SessionName, FString& ConnectInfo, FName PortType = NAME_GamePort) override;
 	virtual bool GetResolvedConnectString(const class FOnlineSessionSearchResult& SearchResult, FName PortType, FString& ConnectInfo) override;
 	virtual FOnlineSessionSettings* GetSessionSettings(FName SessionName) override;
 	virtual bool RegisterPlayer(FName SessionName, const FUniqueNetId& PlayerId, bool bWasInvited) override;
diff --git a/Engine/Plugins/Online/OnlineSubsystemOculus/Source/Private/OnlineVoiceOculus.cpp b/Engine/Plugins/Online/OnlineSubsystemOculus/Source/Private/OnlineVoiceOculus.cpp
--- a/Engine/Plugins/Online/OnlineSubsystemOculus/Source/Private/OnlineVoiceOculus.cpp
+++ b/Engine/Plugins/Online/OnlineSubsystemOculus/Source/Private/OnlineVoiceOculus.cpp
@@ -14,6 +14,7 @@
 
 // Define the voice sample rate the default in OVR_VoipSampleRate
 #define OCULUS_VOICE_SAMPLE_RATE 48000
+#define OCULUS_NUM_VOICE_CHANNELS 1
 
 FRemoteTalkerDataOculus::FRemoteTalkerDataOculus() :
 	LastSeen(0.0),
@@ -206,7 +207,7 @@ void FOnlineVoiceOculus::ProcessRemoteVoicePackets()
 
 			if (QueuedData.AudioComponent == nullptr || QueuedData.AudioComponent->IsPendingKill())
 			{
-				QueuedData.AudioComponent = CreateVoiceAudioComponent(OCULUS_VOICE_SAMPLE_RATE);
+				QueuedData.AudioComponent = CreateVoiceAudioComponent(OCULUS_VOICE_SAMPLE_RATE, OCULUS_NUM_VOICE_CHANNELS);
 				if (QueuedData.AudioComponent)
 				{
 					QueuedData.AudioComponent->AddToRoot(); // make sure this doesn't get deleted by the GC
diff --git a/Engine/Plugins/Online/OnlineSubsystemSteam/Source/Private/OnlineExternalUIInterfaceSteam.h b/Engine/Plugins/Online/OnlineSubsystemSteam/Source/Private/OnlineExternalUIInterfaceSteam.h
--- a/Engine/Plugins/Online/OnlineSubsystemSteam/Source/Private/OnlineExternalUIInterfaceSteam.h
+++ b/Engine/Plugins/Online/OnlineSubsystemSteam/Source/Private/OnlineExternalUIInterfaceSteam.h
@@ -82,7 +82,7 @@ public:
 	// IOnlineExternalUI
 	virtual bool ShowLoginUI(const int ControllerIndex, bool bShowOnlineOnly, bool bShowSkipButton, const FOnLoginUIClosedDelegate& Delegate = FOnLoginUIClosedDelegate()) override;
 	virtual bool ShowFriendsUI(int32 LocalUserNum) override;
-	virtual bool ShowInviteUI(int32 LocalUserNum, FName SessionName = GameSessionName) override;
+	virtual bool ShowInviteUI(int32 LocalUserNum, FName SessionName = NAME_GameSession) override;
 	virtual bool ShowAchievementsUI(int32 LocalUserNum) override;
 	virtual bool ShowLeaderboardUI(const FString& LeaderboardName) override;
 	virtual bool ShowWebURL(const FString& Url, const FShowWebUrlParams& ShowParams, const FOnShowWebUrlClosedDelegate& Delegate = FOnShowWebUrlClosedDelegate()) override;
diff --git a/Engine/Plugins/Online/OnlineSubsystemSteam/Source/Private/OnlineIdentityInterfaceSteam.cpp b/Engine/Plugins/Online/OnlineSubsystemSteam/Source/Private/OnlineIdentityInterfaceSteam.cpp
--- a/Engine/Plugins/Online/OnlineSubsystemSteam/Source/Private/OnlineIdentityInterfaceSteam.cpp
+++ b/Engine/Plugins/Online/OnlineSubsystemSteam/Source/Private/OnlineIdentityInterfaceSteam.cpp
@@ -219,7 +219,7 @@ void FOnlineIdentitySteam::GetUserPrivilege(const FUniqueNetId& UserId, EUserPri
 	Delegate.ExecuteIfBound(UserId, Privilege, (uint32)EPrivilegeResults::NoFailures);
 }
 
-FPlatformUserId FOnlineIdentitySteam::GetPlatformUserIdFromUniqueNetId(const FUniqueNetId& UniqueNetId)
+FPlatformUserId FOnlineIdentitySteam::GetPlatformUserIdFromUniqueNetId(const FUniqueNetId& UniqueNetId) const
 {
 	for (int i = 0; i < MAX_LOCAL_PLAYERS; ++i)
 	{
diff --git a/Engine/Plugins/Online/OnlineSubsystemSteam/Source/Private/OnlineIdentityInterfaceSteam.h b/Engine/Plugins/Online/OnlineSubsystemSteam/Source/Private/OnlineIdentityInterfaceSteam.h
--- a/Engine/Plugins/Online/OnlineSubsystemSteam/Source/Private/OnlineIdentityInterfaceSteam.h
+++ b/Engine/Plugins/Online/OnlineSubsystemSteam/Source/Private/OnlineIdentityInterfaceSteam.h
@@ -45,7 +45,7 @@ public:
 	virtual FString GetAuthToken(int32 LocalUserNum) const override;
 	virtual void RevokeAuthToken(const FUniqueNetId& UserId, const FOnRevokeAuthTokenCompleteDelegate& Delegate) override;
 	virtual void GetUserPrivilege(const FUniqueNetId& UserId, EUserPrivileges::Type Privilege, const FOnGetUserPrivilegeCompleteDelegate& Delegate) override;
-	virtual FPlatformUserId GetPlatformUserIdFromUniqueNetId(const FUniqueNetId& UniqueNetId) override;
+	virtual FPlatformUserId GetPlatformUserIdFromUniqueNetId(const FUniqueNetId& UniqueNetId) const override;
 	virtual FString GetAuthType() const override;
 };
 
diff --git a/Engine/Plugins/Online/OnlineSubsystemSteam/Source/Private/OnlineSessionInterfaceSteam.cpp b/Engine/Plugins/Online/OnlineSubsystemSteam/Source/Private/OnlineSessionInterfaceSteam.cpp
--- a/Engine/Plugins/Online/OnlineSubsystemSteam/Source/Private/OnlineSessionInterfaceSteam.cpp
+++ b/Engine/Plugins/Online/OnlineSubsystemSteam/Source/Private/OnlineSessionInterfaceSteam.cpp
@@ -1215,12 +1215,12 @@ bool FOnlineSessionSteam::GetResolvedConnectString(FName SessionName, FString& C
 	if (Session != NULL)
 	{
 		TSharedPtr<FOnlineSessionInfoSteam> SessionInfo = StaticCastSharedPtr<FOnlineSessionInfoSteam>(Session->SessionInfo);
-		if (PortType == BeaconPort)
+		if (PortType == NAME_BeaconPort)
 		{
 			int32 BeaconListenPort = GetBeaconPortFromSessionSettings(Session->SessionSettings);
 			bSuccess = GetConnectStringFromSessionInfo(SessionInfo, ConnectInfo, BeaconListenPort);
 		}
-		else if (PortType == GamePort)
+		else if (PortType == NAME_GamePort)
 		{
 			bSuccess = GetConnectStringFromSessionInfo(SessionInfo, ConnectInfo);
 		}
@@ -1247,12 +1247,12 @@ bool FOnlineSessionSteam::GetResolvedConnectString(const FOnlineSessionSearchRes
 	{
 		TSharedPtr<FOnlineSessionInfoSteam> SessionInfo = StaticCastSharedPtr<FOnlineSessionInfoSteam>(SearchResult.Session.SessionInfo);
 
-		if (PortType == BeaconPort)
+		if (PortType == NAME_BeaconPort)
 		{
 			int32 BeaconListenPort = GetBeaconPortFromSessionSettings(SearchResult.Session.SessionSettings);
 			bSuccess = GetConnectStringFromSessionInfo(SessionInfo, ConnectInfo, BeaconListenPort);
 		}
-		else if (PortType == GamePort)
+		else if (PortType == NAME_GamePort)
 		{
 			bSuccess = GetConnectStringFromSessionInfo(SessionInfo, ConnectInfo);
 		}
diff --git a/Engine/Plugins/Online/OnlineSubsystemSteam/Source/Private/VoiceEngineSteam.cpp b/Engine/Plugins/Online/OnlineSubsystemSteam/Source/Private/VoiceEngineSteam.cpp
--- a/Engine/Plugins/Online/OnlineSubsystemSteam/Source/Private/VoiceEngineSteam.cpp
+++ b/Engine/Plugins/Online/OnlineSubsystemSteam/Source/Private/VoiceEngineSteam.cpp
@@ -46,7 +46,7 @@ void FVoiceEngineSteam::VoiceCaptureUpdate() const
 		// If no data is available, we have finished capture the last (post-StopRecording) half-second of voice data
 		if (RecordingState == k_EVoiceResultNotRecording)
 		{
-			UE_LOG(LogVoiceEncode, Log, TEXT("Internal voice capture complete."));
+			UE_LOG(LogVoiceEngine, Log, TEXT("Internal voice capture complete."));
 
 			bPendingFinalCapture = false;
 
@@ -66,7 +66,7 @@ void FVoiceEngineSteam::VoiceCaptureUpdate() const
 
 void FVoiceEngineSteam::StartRecording() const
 {
-	UE_LOG(LogVoiceEncode, VeryVerbose, TEXT("VOIP StartRecording"));
+	UE_LOG(LogVoiceEngine, VeryVerbose, TEXT("VOIP StartRecording"));
 	if (SteamUserPtr)
 	{
 		SteamUserPtr->StartVoiceRecording();
@@ -76,7 +76,7 @@ void FVoiceEngineSteam::StartRecording() const
 
 void FVoiceEngineSteam::StopRecording() const
 {
-	UE_LOG(LogVoiceEncode, VeryVerbose, TEXT("VOIP StopRecording"));
+	UE_LOG(LogVoiceEngine, VeryVerbose, TEXT("VOIP StopRecording"));
 	if (SteamUserPtr)
 	{
 		SteamUserPtr->StopVoiceRecording();
@@ -85,7 +85,7 @@ void FVoiceEngineSteam::StopRecording() const
 
 void FVoiceEngineSteam::StoppedRecording() const
 {
-	UE_LOG(LogVoiceEncode, VeryVerbose, TEXT("VOIP StoppedRecording"));
+	UE_LOG(LogVoiceEngine, VeryVerbose, TEXT("VOIP StoppedRecording"));
 	if (SteamFriendsPtr)
 	{
 		SteamFriendsPtr->SetInGameVoiceSpeaking(SteamUserPtr->GetSteamID(), false);
@@ -152,7 +152,7 @@ uint32 FVoiceEngineSteam::StartLocalVoiceProcessing(uint32 LocalUserNum)
 	}
 	else
 	{
-		UE_LOG(LogVoiceEncode, Error, TEXT("StartLocalVoiceProcessing(): Device is currently owned by another user"));
+		UE_LOG(LogVoiceEngine, Error, TEXT("StartLocalVoiceProcessing(): Device is currently owned by another user"));
 	}
 
 	return Return;
@@ -179,7 +179,7 @@ uint32 FVoiceEngineSteam::StopLocalVoiceProcessing(uint32 LocalUserNum)
 	}
 	else
 	{
-		UE_LOG(LogVoiceEncode, Error, TEXT("StopLocalVoiceProcessing: Ignoring stop request for non-owning user"));
+		UE_LOG(LogVoiceEngine, Error, TEXT("StopLocalVoiceProcessing: Ignoring stop request for non-owning user"));
 	}
 
 	return Return;
@@ -217,14 +217,14 @@ uint32 FVoiceEngineSteam::ReadLocalVoiceData(uint32 LocalUserNum, uint8* Data, u
 		EVoiceResult VoiceResult = SteamUserPtr->GetAvailableVoice(&CompressedBytes, NULL, 0);
 		if (VoiceResult != k_EVoiceResultOK && VoiceResult != k_EVoiceResultNoData)
 		{
-			UE_LOG(LogVoiceEncode, Warning, TEXT("ReadLocalVoiceData: GetAvailableVoice failure: VoiceResult: %s"), *SteamVoiceResult(VoiceResult));
+			UE_LOG(LogVoiceEngine, Warning, TEXT("ReadLocalVoiceData: GetAvailableVoice failure: VoiceResult: %s"), *SteamVoiceResult(VoiceResult));
 			return E_FAIL;
 		}
 
 		// This shouldn't happen, but just in case.
 		if (CompressedBytes == 0)
 		{
-			UE_LOG(LogVoiceEncode, VeryVerbose, TEXT("ReadLocalVoiceData: No Data: VoiceResult: %s"), *SteamVoiceResult(VoiceResult));
+			UE_LOG(LogVoiceEngine, VeryVerbose, TEXT("ReadLocalVoiceData: No Data: VoiceResult: %s"), *SteamVoiceResult(VoiceResult));
 			*Size = 0;
 			return S_OK;
 		}
@@ -237,13 +237,13 @@ uint32 FVoiceEngineSteam::ReadLocalVoiceData(uint32 LocalUserNum, uint8* Data, u
 
 		static double LastGetVoiceCallTime = 0.0;
 		double CurTime = FPlatformTime::Seconds();
-		double TimeSinceLastCall = LastGetVoiceCallTime > 0 ? (CurTime - LastGetVoiceCallTime) : 0.0;
+		double TimeSinceLastCall = (LastGetVoiceCallTime > 0) ? (CurTime - LastGetVoiceCallTime) : 0.0;
 		LastGetVoiceCallTime = CurTime;
 
-		UE_LOG(LogVoiceEncode, Log, TEXT("ReadLocalVoiceData: GetVoice: Result: %s, Available: %i, LastCall: %0.3f"), *SteamVoiceResult(VoiceResult), AvailableWritten, TimeSinceLastCall);
+		UE_LOG(LogVoiceEngine, Log, TEXT("ReadLocalVoiceData: GetVoice: Result: %s, Available: %i, LastCall: %0.3f"), *SteamVoiceResult(VoiceResult), AvailableWritten, TimeSinceLastCall * 1000.0);
 		if (VoiceResult != k_EVoiceResultOK)
 		{
-			UE_LOG(LogVoiceEncode, Warning, TEXT("ReadLocalVoiceData: GetVoice failure: VoiceResult: %s"), *SteamVoiceResult(VoiceResult));
+			UE_LOG(LogVoiceEngine, Warning, TEXT("ReadLocalVoiceData: GetVoice failure: VoiceResult: %s"), *SteamVoiceResult(VoiceResult));
 			*Size = 0;
 			CompressedVoiceBuffer.Empty(MAX_COMPRESSED_VOICE_BUFFER_SIZE);
 			return E_FAIL;
@@ -254,7 +254,7 @@ uint32 FVoiceEngineSteam::ReadLocalVoiceData(uint32 LocalUserNum, uint8* Data, u
 			*Size = FMath::Min<int32>(*Size, AvailableWritten);
 			FMemory::Memcpy(Data, CompressedVoiceBuffer.GetData(), *Size);
 
-			UE_LOG(LogVoiceEncode, VeryVerbose, TEXT("ReadLocalVoiceData: Size: %d"), *Size);
+			UE_LOG(LogVoiceEngine, VeryVerbose, TEXT("ReadLocalVoiceData: Size: %d"), *Size);
 			return S_OK;
 		}
 		else
@@ -268,7 +268,7 @@ uint32 FVoiceEngineSteam::ReadLocalVoiceData(uint32 LocalUserNum, uint8* Data, u
 
 uint32 FVoiceEngineSteam::SubmitRemoteVoiceData(const FUniqueNetId& RemoteTalkerId, uint8* Data, uint32* Size) 
 {
-	UE_LOG(LogVoiceDecode, VeryVerbose, TEXT("SubmitRemoteVoiceData(%s) Size: %d received!"), *RemoteTalkerId.ToDebugString(), *Size);
+	UE_LOG(LogVoiceEngine, VeryVerbose, TEXT("SubmitRemoteVoiceData(%s) Size: %d received!"), *RemoteTalkerId.ToDebugString(), *Size);
 
 	const FUniqueNetIdSteam& SteamId = (const FUniqueNetIdSteam&)RemoteTalkerId;
 	FRemoteTalkerDataSteam* QueuedData = RemoteTalkerBuffers.Find(SteamId);
@@ -293,7 +293,7 @@ uint32 FVoiceEngineSteam::SubmitRemoteVoiceData(const FUniqueNetId& RemoteTalker
 
 	if (VoiceResult != k_EVoiceResultOK)
 	{
-		UE_LOG(LogVoiceDecode, Warning, TEXT("SubmitRemoteVoiceData: DecompressVoice failure: VoiceResult: %s"), *SteamVoiceResult(VoiceResult));
+		UE_LOG(LogVoiceEngine, Warning, TEXT("SubmitRemoteVoiceData: DecompressVoice failure: VoiceResult: %s"), *SteamVoiceResult(VoiceResult));
 		*Size = 0;
 		return E_FAIL;
 	}
@@ -313,7 +313,7 @@ uint32 FVoiceEngineSteam::SubmitRemoteVoiceData(const FUniqueNetId& RemoteTalker
 			SerializeHelper = new FVoiceSerializeHelper(this);
 		}
 
-		QueuedData->AudioComponent = CreateVoiceAudioComponent(SteamUserPtr->GetVoiceOptimalSampleRate());
+		QueuedData->AudioComponent = CreateVoiceAudioComponent(SteamUserPtr->GetVoiceOptimalSampleRate(), DEFAULT_NUM_VOICE_CHANNELS);
 		if (QueuedData->AudioComponent)
 		{
 			QueuedData->AudioComponent->OnAudioFinishedNative.AddRaw(this, &FVoiceEngineSteam::OnAudioFinished);
@@ -326,7 +326,7 @@ uint32 FVoiceEngineSteam::SubmitRemoteVoiceData(const FUniqueNetId& RemoteTalker
 		USoundWaveProcedural* SoundStreaming = CastChecked<USoundWaveProcedural>(QueuedData->AudioComponent->Sound);
 		if (SoundStreaming->GetAvailableAudioByteCount() == 0)
 		{
-			UE_LOG(LogVoiceDecode, Log, TEXT("VOIP audio component was starved!"));
+			UE_LOG(LogVoiceEngine, Log, TEXT("VOIP audio component was starved!"));
 		}
 		SoundStreaming->QueueAudio(DecompressedVoiceBuffer.GetData(), BytesWritten);
 	}
@@ -360,7 +360,7 @@ void FVoiceEngineSteam::Tick(float DeltaTime)
 {
 	// Check available voice one a frame, this value changes after calling GetVoice()
 	AvailableVoiceResult = SteamUserPtr->GetAvailableVoice(&CompressedBytesAvailable, NULL, 0);
-	//UE_LOG(LogVoiceEncode, Log, TEXT("TickResult: %s"), *SteamVoiceResult(AvailableVoiceResult));
+	//UE_LOG(LogVoiceEngine, Log, TEXT("TickResult: %s"), *SteamVoiceResult(AvailableVoiceResult));
 	TickTalkers(DeltaTime);
 }
 
@@ -371,12 +371,12 @@ void FVoiceEngineSteam::OnAudioFinished(UAudioComponent* AC)
 		FRemoteTalkerDataSteam& RemoteData = It.Value();
 		if (RemoteData.AudioComponent->IsPendingKill() || AC == RemoteData.AudioComponent)
 		{
-			UE_LOG(LogVoiceDecode, Log, TEXT("Removing VOIP AudioComponent for SteamId: %s"), *It.Key().ToDebugString());
+			UE_LOG(LogVoiceEngine, Log, TEXT("Removing VOIP AudioComponent for SteamId: %s"), *It.Key().ToDebugString());
 			RemoteData.AudioComponent = NULL;
 			break;
 		}
 	}
-	UE_LOG(LogVoiceDecode, Verbose, TEXT("Audio Finished"));
+	UE_LOG(LogVoiceEngine, Verbose, TEXT("Audio Finished"));
 }
 
 FString FVoiceEngineSteam::GetVoiceDebugState() const 
diff --git a/Engine/Plugins/Online/OnlineSubsystemSteam/Source/Private/VoiceInterfaceSteam.cpp b/Engine/Plugins/Online/OnlineSubsystemSteam/Source/Private/VoiceInterfaceSteam.cpp
--- a/Engine/Plugins/Online/OnlineSubsystemSteam/Source/Private/VoiceInterfaceSteam.cpp
+++ b/Engine/Plugins/Online/OnlineSubsystemSteam/Source/Private/VoiceInterfaceSteam.cpp
@@ -732,12 +732,12 @@ void FOnlineVoiceSteam::ProcessLocalVoicePackets()
 						}
 						else
 						{
-							UE_LOG(LogVoiceEncode, Warning, TEXT("Voice data error in ReadLocalVoiceData"));
+							UE_LOG(LogVoiceEngine, Warning, TEXT("Voice data error in ReadLocalVoiceData"));
 						}
 					}
 					else
 					{
-						UE_LOG(LogVoiceEncode, Warning, TEXT("Dropping voice data due to network layer not processing fast enough"));
+						UE_LOG(LogVoiceEngine, Warning, TEXT("Dropping voice data due to network layer not processing fast enough"));
 						// Buffer overflow, so drop previous data
 						VoiceData.LocalPackets[Index].Length = 0;
 					}
@@ -770,7 +770,7 @@ void FOnlineVoiceSteam::ProcessRemoteVoicePackets()
 				uint32 Result = VoiceEngine->SubmitRemoteVoiceData(*VoicePacket->Sender, VoicePacket->Buffer.GetData(), &VoiceBufferSize);
 				if (Result != S_OK)
 				{
-					UE_LOG(LogVoiceDecode, Log,
+					UE_LOG(LogVoiceEngine, Log,
 						TEXT("SubmitRemoteVoiceData(%s) failed with 0x%08X"),
 						*VoicePacket->Sender->ToDebugString(),
 						Result);
diff --git a/Engine/Plugins/Online/OnlineSubsystemTwitch/Source/Private/OnlineExternalUITwitch.h b/Engine/Plugins/Online/OnlineSubsystemTwitch/Source/Private/OnlineExternalUITwitch.h
--- a/Engine/Plugins/Online/OnlineSubsystemTwitch/Source/Private/OnlineExternalUITwitch.h
+++ b/Engine/Plugins/Online/OnlineSubsystemTwitch/Source/Private/OnlineExternalUITwitch.h
@@ -41,7 +41,7 @@ public:
 	// IOnlineExternalUI
 	virtual bool ShowLoginUI(const int ControllerIndex, bool bShowOnlineOnly, bool bShowSkipButton, const FOnLoginUIClosedDelegate& Delegate = FOnLoginUIClosedDelegate()) override;
 	virtual bool ShowFriendsUI(int32 LocalUserNum) override;
-	virtual bool ShowInviteUI(int32 LocalUserNum, FName SessionName = GameSessionName) override;
+	virtual bool ShowInviteUI(int32 LocalUserNum, FName SessionName = NAME_GameSession) override;
 	virtual bool ShowAchievementsUI(int32 LocalUserNum) override;
 	virtual bool ShowLeaderboardUI(const FString& LeaderboardName) override;
 	virtual bool ShowWebURL(const FString& Url, const FShowWebUrlParams& ShowParams, const FOnShowWebUrlClosedDelegate& Delegate = FOnShowWebUrlClosedDelegate()) override;
diff --git a/Engine/Plugins/Online/OnlineSubsystemTwitch/Source/Private/OnlineIdentityTwitch.cpp b/Engine/Plugins/Online/OnlineSubsystemTwitch/Source/Private/OnlineIdentityTwitch.cpp
--- a/Engine/Plugins/Online/OnlineSubsystemTwitch/Source/Private/OnlineIdentityTwitch.cpp
+++ b/Engine/Plugins/Online/OnlineSubsystemTwitch/Source/Private/OnlineIdentityTwitch.cpp
@@ -710,7 +710,7 @@ FOnlineIdentityTwitch::FOnlineIdentityTwitch(FOnlineSubsystemTwitch* InSubsystem
 	check(InSubsystem);
 }
 
-FPlatformUserId FOnlineIdentityTwitch::GetPlatformUserIdFromUniqueNetId(const FUniqueNetId& UniqueNetId)
+FPlatformUserId FOnlineIdentityTwitch::GetPlatformUserIdFromUniqueNetId(const FUniqueNetId& UniqueNetId) const
 {
 	for (int PlayerIdx = 0; PlayerIdx < MAX_LOCAL_PLAYERS; ++PlayerIdx)
 	{
diff --git a/Engine/Plugins/Online/OnlineSubsystemTwitch/Source/Public/OnlineIdentityTwitch.h b/Engine/Plugins/Online/OnlineSubsystemTwitch/Source/Public/OnlineIdentityTwitch.h
--- a/Engine/Plugins/Online/OnlineSubsystemTwitch/Source/Public/OnlineIdentityTwitch.h
+++ b/Engine/Plugins/Online/OnlineSubsystemTwitch/Source/Public/OnlineIdentityTwitch.h
@@ -125,7 +125,7 @@ public:
 	virtual FString GetAuthToken(int32 LocalUserNum) const override;
 	virtual void RevokeAuthToken(const FUniqueNetId& UserId, const FOnRevokeAuthTokenCompleteDelegate& Delegate) override;
 	virtual void GetUserPrivilege(const FUniqueNetId& UserId, EUserPrivileges::Type Privilege, const FOnGetUserPrivilegeCompleteDelegate& Delegate) override;
-	virtual FPlatformUserId GetPlatformUserIdFromUniqueNetId(const FUniqueNetId& UniqueNetId) override;
+	virtual FPlatformUserId GetPlatformUserIdFromUniqueNetId(const FUniqueNetId& UniqueNetId) const override;
 	virtual FString GetAuthType() const override;
 
 	// FOnlineIdentityTwitch
diff --git a/Engine/Plugins/Online/OnlineSubsystemUtils/Source/OnlineSubsystemUtils/Private/CreateSessionCallbackProxy.cpp b/Engine/Plugins/Online/OnlineSubsystemUtils/Source/OnlineSubsystemUtils/Private/CreateSessionCallbackProxy.cpp
--- a/Engine/Plugins/Online/OnlineSubsystemUtils/Source/OnlineSubsystemUtils/Private/CreateSessionCallbackProxy.cpp
+++ b/Engine/Plugins/Online/OnlineSubsystemUtils/Source/OnlineSubsystemUtils/Private/CreateSessionCallbackProxy.cpp
@@ -49,7 +49,7 @@ void UCreateSessionCallbackProxy::Activate()
 			Settings.bUsesPresence = true;
 			Settings.bAllowJoinViaPresence = true;
 
-			Sessions->CreateSession(*Helper.UserID, GameSessionName, Settings);
+			Sessions->CreateSession(*Helper.UserID, NAME_GameSession, Settings);
 
 			// OnCreateCompleted will get called, nothing more to do now
 			return;
@@ -79,7 +79,7 @@ void UCreateSessionCallbackProxy::OnCreateCompleted(FName SessionName, bool bWas
 			if (bWasSuccessful)
 			{
 				StartCompleteDelegateHandle = Sessions->AddOnStartSessionCompleteDelegate_Handle(StartCompleteDelegate);
-				Sessions->StartSession(GameSessionName);
+				Sessions->StartSession(NAME_GameSession);
 
 				// OnStartCompleted will get called, nothing more to do now
 				return;
diff --git a/Engine/Plugins/Online/OnlineSubsystemUtils/Source/OnlineSubsystemUtils/Private/DestroySessionCallbackProxy.cpp b/Engine/Plugins/Online/OnlineSubsystemUtils/Source/OnlineSubsystemUtils/Private/DestroySessionCallbackProxy.cpp
--- a/Engine/Plugins/Online/OnlineSubsystemUtils/Source/OnlineSubsystemUtils/Private/DestroySessionCallbackProxy.cpp
+++ b/Engine/Plugins/Online/OnlineSubsystemUtils/Source/OnlineSubsystemUtils/Private/DestroySessionCallbackProxy.cpp
@@ -35,7 +35,7 @@ void UDestroySessionCallbackProxy::Activate()
 		if (Sessions.IsValid())
 		{
 			DelegateHandle = Sessions->AddOnDestroySessionCompleteDelegate_Handle(Delegate);
-			Sessions->DestroySession(GameSessionName);
+			Sessions->DestroySession(NAME_GameSession);
 
 			// OnCompleted will get called, nothing more to do now
 			return;
diff --git a/Engine/Plugins/Online/OnlineSubsystemUtils/Source/OnlineSubsystemUtils/Private/JoinSessionCallbackProxy.cpp b/Engine/Plugins/Online/OnlineSubsystemUtils/Source/OnlineSubsystemUtils/Private/JoinSessionCallbackProxy.cpp
--- a/Engine/Plugins/Online/OnlineSubsystemUtils/Source/OnlineSubsystemUtils/Private/JoinSessionCallbackProxy.cpp
+++ b/Engine/Plugins/Online/OnlineSubsystemUtils/Source/OnlineSubsystemUtils/Private/JoinSessionCallbackProxy.cpp
@@ -36,7 +36,7 @@ void UJoinSessionCallbackProxy::Activate()
 		if (Sessions.IsValid())
 		{
 			DelegateHandle = Sessions->AddOnJoinSessionCompleteDelegate_Handle(Delegate);
-			Sessions->JoinSession(*Helper.UserID, GameSessionName, OnlineSearchResult);
+			Sessions->JoinSession(*Helper.UserID, NAME_GameSession, OnlineSearchResult);
 
 			// OnCompleted will get called, nothing more to do now
 			return;
@@ -67,7 +67,7 @@ void UJoinSessionCallbackProxy::OnCompleted(FName SessionName, EOnJoinSessionCom
 			{
 				// Client travel to the server
 				FString ConnectString;
-				if (Sessions->GetResolvedConnectString(GameSessionName, ConnectString) && PlayerControllerWeakPtr.IsValid())
+				if (Sessions->GetResolvedConnectString(NAME_GameSession, ConnectString) && PlayerControllerWeakPtr.IsValid())
 				{
 					UE_LOG(LogOnline, Log, TEXT("Join session: traveling to %s"), *ConnectString);
 					PlayerControllerWeakPtr->ClientTravel(ConnectString, TRAVEL_Absolute);
diff --git a/Engine/Plugins/Online/OnlineSubsystemUtils/Source/OnlineSubsystemUtils/Private/OnlineSessionClient.cpp b/Engine/Plugins/Online/OnlineSubsystemUtils/Source/OnlineSubsystemUtils/Private/OnlineSessionClient.cpp
--- a/Engine/Plugins/Online/OnlineSubsystemUtils/Source/OnlineSubsystemUtils/Private/OnlineSessionClient.cpp
+++ b/Engine/Plugins/Online/OnlineSubsystemUtils/Source/OnlineSubsystemUtils/Private/OnlineSessionClient.cpp
@@ -84,7 +84,7 @@ void UOnlineSessionClient::OnSessionUserInviteAccepted(bool bWasSuccessful, int3
 		if (SearchResult.IsValid())
 		{
 			bIsFromInvite = true;
-			JoinSession(GameSessionName, SearchResult);
+			JoinSession(NAME_GameSession, SearchResult);
 		}
 		else
 		{
@@ -264,7 +264,7 @@ bool UOnlineSessionClient::HandleDisconnectInternal(UWorld* World, UNetDriver* N
 		if (!bHandlingDisconnect)
 		{
 			bHandlingDisconnect = true;
-			DestroyExistingSession_Impl(OnDestroyForMainMenuCompleteDelegateHandle, GameSessionName, OnDestroyForMainMenuCompleteDelegate);
+			DestroyExistingSession_Impl(OnDestroyForMainMenuCompleteDelegateHandle, NAME_GameSession, OnDestroyForMainMenuCompleteDelegate);
 		}
 
 		return true;
diff --git a/Engine/Plugins/Online/OnlineSubsystemUtils/Source/OnlineSubsystemUtils/Private/OnlineSubsystemUtils.cpp b/Engine/Plugins/Online/OnlineSubsystemUtils/Source/OnlineSubsystemUtils/Private/OnlineSubsystemUtils.cpp
--- a/Engine/Plugins/Online/OnlineSubsystemUtils/Source/OnlineSubsystemUtils/Private/OnlineSubsystemUtils.cpp
+++ b/Engine/Plugins/Online/OnlineSubsystemUtils/Source/OnlineSubsystemUtils/Private/OnlineSubsystemUtils.cpp
@@ -35,7 +35,7 @@
 #include "Tests/TestVoice.h"
 #include "Tests/TestExternalUIInterface.h"
 
-UAudioComponent* CreateVoiceAudioComponent(uint32 SampleRate)
+UAudioComponent* CreateVoiceAudioComponent(uint32 SampleRate, int32 NumChannels)
 {
 	UAudioComponent* AudioComponent = nullptr;
 	if (GEngine != nullptr)
@@ -44,7 +44,7 @@ UAudioComponent* CreateVoiceAudioComponent(uint32 SampleRate)
 		{
 			USoundWaveProcedural* SoundStreaming = NewObject<USoundWaveProcedural>();
 			SoundStreaming->SampleRate = SampleRate;
-			SoundStreaming->NumChannels = 1;
+			SoundStreaming->NumChannels = NumChannels;
 			SoundStreaming->Duration = INDEFINITELY_LOOPING_DURATION;
 			SoundStreaming->SoundGroup = SOUNDGROUP_Voice;
 			SoundStreaming->bLooping = false;
diff --git a/Engine/Plugins/Online/OnlineSubsystemUtils/Source/OnlineSubsystemUtils/Private/PartyBeaconClient.cpp b/Engine/Plugins/Online/OnlineSubsystemUtils/Source/OnlineSubsystemUtils/Private/PartyBeaconClient.cpp
--- a/Engine/Plugins/Online/OnlineSubsystemUtils/Source/OnlineSubsystemUtils/Private/PartyBeaconClient.cpp
+++ b/Engine/Plugins/Online/OnlineSubsystemUtils/Source/OnlineSubsystemUtils/Private/PartyBeaconClient.cpp
@@ -146,7 +146,7 @@ bool APartyBeaconClient::RequestReservation(const FOnlineSessionSearchResult& De
 			if (SessionInt.IsValid())
 			{
 				FString ConnectInfo;
-				if (SessionInt->GetResolvedConnectString(DesiredHost, BeaconPort, ConnectInfo))
+				if (SessionInt->GetResolvedConnectString(DesiredHost, NAME_BeaconPort, ConnectInfo))
 				{
 					FString SessionId = DesiredHost.Session.SessionInfo->GetSessionId().ToString();
 					return RequestReservation(ConnectInfo, SessionId, RequestingPartyLeader, PartyMembers);
diff --git a/Engine/Plugins/Online/OnlineSubsystemUtils/Source/OnlineSubsystemUtils/Private/Tests/TestSessionInterface.cpp b/Engine/Plugins/Online/OnlineSubsystemUtils/Source/OnlineSubsystemUtils/Private/Tests/TestSessionInterface.cpp
--- a/Engine/Plugins/Online/OnlineSubsystemUtils/Source/OnlineSubsystemUtils/Private/Tests/TestSessionInterface.cpp
+++ b/Engine/Plugins/Online/OnlineSubsystemUtils/Source/OnlineSubsystemUtils/Private/Tests/TestSessionInterface.cpp
@@ -180,7 +180,7 @@ void FTestSessionInterface::Test(UWorld* InWorld, bool bTestLAN, bool bIsPresenc
 		TArray<TSharedRef<const FUniqueNetId>> LocalPlayers;
 		LocalPlayers.Add(UserId.ToSharedRef());
 
-		SessionInt->StartMatchmaking(LocalPlayers, GameSessionName, *HostSettings, SearchSettingsRef);
+		SessionInt->StartMatchmaking(LocalPlayers, NAME_GameSession, *HostSettings, SearchSettingsRef);
 	}
 	// Setup sessions
 	else if (bIsHost)
@@ -189,7 +189,7 @@ void FTestSessionInterface::Test(UWorld* InWorld, bool bTestLAN, bool bIsPresenc
 		HostSettings->AddWorldSettings(InWorld);
 
 		OnCreateSessionCompleteDelegateHandle = SessionInt->AddOnCreateSessionCompleteDelegate_Handle(OnCreateSessionCompleteDelegate);
-		SessionInt->CreateSession(0, GameSessionName, *HostSettings);
+		SessionInt->CreateSession(0, NAME_GameSession, *HostSettings);
 	}
 	else
 	{
@@ -235,7 +235,7 @@ void FTestSessionInterface::OnSessionUserInviteAccepted(bool bWasSuccessful, int
 
 	if (bWasSuccessful)
 	{
-		JoinSession(ControllerId, GameSessionName, SearchResult);
+		JoinSession(ControllerId, NAME_GameSession, SearchResult);
 	}
 }
 
@@ -364,7 +364,7 @@ void FTestSessionInterface::OnFindFriendSessionComplete(int32 LocalUserNum, bool
 		// Can't just use SearchResult.IsValid() here - it's possible the SessionInfo pointer is valid, but not the data until we actually join the session
 		if (SearchResult.Num() > 0 && SearchResult[0].Session.OwningUserId.IsValid() && SearchResult[0].Session.SessionInfo.IsValid())
 		{
-			JoinSession(LocalUserNum, GameSessionName, SearchResult[0]);
+			JoinSession(LocalUserNum, NAME_GameSession, SearchResult[0]);
 		}
 		else
 		{
@@ -452,7 +452,7 @@ bool FTestSessionInterface::Exec(UWorld* InWorld, const TCHAR* Cmd, FOutputDevic
 	{
 		// Get optional session name "GAME" otherwise
 		FString Token;
-		FName SessionName = FParse::Value(Cmd, TEXT("Name="), Token) ? FName(*Token) : GameSessionName;
+		FName SessionName = FParse::Value(Cmd, TEXT("Name="), Token) ? FName(*Token) : NAME_GameSession;
 		if (Token.Len() > 0)
 		{
 			Cmd += FCString::Strlen(TEXT("Name=")) + Token.Len();
diff --git a/Engine/Plugins/Online/OnlineSubsystemUtils/Source/OnlineSubsystemUtils/Private/Tests/TestVoice.cpp b/Engine/Plugins/Online/OnlineSubsystemUtils/Source/OnlineSubsystemUtils/Private/Tests/TestVoice.cpp
--- a/Engine/Plugins/Online/OnlineSubsystemUtils/Source/OnlineSubsystemUtils/Private/Tests/TestVoice.cpp
+++ b/Engine/Plugins/Online/OnlineSubsystemUtils/Source/OnlineSubsystemUtils/Private/Tests/TestVoice.cpp
@@ -11,26 +11,38 @@
 
 #if WITH_DEV_AUTOMATION_TESTS
 
+#if !UE_BUILD_SHIPPING
+#define VOICE_BUFFER_CHECK(Buffer, Size) \
+	check(Buffer.Num() >= (int32)(Size))
+#else
+#define VOICE_BUFFER_CHECK(Buffer, Size)
+#endif
+
 FTestVoice::FTestVoice() :
-	VoiceComp(NULL),
-	VoiceCapture(NULL),
-	VoiceEncoder(NULL),
-	VoiceDecoder(NULL),
+	VoiceComp(nullptr),
+	SoundStreaming(nullptr),
+	VoiceCapture(nullptr),
+	VoiceEncoder(nullptr),
+	VoiceDecoder(nullptr),
+	EncodeHint(EAudioEncodeHint::VoiceEncode_Voice),
+	InputSampleRate(DEFAULT_VOICE_SAMPLE_RATE),
+	OutputSampleRate(DEFAULT_VOICE_SAMPLE_RATE),
+	NumInChannels(DEFAULT_NUM_VOICE_CHANNELS),
+	NumOutChannels(DEFAULT_NUM_VOICE_CHANNELS),
+	bLastWasPlaying(false),
+	StarvedDataCount(0),
+	MaxRawCaptureDataSize(0),
+	MaxCompressedDataSize(0),
+	MaxUncompressedDataSize(0),
+	CurrentUncompressedDataQueueSize(0),
+	MaxUncompressedDataQueueSize(0),
+	MaxRemainderSize(0),
+	LastRemainderSize(0),
 	bUseTestSample(false),
 	bZeroInput(false),
 	bUseDecompressed(true),
 	bZeroOutput(false)
 {
-	MaxRawCaptureDataSize = 100 * 1024;
-	MaxCompressedDataSize = 50 * 1024;
-	MaxUncompressedDataSize = MaxRawCaptureDataSize;
-	MaxRemainderSize = 1 * 1024;
-	LastRemainderSize = 0;
-
-	RawCaptureData.AddUninitialized(MaxRawCaptureDataSize);
-	CompressedData.AddUninitialized(MaxCompressedDataSize);
-	UncompressedData.AddUninitialized(MaxUncompressedDataSize);
-	Remainder.AddUninitialized(MaxRemainderSize);
 }
 
 FTestVoice::~FTestVoice()
@@ -45,39 +57,126 @@ void FTestVoice::Test()
 
 bool FTestVoice::Init()
 {
-	VoiceCapture = FVoiceModule::Get().CreateVoiceCapture();
-	VoiceEncoder = FVoiceModule::Get().CreateVoiceEncoder();
-	VoiceDecoder = FVoiceModule::Get().CreateVoiceDecoder();
+	DeviceName = TEXT("Line 1 (Virtual Audio Cable)");
+	EncodeHint = EAudioEncodeHint::VoiceEncode_Audio;
+	InputSampleRate = 48000;
+	OutputSampleRate = 48000;
+	NumInChannels = 2;
+	NumOutChannels = 2;
+	
+	InitVoiceCapture();
+	InitVoiceEncoder();
+	InitVoiceDecoder();
 	
+	return true;
+}
+
+void FTestVoice::InitVoiceCapture()
+{
+	ensure(!VoiceCapture.IsValid());
+	VoiceCapture = FVoiceModule::Get().CreateVoiceCapture(DeviceName, InputSampleRate, NumInChannels);
 	if (VoiceCapture.IsValid())
 	{
+		MaxRawCaptureDataSize = VoiceCapture->GetBufferSize();
+
+		RawCaptureData.Empty(MaxRawCaptureDataSize);
+		RawCaptureData.AddUninitialized(MaxRawCaptureDataSize);
+
 		VoiceCapture->Start();
 	}
-	
-	return true;
+}
+
+void FTestVoice::InitVoiceEncoder()
+{
+	ensure(!VoiceEncoder.IsValid());
+	VoiceEncoder = FVoiceModule::Get().CreateVoiceEncoder(InputSampleRate, NumInChannels, EncodeHint);
+	if (VoiceEncoder.IsValid())
+	{
+		MaxRemainderSize = VOICE_STARTING_REMAINDER_SIZE;
+		LastRemainderSize = 0;
+		MaxCompressedDataSize = VOICE_MAX_COMPRESSED_BUFFER;
+
+		CompressedData.Empty(MaxCompressedDataSize);
+		CompressedData.AddUninitialized(MaxCompressedDataSize);
+
+		Remainder.Empty(MaxRemainderSize);
+		Remainder.AddUninitialized(MaxRemainderSize);
+	}
+}
+
+void FTestVoice::InitVoiceDecoder()
+{
+	ensure(!VoiceDecoder.IsValid());
+	if (VoiceCapture.IsValid())
+	{
+		VoiceDecoder = FVoiceModule::Get().CreateVoiceDecoder(OutputSampleRate, NumOutChannels);
+		if (VoiceDecoder.IsValid())
+		{
+			// Approx 1 sec worth of data
+			MaxUncompressedDataSize = NumOutChannels * OutputSampleRate * sizeof(uint16);
+
+			UncompressedData.Empty(MaxUncompressedDataSize);
+			UncompressedData.AddUninitialized(MaxUncompressedDataSize);
+
+			MaxUncompressedDataQueueSize = MaxUncompressedDataSize * 5;
+			{
+				FScopeLock ScopeLock(&QueueLock);
+				UncompressedDataQueue.Empty(MaxUncompressedDataQueueSize);
+			}
+		}
+	}
 }
 
 void FTestVoice::Shutdown()
 {
+	RawCaptureData.Empty();
 	CompressedData.Empty();
 	UncompressedData.Empty();
+	Remainder.Empty();
 
+	{
+		FScopeLock ScopeLock(&QueueLock);
+		UncompressedDataQueue.Empty();
+	}
+
+	CleanupVoice();
+	CleanupAudioComponent();
+}
+
+void FTestVoice::CleanupVoice()
+{
 	if (VoiceCapture.IsValid())
 	{
 		VoiceCapture->Shutdown();
-		VoiceCapture = NULL;
+		VoiceCapture = nullptr;
 	}
 
-	VoiceEncoder = NULL;
-	VoiceDecoder = NULL;
+	VoiceEncoder = nullptr;
+	VoiceDecoder = nullptr;
+}
 
-	if (VoiceComp != NULL)
+void FTestVoice::CleanupAudioComponent()
+{
+	if (VoiceComp != nullptr)
 	{
+		VoiceComp->Stop();
+
+		SoundStreaming->OnSoundWaveProceduralUnderflow.Unbind();
+		SoundStreaming = nullptr;
+
 		VoiceComp->RemoveFromRoot();
-		VoiceComp = NULL;
+		VoiceComp = nullptr;
+
+		bLastWasPlaying = false;
 	}
 }
 
+void FTestVoice::CleanupQueue()
+{
+	FScopeLock ScopeLock(&QueueLock);
+	UncompressedDataQueue.Reset();
+}
+
 void FTestVoice::SetStaticVoiceData(TArray<uint8>& VoiceData, uint32& TotalVoiceBytes)
 {
 	static bool bTimeToQueue = true;
@@ -96,6 +195,7 @@ void FTestVoice::SetStaticVoiceData(TArray<uint8>& VoiceData, uint32& TotalVoice
 		VoiceData.Empty(TotalVoiceBytes);
 		VoiceData.AddUninitialized(TotalVoiceBytes);
 
+		VOICE_BUFFER_CHECK(VoiceData, ARRAY_COUNT(RawVoiceTestData));
 		FMemory::Memcpy(VoiceData.GetData(), RawVoiceTestData, ARRAY_COUNT(RawVoiceTestData));
 
 		LastQueueTime = CurrentTime;
@@ -103,27 +203,54 @@ void FTestVoice::SetStaticVoiceData(TArray<uint8>& VoiceData, uint32& TotalVoice
 	}
 }
 
+void FTestVoice::GenerateData(USoundWaveProcedural* InProceduralWave, int32 SamplesRequired)
+{
+	const int32 SampleSize = sizeof(uint16) * NumOutChannels;
+
+	{
+		FScopeLock ScopeLock(&QueueLock);
+		CurrentUncompressedDataQueueSize = UncompressedDataQueue.Num();
+		const int32 AvailableSamples = CurrentUncompressedDataQueueSize / SampleSize;
+		if (AvailableSamples >= SamplesRequired)
+		{
+			InProceduralWave->QueueAudio(UncompressedDataQueue.GetData(), AvailableSamples * SampleSize);
+			UncompressedDataQueue.RemoveAt(0, AvailableSamples * SampleSize, false);
+			CurrentUncompressedDataQueueSize -= (AvailableSamples * SampleSize);
+		}
+	}
+}
+
 bool FTestVoice::Tick(float DeltaTime)
 {
 	if (VoiceCapture.IsValid())
 	{
-		if (!IsRunningDedicatedServer() && VoiceComp == NULL)
+		if (!IsRunningDedicatedServer() && VoiceComp == nullptr)
 		{
-			VoiceComp = CreateVoiceAudioComponent(VOICE_SAMPLE_RATE);
+			VoiceComp = CreateVoiceAudioComponent(OutputSampleRate, NumOutChannels);
 			VoiceComp->AddToRoot();
-			VoiceComp->Play();
+			SoundStreaming = CastChecked<USoundWaveProcedural>(VoiceComp->Sound);
+			if (SoundStreaming)
+			{
+				// Bind the GenerateData callback with FOnSoundWaveProceduralUnderflow object
+				SoundStreaming->OnSoundWaveProceduralUnderflow = FOnSoundWaveProceduralUnderflow::CreateRaw(this, &FTestVoice::GenerateData);
+			}
 		}
 
 		if (VoiceComp)
 		{
-			static bool bLastHasData = true;
-			USoundWaveProcedural* SoundStreaming = CastChecked<USoundWaveProcedural>(VoiceComp->Sound);
-			bool bHasData = SoundStreaming->GetAvailableAudioByteCount() != 0;
+			check(SoundStreaming);
+
+			bool bIsPlaying = VoiceComp->IsPlaying();
+			if (bIsPlaying != bLastWasPlaying)
+			{
+				UE_LOG(LogVoice, Log, TEXT("VOIP audio component %s playing!"), bIsPlaying ? TEXT("is") : TEXT("is not"));
+				bLastWasPlaying = bIsPlaying;
+			}
 
-			if (bHasData != bLastHasData)
+			StarvedDataCount = (!bIsPlaying || (SoundStreaming->GetAvailableAudioByteCount() != 0)) ? 0 : (StarvedDataCount + 1);
+			if (StarvedDataCount > 1)
 			{
-				//UE_LOG(LogVoice, Log, TEXT("VOIP audio component %s starved!"), bHasData ? TEXT("is not") : TEXT("is"));
-				bLastHasData = bHasData;
+				UE_LOG(LogVoice, Log, TEXT("VOIP audio component starved %d frames!"), StarvedDataCount);
 			}
 
 			bool bDoWork = false;
@@ -141,18 +268,19 @@ bool FTestVoice::Tick(float DeltaTime)
 				if (MicState == EVoiceCaptureState::Ok && NewVoiceDataBytes > 0)
 				{
 					//UE_LOG(LogVoice, Log, TEXT("Getting data! %d"), NewVoiceDataBytes);
-					TotalVoiceBytes = NewVoiceDataBytes + LastRemainderSize;
-					RawCaptureData.Empty(MaxRawCaptureDataSize);
-					RawCaptureData.AddUninitialized(TotalVoiceBytes);
-
 					if (LastRemainderSize > 0)
 					{
+						// Add back any data from the previous frame
+						VOICE_BUFFER_CHECK(RawCaptureData, LastRemainderSize);
 						FMemory::Memcpy(RawCaptureData.GetData(), Remainder.GetData(), LastRemainderSize);
 					}
 
+					// Add new data at the beginning of the last frame
 					MicState = VoiceCapture->GetVoiceData(RawCaptureData.GetData() + LastRemainderSize, NewVoiceDataBytes, NewVoiceDataBytes);
 					TotalVoiceBytes = NewVoiceDataBytes + LastRemainderSize;
-					bDoWork = MicState == EVoiceCaptureState::Ok;
+
+					VOICE_BUFFER_CHECK(RawCaptureData, TotalVoiceBytes);
+					bDoWork = (MicState == EVoiceCaptureState::Ok);
 				}
 			}
 
@@ -166,41 +294,90 @@ bool FTestVoice::Tick(float DeltaTime)
 				// ZERO INPUT END
 
 				// COMPRESSION BEGIN
-				uint32 CompressedDataSize = MaxCompressedDataSize;
-				LastRemainderSize = VoiceEncoder->Encode(RawCaptureData.GetData(), TotalVoiceBytes, CompressedData.GetData(), CompressedDataSize);
-
-				if (LastRemainderSize > 0)
+				uint32 CompressedDataSize = 0;
+				if (VoiceEncoder.IsValid())
 				{
-					FMemory::Memcpy(Remainder.GetData(), RawCaptureData.GetData() + (TotalVoiceBytes - LastRemainderSize), LastRemainderSize);
+					CompressedDataSize = MaxCompressedDataSize;
+					LastRemainderSize = VoiceEncoder->Encode(RawCaptureData.GetData(), TotalVoiceBytes, CompressedData.GetData(), CompressedDataSize);
+					VOICE_BUFFER_CHECK(RawCaptureData, CompressedDataSize);
+
+					if (LastRemainderSize > 0)
+					{
+						if (LastRemainderSize > MaxRemainderSize)
+						{
+							UE_LOG(LogVoice, Verbose, TEXT("Overflow!"));
+							Remainder.AddUninitialized(LastRemainderSize - MaxRemainderSize);
+							MaxRemainderSize = Remainder.Num();
+						}
+
+						VOICE_BUFFER_CHECK(Remainder, LastRemainderSize);
+						FMemory::Memcpy(Remainder.GetData(), RawCaptureData.GetData() + (TotalVoiceBytes - LastRemainderSize), LastRemainderSize);
+					}
 				}
-				// COMPRESION END
+				// COMPRESSION END
 
-				// DECOMPRESION BEGIN
-				uint32 UncompressedDataSize = MaxUncompressedDataSize;
-				VoiceDecoder->Decode(CompressedData.GetData(), CompressedDataSize, 
-					UncompressedData.GetData(), UncompressedDataSize);
+				// DECOMPRESSION BEGIN
+				uint32 UncompressedDataSize = 0;
+				if (VoiceDecoder.IsValid() && CompressedDataSize > 0)
+				{
+					UncompressedDataSize = MaxUncompressedDataSize;
+					VoiceDecoder->Decode(CompressedData.GetData(), CompressedDataSize,
+						UncompressedData.GetData(), UncompressedDataSize);
+					VOICE_BUFFER_CHECK(UncompressedData, UncompressedDataSize);
+				}
 				// DECOMPRESSION END
 
+				const uint8* VoiceDataPtr = nullptr;
+				uint32 VoiceDataSize = 0;
+
 				if (bUseDecompressed)
 				{
 					if (UncompressedDataSize > 0)
 					{
-						//UE_LOG(LogVoice, Log, TEXT("Queueing uncompressed data! %d"), UncompressedDataSize);
+						//UE_LOG(LogVoice, Log, TEXT("Queuing uncompressed data! %d"), UncompressedDataSize);
 						if (bZeroOutput)
 						{
 							FMemory::Memzero((uint8*)UncompressedData.GetData(), UncompressedDataSize);
 						}
 
-						SoundStreaming->QueueAudio(UncompressedData.GetData(), UncompressedDataSize);
+						VoiceDataSize = UncompressedDataSize;
+						VoiceDataPtr = UncompressedData.GetData();
 					}
 				}
 				else
 				{
-					//UE_LOG(LogVoice, Log, TEXT("Queueing raw data! %d"), TotalVoiceBytes - LastRemainderSize);
-					SoundStreaming->QueueAudio(RawCaptureData.GetData(), TotalVoiceBytes - LastRemainderSize);
+					//UE_LOG(LogVoice, Log, TEXT("Queuing original data! %d"), UncompressedDataSize);
+					VoiceDataPtr = RawCaptureData.GetData();
+					VoiceDataSize = (TotalVoiceBytes - LastRemainderSize);
 				}
-			}
 
+				if (VoiceDataPtr && VoiceDataSize > 0)
+				{
+					FScopeLock ScopeLock(&QueueLock);
+
+					const int32 OldSize = UncompressedDataQueue.Num();
+					const int32 AmountToBuffer = (OldSize + (int32)VoiceDataSize);
+					if (AmountToBuffer <= MaxUncompressedDataQueueSize)
+					{
+						UncompressedDataQueue.AddUninitialized(VoiceDataSize);
+
+						VOICE_BUFFER_CHECK(UncompressedDataQueue, AmountToBuffer);
+						FMemory::Memcpy(UncompressedDataQueue.GetData() + OldSize, VoiceDataPtr, VoiceDataSize);
+						CurrentUncompressedDataQueueSize += VoiceDataSize;
+					}
+					else
+					{
+						UE_LOG(LogVoice, Warning, TEXT("UncompressedDataQueue Overflow!"));
+					}
+				}
+
+				// Wait for approx 1 sec worth of data before playing
+				if (!bIsPlaying && (CurrentUncompressedDataQueueSize > (MaxUncompressedDataSize / 2)))
+				{
+					UE_LOG(LogVoice, Log, TEXT("Playback started"));
+					VoiceComp->Play();
+				}
+			}
 		}
 	}
 
@@ -215,6 +392,205 @@ bool FTestVoice::Exec(UWorld* InWorld, const TCHAR* Cmd, FOutputDevice& Ar)
 		delete this;
 		bWasHandled = true;
 	}
+	else if (FParse::Command(&Cmd, TEXT("vcstart")))
+	{
+		if (VoiceCapture.IsValid() && !VoiceCapture->IsCapturing())
+		{
+			VoiceCapture->Start();
+		}
+
+		bWasHandled = true;
+	}
+	else if (FParse::Command(&Cmd, TEXT("vcstop")))
+	{
+		if (VoiceCapture.IsValid() && VoiceCapture->IsCapturing())
+		{
+			VoiceCapture->Stop();
+		}
+
+		bWasHandled = true;
+	}
+	else if (FParse::Command(&Cmd, TEXT("vchint")))
+	{
+		// vchint <0/1>
+		FString NewHintStr = FParse::Token(Cmd, false);
+		if (!NewHintStr.IsEmpty())
+		{
+			EAudioEncodeHint NewHint = (EAudioEncodeHint)FPlatformString::Atoi(*NewHintStr);
+			if (NewHint >= EAudioEncodeHint::VoiceEncode_Voice && NewHint <= EAudioEncodeHint::VoiceEncode_Audio)
+			{
+				if (EncodeHint != NewHint)
+				{
+					EncodeHint = NewHint;
+
+					CleanupAudioComponent();
+					CleanupQueue();
+					
+					VoiceEncoder = nullptr;
+					InitVoiceEncoder();
+
+					VoiceDecoder->Reset();
+				}
+			}
+		}
+
+		bWasHandled = true;
+	}
+	else if (FParse::Command(&Cmd, TEXT("vcdevice")))
+	{
+		// vcsample <device name>
+		FString NewDeviceName = FParse::Token(Cmd, false);
+		if (!NewDeviceName.IsEmpty())
+		{
+			bool bQuotesRemoved = false;
+			NewDeviceName = NewDeviceName.TrimQuotes(&bQuotesRemoved);
+			if (VoiceCapture.IsValid())
+			{
+				if (VoiceCapture->ChangeDevice(NewDeviceName, InputSampleRate, NumInChannels))
+				{
+					DeviceName = NewDeviceName;
+					CleanupAudioComponent();
+					CleanupQueue();
+					VoiceEncoder->Reset();
+					VoiceDecoder->Reset();
+					VoiceCapture->Start();
+				}
+				else
+				{
+					UE_LOG(LogVoice, Warning, TEXT("Failed to change device name %s"), *DeviceName);
+				}
+			}
+		}
+
+		bWasHandled = true;
+	}
+	else if (FParse::Command(&Cmd, TEXT("vcin")))
+	{
+		// vcin <rate> <channels>
+		FString SampleRateStr = FParse::Token(Cmd, false);
+		int32 NewInSampleRate = !SampleRateStr.IsEmpty() ? FPlatformString::Atoi(*SampleRateStr) : InputSampleRate;
+
+		FString NumChannelsStr = FParse::Token(Cmd, false);
+		int32 NewNumInChannels = !NumChannelsStr.IsEmpty() ? FPlatformString::Atoi(*NumChannelsStr) : NumInChannels;
+
+		if (NewInSampleRate != InputSampleRate ||
+			NewNumInChannels != NumInChannels)
+		{
+			InputSampleRate = NewInSampleRate;
+			NumInChannels = NewNumInChannels;
+
+			if (VoiceCapture.IsValid())
+			{
+				VoiceCapture->Shutdown();
+				VoiceCapture = nullptr;
+			}
+
+			InitVoiceCapture();
+
+			VoiceEncoder = nullptr;
+			InitVoiceEncoder();
+		}
+
+		bWasHandled = true;
+	}
+	else if (FParse::Command(&Cmd, TEXT("vcout")))
+	{
+		// vcout <rate> <channels>
+		FString SampleRateStr = FParse::Token(Cmd, false);
+		int32 NewOutSampleRate = !SampleRateStr.IsEmpty() ? FPlatformString::Atoi(*SampleRateStr) : OutputSampleRate;
+
+		FString NumChannelsStr = FParse::Token(Cmd, false);
+		int32 NewNumOutChannels = !NumChannelsStr.IsEmpty() ? FPlatformString::Atoi(*NumChannelsStr) : NumOutChannels;
+
+		if (NewOutSampleRate != OutputSampleRate ||
+			NewNumOutChannels != NumOutChannels)
+		{
+			OutputSampleRate = NewOutSampleRate;
+			NumOutChannels = NewNumOutChannels;
+
+			VoiceDecoder = nullptr;
+			InitVoiceDecoder();
+
+			CleanupAudioComponent();
+		}
+
+		bWasHandled = true;
+	}
+	else if (FParse::Command(&Cmd, TEXT("vcvbr")))
+	{
+		// vcvbr <true/false>
+		FString VBRStr = FParse::Token(Cmd, false);
+		int32 ShouldVBR = FPlatformString::Atoi(*VBRStr);
+		bool bVBR = ShouldVBR != 0;
+		if (VoiceEncoder.IsValid())
+		{
+			if (!VoiceEncoder->SetVBR(bVBR))
+			{
+				UE_LOG(LogVoice, Warning, TEXT("Failed to set VBR %d"), bVBR);
+			}
+		}
+
+		bWasHandled = true;
+	}
+	else if (FParse::Command(&Cmd, TEXT("vcbitrate")))
+	{
+		// vcbitrate <bitrate>
+		FString BitrateStr = FParse::Token(Cmd, false);
+		int32 NewBitrate = !BitrateStr.IsEmpty() ? FPlatformString::Atoi(*BitrateStr) : 0;
+		if (VoiceEncoder.IsValid() && NewBitrate > 0)
+		{
+			if (!VoiceEncoder->SetBitrate(NewBitrate))
+			{
+				UE_LOG(LogVoice, Warning, TEXT("Failed to set bitrate %d"), NewBitrate);
+			}
+		}
+
+		bWasHandled = true;
+	}
+	else if (FParse::Command(&Cmd, TEXT("vccomplexity")))
+	{
+		// vccomplexity <complexity>
+		FString ComplexityStr = FParse::Token(Cmd, false);
+		int32 NewComplexity = !ComplexityStr.IsEmpty() ? FPlatformString::Atoi(*ComplexityStr) : -1;
+		if (VoiceEncoder.IsValid() && NewComplexity >= 0)
+		{
+			if (!VoiceEncoder->SetComplexity(NewComplexity))
+			{
+				UE_LOG(LogVoice, Warning, TEXT("Failed to set complexity %d"), NewComplexity);
+			}
+		}
+
+		bWasHandled = true;
+	}
+	else if (FParse::Command(&Cmd, TEXT("vcdecompress")))
+	{
+		// vcdecompress <0/1>
+		FString DecompressStr = FParse::Token(Cmd, false);
+		int32 ShouldDecompress = FPlatformString::Atoi(*DecompressStr);
+		bUseDecompressed = (ShouldDecompress != 0);
+
+		bWasHandled = true;
+	}
+	else if (FParse::Command(&Cmd, TEXT("vcdump")))
+	{
+		if (VoiceCapture.IsValid())
+		{
+			VoiceCapture->DumpState();
+		}
+
+		if (VoiceEncoder.IsValid())
+		{
+			VoiceEncoder->DumpState();
+		}
+
+		if (VoiceDecoder.IsValid())
+		{
+			VoiceDecoder->DumpState();
+		}
+
+		bWasHandled = true;
+	}
+
 	return bWasHandled;
 }
 
diff --git a/Engine/Plugins/Online/OnlineSubsystemUtils/Source/OnlineSubsystemUtils/Private/Tests/TestVoice.h b/Engine/Plugins/Online/OnlineSubsystemUtils/Source/OnlineSubsystemUtils/Private/Tests/TestVoice.h
--- a/Engine/Plugins/Online/OnlineSubsystemUtils/Source/OnlineSubsystemUtils/Private/Tests/TestVoice.h
+++ b/Engine/Plugins/Online/OnlineSubsystemUtils/Source/OnlineSubsystemUtils/Private/Tests/TestVoice.h
@@ -6,12 +6,18 @@
 #include "Misc/CoreMisc.h"
 #include "Containers/Ticker.h"
 
+class UAudioComponent;
+class USoundWaveProcedural;
 class IVoiceCapture;
-class IVoiceDecoder;
 class IVoiceEncoder;
+class IVoiceDecoder;
+enum class EAudioEncodeHint : uint8;
 
 #if WITH_DEV_AUTOMATION_TESTS
 
+#define VOICE_MAX_COMPRESSED_BUFFER 20 * 1024
+#define VOICE_STARTING_REMAINDER_SIZE 1 * 1024
+
 class FTestVoice : public FTickerObjectBase, public FSelfRegisteringExec
 {
 public:
@@ -29,29 +35,105 @@ public:
 
 private:
 
-	class UAudioComponent* VoiceComp;
-	TSharedPtr<class IVoiceCapture> VoiceCapture;
-	TSharedPtr<class IVoiceEncoder> VoiceEncoder;
-	TSharedPtr<class IVoiceDecoder> VoiceDecoder;
-
+	/** Reference to audio component playing active streaming sound */
+	UAudioComponent* VoiceComp;
+	/** Internal reference to audio component's streaming wave object */
+	USoundWaveProcedural* SoundStreaming;
+
+	/** Active audio capture */
+	TSharedPtr<IVoiceCapture> VoiceCapture;
+	/** Active audio encoder */
+	TSharedPtr<IVoiceEncoder> VoiceEncoder;
+	/** Active audio decoder */
+	TSharedPtr<IVoiceDecoder> VoiceDecoder;
+
+	/** Name of current device under capture */
+	FString DeviceName;
+	/** Current type of audio under capture */
+	EAudioEncodeHint EncodeHint;
+	/** Current input sample rate provided */
+	int32 InputSampleRate;
+	/** Desired output sample rate */
+	int32 OutputSampleRate;
+	/** Current number of channels captured */
+	int32 NumInChannels;
+	/** Desired number of output channels */
+	int32 NumOutChannels;
+
+	/** Was the audio component playing last frame */
+	bool bLastWasPlaying;
+	/** Number of consecutive frames that the playback has been starved */
+	int32 StarvedDataCount;
+
+	/** Buffer for pre-encoded audio data */
 	TArray<uint8> RawCaptureData;
+	/** Maximum size of a single raw capture packet */
 	int32 MaxRawCaptureDataSize;
+	/** Buffer for compressed audio data */
 	TArray<uint8> CompressedData;
+	/** Maximum size of a single encoded packet */
 	int32 MaxCompressedDataSize;
+	/** Buffer for uncompressed audio data (valid during Tick only) */
 	TArray<uint8> UncompressedData;
+	/** Maximum size of a single decoded packet */
 	int32 MaxUncompressedDataSize;
+
+	/** Buffer for outgoing audio intended for procedural streaming */
+	mutable FCriticalSection QueueLock;
+	TArray<uint8> UncompressedDataQueue;
+	/** Amount of data currently in the outgoing playback queue */
+	int32 CurrentUncompressedDataQueueSize;
+	/** Maximum size of the outgoing playback queue */
+	int32 MaxUncompressedDataQueueSize;
+
+	/** Buffer for data left uncompressed after call to Encode to be used next encoding */
 	TArray<uint8> Remainder;
+	/** Maximum size of the remainder buffer */
 	int32 MaxRemainderSize;
+	/** Current amount of raw data left over from last encode */
 	int32 LastRemainderSize;
 
+	/** Use hard coded audio test data in TestVoiceData.h */
 	bool bUseTestSample;
+	/** Zero out input before encoding */
 	bool bZeroInput;
+	/** Pass originating audio capture data directly to the audio component (skip Encode/Decode) */
 	bool bUseDecompressed;
+	/** Zero out output data before playback */
 	bool bZeroOutput;
 
+	/** First time initialization of all components necessary for end to end capture -> encode -> decode -> playback */
 	bool Init();
+	/** (Re)Initialize the audio capture object with current settings, reallocating buffers */
+	void InitVoiceCapture();
+	/** (Re)Initialize the audio encoder with current settings, reallocating buffers */
+	void InitVoiceEncoder();
+	/** (Re)Initialize the audio decoder with current settings, reallocating buffers */
+	void InitVoiceDecoder();
+	/** Cleanup and shutdown the entire object */
 	void Shutdown();
 
+	/** Free all audio objects (capture/encode/decode) */
+	void CleanupVoice();
+	/** Free audio component */
+	void CleanupAudioComponent();
+	/** Empty and reset the outgoing audio data queue */
+	void CleanupQueue();
+
+	/**
+	 * Callback from streaming audio when data is requested for playback
+	 *
+	 * @param InProceduralWave SoundWave requesting more data
+	 * @param SamplesRequired number of samples needed for immediate playback
+	 */
+	void GenerateData(USoundWaveProcedural* InProceduralWave, int32 SamplesRequired);
+
+	/**
+	 * Play static audio from TestVoiceData.h in a loop
+	 *
+	 * @param VoiceData out buffer to fill with data
+	 * @param TotalVoiceBytes amount of data added to the buffer
+	 */
 	void SetStaticVoiceData(TArray<uint8>& VoiceData, uint32& TotalVoiceBytes);
 };
 
diff --git a/Engine/Plugins/Online/OnlineSubsystemUtils/Source/OnlineSubsystemUtils/Private/Tests/TestVoiceData.h b/Engine/Plugins/Online/OnlineSubsystemUtils/Source/OnlineSubsystemUtils/Private/Tests/TestVoiceData.h
--- a/Engine/Plugins/Online/OnlineSubsystemUtils/Source/OnlineSubsystemUtils/Private/Tests/TestVoiceData.h
+++ b/Engine/Plugins/Online/OnlineSubsystemUtils/Source/OnlineSubsystemUtils/Private/Tests/TestVoiceData.h
@@ -2010,4 +2010,4 @@ uint8 RawVoiceTestData[] = {
 	0x99, 0xff, 0x9b, 0xff, 0x99, 0xff, 0xa3, 0xff, 0xa9, 0xff, 0xa2, 0xff, 0x97, 0xff, 0x98, 0xff, 0xa5, 0xff, 0x96, 0xff, 0x8f, 0xff, 0xae, 0xff, 0xaf, 0xff, 0xa3, 0xff, 0x9d, 0xff, 0xa3, 0xff
 };
 
-#endif //UE_BUILD_SHIPPING
+#endif // WITH_DEV_AUTOMATION_TESTS
diff --git a/Engine/Plugins/Online/OnlineSubsystemUtils/Source/OnlineSubsystemUtils/Private/VoiceEngineImpl.cpp b/Engine/Plugins/Online/OnlineSubsystemUtils/Source/OnlineSubsystemUtils/Private/VoiceEngineImpl.cpp
--- a/Engine/Plugins/Online/OnlineSubsystemUtils/Source/OnlineSubsystemUtils/Private/VoiceEngineImpl.cpp
+++ b/Engine/Plugins/Online/OnlineSubsystemUtils/Source/OnlineSubsystemUtils/Private/VoiceEngineImpl.cpp
@@ -16,17 +16,102 @@
 #define MAX_VOICE_REMAINDER_SIZE 1 * 1024
 
 FRemoteTalkerDataImpl::FRemoteTalkerDataImpl() :
+	MaxUncompressedDataSize(0),
+	MaxUncompressedDataQueueSize(0),
+	CurrentUncompressedDataQueueSize(0),
 	LastSeen(0.0),
+	NumFramesStarved(0),
 	AudioComponent(nullptr),
 	VoiceDecoder(nullptr)
 {
-	VoiceDecoder = FVoiceModule::Get().CreateVoiceDecoder();
+	int32 SampleRate = DEFAULT_VOICE_SAMPLE_RATE;
+	int32 NumChannels = DEFAULT_NUM_VOICE_CHANNELS;
+	VoiceDecoder = FVoiceModule::Get().CreateVoiceDecoder(SampleRate, NumChannels);
 	check(VoiceDecoder.IsValid());
+
+	// Approx 1 sec worth of data
+	MaxUncompressedDataSize = NumChannels * SampleRate * sizeof(uint16);
+	MaxUncompressedDataQueueSize = MaxUncompressedDataSize * 5;
+	{
+		FScopeLock ScopeLock(&QueueLock);
+		UncompressedDataQueue.Empty(MaxUncompressedDataQueueSize);
+	}
+}
+
+FRemoteTalkerDataImpl::FRemoteTalkerDataImpl(FRemoteTalkerDataImpl&& Other)
+{
+	LastSeen = Other.LastSeen;
+	Other.LastSeen = 0.0;
+
+	NumFramesStarved = Other.NumFramesStarved;
+	Other.NumFramesStarved = 0;
+
+	AudioComponent = Other.AudioComponent;
+	Other.AudioComponent = nullptr;
+
+	VoiceDecoder = MoveTemp(Other.VoiceDecoder);
+	Other.VoiceDecoder = nullptr;
+
+	MaxUncompressedDataSize = Other.MaxUncompressedDataSize;
+	Other.MaxUncompressedDataSize = 0;
+
+	MaxUncompressedDataQueueSize = Other.MaxUncompressedDataQueueSize;
+	Other.MaxUncompressedDataQueueSize = 0;
+
+	CurrentUncompressedDataQueueSize = Other.CurrentUncompressedDataQueueSize;
+	Other.CurrentUncompressedDataQueueSize = 0;
+
+	{
+		FScopeLock ScopeLock(&Other.QueueLock);
+		UncompressedDataQueue = MoveTemp(Other.UncompressedDataQueue);
+	}
 }
 
 FRemoteTalkerDataImpl::~FRemoteTalkerDataImpl()
 {
 	VoiceDecoder = nullptr;
+
+	CurrentUncompressedDataQueueSize = 0;
+
+	{
+		FScopeLock ScopeLock(&QueueLock);
+		UncompressedDataQueue.Empty();
+	}
+}
+
+void FRemoteTalkerDataImpl::Reset()
+{
+	// Set to large number so TickTalkers doesn't come in here
+	LastSeen = MAX_FLT;
+	NumFramesStarved = 0;
+
+	if (AudioComponent && !AudioComponent->IsPendingKill())
+	{
+		AudioComponent->Stop();
+	}
+
+	CurrentUncompressedDataQueueSize = 0;
+
+	{
+		FScopeLock ScopeLock(&QueueLock);
+		UncompressedDataQueue.Empty();
+	}
+}
+
+void FRemoteTalkerDataImpl::Cleanup()
+{
+	if (AudioComponent && !AudioComponent->IsPendingKill())
+	{
+		AudioComponent->Stop();
+		USoundWaveProcedural* SoundStreaming = CastChecked<USoundWaveProcedural>(AudioComponent->Sound);
+		if (SoundStreaming)
+		{
+			SoundStreaming->OnSoundWaveProceduralUnderflow.Unbind();
+			AudioComponent->Sound = nullptr;
+		}
+	}
+
+	AudioComponent = nullptr;
 }
 
 FVoiceEngineImpl::FVoiceEngineImpl(IOnlineSubsystem* InSubsystem) :
@@ -66,7 +151,7 @@ void FVoiceEngineImpl::VoiceCaptureUpdate() const
 		// If no data is available, we have finished capture the last (post-StopRecording) half-second of voice data
 		if (RecordingState == EVoiceCaptureState::NotCapturing)
 		{
-			UE_LOG(LogVoiceEncode, Log, TEXT("Internal voice capture complete."));
+			UE_LOG(LogVoiceEngine, Log, TEXT("Internal voice capture complete."));
 
 			bPendingFinalCapture = false;
 
@@ -86,19 +171,19 @@ void FVoiceEngineImpl::VoiceCaptureUpdate() const
 
 void FVoiceEngineImpl::StartRecording() const
 {
-	UE_LOG(LogVoiceEncode, VeryVerbose, TEXT("VOIP StartRecording"));
+	UE_LOG(LogVoiceEngine, VeryVerbose, TEXT("VOIP StartRecording"));
 	if (VoiceCapture.IsValid())
 	{
 		if (!VoiceCapture->Start())
 		{
-			UE_LOG(LogVoiceEncode, Warning, TEXT("Failed to start voice recording"));
+			UE_LOG(LogVoiceEngine, Warning, TEXT("Failed to start voice recording"));
 		}
 	}
 }
 
 void FVoiceEngineImpl::StopRecording() const
 {
-	UE_LOG(LogVoiceEncode, VeryVerbose, TEXT("VOIP StopRecording"));
+	UE_LOG(LogVoiceEngine, VeryVerbose, TEXT("VOIP StopRecording"));
 	if (VoiceCapture.IsValid())
 	{
 		VoiceCapture->Stop();
@@ -107,7 +192,7 @@ void FVoiceEngineImpl::StopRecording() const
 
 void FVoiceEngineImpl::StoppedRecording() const
 {
-	UE_LOG(LogVoiceEncode, VeryVerbose, TEXT("VOIP StoppedRecording"));
+	UE_LOG(LogVoiceEngine, VeryVerbose, TEXT("VOIP StoppedRecording"));
 }
 
 bool FVoiceEngineImpl::Init(int32 MaxLocalTalkers, int32 MaxRemoteTalkers)
@@ -170,7 +255,7 @@ uint32 FVoiceEngineImpl::StartLocalVoiceProcessing(uint32 LocalUserNum)
 	}
 	else
 	{
-		UE_LOG(LogVoiceEncode, Error, TEXT("StartLocalVoiceProcessing(): Device is currently owned by another user"));
+		UE_LOG(LogVoiceEngine, Error, TEXT("StartLocalVoiceProcessing(): Device is currently owned by another user"));
 	}
 
 	return Return;
@@ -197,7 +282,7 @@ uint32 FVoiceEngineImpl::StopLocalVoiceProcessing(uint32 LocalUserNum)
 	}
 	else
 	{
-		UE_LOG(LogVoiceEncode, Error, TEXT("StopLocalVoiceProcessing: Ignoring stop request for non-owning user"));
+		UE_LOG(LogVoiceEngine, Error, TEXT("StopLocalVoiceProcessing: Ignoring stop request for non-owning user"));
 	}
 
 	return Return;
@@ -210,12 +295,7 @@ uint32 FVoiceEngineImpl::UnregisterRemoteTalker(const FUniqueNetId& UniqueId)
 	if (RemoteData != nullptr)
 	{
 		// Dump the whole talker
-		if (RemoteData->AudioComponent)
-		{
-			RemoteData->AudioComponent->Stop();
-			RemoteData->AudioComponent = nullptr;
-		}
-		
+		RemoteData->Cleanup();
 		RemoteTalkerBuffers.Remove(RemoteTalkerId);
 	}
 
@@ -255,22 +335,22 @@ uint32 FVoiceEngineImpl::ReadLocalVoiceData(uint32 LocalUserNum, uint8* Data, ui
 		EVoiceCaptureState::Type VoiceResult = VoiceCapture->GetCaptureState(NewVoiceDataBytes);
 		if (VoiceResult != EVoiceCaptureState::Ok && VoiceResult != EVoiceCaptureState::NoData)
 		{
-			UE_LOG(LogVoiceEncode, Warning, TEXT("ReadLocalVoiceData: GetAvailableVoice failure: VoiceResult: %s"), EVoiceCaptureState::ToString(VoiceResult));
+			UE_LOG(LogVoiceEngine, Warning, TEXT("ReadLocalVoiceData: GetAvailableVoice failure: VoiceResult: %s"), EVoiceCaptureState::ToString(VoiceResult));
 			return E_FAIL;
 		}
 
- 		if (NewVoiceDataBytes == 0)
- 		{
- 			UE_LOG(LogVoiceEncode, VeryVerbose, TEXT("ReadLocalVoiceData: No Data: VoiceResult: %s"), EVoiceCaptureState::ToString(VoiceResult));
- 			*Size = 0;
- 			return S_OK;
- 		}
+		if (NewVoiceDataBytes == 0)
+		{
+			UE_LOG(LogVoiceEngine, VeryVerbose, TEXT("ReadLocalVoiceData: No Data: VoiceResult: %s"), EVoiceCaptureState::ToString(VoiceResult));
+			*Size = 0;
+			return S_OK;
+		}
 
 		// Make space for new and any previously remaining data
 		uint32 TotalVoiceBytes = NewVoiceDataBytes + PlayerVoiceData[LocalUserNum].VoiceRemainderSize;
 		if (TotalVoiceBytes > MAX_UNCOMPRESSED_VOICE_BUFFER_SIZE)
 		{
-			UE_LOG(LogVoiceEncode, Warning, TEXT("Exceeded uncompressed voice buffer size, clamping"))
+			UE_LOG(LogVoiceEngine, Warning, TEXT("Exceeded uncompressed voice buffer size, clamping"))
 			TotalVoiceBytes = MAX_UNCOMPRESSED_VOICE_BUFFER_SIZE;
 		}
 
@@ -300,7 +380,7 @@ uint32 FVoiceEngineImpl::ReadLocalVoiceData(uint32 LocalUserNum, uint8* Data, ui
 			{
 				if (PlayerVoiceData[LocalUserNum].VoiceRemainderSize > MAX_VOICE_REMAINDER_SIZE)
 				{
-					UE_LOG(LogVoiceEncode, Warning, TEXT("Exceeded voice remainder buffer size, clamping"))
+					UE_LOG(LogVoiceEngine, Warning, TEXT("Exceeded voice remainder buffer size, clamping"))
 						PlayerVoiceData[LocalUserNum].VoiceRemainderSize = MAX_VOICE_REMAINDER_SIZE;
 				}
 
@@ -310,16 +390,16 @@ uint32 FVoiceEngineImpl::ReadLocalVoiceData(uint32 LocalUserNum, uint8* Data, ui
 
 			static double LastGetVoiceCallTime = 0.0;
 			double CurTime = FPlatformTime::Seconds();
-			double TimeSinceLastCall = LastGetVoiceCallTime > 0 ? (CurTime - LastGetVoiceCallTime) : 0.0;
+			double TimeSinceLastCall = (LastGetVoiceCallTime > 0) ? (CurTime - LastGetVoiceCallTime) : 0.0;
 			LastGetVoiceCallTime = CurTime;
 
-			UE_LOG(LogVoiceEncode, Log, TEXT("ReadLocalVoiceData: GetVoice: Result: %s, Available: %i, LastCall: %0.3f"), EVoiceCaptureState::ToString(VoiceResult), CompressedBytesAvailable, TimeSinceLastCall);
+			UE_LOG(LogVoiceEngine, Log, TEXT("ReadLocalVoiceData: GetVoice: Result: %s, Available: %i, LastCall: %0.3f ms"), EVoiceCaptureState::ToString(VoiceResult), CompressedBytesAvailable, TimeSinceLastCall * 1000.0);
 			if (CompressedBytesAvailable > 0)
 			{
 				*Size = FMath::Min<int32>(*Size, CompressedBytesAvailable);
 				FMemory::Memcpy(Data, CompressedVoiceBuffer.GetData(), *Size);
 
-				UE_LOG(LogVoiceEncode, VeryVerbose, TEXT("ReadLocalVoiceData: Size: %d"), *Size);
+				UE_LOG(LogVoiceEngine, VeryVerbose, TEXT("ReadLocalVoiceData: Size: %d"), *Size);
 				return S_OK;
 			}
 			else
@@ -327,7 +407,7 @@ uint32 FVoiceEngineImpl::ReadLocalVoiceData(uint32 LocalUserNum, uint8* Data, ui
 				*Size = 0;
 				CompressedVoiceBuffer.Empty(MAX_COMPRESSED_VOICE_BUFFER_SIZE);
 
-				UE_LOG(LogVoiceEncode, Warning, TEXT("ReadLocalVoiceData: GetVoice failure: VoiceResult: %s"), EVoiceCaptureState::ToString(VoiceResult));
+				UE_LOG(LogVoiceEngine, Warning, TEXT("ReadLocalVoiceData: GetVoice failure: VoiceResult: %s"), EVoiceCaptureState::ToString(VoiceResult));
 				return E_FAIL;
 			}
 		}
@@ -338,7 +418,7 @@ uint32 FVoiceEngineImpl::ReadLocalVoiceData(uint32 LocalUserNum, uint8* Data, ui
 
 uint32 FVoiceEngineImpl::SubmitRemoteVoiceData(const FUniqueNetId& RemoteTalkerId, uint8* Data, uint32* Size)
 {
-	UE_LOG(LogVoiceDecode, VeryVerbose, TEXT("SubmitRemoteVoiceData(%s) Size: %d received!"), *RemoteTalkerId.ToDebugString(), *Size);
+	UE_LOG(LogVoiceEngine, VeryVerbose, TEXT("SubmitRemoteVoiceData(%s) Size: %d received!"), *RemoteTalkerId.ToDebugString(), *Size);
 
 	const FUniqueNetIdString& TalkerId = (const FUniqueNetIdString&)RemoteTalkerId;
 	FRemoteTalkerDataImpl& QueuedData = RemoteTalkerBuffers.FindOrAdd(TalkerId);
@@ -368,31 +448,47 @@ uint32 FVoiceEngineImpl::SubmitRemoteVoiceData(const FUniqueNetId& RemoteTalkerI
 			SerializeHelper = new FVoiceSerializeHelper(this);
 		}
 
-		QueuedData.AudioComponent = CreateVoiceAudioComponent(VOICE_SAMPLE_RATE);
+		QueuedData.AudioComponent = CreateVoiceAudioComponent(DEFAULT_VOICE_SAMPLE_RATE, DEFAULT_NUM_VOICE_CHANNELS);
 		if (QueuedData.AudioComponent)
 		{
 			QueuedData.AudioComponent->OnAudioFinishedNative.AddRaw(this, &FVoiceEngineImpl::OnAudioFinished);
+			USoundWaveProcedural* SoundStreaming = CastChecked<USoundWaveProcedural>(QueuedData.AudioComponent->Sound);
+			if (SoundStreaming)
+			{
+				// Bind the GenerateVoiceData callback with FOnSoundWaveProceduralUnderflow object
+				SoundStreaming->OnSoundWaveProceduralUnderflow = FOnSoundWaveProceduralUnderflow::CreateRaw(this, &FVoiceEngineImpl::GenerateVoiceData, TalkerId);
+			}
 		}
 	}
 	
 	if (QueuedData.AudioComponent != nullptr)
 	{
-		if (!QueuedData.AudioComponent->IsActive())
 		{
-			QueuedData.AudioComponent->Play();
-		}
+			FScopeLock ScopeLock(&QueuedData.QueueLock);
 
-		USoundWaveProcedural* SoundStreaming = CastChecked<USoundWaveProcedural>(QueuedData.AudioComponent->Sound);
-
-		if (0)
-		{
-			if (!bAudioComponentCreated && SoundStreaming->GetAvailableAudioByteCount() == 0)
+			const int32 OldSize = QueuedData.UncompressedDataQueue.Num();
+			const int32 AmountToBuffer = (OldSize + (int32)BytesWritten);
+			if (AmountToBuffer <= QueuedData.MaxUncompressedDataQueueSize)
+			{
+				QueuedData.UncompressedDataQueue.AddUninitialized(BytesWritten);
+				FMemory::Memcpy(QueuedData.UncompressedDataQueue.GetData() + OldSize, DecompressedVoiceBuffer.GetData(), BytesWritten);
+				QueuedData.CurrentUncompressedDataQueueSize += BytesWritten;
+			}
+			else
 			{
-				UE_LOG(LogVoiceDecode, Log, TEXT("VOIP audio component was starved!"));
+				UE_LOG(LogVoiceEngine, Warning, TEXT("UncompressedDataQueue Overflow!"));
+				QueuedData.UncompressedDataQueue.Empty(QueuedData.MaxUncompressedDataQueueSize);
+				FMemory::Memcpy(QueuedData.UncompressedDataQueue.GetData(), DecompressedVoiceBuffer.GetData(), BytesWritten);
+				QueuedData.CurrentUncompressedDataQueueSize = BytesWritten;
 			}
 		}
 
-		SoundStreaming->QueueAudio(DecompressedVoiceBuffer.GetData(), BytesWritten);
+		// Wait for approx .5 sec worth of data before playing
+		if (!QueuedData.AudioComponent->IsActive() && (QueuedData.CurrentUncompressedDataQueueSize > (QueuedData.MaxUncompressedDataSize / 2)))
+		{
+			UE_LOG(LogVoiceEngine, Log, TEXT("Playback started"));
+			QueuedData.AudioComponent->Play();
+		}
 	}
 
 	return S_OK;
@@ -409,11 +505,33 @@ void FVoiceEngineImpl::TickTalkers(float DeltaTime)
 		if (TimeSince >= 1.0)
 		{
 			// Dump the whole talker
-			if (RemoteData.AudioComponent)
+			RemoteData.Reset();
+		}
+
+#if !UE_BUILD_SHIPPING
+		if (RemoteData.AudioComponent && RemoteData.AudioComponent->IsPlaying())
+		{
+			USoundWaveProcedural* SoundStreaming = CastChecked<USoundWaveProcedural>(RemoteData.AudioComponent->Sound);
+			if (SoundStreaming)
 			{
-				RemoteData.AudioComponent->Stop();
+				RemoteData.NumFramesStarved = (SoundStreaming->GetAvailableAudioByteCount() == 0) ? (RemoteData.NumFramesStarved + 1) : 0;
+				if (RemoteData.NumFramesStarved > 1)
+				{
+					int32 Data1, Data2;
+					{
+						FScopeLock ScopeLock(&RemoteData.QueueLock);
+						Data1 = RemoteData.UncompressedDataQueue.Num();
+						Data2 = RemoteData.CurrentUncompressedDataQueueSize;
+					}
+					UE_LOG(LogVoiceEngine, Log, TEXT("VOIP audio component starved %d frames! %d / %d"), RemoteData.NumFramesStarved, Data1, Data2);
+				}
 			}
 		}
+		else
+		{
+			RemoteData.NumFramesStarved = 0;
+		}
+#endif
 	}
 }
 
@@ -425,6 +543,33 @@ void FVoiceEngineImpl::Tick(float DeltaTime)
 	TickTalkers(DeltaTime);
 }
 
+void FVoiceEngineImpl::GenerateVoiceData(USoundWaveProcedural* InProceduralWave, int32 SamplesRequired, FUniqueNetIdString TalkerId)
+{
+	FRemoteTalkerDataImpl* QueuedData = RemoteTalkerBuffers.Find(TalkerId);
+	if (QueuedData)
+	{
+		const int32 SampleSize = sizeof(uint16) * DEFAULT_NUM_VOICE_CHANNELS;
+
+		{
+			FScopeLock ScopeLock(&QueuedData->QueueLock);
+			QueuedData->CurrentUncompressedDataQueueSize = QueuedData->UncompressedDataQueue.Num();
+			const int32 AvailableSamples = QueuedData->CurrentUncompressedDataQueueSize / SampleSize;
+			if (AvailableSamples >= SamplesRequired)
+			{
+				UE_LOG(LogVoiceEngine, Verbose, TEXT("GenerateVoiceData %d / %d"), AvailableSamples, SamplesRequired);
+				const int32 SamplesBytesTaken = AvailableSamples * SampleSize;
+				InProceduralWave->QueueAudio(QueuedData->UncompressedDataQueue.GetData(), SamplesBytesTaken);
+				QueuedData->UncompressedDataQueue.RemoveAt(0, SamplesBytesTaken, false);
+				QueuedData->CurrentUncompressedDataQueueSize -= (SamplesBytesTaken);
+			}
+			else
+			{
+				UE_LOG(LogVoiceEngine, Verbose, TEXT("Voice underflow"));
+			}
+		}
+	}
+}
+
 void FVoiceEngineImpl::OnAudioFinished(UAudioComponent* AC)
 {
 	for (FRemoteTalkerData::TIterator It(RemoteTalkerBuffers); It; ++It)
@@ -432,12 +577,12 @@ void FVoiceEngineImpl::OnAudioFinished(UAudioComponent* AC)
 		FRemoteTalkerDataImpl& RemoteData = It.Value();
 		if (RemoteData.AudioComponent->IsPendingKill() || AC == RemoteData.AudioComponent)
 		{
-			UE_LOG(LogVoiceDecode, Log, TEXT("Removing VOIP AudioComponent for Id: %s"), *It.Key().ToDebugString());
+			UE_LOG(LogVoiceEngine, Log, TEXT("Removing VOIP AudioComponent for Id: %s"), *It.Key().ToDebugString());
 			RemoteData.AudioComponent = nullptr;
 			break;
 		}
 	}
-	UE_LOG(LogVoiceDecode, Verbose, TEXT("Audio Finished"));
+	UE_LOG(LogVoiceEngine, Verbose, TEXT("Audio Finished"));
 }
 
 FString FVoiceEngineImpl::GetVoiceDebugState() const
@@ -460,4 +605,81 @@ FString FVoiceEngineImpl::GetVoiceDebugState() const
 	return Output;
 }
 
+bool FVoiceEngineImpl::Exec(UWorld* InWorld, const TCHAR* Cmd, FOutputDevice& Ar)
+{
+	bool bWasHandled = false;
+
+	if (FParse::Command(&Cmd, TEXT("vcvbr")))
+	{
+		// vcvbr <true/false>
+		FString VBRStr = FParse::Token(Cmd, false);
+		int32 ShouldVBR = FPlatformString::Atoi(*VBRStr);
+		bool bVBR = ShouldVBR != 0;
+		if (VoiceEncoder.IsValid())
+		{
+			if (!VoiceEncoder->SetVBR(bVBR))
+			{
+				UE_LOG(LogVoice, Warning, TEXT("Failed to set VBR %d"), bVBR);
+			}
+		}
+
+		bWasHandled = true;
+	}
+	else if (FParse::Command(&Cmd, TEXT("vcbitrate")))
+	{
+		// vcbitrate <bitrate>
+		FString BitrateStr = FParse::Token(Cmd, false);
+		int32 NewBitrate = !BitrateStr.IsEmpty() ? FPlatformString::Atoi(*BitrateStr) : 0;
+		if (VoiceEncoder.IsValid() && NewBitrate > 0)
+		{
+			if (!VoiceEncoder->SetBitrate(NewBitrate))
+			{
+				UE_LOG(LogVoice, Warning, TEXT("Failed to set bitrate %d"), NewBitrate);
+			}
+		}
+
+		bWasHandled = true;
+	}
+	else if (FParse::Command(&Cmd, TEXT("vccomplexity")))
+	{
+		// vccomplexity <complexity>
+		FString ComplexityStr = FParse::Token(Cmd, false);
+		int32 NewComplexity = !ComplexityStr.IsEmpty() ? FPlatformString::Atoi(*ComplexityStr) : -1;
+		if (VoiceEncoder.IsValid() && NewComplexity >= 0)
+		{
+			if (!VoiceEncoder->SetComplexity(NewComplexity))
+			{
+				UE_LOG(LogVoice, Warning, TEXT("Failed to set complexity %d"), NewComplexity);
+			}
+		}
+
+		bWasHandled = true;
+	}
+	else if (FParse::Command(&Cmd, TEXT("vcdump")))
+	{
+		if (VoiceCapture.IsValid())
+		{
+			VoiceCapture->DumpState();
+		}
+
+		if (VoiceEncoder.IsValid())
+		{
+			VoiceEncoder->DumpState();
+		}
+
+		for (FRemoteTalkerData::TIterator It(RemoteTalkerBuffers); It; ++It)
+		{
+			FRemoteTalkerDataImpl& RemoteData = It.Value();
+			if (RemoteData.VoiceDecoder.IsValid())
+			{
+				RemoteData.VoiceDecoder->DumpState();
+			}
+		}
+
+		bWasHandled = true;
+	}
+
+	return bWasHandled;
+}
+
 
diff --git a/Engine/Plugins/Online/OnlineSubsystemUtils/Source/OnlineSubsystemUtils/Private/VoiceInterfaceImpl.cpp b/Engine/Plugins/Online/OnlineSubsystemUtils/Source/OnlineSubsystemUtils/Private/VoiceInterfaceImpl.cpp
--- a/Engine/Plugins/Online/OnlineSubsystemUtils/Source/OnlineSubsystemUtils/Private/VoiceInterfaceImpl.cpp
+++ b/Engine/Plugins/Online/OnlineSubsystemUtils/Source/OnlineSubsystemUtils/Private/VoiceInterfaceImpl.cpp
@@ -762,12 +762,12 @@ void FOnlineVoiceImpl::ProcessLocalVoicePackets()
 						}
 						else
 						{
-							UE_LOG(LogVoiceEncode, Warning, TEXT("Voice data error in ReadLocalVoiceData"));
+							UE_LOG(LogVoiceEngine, Warning, TEXT("Voice data error in ReadLocalVoiceData"));
 						}
 					}
 					else
 					{
-						UE_LOG(LogVoiceEncode, Warning, TEXT("Dropping voice data due to network layer not processing fast enough"));
+						UE_LOG(LogVoiceEngine, Warning, TEXT("Dropping voice data due to network layer not processing fast enough"));
 						// Buffer overflow, so drop previous data
 						VoiceData.LocalPackets[Index].Length = 0;
 					}
@@ -800,7 +800,7 @@ void FOnlineVoiceImpl::ProcessRemoteVoicePackets()
 				uint32 Result = VoiceEngine->SubmitRemoteVoiceData(*VoicePacket->Sender, VoicePacket->Buffer.GetData(), &VoiceBufferSize);
 				if (Result != S_OK)
 				{
-					UE_LOG(LogVoiceDecode, Warning,
+					UE_LOG(LogVoiceEngine, Warning,
 						TEXT("SubmitRemoteVoiceData(%s) failed with 0x%08X"),
 						*VoicePacket->Sender->ToDebugString(),
 						Result);
diff --git a/Engine/Plugins/Online/OnlineSubsystemUtils/Source/OnlineSubsystemUtils/Public/OnlineSubsystemUtils.h b/Engine/Plugins/Online/OnlineSubsystemUtils/Source/OnlineSubsystemUtils/Public/OnlineSubsystemUtils.h
--- a/Engine/Plugins/Online/OnlineSubsystemUtils/Source/OnlineSubsystemUtils/Public/OnlineSubsystemUtils.h
+++ b/Engine/Plugins/Online/OnlineSubsystemUtils/Source/OnlineSubsystemUtils/Public/OnlineSubsystemUtils.h
@@ -14,7 +14,7 @@
 #ifdef ONLINESUBSYSTEMUTILS_API
 
 /** @return an initialized audio component specifically for use with VoIP */
-ONLINESUBSYSTEMUTILS_API class UAudioComponent* CreateVoiceAudioComponent(uint32 SampleRate);
+ONLINESUBSYSTEMUTILS_API class UAudioComponent* CreateVoiceAudioComponent(uint32 SampleRate, int32 NumChannels);
 
 /** @return the world associated with a named online subsystem instance */
 ONLINESUBSYSTEMUTILS_API UWorld* GetWorldForOnline(FName InstanceName);
diff --git a/Engine/Plugins/Online/OnlineSubsystemUtils/Source/OnlineSubsystemUtils/Public/VoiceEngineImpl.h b/Engine/Plugins/Online/OnlineSubsystemUtils/Source/OnlineSubsystemUtils/Public/VoiceEngineImpl.h
--- a/Engine/Plugins/Online/OnlineSubsystemUtils/Source/OnlineSubsystemUtils/Public/VoiceEngineImpl.h
+++ b/Engine/Plugins/Online/OnlineSubsystemUtils/Source/OnlineSubsystemUtils/Public/VoiceEngineImpl.h
@@ -3,6 +3,7 @@
 #pragma once
 
 #include "CoreMinimal.h"
+#include "CoreMisc.h"
 #include "UObject/GCObject.h"
 #include "OnlineSubsystemTypes.h"
 #include "Interfaces/VoiceInterface.h"
@@ -14,6 +15,11 @@
 
 class IOnlineSubsystem;
 class UAudioComponent;
+class USoundWaveProcedural;
+class FUniqueNetIdString;
+class IVoiceDecoder;
+class IVoiceEncoder;
+class IVoiceCapture;
 
 /**
  * Container for unprocessed voice data
@@ -39,27 +45,46 @@ class FRemoteTalkerDataImpl
 public:
 
 	FRemoteTalkerDataImpl();
+	/** Required for TMap FindOrAdd() */
+	FRemoteTalkerDataImpl(FRemoteTalkerDataImpl&& Other);
 	~FRemoteTalkerDataImpl();
 
+	/** Reset the talker after long periods of silence */
+	void Reset();
+	/** Cleanup the talker before unregistration */
+	void Cleanup();
+
+	/** Maximum size of a single decoded packet */
+	int32 MaxUncompressedDataSize;
+	/** Maximum size of the outgoing playback queue */
+	int32 MaxUncompressedDataQueueSize;
+	/** Amount of data currently in the outgoing playback queue */
+	int32 CurrentUncompressedDataQueueSize;
+
 	/** Receive side timestamp since last voice packet fragment */
 	double LastSeen;
+	/** Number of frames starved of audio */
+	int32 NumFramesStarved;
 	/** Audio component playing this buffer (only valid on remote instances) */
-	class UAudioComponent* AudioComponent;
+	UAudioComponent* AudioComponent;
+	/** Buffer for outgoing audio intended for procedural streaming */
+	mutable FCriticalSection QueueLock;
+	TArray<uint8> UncompressedDataQueue;
 	/** Per remote talker voice decoding state */
-	TSharedPtr<class IVoiceDecoder> VoiceDecoder;
+	TSharedPtr<IVoiceDecoder> VoiceDecoder;
 };
 
 /**
  * Generic implementation of voice engine, using Voice module for capture/codec
  */
-class FVoiceEngineImpl : public IVoiceEngine
+class FVoiceEngineImpl : public IVoiceEngine, public FSelfRegisteringExec
 {
 	class FVoiceSerializeHelper : public FGCObject
 	{
 		/** Reference to audio components */
 		FVoiceEngineImpl* VoiceEngine;
 		FVoiceSerializeHelper() :
-			VoiceEngine(NULL)
+			VoiceEngine(nullptr)
 		{}
 
 	public:
@@ -87,18 +112,18 @@ class FVoiceEngineImpl : public IVoiceEngine
 	friend class FVoiceSerializeHelper;
 
 	/** Mapping of UniqueIds to the incoming voice data and their audio component */
-	typedef TMap<class FUniqueNetIdString, FRemoteTalkerDataImpl> FRemoteTalkerData;
+	typedef TMap<FUniqueNetIdString, FRemoteTalkerDataImpl> FRemoteTalkerData;
 
 	/** Reference to the main online subsystem */
-	class IOnlineSubsystem* OnlineSubsystem;
+	IOnlineSubsystem* OnlineSubsystem;
 
 	FLocalVoiceData PlayerVoiceData[MAX_SPLITSCREEN_TALKERS];
 	/** Reference to voice capture device */
-	TSharedPtr<class IVoiceCapture> VoiceCapture;
+	TSharedPtr<IVoiceCapture> VoiceCapture;
 	/** Reference to voice encoding object */
-	TSharedPtr<class IVoiceEncoder> VoiceEncoder;
+	TSharedPtr<IVoiceEncoder> VoiceEncoder;
 
-    /** User index currently holding onto the voice interface */
+	/** User index currently holding onto the voice interface */
 	int32 OwningUserIndex;
 	/** Amount of uncompressed data available this frame */
 	uint32 UncompressedBytesAvailable;
@@ -106,19 +131,19 @@ class FVoiceEngineImpl : public IVoiceEngine
 	uint32 CompressedBytesAvailable;
 	/** Current frame state of voice capture */
 	EVoiceCaptureState::Type AvailableVoiceResult;
-    /** Have we stopped capturing voice but are waiting for its completion */
+	/** Have we stopped capturing voice but are waiting for its completion */
 	mutable bool bPendingFinalCapture;
-    /** State of voice recording */
+	/** State of voice recording */
 	bool bIsCapturing;
 
 	/** Data from voice codec, waiting to send to network. */
 	TArray<uint8> CompressedVoiceBuffer;
 	/** Data from network playing on an audio component. */
 	FRemoteTalkerData RemoteTalkerBuffers;
-	/** Voice decompression buffer, shared by all talkers */
+	/** Voice decompression buffer, shared by all talkers, valid during SubmitRemoteVoiceData */
 	TArray<uint8> DecompressedVoiceBuffer;
 	/** Serialization helper */
-	class FVoiceSerializeHelper* SerializeHelper;
+	FVoiceSerializeHelper* SerializeHelper;
 
 	/**
 	 * Determines if the specified index is the owner or not
@@ -150,20 +175,29 @@ class FVoiceEngineImpl : public IVoiceEngine
 	/** @return is active recording occurring at the moment */
 	bool IsRecording() const { return bIsCapturing || bPendingFinalCapture; }
 
+	/**
+	 * Callback from streaming audio when data is requested for playback
+	 *
+	 * @param InProceduralWave SoundWave requesting more data
+	 * @param SamplesRequired number of samples needed for immediate playback
+	 * @param TalkerId id of the remote talker to allocate voice data for
+	 */
+	void GenerateVoiceData(USoundWaveProcedural* InProceduralWave, int32 SamplesRequired, FUniqueNetIdString TalkerId);
+
 PACKAGE_SCOPE:
 
 	/** Constructor */
 	FVoiceEngineImpl() :
-		OnlineSubsystem(NULL),
-		VoiceCapture(NULL),
-		VoiceEncoder(NULL),
+		OnlineSubsystem(nullptr),
+		VoiceCapture(nullptr),
+		VoiceEncoder(nullptr),
 		OwningUserIndex(INVALID_INDEX),
 		UncompressedBytesAvailable(0),
 		CompressedBytesAvailable(0),
 		AvailableVoiceResult(EVoiceCaptureState::UnInitialized),
 		bPendingFinalCapture(false),
 		bIsCapturing(false),
-		SerializeHelper(NULL)
+		SerializeHelper(nullptr)
 	{};
 
 	// IVoiceEngine
@@ -175,22 +209,22 @@ public:
 	virtual ~FVoiceEngineImpl();
 
 	// IVoiceEngine
-    virtual uint32 StartLocalVoiceProcessing(uint32 LocalUserNum) override;
-    virtual uint32 StopLocalVoiceProcessing(uint32 LocalUserNum) override;
-    
+	virtual uint32 StartLocalVoiceProcessing(uint32 LocalUserNum) override;
+	virtual uint32 StopLocalVoiceProcessing(uint32 LocalUserNum) override;
+
 	virtual uint32 StartRemoteVoiceProcessing(const FUniqueNetId& UniqueId) override
 	{
 		// Not needed
 		return S_OK;
 	}
 
-    virtual uint32 StopRemoteVoiceProcessing(const FUniqueNetId& UniqueId) override
+	virtual uint32 StopRemoteVoiceProcessing(const FUniqueNetId& UniqueId) override
 	{
 		// Not needed
 		return S_OK;
 	}
 
-    virtual uint32 RegisterLocalTalker(uint32 LocalUserNum) override
+	virtual uint32 RegisterLocalTalker(uint32 LocalUserNum) override
 	{
 		if (OwningUserIndex == INVALID_INDEX)
 		{
@@ -201,7 +235,7 @@ public:
 		return E_FAIL;
 	}
 
-    virtual uint32 UnregisterLocalTalker(uint32 LocalUserNum) override
+	virtual uint32 UnregisterLocalTalker(uint32 LocalUserNum) override
 	{
 		if (IsOwningUser(LocalUserNum))
 		{
@@ -212,7 +246,7 @@ public:
 		return E_FAIL;
 	}
 
-    virtual uint32 RegisterRemoteTalker(const FUniqueNetId& UniqueId) override
+	virtual uint32 RegisterRemoteTalker(const FUniqueNetId& UniqueId) override
 	{
 		// Not needed
 		return S_OK;
@@ -220,19 +254,19 @@ public:
 
 	virtual uint32 UnregisterRemoteTalker(const FUniqueNetId& UniqueId) override;
 
-    virtual bool IsHeadsetPresent(uint32 LocalUserNum) override
+	virtual bool IsHeadsetPresent(uint32 LocalUserNum) override
 	{
 		return IsOwningUser(LocalUserNum) ? true : false;
 	}
 
-    virtual bool IsLocalPlayerTalking(uint32 LocalUserNum) override
+	virtual bool IsLocalPlayerTalking(uint32 LocalUserNum) override
 	{
 		return (GetVoiceDataReadyFlags() & (LocalUserNum << 1)) != 0;
 	}
 
 	virtual bool IsRemotePlayerTalking(const FUniqueNetId& UniqueId) override
 	{
-		return RemoteTalkerBuffers.Find((const FUniqueNetIdString&)UniqueId) != NULL;
+		return RemoteTalkerBuffers.Find((const FUniqueNetIdString&)UniqueId) != nullptr;
 	}
 
 	virtual uint32 GetVoiceDataReadyFlags() const override;
@@ -247,6 +281,11 @@ public:
 	virtual void Tick(float DeltaTime) override;
 	FString GetVoiceDebugState() const override;
 
+	// FSelfRegisteringExec
+	virtual bool Exec(UWorld* InWorld, const TCHAR* Cmd, FOutputDevice& Ar) override;
+
+private:
+
 	/**
 	 * Update the state of all remote talkers, possibly dropping data or the talker entirely
 	 */
diff --git a/Engine/Source/Runtime/Core/Public/Logging/LogMacros.h b/Engine/Source/Runtime/Core/Public/Logging/LogMacros.h
--- a/Engine/Source/Runtime/Core/Public/Logging/LogMacros.h
+++ b/Engine/Source/Runtime/Core/Public/Logging/LogMacros.h
@@ -39,7 +39,7 @@ struct CORE_API FMsg
 ----------------------------------------------------------------------------*/
 
 #ifndef COMPILED_IN_MINIMUM_VERBOSITY
-	#define COMPILED_IN_MINIMUM_VERBOSITY Log
+	#define COMPILED_IN_MINIMUM_VERBOSITY VeryVerbose
 #else
 	#if !IS_MONOLITHIC
 		#error COMPILED_IN_MINIMUM_VERBOSITY can only be defined in monolithic builds.
diff --git a/Engine/Source/Runtime/CoreUObject/Public/UObject/CoreOnline.h b/Engine/Source/Runtime/CoreUObject/Public/UObject/CoreOnline.h
--- a/Engine/Source/Runtime/CoreUObject/Public/UObject/CoreOnline.h
+++ b/Engine/Source/Runtime/CoreUObject/Public/UObject/CoreOnline.h
@@ -5,11 +5,10 @@
 #include "CoreMinimal.h"
 #include "UObject/ObjectMacros.h"
 
-const FName GameSessionName(TEXT("Game"));
-const FName PartySessionName(TEXT("Party"));
-
-const FName GamePort(TEXT("GamePort"));
-const FName BeaconPort(TEXT("BeaconPort"));
+#define GameSessionName NAME_GameSession
+#define PartySessionName NAME_PartySession
+#define GamePort NAME_GamePort
+#define BeaconPort NAME_BeaconPort
 
 #if !CPP
 // Circular dependency on Core vs UHT means we have to noexport these structs so tools can build
diff --git a/Engine/Source/Runtime/Engine/Private/PlayerState.cpp b/Engine/Source/Runtime/Engine/Private/PlayerState.cpp
--- a/Engine/Source/Runtime/Engine/Private/PlayerState.cpp
+++ b/Engine/Source/Runtime/Engine/Private/PlayerState.cpp
@@ -30,7 +30,7 @@ APlayerState::APlayerState(const FObjectInitializer& ObjectInitializer)
 	bNetLoadOnClient = false;
 
 	EngineMessageClass = UEngineMessage::StaticClass();
-	SessionName = GameSessionName;
+	SessionName = NAME_GameSession;
 }
 
 void APlayerState::UpdatePing(float InPing)
diff --git a/Engine/Source/Runtime/Engine/Private/VoiceChannel.cpp b/Engine/Source/Runtime/Engine/Private/VoiceChannel.cpp
--- a/Engine/Source/Runtime/Engine/Private/VoiceChannel.cpp
+++ b/Engine/Source/Runtime/Engine/Private/VoiceChannel.cpp
@@ -105,12 +105,12 @@ void UVoiceChannel::Tick()
 			}
 		}
 
-		if(Index >= VoicePackets.Num())
+		if (Index >= VoicePackets.Num())
 		{
 			// all sent, can throw everything away
 			VoicePackets.Empty();
 		}
-		else if(Index > 0)
+		else if (Index > 0)
 		{
 			// didn't send everything, just remove all packets actually sent
 			VoicePackets.RemoveAt(0,Index);
@@ -119,16 +119,16 @@ void UVoiceChannel::Tick()
 
 	// make sure we keep any reliable messages to try again next time
 	// but ditch any unreliable messages we've not managed to send
-	int PacketLoss = 0;
-	for(int i=VoicePackets.Num() - 1; i >= 0; i--)
+	int32 PacketLoss = 0;
+	for (int32 i = VoicePackets.Num() - 1; i >= 0; i--)
 	{
-		if(!VoicePackets[i]->IsReliable())
+		if (!VoicePackets[i]->IsReliable())
 		{
 			VoicePackets.RemoveAt(i,1,false);
 			PacketLoss++;
 		}
 	}
-	if(bHandshakeCompleted && PacketLoss > 0)
+	if (bHandshakeCompleted && PacketLoss > 0)
 	{
 		UE_LOG(LogNet, Log, TEXT("Dropped %d packets due to congestion in the voicechannel"), PacketLoss);
 	}
diff --git a/Engine/Source/Runtime/Online/Voice/Private/Android/VoiceModuleAndroid.cpp b/Engine/Source/Runtime/Online/Voice/Private/Android/VoiceModuleAndroid.cpp
--- a/Engine/Source/Runtime/Online/Voice/Private/Android/VoiceModuleAndroid.cpp
+++ b/Engine/Source/Runtime/Online/Voice/Private/Android/VoiceModuleAndroid.cpp
@@ -239,9 +239,8 @@ public:
 	
 	
 	// IVoiceCapture
-	virtual bool Init(int32 InSampleRate, int32 InNumChannels) override
+	virtual bool Init(const FString& DeviceName, int32 InSampleRate, int32 InNumChannels) override
 	{		
-
 #if ANDROIDVOICE_SUPPORTED_PLATFORMS
 		UE_LOG(LogVoiceCapture, Warning, TEXT("VoiceModuleAndroid Init"));
 		check(VoiceCaptureState == EVoiceCaptureState::UnInitialized);
@@ -258,14 +257,15 @@ public:
 			return false;
 		}
 		
-		if((inBufSamples = BUFFERFRAMES*InNumChannels) != 0){
+		if ((inBufSamples = BUFFERFRAMES*InNumChannels) != 0)
+		{
 			if((inputBuffer = (short *) calloc(inBufSamples, sizeof(short))) == NULL)
 			{
 				return false;
 			}
 		}
 				
-		if((inrb = create_circular_buffer(inBufSamples*sizeof(short)*4)) == NULL) 
+		if ((inrb = create_circular_buffer(inBufSamples*sizeof(short)*4)) == NULL) 
 		{
 			return false;
 		}
@@ -309,13 +309,13 @@ public:
 		
 		result = (*SL_RecorderObject)->Realize(SL_RecorderObject,
 											   SL_BOOLEAN_FALSE);
-        
+		
 		result = (*SL_RecorderObject)->GetInterface(SL_RecorderObject,
 													SL_IID_RECORD, &(SL_RecorderRecord));
 		
 		result = (*SL_RecorderObject)->GetInterface(SL_RecorderObject,
 													SL_IID_ANDROIDSIMPLEBUFFERQUEUE, &(SL_RecorderBufferQueue));
-        
+		
 		if(result != SL_RESULT_SUCCESS) 
 		{
 			UE_LOG( LogVoiceCapture, Warning,TEXT("FAILED OPENSL BUFFER QUEUE GetInterface 0x%x "), result); 
@@ -384,12 +384,19 @@ public:
 
 		return true;
 	}
+
 	virtual void Stop() override
 	{
 		UE_LOG(LogVoiceCapture, Warning, TEXT("Stop"));
 		VoiceCaptureState = EVoiceCaptureState::NotCapturing;
 		return;
 	}
+
+	virtual bool ChangeDevice(const FString& DeviceName, int32 SampleRate, int32 NumChannels)
+	{
+		/** NYI */
+		return false;
+	}
 	
 	virtual bool IsCapturing() override
 	{
@@ -455,6 +462,19 @@ public:
 		}
 		return State;
 	}
+
+	virtual int32 GetBufferSize() const
+	{
+		/** NYI */
+		/** possibly inBufSamples*sizeof(short)*4 */
+		return 0;
+	}
+
+	virtual void DumpState() const
+	{
+		/** NYI */
+		UE_LOG(LogVoiceCapture, Display, TEXT("NYI"));
+	}
 	
 private:
 
@@ -484,11 +504,11 @@ void ShutdownVoiceCapture()
 {
 }
 
-IVoiceCapture* CreateVoiceCaptureObject()
+IVoiceCapture* CreateVoiceCaptureObject(const FString& DeviceName, int32 SampleRate, int32 NumChannels)
 {
 #if ANDROIDVOICE_SUPPORTED_PLATFORMS
 	IVoiceCapture* Capture = new FVoiceCaptureOpenSLES;
-	if (!Capture || !Capture->Init(VOICE_SAMPLE_RATE, NUM_VOICE_CHANNELS))
+	if (!Capture || !Capture->Init(DeviceName, SampleRate, NumChannels))
 	{
 		delete Capture;
 		Capture = nullptr;
@@ -499,12 +519,11 @@ IVoiceCapture* CreateVoiceCaptureObject()
 #endif 
 }
 
-IVoiceEncoder* CreateVoiceEncoderObject()
+IVoiceEncoder* CreateVoiceEncoderObject(int32 SampleRate, int32 NumChannels, EAudioEncodeHint EncodeHint)
 {
-
 #if ANDROIDVOICE_SUPPORTED_PLATFORMS
 	FVoiceEncoderOpus* NewEncoder = new FVoiceEncoderOpus;
-	if (!NewEncoder->Init(VOICE_SAMPLE_RATE, NUM_VOICE_CHANNELS))
+	if (!NewEncoder->Init(SampleRate, NumChannels, EncodeHint))
 	{
 		delete NewEncoder;
 		NewEncoder = nullptr;
@@ -516,11 +535,11 @@ IVoiceEncoder* CreateVoiceEncoderObject()
 #endif
 }
 
-IVoiceDecoder* CreateVoiceDecoderObject()
+IVoiceDecoder* CreateVoiceDecoderObject(int32 SampleRate, int32 NumChannels)
 {
 #if ANDROIDVOICE_SUPPORTED_PLATFORMS
 	FVoiceDecoderOpus* NewDecoder = new FVoiceDecoderOpus;
-	if (!NewDecoder->Init(VOICE_SAMPLE_RATE, NUM_VOICE_CHANNELS))
+	if (!NewDecoder->Init(SampleRate, NumChannels))
 	{
 		delete NewDecoder;
 		NewDecoder = nullptr;
diff --git a/Engine/Source/Runtime/Online/Voice/Private/Linux/VoiceModuleLinux.cpp b/Engine/Source/Runtime/Online/Voice/Private/Linux/VoiceModuleLinux.cpp
--- a/Engine/Source/Runtime/Online/Voice/Private/Linux/VoiceModuleLinux.cpp
+++ b/Engine/Source/Runtime/Online/Voice/Private/Linux/VoiceModuleLinux.cpp
@@ -13,17 +13,17 @@ void ShutdownVoiceCapture()
 {
 }
 
-IVoiceCapture* CreateVoiceCaptureObject()
+IVoiceCapture* CreateVoiceCaptureObject(const FString& DeviceName, int32 SampleRate, int32 NumChannels)
 {
-	return NULL; 
+	return nullptr; 
 }
 
-IVoiceEncoder* CreateVoiceEncoderObject()
+IVoiceEncoder* CreateVoiceEncoderObject(int32 SampleRate, int32 NumChannels, EAudioEncodeHint EncodeHint)
 {
-	return NULL; 
+	return nullptr; 
 }
 
-IVoiceDecoder* CreateVoiceDecoderObject()
+IVoiceDecoder* CreateVoiceDecoderObject(int32 SampleRate, int32 NumChannels)
 {
-	return NULL; 
+	return nullptr; 
 }
diff --git a/Engine/Source/Runtime/Online/Voice/Private/Mac/VoiceModuleMac.cpp b/Engine/Source/Runtime/Online/Voice/Private/Mac/VoiceModuleMac.cpp
--- a/Engine/Source/Runtime/Online/Voice/Private/Mac/VoiceModuleMac.cpp
+++ b/Engine/Source/Runtime/Online/Voice/Private/Mac/VoiceModuleMac.cpp
@@ -94,7 +94,7 @@
 	}
 	
 	// IVoiceCapture
-	virtual bool Init(int32 InSampleRate, int32 InNumChannels) override
+	virtual bool Init(const FString& DeviceName, int32 InSampleRate, int32 InNumChannels) override
 	{
 		check(VoiceCaptureState == EVoiceCaptureState::UnInitialized);
 		
@@ -268,6 +268,7 @@
 		
 		return true;
 	}
+
 	virtual void Shutdown() override
 	{
 		switch (VoiceCaptureState)
@@ -301,6 +302,7 @@
 				break;
 		}
 	}
+
 	virtual bool Start() override
 	{
 		check(StreamComponent && VoiceCaptureState == EVoiceCaptureState::NotCapturing);
@@ -326,6 +328,7 @@
 			return false;
 		}
 	}
+
 	virtual void Stop() override
 	{
 		check(IsCapturing());
@@ -345,10 +348,18 @@
 			FMemory::Memzero(BufferList[i].mData, BufferSize);
 		}
 	}
+
+	virtual bool ChangeDevice(const FString& DeviceName, int32 SampleRate, int32 NumChannels)
+	{
+		/** NYI */
+		return false;
+	}
+
 	virtual bool IsCapturing() override
 	{
 		return StreamComponent && VoiceCaptureState > EVoiceCaptureState::NotCapturing;
 	}
+
 	virtual EVoiceCaptureState::Type GetCaptureState(uint32& OutAvailableVoiceData) const override
 	{
 		if (VoiceCaptureState != EVoiceCaptureState::UnInitialized &&
@@ -364,6 +375,86 @@
 		return VoiceCaptureState;
 	}
 	
+	virtual EVoiceCaptureState::Type GetVoiceData(uint8* OutVoiceBuffer, uint32 InVoiceBufferSize, uint32& OutAvailableVoiceData) override
+	{
+		EVoiceCaptureState::Type State = VoiceCaptureState;
+		
+		if (VoiceCaptureState != EVoiceCaptureState::UnInitialized &&
+			VoiceCaptureState != EVoiceCaptureState::NotCapturing)
+		{
+			GetCaptureState(OutAvailableVoiceData);
+			
+			if ( OutAvailableVoiceData > 0 )
+			{
+				uint32 ConvertedVoiceData = 0;
+				if ( InVoiceBufferSize >= OutAvailableVoiceData )
+				{
+					check(ReadOffset <= BufferSize);
+					uint32 BytesToRead = OutAvailableVoiceData;
+					uint8* Data = OutVoiceBuffer;
+					while (BytesToRead)
+					{
+						uint32 CurrentRead = FMath::Min(FMath::Min(BufferList[ReadBuffer].mDataByteSize - ReadOffset, BytesToRead), BufferSize);
+						uint32 CurrentWrite = 0;
+						State = CopyBuffer(ReadBuffer, Data, ReadOffset, &CurrentRead, &CurrentWrite);
+						
+						if ( State == EVoiceCaptureState::Ok )
+						{
+							Data += CurrentWrite;
+							ConvertedVoiceData += CurrentWrite;
+							BytesToRead = ( CurrentRead != 0 ) ? BytesToRead - CurrentRead : 0;
+							ReadOffset += CurrentRead;
+							if( ReadOffset == BufferList[ReadBuffer].mDataByteSize )
+							{
+								BufferList[ReadBuffer].mDataByteSize = 0;
+								FMemory::Memzero(BufferList[ReadBuffer].mData, BufferSize);
+								ReadBuffer = (ReadBuffer + 1) % BufferList.Num();
+								ReadOffset = 0;
+							}
+							check(ReadOffset <= BufferSize);
+						}
+						else
+						{
+							break;
+						}
+					}
+					int32 BytesRead = OutAvailableVoiceData;
+					FPlatformAtomics::InterlockedAdd(&ReadableBytes, -BytesRead);
+					
+					OutAvailableVoiceData  = ConvertedVoiceData;
+				}
+				else
+				{
+					State = EVoiceCaptureState::BufferTooSmall;
+				}
+			}
+			else
+			{
+				VoiceCaptureState = EVoiceCaptureState::NoData;
+				State = EVoiceCaptureState::Ok;
+			}
+		}
+		else
+		{
+			OutAvailableVoiceData = 0;
+		}
+		return State;
+	}
+
+	virtual int32 GetBufferSize() const
+	{
+		/** NYI */
+		return 0;
+	}
+
+	virtual void DumpState() const
+	{
+		/** NYI */
+		UE_LOG(LogVoiceCapture, Display, TEXT("NYI"));
+	}
+
+private:
+
 	struct FAudioFileIO
 	{
 		uint32 SrcReadFrames;
@@ -371,23 +462,23 @@
 		uint32 SrcBufferSize;
 		uint32 SrcBytesPerFrame;
 	};
-	
-	static OSStatus ConvertInputFormat( AudioConverterRef AudioConverter,
+
+	static OSStatus ConvertInputFormat(AudioConverterRef AudioConverter,
 										UInt32* IONumberDataPackets,
 										AudioBufferList* IOData,
 										AudioStreamPacketDescription** OutDataPacketDescription,
-										void* InUserData )
+										void* InUserData)
 	{
 		FAudioFileIO* Input = (FAudioFileIO*)InUserData;
-		
-		if ( Input && Input->SrcBuffer )
+
+		if (Input && Input->SrcBuffer)
 		{
 			*IOData = *Input->SrcBuffer;
-			
+
 			int32 Num = *IONumberDataPackets;
-			if ( ( Input->SrcReadFrames * (Input->SrcBytesPerFrame)) < Input->SrcBufferSize )
+			if ((Input->SrcReadFrames * (Input->SrcBytesPerFrame)) < Input->SrcBufferSize)
 			{
-				if ( (Num * Input->SrcBytesPerFrame) < IOData->mBuffers[0].mDataByteSize )
+				if ((Num * Input->SrcBytesPerFrame) < IOData->mBuffers[0].mDataByteSize)
 				{
 					IOData->mBuffers[0].mDataByteSize = Num * Input->SrcBytesPerFrame;
 				}
@@ -404,18 +495,18 @@
 				return eofErr;
 			}
 		}
-		
+
 		return noErr;
 	}
-	
+
 	EVoiceCaptureState::Type CopyBuffer(uint8 InReadBuffer, uint8* OutVoiceBuffer, uint32 InReadOffset, uint32* ReadLen, uint32* WriteLen)
 	{
 		EVoiceCaptureState::Type State = VoiceCaptureState;
-		if ( *ReadLen )
+		if (*ReadLen)
 		{
 			uint8* Data = (uint8*)(BufferList[InReadBuffer].mData) + InReadOffset;
-			
-			if ( StreamDesc.mSampleRate == OutputDesc.mSampleRate )
+
+			if (StreamDesc.mSampleRate == OutputDesc.mSampleRate)
 			{
 				FMemory::Memcpy(OutVoiceBuffer, Data, *ReadLen);
 				*WriteLen = *ReadLen;
@@ -429,21 +520,21 @@
 				OutputBuffer.mBuffers[0].mNumberChannels = StreamDesc.mChannelsPerFrame;
 				OutputBuffer.mBuffers[0].mDataByteSize = *ReadLen;
 				OutputBuffer.mBuffers[0].mData = OutVoiceBuffer;
-				
+
 				AudioBufferList InputBuffer;
 				InputBuffer.mNumberBuffers = 1;
-				InputBuffer.mBuffers[0].mNumberChannels =  StreamDesc.mChannelsPerFrame;
+				InputBuffer.mBuffers[0].mNumberChannels = StreamDesc.mChannelsPerFrame;
 				InputBuffer.mBuffers[0].mDataByteSize = BufferList[InReadBuffer].mDataByteSize;
 				InputBuffer.mBuffers[0].mData = Data;
-				
+
 				FAudioFileIO inUserData;
 				inUserData.SrcBuffer = &InputBuffer;
 				inUserData.SrcBytesPerFrame = StreamDesc.mBytesPerFrame;
 				inUserData.SrcBufferSize = *ReadLen;
 				inUserData.SrcReadFrames = 0;
-				
+
 				OSStatus result = AudioConverterFillComplexBuffer(StreamConverter, &FVoiceCaptureCoreAudio::ConvertInputFormat, &inUserData, &FramesToCopy, &OutputBuffer, nullptr);
-				if ( result == 0 || result == eofErr )
+				if (result == 0 || result == eofErr)
 				{
 					State = EVoiceCaptureState::Ok;
 					*ReadLen = inUserData.SrcReadFrames * StreamDesc.mBytesPerFrame;
@@ -459,72 +550,6 @@
 		return State;
 	}
 	
-	virtual EVoiceCaptureState::Type GetVoiceData(uint8* OutVoiceBuffer, uint32 InVoiceBufferSize, uint32& OutAvailableVoiceData) override
-	{
-		EVoiceCaptureState::Type State = VoiceCaptureState;
-		
-		if (VoiceCaptureState != EVoiceCaptureState::UnInitialized &&
-			VoiceCaptureState != EVoiceCaptureState::NotCapturing)
-		{
-			GetCaptureState(OutAvailableVoiceData);
-			
-			if ( OutAvailableVoiceData > 0 )
-			{
-				uint32 ConvertedVoiceData = 0;
-				if ( InVoiceBufferSize >= OutAvailableVoiceData )
-				{
-					check(ReadOffset <= BufferSize);
-					uint32 BytesToRead = OutAvailableVoiceData;
-					uint8* Data = OutVoiceBuffer;
-					while (BytesToRead)
-					{
-						uint32 CurrentRead = FMath::Min(FMath::Min(BufferList[ReadBuffer].mDataByteSize - ReadOffset, BytesToRead), BufferSize);
-						uint32 CurrentWrite = 0;
-						State = CopyBuffer(ReadBuffer, Data, ReadOffset, &CurrentRead, &CurrentWrite);
-						
-						if ( State == EVoiceCaptureState::Ok )
-						{
-							Data += CurrentWrite;
-							ConvertedVoiceData += CurrentWrite;
-							BytesToRead = ( CurrentRead != 0 ) ? BytesToRead - CurrentRead : 0;
-							ReadOffset += CurrentRead;
-							if( ReadOffset == BufferList[ReadBuffer].mDataByteSize )
-							{
-								BufferList[ReadBuffer].mDataByteSize = 0;
-								FMemory::Memzero(BufferList[ReadBuffer].mData, BufferSize);
-								ReadBuffer = (ReadBuffer + 1) % BufferList.Num();
-								ReadOffset = 0;
-							}
-							check(ReadOffset <= BufferSize);
-						}
-						else
-						{
-							break;
-						}
-					}
-					int32 BytesRead = OutAvailableVoiceData;
-					FPlatformAtomics::InterlockedAdd(&ReadableBytes, -BytesRead);
-					
-					OutAvailableVoiceData  = ConvertedVoiceData;
-				}
-				else
-				{
-					State = EVoiceCaptureState::BufferTooSmall;
-				}
-			}
-			else
-			{
-				VoiceCaptureState = EVoiceCaptureState::NoData;
-				State = EVoiceCaptureState::Ok;
-			}
-		}
-		else
-		{
-			OutAvailableVoiceData = 0;
-		}
-		return State;
-	}
-	
 private:
 	/** Descriptor for stream format */
 	AudioStreamBasicDescription StreamDesc;
@@ -563,10 +588,10 @@
 {
 }
 
-IVoiceCapture* CreateVoiceCaptureObject()
+IVoiceCapture* CreateVoiceCaptureObject(const FString& DeviceName, int32 SampleRate, int32 NumChannels)
 {
 	IVoiceCapture* Capture = new FVoiceCaptureCoreAudio;
-	if ( !Capture || !Capture->Init(VOICE_SAMPLE_RATE, NUM_VOICE_CHANNELS) )
+	if ( !Capture || !Capture->Init(DeviceName, SampleRate, NumChannels) )
 	{
 		delete Capture;
 		Capture = nullptr;
@@ -574,25 +599,25 @@
 	return Capture;
 }
 
-IVoiceEncoder* CreateVoiceEncoderObject()
+IVoiceEncoder* CreateVoiceEncoderObject(int32 SampleRate, int32 NumChannels, EAudioEncodeHint EncodeHint)
 {
 	FVoiceEncoderOpus* NewEncoder = new FVoiceEncoderOpus;
-	if (!NewEncoder->Init(VOICE_SAMPLE_RATE, NUM_VOICE_CHANNELS))
+	if (!NewEncoder->Init(SampleRate, NumChannels, EncodeHint))
 	{
 		delete NewEncoder;
-		NewEncoder = NULL;
+		NewEncoder = nullptr;
 	}
 	
 	return NewEncoder;
 }
 
-IVoiceDecoder* CreateVoiceDecoderObject()
+IVoiceDecoder* CreateVoiceDecoderObject(int32 SampleRate, int32 NumChannels)
 {
 	FVoiceDecoderOpus* NewDecoder = new FVoiceDecoderOpus;
-	if (!NewDecoder->Init(VOICE_SAMPLE_RATE, NUM_VOICE_CHANNELS))
+	if (!NewDecoder->Init(SampleRate, NumChannels))
 	{
 		delete NewDecoder;
-		NewDecoder = NULL;
+		NewDecoder = nullptr;
 	}
 	
 	return NewDecoder;
diff --git a/Engine/Source/Runtime/Online/Voice/Private/VoiceCodecOpus.cpp b/Engine/Source/Runtime/Online/Voice/Private/VoiceCodecOpus.cpp
--- a/Engine/Source/Runtime/Online/Voice/Private/VoiceCodecOpus.cpp
+++ b/Engine/Source/Runtime/Online/Voice/Private/VoiceCodecOpus.cpp
@@ -30,6 +30,12 @@ THIRD_PARTY_INCLUDES_END
 /** 20ms frame sizes are a good choice for most applications (1000ms / 20ms = 50) */
 #define NUM_OPUS_FRAMES_PER_SEC 50
 
+#define OPUS_CHECK_CTL(Category, CTL) \
+	if (ErrCode != OPUS_OK) \
+	{ \
+		UE_LOG(Category, Warning, TEXT("Failure to get CTL %s"), #CTL); \
+	}
+
 /**
  * Output debug information regarding the state of the Opus encoder
  */
@@ -39,21 +45,55 @@ void DebugEncoderInfo(OpusEncoder* Encoder)
 
 	int32 BitRate = 0;
 	ErrCode = opus_encoder_ctl(Encoder, OPUS_GET_BITRATE(&BitRate));
+	OPUS_CHECK_CTL(LogVoiceEncode, BitRate);
 
 	int32 Vbr = 0;
 	ErrCode = opus_encoder_ctl(Encoder, OPUS_GET_VBR(&Vbr));
+	OPUS_CHECK_CTL(LogVoiceEncode, Vbr);
 
 	int32 SampleRate = 0;
 	ErrCode = opus_encoder_ctl(Encoder, OPUS_GET_SAMPLE_RATE(&SampleRate));
+	OPUS_CHECK_CTL(LogVoiceEncode, SampleRate);
 
 	int32 Application = 0;
 	ErrCode = opus_encoder_ctl(Encoder, OPUS_GET_APPLICATION(&Application));
+	OPUS_CHECK_CTL(LogVoiceEncode, Application);
+
+	int32 Signal = 0;
+	ErrCode = opus_encoder_ctl(Encoder, OPUS_GET_SIGNAL(&Signal));
+	OPUS_CHECK_CTL(LogVoiceEncode, Signal);
 
 	int32 Complexity = 0;
 	ErrCode = opus_encoder_ctl(Encoder, OPUS_GET_COMPLEXITY(&Complexity));
+	OPUS_CHECK_CTL(LogVoiceEncode, Complexity);
+
+	UE_LOG(LogVoiceEncode, Display, TEXT("Opus Encoder Details"));
+	UE_LOG(LogVoiceEncode, Display, TEXT("- Application: %d"), Application);
+	UE_LOG(LogVoiceEncode, Display, TEXT("- Signal: %d"), Signal);
+	UE_LOG(LogVoiceEncode, Display, TEXT("- BitRate: %d"), BitRate);
+	UE_LOG(LogVoiceEncode, Display, TEXT("- SampleRate: %d"), SampleRate);
+	UE_LOG(LogVoiceEncode, Display, TEXT("- Vbr: %d"), Vbr);
+	UE_LOG(LogVoiceEncode, Display, TEXT("- Complexity: %d"), Complexity);
+}
+
+/**
+ * Output debug information regarding the state of the Opus decoder
+ */
+void DebugDecoderInfo(OpusDecoder* Decoder)
+{
+	int32 ErrCode = 0;
 
-	UE_LOG(LogVoiceEncode, Display, TEXT("Opus Encoder Details:\n Application: %d\n Bitrate: %d\n SampleRate: %d\n VBR: %d\n Complexity: %d"),
-		Application, BitRate, SampleRate, Vbr, Complexity);
+	int32 Gain = 0;
+	ErrCode = opus_decoder_ctl(Decoder, OPUS_GET_GAIN(&Gain));
+	OPUS_CHECK_CTL(LogVoiceDecode, Gain);
+
+	int32 Pitch = 0;
+	ErrCode = opus_decoder_ctl(Decoder, OPUS_GET_PITCH(&Pitch));
+	OPUS_CHECK_CTL(LogVoiceDecode, Pitch);
+
+	UE_LOG(LogVoiceDecode, Display, TEXT("Opus Decoder Details"));
+	UE_LOG(LogVoiceDecode, Display, TEXT("- Gain: %d"), Gain);
+	UE_LOG(LogVoiceDecode, Display, TEXT("- Pitch: %d"), Pitch);
 }
 
 /**
@@ -76,7 +116,7 @@ void DebugFrameInfoInternal(const uint8* PacketData, uint32 PacketLength, uint32
 	int32 NumSamplesPerFrame = opus_packet_get_samples_per_frame(PacketData, SampleRate);
 	int32 Bandwidth = opus_packet_get_bandwidth(PacketData);
 
-	const TCHAR* BandwidthStr = NULL;
+	const TCHAR* BandwidthStr = nullptr;
 	switch (Bandwidth)
 	{
 	case OPUS_BANDWIDTH_NARROWBAND: // Narrowband (4kHz bandpass)
@@ -107,7 +147,7 @@ void DebugFrameInfoInternal(const uint8* PacketData, uint32 PacketLength, uint32
 	 *	| config  |s| c |
 	 *	+-+-+-+-+-+-+-+-+
 	 */
-	uint8 TOC;
+	uint8 TOC = 0;
 	// (max 48 x 2.5ms frames in a packet = 120ms)
 	const uint8* frames[48];
 	int16 size[48];
@@ -123,7 +163,7 @@ void DebugFrameInfoInternal(const uint8* PacketData, uint32 PacketLength, uint32
 	// Number of frames and their configuration
 	// 0: 1 frame in the packet
 	// 1: 2 frames in the packet, each with equal compressed size
-    // 2: 2 frames in the packet, with different compressed sizes
+	// 2: 2 frames in the packet, with different compressed sizes
 	// 3: an arbitrary number of frames in the packet
 	int32 TOCMode = TOC & 0x3;
 
@@ -153,7 +193,7 @@ FVoiceEncoderOpus::FVoiceEncoderOpus() :
 	SampleRate(0),
 	NumChannels(0),
 	FrameSize(0),
-	Encoder(NULL),
+	Encoder(nullptr),
 	LastEntropyIdx(0),
 	Generation(0)
 {
@@ -165,52 +205,73 @@ FVoiceEncoderOpus::~FVoiceEncoderOpus()
 	Destroy();
 }
 
-bool FVoiceEncoderOpus::Init(int32 InSampleRate, int32 InNumChannels)
+bool FVoiceEncoderOpus::Init(int32 InSampleRate, int32 InNumChannels, EAudioEncodeHint EncodeHint)
 {
 	UE_LOG(LogVoiceEncode, Display, TEXT("EncoderVersion: %s"), ANSI_TO_TCHAR(opus_get_version_string()));
 
+	if (InSampleRate != 8000 &&
+		InSampleRate != 12000 &&
+		InSampleRate != 16000 &&
+		InSampleRate != 24000 &&
+		InSampleRate != 48000)
+	{
+		UE_LOG(LogVoiceEncode, Warning, TEXT("Voice encoder doesn't support %d hz"), InSampleRate);
+		return false;
+	}
+
+	if (InNumChannels < 1 || InNumChannels > 2)
+	{
+		UE_LOG(LogVoiceEncode, Warning, TEXT("Voice encoder only supports 1 or 2 channels"));
+		return false;
+	}
+
 	SampleRate = InSampleRate;
 	NumChannels = InNumChannels;
-
+	
 	// 20ms frame sizes are a good choice for most applications (1000ms / 20ms = 50)
 	FrameSize = SampleRate / NUM_OPUS_FRAMES_PER_SEC;
 	//MaxFrameSize = FrameSize * MAX_OPUS_FRAMES;
 
 	int32 EncError = 0;
 
+	const int32 Application = (EncodeHint == EAudioEncodeHint::VoiceEncode_Audio) ? OPUS_APPLICATION_AUDIO : OPUS_APPLICATION_VOIP;
+
 #if USE_UE4_MEM_ALLOC
-	int32 EncSize = opus_encoder_get_size(NumChannels);
+	const int32 EncSize = opus_encoder_get_size(NumChannels);
 	Encoder = (OpusEncoder*)FMemory::Malloc(EncSize);
-	EncError = opus_encoder_init(Encoder, SampleRate, NumChannels, OPUS_APPLICATION_VOIP);
+	EncError = opus_encoder_init(Encoder, SampleRate, NumChannels, Application);
 #else
-	Encoder = opus_encoder_create(SampleRate, NumChannels, OPUS_APPLICATION_VOIP, &EncError);
+	Encoder = opus_encoder_create(SampleRate, NumChannels, Application, &EncError);
 #endif
 
-	if (EncError != OPUS_OK)
+	if (EncError == OPUS_OK)
 	{
-		UE_LOG(LogVoiceEncode, Warning, TEXT("Failed to init Opus Encoder: %s"), ANSI_TO_TCHAR(opus_strerror(EncError)));
-		Destroy();
-	}
+		// Turn on variable bit rate encoding
+		const int32 UseVbr = 1;
+		opus_encoder_ctl(Encoder, OPUS_SET_VBR(UseVbr));
 
-	// Turn on variable bit rate encoding
-	int32 UseVbr = 1;
-	opus_encoder_ctl(Encoder, OPUS_SET_VBR(UseVbr));
+		// Turn off constrained VBR
+		const int32 UseCVbr = 0;
+		opus_encoder_ctl(Encoder, OPUS_SET_VBR_CONSTRAINT(UseCVbr));
 
-	// Turn off constrained VBR
-	int32 UseCVbr = 0;
-	opus_encoder_ctl(Encoder, OPUS_SET_VBR_CONSTRAINT(UseCVbr));
+		// Complexity (1-10)
+		const int32 Complexity = 1;
+		opus_encoder_ctl(Encoder, OPUS_SET_COMPLEXITY(Complexity));
 
-	// Complexity (1-10)
-	int32 Complexity = 1;
-	opus_encoder_ctl(Encoder, OPUS_SET_COMPLEXITY(Complexity));
-
-	// Forward error correction
-	int32 InbandFEC = 0;
-	opus_encoder_ctl(Encoder, OPUS_SET_INBAND_FEC(InbandFEC));
+		// Forward error correction
+		const int32 InbandFEC = 0;
+		opus_encoder_ctl(Encoder, OPUS_SET_INBAND_FEC(InbandFEC));
 
 #if DEBUG_OPUS
-	DebugEncoderInfo(Encoder);
+		DebugEncoderInfo(Encoder);
 #endif // DEBUG_OPUS
+	}
+	else
+	{
+		UE_LOG(LogVoiceEncode, Warning, TEXT("Failed to init Opus Encoder: %s"), ANSI_TO_TCHAR(opus_strerror(EncError)));
+		Destroy();
+	}
+
 	return EncError == OPUS_OK;
 }
 
@@ -218,34 +279,43 @@ int32 FVoiceEncoderOpus::Encode(const uint8* RawPCMData, uint32 RawDataSize, uin
 {
 	check(Encoder);
 	SCOPE_CYCLE_COUNTER(STAT_Voice_Encoding);
+	SET_DWORD_STAT(STAT_Encode_SampleRate, SampleRate);
+	SET_DWORD_STAT(STAT_Encode_NumChannels, NumChannels);
 
 	int32 HeaderSize = 0;
-	int32 BytesPerFrame = FrameSize * NumChannels * sizeof(opus_int16);
-	int32 MaxFramesEncoded = MAX_OPUS_UNCOMPRESSED_BUFFER_SIZE / BytesPerFrame;
+	const int32 BytesPerFrame = FrameSize * NumChannels * sizeof(opus_int16);
+	const int32 MaxFramesEncoded = MAX_OPUS_UNCOMPRESSED_BUFFER_SIZE / BytesPerFrame;
 
 	// total bytes / bytes per frame
-	int32 NumFramesToEncode = FMath::Min((int32)RawDataSize / BytesPerFrame, MaxFramesEncoded); 
-	int32 DataRemainder = RawDataSize % BytesPerFrame;
-	int32 RawDataStride = BytesPerFrame;
+	const int32 NumFramesToEncode = FMath::Min((int32)RawDataSize / BytesPerFrame, MaxFramesEncoded);
+	const int32 DataRemainder = RawDataSize % BytesPerFrame;
+	const int32 RawDataStride = BytesPerFrame;
+
+	if (NumFramesToEncode == 0)
+	{
+		// We can avoid saving out an empty header if we know we're not going to send anything
+		check(DataRemainder == RawDataSize);
+		OutCompressedDataSize = 0;
+		return DataRemainder;
+	}
 
 	// Store the number of frames to be encoded
 	check(NumFramesToEncode < MAX_uint8);
-	OutCompressedData[0] = NumFramesToEncode;
+	OutCompressedData[0] = (uint8)NumFramesToEncode;
 	OutCompressedData[1] = Generation;
 	HeaderSize += 2 * sizeof(uint8);
 	
 	// Store the offset to each encoded frame
 	uint16* CompressedOffsets = (uint16*)(OutCompressedData + HeaderSize);
-	uint32 LengthOfCompressedOffsets = NumFramesToEncode * sizeof(uint16);
+	const uint32 LengthOfCompressedOffsets = NumFramesToEncode * sizeof(uint16);
 	HeaderSize += LengthOfCompressedOffsets;
 
 	// Store the entropy to each encoded frame
-	uint32 LengthOfEntropyOffsets = 0;
 #if ADD_ENTROPY_TO_PACKET
 	uint32* EntropyOffsets = (uint32*)(OutCompressedData + HeaderSize);
-	LengthOfEntropyOffsets = NumFramesToEncode * sizeof(uint32);
-#endif
+	uint32 LengthOfEntropyOffsets = NumFramesToEncode * sizeof(uint32);
 	HeaderSize += LengthOfEntropyOffsets;
+#endif
 
 	// Space available after overhead
 	int32 AvailableBufferSize = OutCompressedDataSize - HeaderSize;
@@ -255,14 +325,13 @@ int32 FVoiceEncoderOpus::Encode(const uint8* RawPCMData, uint32 RawDataSize, uin
 	int32 CompressedBufferOffset = 0;
 	for (int32 i = 0; i < NumFramesToEncode; i++)
 	{
-		int32 CompressedLength = 0;
-			CompressedLength = opus_encode(Encoder, (const opus_int16*)(RawPCMData + (i * RawDataStride)), FrameSize, CompressedDataStart + CompressedBufferOffset, AvailableBufferSize);
-
+		int32 CompressedLength = opus_encode(Encoder, (const opus_int16*)(RawPCMData + (i * RawDataStride)), FrameSize, CompressedDataStart + CompressedBufferOffset, AvailableBufferSize);
 		if (CompressedLength < 0)
 		{
 			const char* ErrorStr = opus_strerror(CompressedLength);
 			UE_LOG(LogVoiceEncode, Warning, TEXT("Failed to encode: [%d] %s"), CompressedLength, ANSI_TO_TCHAR(ErrorStr));
 
+			// Mark header as nothing encoded
 			OutCompressedData[0] = 0;
 			OutCompressedDataSize = 0;
 			return 0;
@@ -290,18 +359,104 @@ int32 FVoiceEncoderOpus::Encode(const uint8* RawPCMData, uint32 RawDataSize, uin
 		else
 		{
 			UE_LOG(LogVoiceEncode, Warning, TEXT("Nothing to encode!"));
+			CompressedOffsets[i] = 0;
+#if ADD_ENTROPY_TO_PACKET
+			EntropyOffsets[i] = 0;
+#endif
 		}
 	}
 
 	// End of buffer
 	OutCompressedDataSize = HeaderSize + CompressedBufferOffset;
 
-	UE_LOG(LogVoiceEncode, VeryVerbose, TEXT("OpusEncode[%d]: RawSize: %d CompressedSize: %d NumFramesEncoded: %d Remains: %d"), Generation, RawDataSize, OutCompressedDataSize, NumFramesToEncode, DataRemainder);
+	UE_LOG(LogVoiceEncode, VeryVerbose, TEXT("OpusEncode[%d]: RawSize: %d HeaderSize: %d CompressedSize: %d NumFramesEncoded: %d Remains: %d"), Generation, RawDataSize, HeaderSize, OutCompressedDataSize, NumFramesToEncode, DataRemainder);
+
+#if STATS
+	int32 BitRate = 0;
+	int32 ErrCode = opus_encoder_ctl(Encoder, OPUS_GET_BITRATE(&BitRate));
+	SET_DWORD_STAT(STAT_Encode_Bitrate, BitRate);
+	SET_FLOAT_STAT(STAT_Encode_CompressionRatio, (float)OutCompressedDataSize / (float)RawDataSize);
+	SET_DWORD_STAT(STAT_Encode_OutSize, OutCompressedDataSize);
+#endif // STATS
 
 	Generation = (Generation + 1) % MAX_uint8;
 	return DataRemainder;
 }
 
+bool FVoiceEncoderOpus::SetBitrate(int32 InBitRate)
+{
+	if (InBitRate >= 500 &&
+		InBitRate <= 512000)
+	{
+		if (Encoder)
+		{
+			int32 EncError = opus_encoder_ctl(Encoder, OPUS_SET_BITRATE(InBitRate));
+
+#if DEBUG_OPUS
+			DebugEncoderInfo(Encoder);
+#endif // DEBUG_OPUS
+
+			return EncError == OPUS_OK;
+		}
+	}
+	
+	return false;
+}
+
+bool FVoiceEncoderOpus::SetVBR(bool bEnableVBR)
+{
+	if (Encoder)
+	{
+		const int32 UseVbr = bEnableVBR ? 1 : 0;
+		int32 EncError = opus_encoder_ctl(Encoder, OPUS_SET_VBR(UseVbr));
+
+#if DEBUG_OPUS
+		DebugEncoderInfo(Encoder);
+#endif // DEBUG_OPUS
+
+		return EncError == OPUS_OK;
+	}
+
+	return false;
+}
+
+bool FVoiceEncoderOpus::SetComplexity(int32 InComplexity)
+{
+	if (InComplexity >= 0 &&
+		InComplexity <= 10)
+	{
+		if (Encoder)
+		{
+			int32 EncError = opus_encoder_ctl(Encoder, OPUS_SET_COMPLEXITY(InComplexity));
+
+#if DEBUG_OPUS
+			DebugEncoderInfo(Encoder);
+#endif // DEBUG_OPUS
+
+			return EncError == OPUS_OK;
+		}
+	}
+
+	return false;
+}
+
+void FVoiceEncoderOpus::Reset()
+{
+	int32 EncError = 0;
+	if (Encoder)
+	{
+		EncError = opus_encoder_ctl(Encoder, OPUS_RESET_STATE);
+		if (EncError != OPUS_OK)
+		{
+			UE_LOG(LogVoiceEncode, Warning, TEXT("Failure to reset Opus encoder"));
+		}
+
+#if DEBUG_OPUS
+		DebugEncoderInfo(Encoder);
+#endif // DEBUG_OPUS
+	}
+}
+
 void FVoiceEncoderOpus::Destroy()
 {
 #if USE_UE4_MEM_ALLOC
@@ -309,14 +464,26 @@ void FVoiceEncoderOpus::Destroy()
 #else
 	opus_encoder_destroy(Encoder);
 #endif
-	Encoder = NULL;
+	Encoder = nullptr;
+}
+
+void FVoiceEncoderOpus::DumpState() const
+{
+	if (Encoder)
+	{
+		DebugEncoderInfo(Encoder);
+	}
+	else
+	{
+		UE_LOG(LogVoiceEncode, Display, TEXT("No encoder to dump state"));
+	}
 }
 
 FVoiceDecoderOpus::FVoiceDecoderOpus() :
 	SampleRate(0),
 	NumChannels(0),
 	FrameSize(0),
-	Decoder(NULL),
+	Decoder(nullptr),
 	LastEntropyIdx(0),
 	LastGeneration(0)
 {
@@ -332,22 +499,44 @@ bool FVoiceDecoderOpus::Init(int32 InSampleRate, int32 InNumChannels)
 {
 	UE_LOG(LogVoiceDecode, Display, TEXT("DecoderVersion: %s"), ANSI_TO_TCHAR(opus_get_version_string()));
 
+	if (InSampleRate != 8000 &&
+		InSampleRate != 12000 &&
+		InSampleRate != 16000 &&
+		InSampleRate != 24000 &&
+		InSampleRate != 48000)
+	{
+		UE_LOG(LogVoiceDecode, Warning, TEXT("Voice decoder doesn't support %d hz"), InSampleRate);
+		return false;
+	}
+
+	if (InNumChannels < 1 || InNumChannels > 2)
+	{
+		UE_LOG(LogVoiceDecode, Warning, TEXT("Voice decoder only supports 1 or 2 channels"));
+		return false;
+	}
+
 	SampleRate = InSampleRate;
 	NumChannels = InNumChannels;
 
-	FrameSize = SampleRate / 50;
+	// 20ms frame sizes are a good choice for most applications (1000ms / 20ms = 50)
+	FrameSize = SampleRate / NUM_OPUS_FRAMES_PER_SEC;
 
-	int32 DecSize = 0;
 	int32 DecError = 0;
 
 #if USE_UE4_MEM_ALLOC
-	DecSize = opus_decoder_get_size(NumChannels);
+	const int32 DecSize = opus_decoder_get_size(NumChannels);
 	Decoder = (OpusDecoder*)FMemory::Malloc(DecSize);
 	DecError = opus_decoder_init(Decoder, SampleRate, NumChannels);
 #else
 	Decoder = opus_decoder_create(SampleRate, NumChannels, &DecError);
 #endif
-	if (DecError != OPUS_OK)
+	if (DecError == OPUS_OK)
+	{
+#if DEBUG_OPUS
+		DebugDecoderInfo(Decoder);
+#endif // DEBUG_OPUS
+	}
+	else
 	{
 		UE_LOG(LogVoiceDecode, Warning, TEXT("Failed to init Opus Decoder: %s"), ANSI_TO_TCHAR(opus_strerror(DecError)));
 		Destroy();
@@ -356,6 +545,23 @@ bool FVoiceDecoderOpus::Init(int32 InSampleRate, int32 InNumChannels)
 	return DecError == OPUS_OK;
 }
 
+void FVoiceDecoderOpus::Reset()
+{
+	int32 DecError = 0;
+	if (Decoder)
+	{
+		DecError = opus_decoder_ctl(Decoder, OPUS_RESET_STATE);
+		if (DecError != OPUS_OK)
+		{
+			UE_LOG(LogVoiceDecode, Warning, TEXT("Failure to reset Opus decoder"));
+		}
+
+#if DEBUG_OPUS
+		DebugDecoderInfo(Decoder);
+#endif // DEBUG_OPUS
+	}
+}
+
 void FVoiceDecoderOpus::Destroy()
 {
 #if USE_UE4_MEM_ALLOC
@@ -363,102 +569,157 @@ void FVoiceDecoderOpus::Destroy()
 #else
 	opus_decoder_destroy(Decoder);
 #endif 
-	Decoder = NULL;
+	Decoder = nullptr;
+}
+
+inline bool SanityCheckHeader(uint32 HeaderSize, uint32 CompressedDataSize, int32 NumFramesToDecode, const uint16* CompressedOffsets)
+{
+	bool bHeaderDataOk = (HeaderSize <= CompressedDataSize);
+	if (bHeaderDataOk)
+	{
+		// Validate that the sum of the encoded data sizes fit under the given amount of compressed data
+		uint16 LastCompressedOffset = 0;
+		int32 TotalCompressedBufferSize = 0;
+		for (int32 Idx = 0; Idx < NumFramesToDecode; Idx++)
+		{
+			// Offsets should be monotonically increasing (prevent later values intentionally reducing bad previous values)
+			if (CompressedOffsets[Idx] >= LastCompressedOffset)
+			{
+				TotalCompressedBufferSize += (CompressedOffsets[Idx] - LastCompressedOffset);
+				LastCompressedOffset = CompressedOffsets[Idx];
+			}
+			else
+			{
+				bHeaderDataOk = false;
+				break;
+			}
+		}
+
+		bHeaderDataOk = bHeaderDataOk && ((HeaderSize + TotalCompressedBufferSize) <= CompressedDataSize);
+	}
+
+	return bHeaderDataOk;
 }
 
 void FVoiceDecoderOpus::Decode(const uint8* InCompressedData, uint32 CompressedDataSize, uint8* OutRawPCMData, uint32& OutRawDataSize)
 {
 	SCOPE_CYCLE_COUNTER(STAT_Voice_Decoding);
+	SET_DWORD_STAT(STAT_Decode_SampleRate, SampleRate);
+	SET_DWORD_STAT(STAT_Decode_NumChannels, NumChannels);
 
-	int32 HeaderSize = 0;
-	int32 BytesPerFrame = FrameSize * NumChannels * sizeof(opus_int16);
-	int32 MaxFramesEncoded = MAX_OPUS_UNCOMPRESSED_BUFFER_SIZE / BytesPerFrame;
+	uint32 HeaderSize = (2 * sizeof(uint8));
+	if (!InCompressedData || (CompressedDataSize < HeaderSize))
+	{
+		OutRawDataSize = 0;
+		return;
+	}
 
-	int32 NumFramesToDecode = InCompressedData[0];
-	int32 PacketGeneration = InCompressedData[1];
-	HeaderSize += 2 * sizeof(uint8);
+	const int32 BytesPerFrame = FrameSize * NumChannels * sizeof(opus_int16);
+	const int32 MaxFramesEncoded = MAX_OPUS_UNCOMPRESSED_BUFFER_SIZE / BytesPerFrame;
+
+	const int32 NumFramesToDecode = InCompressedData[0];
+	const int32 PacketGeneration = InCompressedData[1];
 
 	if (PacketGeneration != LastGeneration + 1)
 	{
 		UE_LOG(LogVoiceDecode, Verbose, TEXT("Packet generation skipped from %d to %d"), LastGeneration, PacketGeneration);
 	}
 
-	if (NumFramesToDecode > 0 && NumFramesToDecode <= MaxFramesEncoded)
+	if ((NumFramesToDecode > 0) && (NumFramesToDecode <= MaxFramesEncoded))
 	{
 		// Start of compressed data offsets
 		const uint16* CompressedOffsets = (const uint16*)(InCompressedData + HeaderSize);
 		uint32 LengthOfCompressedOffsets = NumFramesToDecode * sizeof(uint16);
 		HeaderSize += LengthOfCompressedOffsets;
 
-		uint32 LengthOfEntropyOffsets = 0;
 #if ADD_ENTROPY_TO_PACKET
 		// Start of the entropy to each encoded frame
 		const uint32* EntropyOffsets = (uint32*)(InCompressedData + HeaderSize);
-		LengthOfEntropyOffsets = NumFramesToDecode * sizeof(uint32);
+		uint32 LengthOfEntropyOffsets = NumFramesToDecode * sizeof(uint32);
 		HeaderSize += LengthOfEntropyOffsets;
 #endif
 
-		// Start of compressed data
-		const uint8* CompressedDataStart = (InCompressedData + HeaderSize);
-		
-		int32 CompressedBufferOffset = 0;
-		int32 DecompressedBufferOffset = 0;
-		uint16 LastCompressedOffset = 0;
-
-		for (int32 i=0; i < NumFramesToDecode; i++)
+		// At this point we have all our pointer fix up complete, but the data it references may be invalid in corrupt/spoofed packets
+		// Sanity check the numbers to make sure everything works out
+		if (SanityCheckHeader(HeaderSize, CompressedDataSize, NumFramesToDecode, CompressedOffsets))
 		{
-			int32 UncompressedBufferAvail = OutRawDataSize - DecompressedBufferOffset;
+			// Start of compressed data
+			const uint8* CompressedDataStart = (InCompressedData + HeaderSize);
 
-			if (UncompressedBufferAvail >= (MAX_OPUS_FRAMES * BytesPerFrame))
+			int32 CompressedBufferOffset = 0;
+			int32 DecompressedBufferOffset = 0;
+			uint16 LastCompressedOffset = 0;
+
+			for (int32 i = 0; i < NumFramesToDecode; i++)
 			{
-				int32 CompressedBufferSize = CompressedOffsets[i] - LastCompressedOffset;
+				const int32 UncompressedBufferAvail = (OutRawDataSize - DecompressedBufferOffset);
+
+				if (UncompressedBufferAvail >= (MAX_OPUS_FRAMES * BytesPerFrame))
+				{
+					if (CompressedOffsets[i] > 0)
+					{
+						const int32 CompressedBufferSize = (CompressedOffsets[i] - LastCompressedOffset);
 
-				check(Decoder);
-				int32 NumDecompressedSamples = opus_decode(Decoder, 
-					CompressedDataStart + CompressedBufferOffset, CompressedBufferSize, 
-					(opus_int16*)(OutRawPCMData + DecompressedBufferOffset), MAX_OPUS_FRAME_SIZE, 0);
+						check(Decoder);
+						const int32 NumDecompressedSamples = opus_decode(Decoder,
+							CompressedDataStart + CompressedBufferOffset, CompressedBufferSize,
+							(opus_int16*)(OutRawPCMData + DecompressedBufferOffset), MAX_OPUS_FRAME_SIZE, 0);
 
 #if DEBUG_OPUS
-				DebugFrameDecodeInfo(CompressedDataStart + CompressedBufferOffset, CompressedBufferSize, SampleRate);
+						DebugFrameDecodeInfo(CompressedDataStart + CompressedBufferOffset, CompressedBufferSize, SampleRate);
 #endif // DEBUG_OPUS
 
-				if (NumDecompressedSamples < 0)
-				{
-					const char* ErrorStr = opus_strerror(NumDecompressedSamples);
-					UE_LOG(LogVoiceDecode, Warning, TEXT("Failed to decode: [%d] %s"), NumDecompressedSamples, ANSI_TO_TCHAR(ErrorStr));
-					break;
-				}
-				else
-				{
-					if (NumDecompressedSamples != FrameSize)
-					{
-						UE_LOG(LogVoiceDecode, Warning, TEXT("Unexpected decode result NumSamplesDecoded %d != FrameSize %d"), NumDecompressedSamples, FrameSize);
-					}
+						if (NumDecompressedSamples < 0)
+						{
+							const char* ErrorStr = opus_strerror(NumDecompressedSamples);
+							UE_LOG(LogVoiceDecode, Warning, TEXT("Failed to decode: [%d] %s"), NumDecompressedSamples, ANSI_TO_TCHAR(ErrorStr));
+						}
+						else
+						{
+							if (NumDecompressedSamples != FrameSize)
+							{
+								UE_LOG(LogVoiceDecode, Warning, TEXT("Unexpected decode result NumSamplesDecoded %d != FrameSize %d"), NumDecompressedSamples, FrameSize);
+							}
 
-					opus_decoder_ctl(Decoder, OPUS_GET_FINAL_RANGE(&Entropy[LastEntropyIdx]));
+							opus_decoder_ctl(Decoder, OPUS_GET_FINAL_RANGE(&Entropy[LastEntropyIdx]));
 
 #if ADD_ENTROPY_TO_PACKET
-					if (Entropy[LastEntropyIdx] != EntropyOffsets[i])
-					{
-						UE_LOG(LogVoiceDecode, Verbose, TEXT("Decoder Entropy[%d/%d] = %d expected %d"), i, NumFramesToDecode-1, Entropy[LastEntropyIdx], EntropyOffsets[i]);
-					}
+							if (Entropy[LastEntropyIdx] != EntropyOffsets[i])
+							{
+								UE_LOG(LogVoiceDecode, Verbose, TEXT("Decoder Entropy[%d/%d] = %d expected %d"), i, NumFramesToDecode - 1, Entropy[LastEntropyIdx], EntropyOffsets[i]);
+							}
 #endif
 
-					LastEntropyIdx = (LastEntropyIdx + 1) % NUM_ENTROPY_VALUES;
+							LastEntropyIdx = (LastEntropyIdx + 1) % NUM_ENTROPY_VALUES;
+
+							// Advance within the decompressed output stream
+							DecompressedBufferOffset += (NumDecompressedSamples * NumChannels * sizeof(opus_int16));
+						}
 
-					DecompressedBufferOffset += NumDecompressedSamples * sizeof(opus_int16);
-					CompressedBufferOffset += CompressedBufferSize;
-					LastCompressedOffset = CompressedOffsets[i];
+						// Advance within the compressed input stream
+						CompressedBufferOffset += CompressedBufferSize;
+						LastCompressedOffset = CompressedOffsets[i];
+					}
+					else
+					{
+						UE_LOG(LogVoiceDecode, Verbose, TEXT("Decompression buffer skipped a frame"));
+						// Nothing to advance within the compressed input stream
+					}
+				}
+				else
+				{
+					UE_LOG(LogVoiceDecode, Warning, TEXT("Decompression buffer too small to decode voice"));
+					break;
 				}
 			}
-			else
-			{
-				UE_LOG(LogVoiceDecode, Warning, TEXT("Decompression buffer too small to decode voice"));
-				break;
-			}
-		}
 
-		OutRawDataSize = DecompressedBufferOffset;
+			OutRawDataSize = DecompressedBufferOffset;
+		}
+		else
+		{
+			UE_LOG(LogVoiceDecode, Warning, TEXT("Failed to decode: header corrupted"));
+			OutRawDataSize = 0;
+		}
 	}
 	else
 	{
@@ -466,8 +727,23 @@ void FVoiceDecoderOpus::Decode(const uint8* InCompressedData, uint32 CompressedD
 		OutRawDataSize = 0;
 	}
 
-	UE_LOG(LogVoiceDecode, VeryVerbose, TEXT("OpusDecode[%d]: RawSize: %d CompressedSize: %d NumFramesEncoded: %d "), PacketGeneration, OutRawDataSize, CompressedDataSize, NumFramesToDecode);
+	UE_LOG(LogVoiceDecode, VeryVerbose, TEXT("OpusDecode[%d]: RawSize: %d HeaderSize: %d CompressedSize: %d NumFramesDecoded: %d "), PacketGeneration, OutRawDataSize, HeaderSize, CompressedDataSize, NumFramesToDecode);
+	SET_FLOAT_STAT(STAT_Decode_CompressionRatio, (float)CompressedDataSize / (float)OutRawDataSize);
+	SET_DWORD_STAT(STAT_Decode_OutSize, OutRawDataSize);
+
 	LastGeneration = PacketGeneration;
 }
 
+void FVoiceDecoderOpus::DumpState() const
+{
+	if (Decoder)
+	{
+		DebugDecoderInfo(Decoder);
+	}
+	else
+	{
+		UE_LOG(LogVoiceDecode, Display, TEXT("No decoder to dump state"));
+	}
+}
+
 #endif // PLATFORM_SUPPORTS_VOICE_CAPTURE
diff --git a/Engine/Source/Runtime/Online/Voice/Private/VoiceCodecOpus.h b/Engine/Source/Runtime/Online/Voice/Private/VoiceCodecOpus.h
--- a/Engine/Source/Runtime/Online/Voice/Private/VoiceCodecOpus.h
+++ b/Engine/Source/Runtime/Online/Voice/Private/VoiceCodecOpus.h
@@ -23,9 +23,14 @@ public:
 	~FVoiceEncoderOpus();
 
 	// IVoiceEncoder
-	virtual bool Init(int32 SampleRate, int32 NumChannels) override;
+	virtual bool Init(int32 SampleRate, int32 NumChannels, EAudioEncodeHint EncodeHint) override;
 	virtual int32 Encode(const uint8* RawPCMData, uint32 RawDataSize, uint8* OutCompressedData, uint32& OutCompressedDataSize) override;
+	virtual bool SetBitrate(int32 InBitRate) override;
+	virtual bool SetVBR(bool bEnableVBR) override;
+	virtual bool SetComplexity(int32 InComplexity) override;
+	virtual void Reset() override;
 	virtual void Destroy() override;
+	virtual void DumpState() const override;
 
 private:
 
@@ -64,7 +69,9 @@ public:
 	//IVoiceDecoder
 	virtual bool Init(int32 SampleRate, int32 NumChannels) override;
 	virtual void Decode(const uint8* CompressedData, uint32 CompressedDataSize, uint8* OutRawPCMData, uint32& OutRawDataSize) override;
+	virtual void Reset() override;
 	virtual void Destroy() override;
+	virtual void DumpState() const override;
 
 private:
 
diff --git a/Engine/Source/Runtime/Online/Voice/Private/VoiceModule.cpp b/Engine/Source/Runtime/Online/Voice/Private/VoiceModule.cpp
--- a/Engine/Source/Runtime/Online/Voice/Private/VoiceModule.cpp
+++ b/Engine/Source/Runtime/Online/Voice/Private/VoiceModule.cpp
@@ -13,14 +13,23 @@ DEFINE_LOG_CATEGORY(LogVoiceCapture);
 
 DEFINE_STAT(STAT_Voice_Encoding);
 DEFINE_STAT(STAT_Voice_Decoding);
+DEFINE_STAT(STAT_Encode_SampleRate);
+DEFINE_STAT(STAT_Encode_NumChannels);
+DEFINE_STAT(STAT_Encode_Bitrate);
+DEFINE_STAT(STAT_Encode_CompressionRatio);
+DEFINE_STAT(STAT_Encode_OutSize);
+DEFINE_STAT(STAT_Decode_SampleRate);
+DEFINE_STAT(STAT_Decode_NumChannels);
+DEFINE_STAT(STAT_Decode_CompressionRatio);
+DEFINE_STAT(STAT_Decode_OutSize);
 
 #if PLATFORM_SUPPORTS_VOICE_CAPTURE
 /** Implement these functions per platform to create the voice objects */
 extern bool InitVoiceCapture();
 extern void ShutdownVoiceCapture();
-extern IVoiceCapture* CreateVoiceCaptureObject();
-extern IVoiceEncoder* CreateVoiceEncoderObject();
-extern IVoiceDecoder* CreateVoiceDecoderObject();
+extern IVoiceCapture* CreateVoiceCaptureObject(const FString& DeviceName, int32 SampleRate, int32 NumChannels);
+extern IVoiceEncoder* CreateVoiceEncoderObject(int32 SampleRate, int32 NumChannels, EAudioEncodeHint EncodeHint);
+extern IVoiceDecoder* CreateVoiceDecoderObject(int32 SampleRate, int32 NumChannels);
 #endif
 
 void FVoiceModule::StartupModule()
@@ -64,42 +73,42 @@ bool FVoiceModule::Exec(UWorld* InWorld, const TCHAR* Cmd, FOutputDevice& Ar)
 	return false;
 }
 
-TSharedPtr<IVoiceCapture> FVoiceModule::CreateVoiceCapture()
+TSharedPtr<IVoiceCapture> FVoiceModule::CreateVoiceCapture(const FString& DeviceName, int32 SampleRate, int32 NumChannels)
 {
 #if PLATFORM_SUPPORTS_VOICE_CAPTURE
 	if (bEnabled)
 	{
 		// Create the platform specific instance
-		return TSharedPtr<IVoiceCapture>(CreateVoiceCaptureObject());
+		return TSharedPtr<IVoiceCapture>(CreateVoiceCaptureObject(DeviceName, SampleRate, NumChannels));
 	}
 #endif
-	return NULL;
+	return nullptr;
 }
 
-TSharedPtr<IVoiceEncoder> FVoiceModule::CreateVoiceEncoder()
+TSharedPtr<IVoiceEncoder> FVoiceModule::CreateVoiceEncoder(int32 SampleRate, int32 NumChannels, EAudioEncodeHint EncodeHint)
 {
 #if PLATFORM_SUPPORTS_VOICE_CAPTURE
 	if (bEnabled)
 	{
 		// Create the platform specific instance
-		return TSharedPtr<IVoiceEncoder>(CreateVoiceEncoderObject());
+		return TSharedPtr<IVoiceEncoder>(CreateVoiceEncoderObject(SampleRate, NumChannels, EncodeHint));
 	}
 #endif
 
-	return NULL;
+	return nullptr;
 }
 
-TSharedPtr<IVoiceDecoder> FVoiceModule::CreateVoiceDecoder()
+TSharedPtr<IVoiceDecoder> FVoiceModule::CreateVoiceDecoder(int32 SampleRate, int32 NumChannels)
 {
 #if PLATFORM_SUPPORTS_VOICE_CAPTURE
 	if (bEnabled)
 	{
 		// Create the platform specific instance
-		return TSharedPtr<IVoiceDecoder>(CreateVoiceDecoderObject());
+		return TSharedPtr<IVoiceDecoder>(CreateVoiceDecoderObject(SampleRate, NumChannels));
 	}
 #endif
 
-	return NULL;
+	return nullptr;
 }
 
 
diff --git a/Engine/Source/Runtime/Online/Voice/Private/Windows/VoiceCaptureWindows.cpp b/Engine/Source/Runtime/Online/Voice/Private/Windows/VoiceCaptureWindows.cpp
--- a/Engine/Source/Runtime/Online/Voice/Private/Windows/VoiceCaptureWindows.cpp
+++ b/Engine/Source/Runtime/Online/Voice/Private/Windows/VoiceCaptureWindows.cpp
@@ -4,10 +4,14 @@
 #include "VoicePrivate.h"
 #include "VoiceModule.h"
 
+#if PLATFORM_SUPPORTS_VOICE_CAPTURE
+
 #include "AllowWindowsPlatformTypes.h"
 
 struct FVoiceCaptureWindowsVars
 {
+	/** GUID of current voice capture device */
+	GUID VoiceCaptureDeviceGuid;
 	/** Voice capture device */
 	LPDIRECTSOUNDCAPTURE8 VoiceCaptureDev;
 	/** Voice capture device caps */
@@ -21,26 +25,48 @@ struct FVoiceCaptureWindowsVars
 	/** Buffer caps */
 	DSCBCAPS VoiceCaptureBufferCaps8;
 	/** Notification events */
-	HANDLE Events[NUM_EVENTS];
+	HANDLE StopEvent;
 	/** Current audio position of valid data in capture buffer */
 	DWORD NextCaptureOffset;
 
 	FVoiceCaptureWindowsVars() :
-		VoiceCaptureDev(NULL),
-		VoiceCaptureBuffer8(NULL),
+		VoiceCaptureDev(nullptr),
+		VoiceCaptureBuffer8(nullptr),
 		NextCaptureOffset(0)
 	{
+		StopEvent = INVALID_HANDLE_VALUE;
+		Reset();
+	}
+
+	void Reset()
+	{
+		if (StopEvent != INVALID_HANDLE_VALUE)
+		{
+			CloseHandle(StopEvent);
+			StopEvent = INVALID_HANDLE_VALUE;
+		}
+
+		// Free up DirectSound resources
+		if (VoiceCaptureBuffer8)
+		{
+			VoiceCaptureBuffer8->Release();
+			VoiceCaptureBuffer8 = nullptr;
+		}
+
+		if (VoiceCaptureDev)
+		{
+			VoiceCaptureDev->Release();
+			VoiceCaptureDev = nullptr;
+		}
+
+		NextCaptureOffset = 0;
+
+		FMemory::Memzero(&VoiceCaptureDeviceGuid, sizeof(GUID));
 		FMemory::Memzero(&VoiceCaptureDevCaps, sizeof(VoiceCaptureDevCaps));
 
 		FMemory::Memzero(&WavFormat, sizeof(WavFormat));
 		FMemory::Memzero(&VoiceCaptureBufferDesc, sizeof(VoiceCaptureBufferDesc));
 		FMemory::Memzero(&VoiceCaptureBufferCaps8, sizeof(VoiceCaptureBufferCaps8));
-		FMemory::Memzero(&Events, ARRAY_COUNT(Events) * sizeof(HANDLE));
-
-		for (int i = 0; i < NUM_EVENTS; i++)
-		{
-			Events[i] = INVALID_HANDLE_VALUE;
-		}
 	}
 };
 
@@ -78,12 +104,11 @@ bool IsSilence(T* Buffer, int32 BuffSize)
 }
 
 FVoiceCaptureWindows::FVoiceCaptureWindows() :
-	CV(NULL),
-	LastEventTriggered(STOP_EVENT - 1),
+	CV(nullptr),
+	LastCaptureTime(0.0),
 	VoiceCaptureState(EVoiceCaptureState::UnInitialized)
 {
 	CV = new FVoiceCaptureWindowsVars();
-	UncompressedAudioBuffer.Empty(MAX_UNCOMPRESSED_VOICE_BUFFER_SIZE);
 }
 
 FVoiceCaptureWindows::~FVoiceCaptureWindows()
@@ -99,12 +124,33 @@ FVoiceCaptureWindows::~FVoiceCaptureWindows()
 	if (CV)
 	{
 		delete CV;
-		CV = NULL;
+		CV = nullptr;
 	}	
 }
 
-bool FVoiceCaptureWindows::Init(int32 SampleRate, int32 NumChannels)
+bool FVoiceCaptureWindows::Init(const FString& DeviceName, int32 SampleRate, int32 NumChannels)
 {
+	FVoiceCaptureDeviceWindows* VoiceDev = FVoiceCaptureDeviceWindows::Get();
+	if (!VoiceDev)
+	{
+		UE_LOG(LogVoiceCapture, Warning, TEXT("No voice capture interface."));
+		return false;
+	}
+	
+	if (DeviceName.IsEmpty())
+	{
+		return CreateCaptureBuffer(VoiceDev->DefaultVoiceCaptureDevice.DeviceName, SampleRate, NumChannels);
+	}
+
+	return CreateCaptureBuffer(DeviceName, SampleRate, NumChannels);
+}
+
+bool FVoiceCaptureWindows::CreateCaptureBuffer(const FString& DeviceName, int32 SampleRate, int32 NumChannels)
+{
+	// Free the previous buffer
+	FreeCaptureBuffer();
+	VoiceCaptureState = EVoiceCaptureState::NotCapturing;
+
 	if (SampleRate < 8000 || SampleRate > 48000)
 	{
 		UE_LOG(LogVoiceCapture, Warning, TEXT("Voice capture doesn't support %d hz"), SampleRate);
@@ -124,125 +170,144 @@ bool FVoiceCaptureWindows::Init(int32 SampleRate, int32 NumChannels)
 		return false;
 	}
 
-	// DSDEVID_DefaultCapture WAVEINCAPS 
-	HRESULT hr = DirectSoundCaptureCreate8(&VoiceDev->VoiceCaptureDeviceGuid, &CV->VoiceCaptureDev, NULL);
-	if (FAILED(hr))
+	FVoiceCaptureDeviceWindows::FCaptureDeviceInfo* DeviceInfo = nullptr;
+	if (DeviceName.IsEmpty())
 	{
-		//DSERR_ALLOCATED, DSERR_INVALIDPARAM, DSERR_NOAGGREGATION, DSERR_OUTOFMEMORY
-		UE_LOG(LogVoiceCapture, Warning, TEXT("Failed to create capture device 0x%08x"), hr);
-		return false;
+		DeviceInfo = VoiceDev->Devices.Find(DEFAULT_DEVICE_NAME);
 	}
-
-	// Device capabilities
-	CV->VoiceCaptureDevCaps.dwSize = sizeof(DSCCAPS);
-	hr = CV->VoiceCaptureDev->GetCaps(&CV->VoiceCaptureDevCaps);
-	if (FAILED(hr))
+	else
 	{
-		UE_LOG(LogVoiceCapture, Warning, TEXT("Failed to get mic device caps 0x%08x"), hr);
-		return false;
+		DeviceInfo = VoiceDev->Devices.Find(DeviceName);
 	}
 
-	// Wave format setup
-	CV->WavFormat.wFormatTag = WAVE_FORMAT_PCM;
-	CV->WavFormat.nChannels = NumChannels;
-	CV->WavFormat.wBitsPerSample = 16;
-	CV->WavFormat.nSamplesPerSec = SampleRate;
-	CV->WavFormat.nBlockAlign = (CV->WavFormat.nChannels * CV->WavFormat.wBitsPerSample) / 8;
-	CV->WavFormat.nAvgBytesPerSec = CV->WavFormat.nBlockAlign * CV->WavFormat.nSamplesPerSec;
-	CV->WavFormat.cbSize = 0;
-
-	// Buffer setup
-	CV->VoiceCaptureBufferDesc.dwSize = sizeof(DSCBUFFERDESC);
-	CV->VoiceCaptureBufferDesc.dwFlags = 0;
-	CV->VoiceCaptureBufferDesc.dwBufferBytes = CV->WavFormat.nAvgBytesPerSec / 2; // .5 sec buffer
-	CV->VoiceCaptureBufferDesc.dwReserved = 0;
-	CV->VoiceCaptureBufferDesc.lpwfxFormat = &CV->WavFormat;
-	CV->VoiceCaptureBufferDesc.dwFXCount = 0;
-	CV->VoiceCaptureBufferDesc.lpDSCFXDesc = NULL;
-
-	LPDIRECTSOUNDCAPTUREBUFFER VoiceBuffer = NULL;
-
-	hr = CV->VoiceCaptureDev->CreateCaptureBuffer(&CV->VoiceCaptureBufferDesc, &VoiceBuffer, NULL);
-	if (FAILED(hr))
+	if (DeviceInfo)
 	{
-		UE_LOG(LogVoiceCapture, Warning, TEXT("Failed to create voice capture buffer 0x%08x"), hr);
-		return false;
-	}
+		UE_LOG(LogVoiceCapture, Display, TEXT("Creating capture %s [%d:%d]"), *DeviceInfo->DeviceName, SampleRate, NumChannels);
+		CV->VoiceCaptureDeviceGuid = DeviceInfo->DeviceId;
 
-	hr = VoiceBuffer->QueryInterface(IID_IDirectSoundCaptureBuffer8, (LPVOID*)&CV->VoiceCaptureBuffer8);
-	VoiceBuffer->Release(); 
-	VoiceBuffer = NULL;
-	if (FAILED(hr))
-	{
-		UE_LOG(LogVoiceCapture, Warning, TEXT("Failed to create voice capture buffer 0x%08x"), hr);
-		return false;
-	}
+		// DSDEVID_DefaultCapture WAVEINCAPS 
+		HRESULT hr = DirectSoundCaptureCreate8(&DeviceInfo->DeviceId, &CV->VoiceCaptureDev, nullptr);
+		if (FAILED(hr))
+		{
+			//DSERR_ALLOCATED, DSERR_INVALIDPARAM, DSERR_NOAGGREGATION, DSERR_OUTOFMEMORY
+			UE_LOG(LogVoiceCapture, Warning, TEXT("Failed to create capture device 0x%08x"), hr);
+			return false;
+		}
 
-	CV->VoiceCaptureBufferCaps8.dwSize = sizeof(DSCBCAPS);
-	hr = CV->VoiceCaptureBuffer8->GetCaps(&CV->VoiceCaptureBufferCaps8);
-	if (FAILED(hr))
-	{
-		UE_LOG(LogVoiceCapture, Warning, TEXT("Failed to get voice buffer caps 0x%08x"), hr);
-		return false;
-	}
+		// Device capabilities
+		CV->VoiceCaptureDevCaps.dwSize = sizeof(DSCCAPS);
+		hr = CV->VoiceCaptureDev->GetCaps(&CV->VoiceCaptureDevCaps);
+		if (FAILED(hr))
+		{
+			UE_LOG(LogVoiceCapture, Warning, TEXT("Failed to get mic device caps 0x%08x"), hr);
+			return false;
+		}
 
-	// TEST ------------------------
-	if (0)
-	{
-		DWORD SizeWritten8 = 0;
-		CV->VoiceCaptureBuffer8->GetFormat(NULL, sizeof(WAVEFORMATEX), &SizeWritten8);
+		// Wave format setup
+		CV->WavFormat.wFormatTag = WAVE_FORMAT_PCM;
+		CV->WavFormat.nChannels = NumChannels;
+		CV->WavFormat.wBitsPerSample = 16;
+		CV->WavFormat.nSamplesPerSec = SampleRate;
+		CV->WavFormat.nBlockAlign = (CV->WavFormat.nChannels * CV->WavFormat.wBitsPerSample) / 8;
+		CV->WavFormat.nAvgBytesPerSec = CV->WavFormat.nBlockAlign * CV->WavFormat.nSamplesPerSec;
+		CV->WavFormat.cbSize = 0;
+
+		// Buffer setup
+		CV->VoiceCaptureBufferDesc.dwSize = sizeof(DSCBUFFERDESC);
+		CV->VoiceCaptureBufferDesc.dwFlags = 0;
+		CV->VoiceCaptureBufferDesc.dwBufferBytes = CV->WavFormat.nAvgBytesPerSec / 2; // .5 sec buffer
+		CV->VoiceCaptureBufferDesc.dwReserved = 0;
+		CV->VoiceCaptureBufferDesc.lpwfxFormat = &CV->WavFormat;
+		CV->VoiceCaptureBufferDesc.dwFXCount = 0;
+		CV->VoiceCaptureBufferDesc.lpDSCFXDesc = nullptr;
+
+		LPDIRECTSOUNDCAPTUREBUFFER VoiceBuffer = nullptr;
+
+		hr = CV->VoiceCaptureDev->CreateCaptureBuffer(&CV->VoiceCaptureBufferDesc, &VoiceBuffer, nullptr);
+		if (FAILED(hr))
+		{
+			UE_LOG(LogVoiceCapture, Warning, TEXT("Failed to create voice capture buffer 0x%08x"), hr);
+			return false;
+		}
 
-		LPWAVEFORMATEX BufferFormat8 = (WAVEFORMATEX*)FMemory::Malloc(SizeWritten8);
-		CV->VoiceCaptureBuffer8->GetFormat(BufferFormat8, SizeWritten8, &SizeWritten8);
-		FMemory::Free(BufferFormat8);
-	}
-	// TEST ------------------------
+		hr = VoiceBuffer->QueryInterface(IID_IDirectSoundCaptureBuffer8, (LPVOID*)&CV->VoiceCaptureBuffer8);
+		VoiceBuffer->Release();
+		VoiceBuffer = nullptr;
+		if (FAILED(hr))
+		{
+			UE_LOG(LogVoiceCapture, Warning, TEXT("Failed to create voice capture buffer 0x%08x"), hr);
+			return false;
+		}
 
-	if (!CreateNotifications(CV->VoiceCaptureBufferCaps8.dwBufferBytes))
-	{
-		UE_LOG(LogVoiceCapture, Warning, TEXT("Failed to create voice buffer notifications"));
-		return false;
-	}
+		CV->VoiceCaptureBufferCaps8.dwSize = sizeof(DSCBCAPS);
+		hr = CV->VoiceCaptureBuffer8->GetCaps(&CV->VoiceCaptureBufferCaps8);
+		if (FAILED(hr))
+		{
+			UE_LOG(LogVoiceCapture, Warning, TEXT("Failed to get voice buffer caps 0x%08x"), hr);
+			return false;
+		}
 
-	VoiceCaptureState = EVoiceCaptureState::NotCapturing;
-	return true;
-}
+		// TEST ------------------------
+		if (0)
+		{
+			DWORD SizeWritten8 = 0;
+			CV->VoiceCaptureBuffer8->GetFormat(nullptr, sizeof(WAVEFORMATEX), &SizeWritten8);
 
-void FVoiceCaptureWindows::Shutdown()
-{
-	Stop();
+			LPWAVEFORMATEX BufferFormat8 = (WAVEFORMATEX*)FMemory::Malloc(SizeWritten8);
+			CV->VoiceCaptureBuffer8->GetFormat(BufferFormat8, SizeWritten8, &SizeWritten8);
+			FMemory::Free(BufferFormat8);
+		}
+		// TEST ------------------------
 
-	for (int i = 0; i < ARRAY_COUNT(CV->Events); i++)
-	{
-		if (CV->Events[i] != INVALID_HANDLE_VALUE)
+		if (CreateNotifications(CV->VoiceCaptureBufferCaps8.dwBufferBytes))
 		{
-			CloseHandle(CV->Events[i]);
-			CV->Events[i] = INVALID_HANDLE_VALUE;
+			// Reset notification related values
+			LastCaptureTime = FPlatformTime::Seconds();
+		}
+		else
+		{
+			UE_LOG(LogVoiceCapture, Warning, TEXT("Failed to create voice buffer notifications"));
+			return false;
 		}
 	}
-
-	// Free up DirectSound resources
-	if (CV->VoiceCaptureBuffer8)
+	else
 	{
-		CV->VoiceCaptureBuffer8->Release();
-		CV->VoiceCaptureBuffer8 = NULL;
+		UE_LOG(LogVoiceCapture, Warning, TEXT("No voice capture device %s found."), *DeviceName);
+		return false;
 	}
 
-	if (CV->VoiceCaptureDev)
-	{
-		CV->VoiceCaptureDev->Release();
-		CV->VoiceCaptureDev = NULL;
-	}
+	UncompressedAudioBuffer.Empty(CV->VoiceCaptureBufferDesc.dwBufferBytes);
+	check(UncompressedAudioBuffer.Max() >= (int32)CV->VoiceCaptureBufferCaps8.dwBufferBytes);
+
+	return true;
+}
+
+void FVoiceCaptureWindows::FreeCaptureBuffer()
+{
+	// Stop playback
+	Stop();
 
-	FMemory::Memzero(&CV->VoiceCaptureDevCaps, sizeof(CV->VoiceCaptureDevCaps));
+	// Release all D3D8 resources
+	CV->Reset();
 
 	VoiceCaptureState = EVoiceCaptureState::UnInitialized;
 }
 
+void FVoiceCaptureWindows::Shutdown()
+{
+	FreeCaptureBuffer();
+}
+
 bool FVoiceCaptureWindows::Start()
 {
 	check(VoiceCaptureState != EVoiceCaptureState::UnInitialized);
 
+	if (CV->VoiceCaptureBuffer8 == nullptr)
+	{
+		UE_LOG(LogVoiceCapture, Warning, TEXT("CV->VoiceCaptureBuffer8 == nullptr"));
+		return false;
+	}
+
 	HRESULT hr = CV->VoiceCaptureBuffer8->Start(DSCBSTART_LOOPING);
 	if (FAILED(hr))
 	{
@@ -265,6 +330,19 @@ void FVoiceCaptureWindows::Stop()
 	}
 }
 
+bool FVoiceCaptureWindows::ChangeDevice(const FString& DeviceName, int32 SampleRate, int32 NumChannels)
+{
+	if (VoiceCaptureState != EVoiceCaptureState::UnInitialized)
+	{
+		return CreateCaptureBuffer(DeviceName, SampleRate, NumChannels);
+	}
+	else
+	{
+		UE_LOG(LogVoiceCapture, Warning, TEXT("Unable to change device, not initialized"));
+		return false;
+	}
+}
+
 bool FVoiceCaptureWindows::IsCapturing()
 {		
 	if (CV->VoiceCaptureBuffer8)
@@ -315,23 +393,36 @@ void FVoiceCaptureWindows::ProcessData()
 	DWORD LockSize = ((CurrentReadPos - CV->NextCaptureOffset) + CV->VoiceCaptureBufferCaps8.dwBufferBytes) % CV->VoiceCaptureBufferCaps8.dwBufferBytes;
 	if(LockSize != 0) 
 	{ 
+		//UE_LOG( LogVoiceCapture, Log, TEXT( "LockSize: %i, CurrentCapturePos: %i, CurrentReadPos: %i, NextCaptureOffset: %i" ), LockSize, CurrentCapturePos, CurrentReadPos, CV->NextCaptureOffset );
+
 		DWORD CaptureFlags = 0;
 		DWORD CaptureLength = 0;
-		void* CaptureData = NULL;
+		void* CaptureData = nullptr;
 		DWORD CaptureLength2 = 0;
-		void* CaptureData2 = NULL;
+		void* CaptureData2 = nullptr;
 		hr = CV->VoiceCaptureBuffer8->Lock(CV->NextCaptureOffset, LockSize,
 			&CaptureData, &CaptureLength, 
 			&CaptureData2, &CaptureLength2, CaptureFlags);
 		if (SUCCEEDED(hr))
 		{
-			CaptureLength = FMath::Min(CaptureLength, (DWORD)MAX_UNCOMPRESSED_VOICE_BUFFER_SIZE);
-			CaptureLength2 = FMath::Min(CaptureLength2, (DWORD)MAX_UNCOMPRESSED_VOICE_BUFFER_SIZE - CaptureLength);
+			const DWORD OriginalCaptureLength = CaptureLength;
+			const DWORD OriginalCaptureLength2 = CaptureLength2;
+
+			if (UncompressedAudioBuffer.Num() + CaptureLength + CaptureLength2 > (DWORD)UncompressedAudioBuffer.Max())
+			{ 
+				UE_LOG(LogVoiceCapture, Warning, TEXT("Resetting UncompressedAudioBuffer."));
+				UncompressedAudioBuffer.Empty(UncompressedAudioBuffer.Max());
+				VoiceCaptureState = EVoiceCaptureState::NoData;
+			}
+
+			const int32 Offset = UncompressedAudioBuffer.Num();
+
+			CaptureLength = FMath::Min(CaptureLength, (DWORD)UncompressedAudioBuffer.Max());
+			CaptureLength2 = FMath::Min(CaptureLength2, (DWORD)UncompressedAudioBuffer.Max() - CaptureLength);
 
-			UncompressedAudioBuffer.Empty(MAX_UNCOMPRESSED_VOICE_BUFFER_SIZE);
 			UncompressedAudioBuffer.AddUninitialized(CaptureLength + CaptureLength2);
 
-			uint8* AudioBuffer = UncompressedAudioBuffer.GetData();
+			uint8* AudioBuffer = UncompressedAudioBuffer.GetData() + Offset;
 			FMemory::Memcpy(AudioBuffer, CaptureData, CaptureLength);
 
 			if (CaptureData2 && CaptureLength2 > 0)
@@ -340,15 +431,17 @@ void FVoiceCaptureWindows::ProcessData()
 			}
 
 			CA_SUPPRESS(6385);
-			CV->VoiceCaptureBuffer8->Unlock(CaptureData, CaptureLength, CaptureData2, CaptureLength2);
+			CV->VoiceCaptureBuffer8->Unlock(CaptureData, OriginalCaptureLength, CaptureData2, OriginalCaptureLength2);
 
 			// Move the capture offset forward.
 			CV->NextCaptureOffset = (CV->NextCaptureOffset + CaptureLength) % CV->VoiceCaptureBufferCaps8.dwBufferBytes;
 			CV->NextCaptureOffset = (CV->NextCaptureOffset + CaptureLength2) % CV->VoiceCaptureBufferCaps8.dwBufferBytes;
 
-			if (IsSilence((int16*)AudioBuffer, UncompressedAudioBuffer.Num()))
+			// If Offset > 0 then we know a previous check guarantees it's not silent data
+			if (Offset == 0 && IsSilence((int16*)AudioBuffer, UncompressedAudioBuffer.Num()))
 			{
 				VoiceCaptureState = EVoiceCaptureState::NoData;
+				UncompressedAudioBuffer.Empty(UncompressedAudioBuffer.Max());
 			}
 			else
 			{
@@ -356,10 +449,9 @@ void FVoiceCaptureWindows::ProcessData()
 			}
 
 #if !UE_BUILD_SHIPPING
-			static double LastCapture = 0.0;
-			double NewTime = FPlatformTime::Seconds();
-			UE_LOG(LogVoiceCapture, VeryVerbose, TEXT("LastCapture: %f %s"), (NewTime - LastCapture) * 1000.0, EVoiceCaptureState::ToString(VoiceCaptureState));
-			LastCapture = NewTime;
+			const double NewTime = FPlatformTime::Seconds();
+			UE_LOG(LogVoiceCapture, VeryVerbose, TEXT("LastCapture: %f %s"), (NewTime - LastCaptureTime) * 1000.0, EVoiceCaptureState::ToString(VoiceCaptureState));
+			LastCaptureTime = NewTime;
 #endif
 		}
 		else
@@ -378,11 +470,14 @@ EVoiceCaptureState::Type FVoiceCaptureWindows::GetVoiceData(uint8* OutVoiceBuffe
 	if (VoiceCaptureState == EVoiceCaptureState::Ok ||
 		VoiceCaptureState == EVoiceCaptureState::Stopping)
 	{
+		check(UncompressedAudioBuffer.Num() <= (int32)CV->VoiceCaptureBufferCaps8.dwBufferBytes);
+
 		OutAvailableVoiceData = UncompressedAudioBuffer.Num();
 		if (InVoiceBufferSize >= OutAvailableVoiceData)
 		{
 			FMemory::Memcpy(OutVoiceBuffer, UncompressedAudioBuffer.GetData(), OutAvailableVoiceData);
 			VoiceCaptureState = EVoiceCaptureState::NoData;
+			UncompressedAudioBuffer.Empty(UncompressedAudioBuffer.Max());
 		}
 		else
 		{
@@ -393,52 +488,51 @@ EVoiceCaptureState::Type FVoiceCaptureWindows::GetVoiceData(uint8* OutVoiceBuffe
 	return NewMicState;
 }
 
+int32 FVoiceCaptureWindows::GetBufferSize() const
+{
+	if (VoiceCaptureState != EVoiceCaptureState::UnInitialized)
+	{
+		return CV->VoiceCaptureBufferCaps8.dwBufferBytes;
+	}
+
+	return 0;
+}
+
 bool FVoiceCaptureWindows::CreateNotifications(uint32 BufferSize)
 {
-	bool bSuccess = true;
+	bool bSuccess = false;
 
-	LPDIRECTSOUNDNOTIFY8 NotifyInt = NULL;
+	LPDIRECTSOUNDNOTIFY8 NotifyInt = nullptr;
 
 	HRESULT hr = CV->VoiceCaptureBuffer8->QueryInterface(IID_IDirectSoundNotify, (LPVOID*)&NotifyInt);
 	if (SUCCEEDED(hr))
 	{
-		DSBPOSITIONNOTIFY NotifyEvents[NUM_EVENTS];
-
-		// Create events.
-		for (int i = 0; i < NUM_EVENTS; i++)
+		// Create stop event
+		CV->StopEvent = CreateEvent(nullptr, true, false, nullptr);
+		if (CV->StopEvent == NULL || CV->StopEvent == INVALID_HANDLE_VALUE)
 		{
-			CV->Events[i] = CreateEvent(NULL, true, false, NULL);
-			if (CV->Events[i] == INVALID_HANDLE_VALUE)
-			{
-				UE_LOG(LogVoiceCapture, Warning, TEXT("Error creating sync event"));
-				bSuccess = false;
-			}
+			UE_LOG(LogVoiceCapture, Warning, TEXT("Error creating stop event"));
 		}
-
-		if (bSuccess)
+		else
 		{
-			// Evenly distribute notifications throughout the buffer
-			uint32 NumSegments = NUM_EVENTS - 1; 
-			uint32 BufferSegSize = BufferSize / NumSegments;
-			for (uint32 i = 0; i < NumSegments; i++)
-			{
-				NotifyEvents[i].dwOffset = ((i + 1) * BufferSegSize) - 1;
-				NotifyEvents[i].hEventNotify = CV->Events[i];
-			}
+			DSBPOSITIONNOTIFY StopEvent;
 
 			// when buffer stops
-			NotifyEvents[STOP_EVENT].dwOffset = DSBPN_OFFSETSTOP;
-			NotifyEvents[STOP_EVENT].hEventNotify = CV->Events[STOP_EVENT];
+			StopEvent.dwOffset = DSBPN_OFFSETSTOP;
+			StopEvent.hEventNotify = CV->StopEvent;
 
-			hr = NotifyInt->SetNotificationPositions(NUM_EVENTS, NotifyEvents);
-			if (FAILED(hr))
+			hr = NotifyInt->SetNotificationPositions(1, &StopEvent);
+			if (SUCCEEDED(hr))
 			{
-				UE_LOG(LogVoiceCapture, Warning, TEXT("Failed to set notifications 0x%08x"), hr);
-				bSuccess = false;
+				bSuccess = true;
 			}
+			else
+			{
+				UE_LOG(LogVoiceCapture, Warning, TEXT("Failed to set stop notifications 0x%08x"), hr);
 
-			float BufferSegMs = BufferSegSize / ((CV->WavFormat.nSamplesPerSec / 1000.0) * (CV->WavFormat.wBitsPerSample / 8.0));
-			UE_LOG(LogVoiceCapture, Display, TEXT("%d notifications, every %d bytes [%f ms]"), NumSegments, BufferSegSize, BufferSegMs);
+				CloseHandle(CV->StopEvent);
+				CV->StopEvent = INVALID_HANDLE_VALUE;
+			}
 		}
 
 		NotifyInt->Release();
@@ -448,18 +542,6 @@ bool FVoiceCaptureWindows::CreateNotifications(uint32 BufferSize)
 		UE_LOG(LogVoiceCapture, Warning, TEXT("Failed to create voice notification interface 0x%08x"), hr);
 	}
 
-	if (!bSuccess)
-	{
-		for (int i = 0; i < NUM_EVENTS; i++)
-		{
-			if (CV->Events[i] != INVALID_HANDLE_VALUE)
-			{
-				CloseHandle(CV->Events[i]);
-				CV->Events[i] = INVALID_HANDLE_VALUE;
-			}
-		}
-	}
-
 	return bSuccess;
 }
 
@@ -468,47 +550,60 @@ bool FVoiceCaptureWindows::Tick(float DeltaTime)
 	if (VoiceCaptureState != EVoiceCaptureState::UnInitialized &&
 		VoiceCaptureState != EVoiceCaptureState::NotCapturing)
 	{
-#if 1
-		uint32 NextEvent = (LastEventTriggered + 1) % (NUM_EVENTS - 1);
-		if (WaitForSingleObject(CV->Events[NextEvent], 0) == WAIT_OBJECT_0)
-		{
-#if !UE_BUILD_SHIPPING
-			UE_LOG(LogVoiceCapture, VeryVerbose, TEXT("Event %d Signalled"), NextEvent);
-#endif
-			ProcessData();
+		ProcessData();
 
-			LastEventTriggered = NextEvent;
-			ResetEvent(CV->Events[NextEvent]);
-		}
-#else
-		// All "non stop" events
-		for (int32 EventIdx = 0; EventIdx < STOP_EVENT; EventIdx++)
-		{
-			if (WaitForSingleObject(Events[EventIdx], 0) == WAIT_OBJECT_0)
-			{
-#if !UE_BUILD_SHIPPING
-				UE_LOG(LogVoiceCapture, VeryVerbose, TEXT("Event %d Signalled"), EventIdx);
-#endif
-				if (LastEventTriggered != EventIdx)
-				{
-					ProcessData();
-				}
-
-				LastEventTriggered = EventIdx;
-				ResetEvent(CV->Events[EventIdx]);
-			}
-		}
-#endif
-
-		if (CV->Events[STOP_EVENT] != INVALID_HANDLE_VALUE && WaitForSingleObject(CV->Events[STOP_EVENT], 0) == WAIT_OBJECT_0)
+		if (CV->StopEvent != INVALID_HANDLE_VALUE && WaitForSingleObject(CV->StopEvent, 0) == WAIT_OBJECT_0)
 		{
 			UE_LOG(LogVoiceCapture, Verbose, TEXT("Voice capture stopped"));
-			ResetEvent(CV->Events[STOP_EVENT]);
+			ResetEvent(CV->StopEvent);
 			VoiceCaptureState = EVoiceCaptureState::NotCapturing;
+			UncompressedAudioBuffer.Empty(UncompressedAudioBuffer.Max());
 		}
 	}
 
 	return true;
 }
 
-#include "HideWindowsPlatformTypes.h"
\ No newline at end of file
+void FVoiceCaptureWindows::DumpState() const
+{
+	if (CV)
+	{
+		extern FString PrintMSGUID(LPGUID Guid);
+		UE_LOG(LogVoiceCapture, Display, TEXT("Device %s"), *PrintMSGUID(&CV->VoiceCaptureDeviceGuid));
+
+		UE_LOG(LogVoiceCapture, Display, TEXT("CaptureDev: 0x%08x"), CV->VoiceCaptureDev);
+		UE_LOG(LogVoiceCapture, Display, TEXT("CaptureBuffer: 0x%08x"), CV->VoiceCaptureBuffer8);
+
+		UE_LOG(LogVoiceCapture, Display, TEXT("Capture Format"));
+		UE_LOG(LogVoiceCapture, Display, TEXT("- Tag: %d"), CV->WavFormat.wFormatTag);
+		UE_LOG(LogVoiceCapture, Display, TEXT("- Channels: %d"), CV->WavFormat.nChannels);
+		UE_LOG(LogVoiceCapture, Display, TEXT("- BitsPerSample: %d"), CV->WavFormat.wBitsPerSample);
+		UE_LOG(LogVoiceCapture, Display, TEXT("- SamplesPerSec: %d"), CV->WavFormat.nSamplesPerSec);
+		UE_LOG(LogVoiceCapture, Display, TEXT("- BlockAlign: %d"), CV->WavFormat.nBlockAlign);
+		UE_LOG(LogVoiceCapture, Display, TEXT("- AvgBytesPerSec: %d"), CV->WavFormat.nAvgBytesPerSec);
+
+		UE_LOG(LogVoiceCapture, Display, TEXT("Capture Buffer"));
+		UE_LOG(LogVoiceCapture, Display, TEXT("- Flags: 0x%08x"), CV->VoiceCaptureBufferDesc.dwFlags);
+		UE_LOG(LogVoiceCapture, Display, TEXT("- BufferBytes: %d"), CV->VoiceCaptureBufferDesc.dwBufferBytes);
+		UE_LOG(LogVoiceCapture, Display, TEXT("- Format: 0x%08x"), CV->VoiceCaptureBufferDesc.lpwfxFormat);
+
+		UE_LOG(LogVoiceCapture, Display, TEXT("Device Caps"));
+		UE_LOG(LogVoiceCapture, Display, TEXT("- Size: %d"), CV->VoiceCaptureDevCaps.dwSize);
+		UE_LOG(LogVoiceCapture, Display, TEXT("- Flags: 0x%08x"), CV->VoiceCaptureDevCaps.dwFlags);
+		UE_LOG(LogVoiceCapture, Display, TEXT("- Formats: %d"), CV->VoiceCaptureDevCaps.dwFormats);
+		UE_LOG(LogVoiceCapture, Display, TEXT("- Channels: %d"), CV->VoiceCaptureDevCaps.dwChannels);
+
+		UE_LOG(LogVoiceCapture, Display, TEXT("D3D8 Caps"));
+		UE_LOG(LogVoiceCapture, Display, TEXT("- Size: %d"), CV->VoiceCaptureBufferCaps8.dwSize);
+		UE_LOG(LogVoiceCapture, Display, TEXT("- Flags: 0x%08x"), CV->VoiceCaptureBufferCaps8.dwFlags);
+		UE_LOG(LogVoiceCapture, Display, TEXT("- BufferBytes: %d"), CV->VoiceCaptureBufferCaps8.dwBufferBytes);
+	}
+	else
+	{
+		UE_LOG(LogVoiceCapture, Display, TEXT("No capture device to dump state"));
+	}
+}
+
+#include "HideWindowsPlatformTypes.h"
+
+#endif // PLATFORM_SUPPORTS_VOICE_CAPTURE
\ No newline at end of file
diff --git a/Engine/Source/Runtime/Online/Voice/Private/Windows/VoiceCaptureWindows.h b/Engine/Source/Runtime/Online/Voice/Private/Windows/VoiceCaptureWindows.h
--- a/Engine/Source/Runtime/Online/Voice/Private/Windows/VoiceCaptureWindows.h
+++ b/Engine/Source/Runtime/Online/Voice/Private/Windows/VoiceCaptureWindows.h
@@ -14,13 +14,6 @@ class IVoiceCapture;
 class FVoiceCaptureWindows;
 struct FVoiceCaptureWindowsVars;
 
-/** Number of notification events distributed throughout the capture buffer */
-#define NUM_EVENTS 5
-/** The last notificatione event is reserved for the stop capturing event */
-#define STOP_EVENT NUM_EVENTS - 1
-/** Maximum buffer size for storing raw uncompressed audio from the system */
-#define MAX_UNCOMPRESSED_VOICE_BUFFER_SIZE 30 * 1024
-
 /**
  * DirectSound wrapper for initialization / shutdown
  */ 
@@ -61,13 +54,24 @@ PACKAGE_SCOPE:
 
 public:
 
+	struct FCaptureDeviceInfo
+	{
+		/** Enumerated capture device name */
+		FString DeviceName;
+		/** Enumerated capture device guid */
+		GUID DeviceId;
+	};
+
 	/** DirectSound8 Interface */
 	LPDIRECTSOUND8 DirectSound;
 
 	/** HMD audio input device to use */
 	FString HMDAudioInputDevice;
 	/** GUID of selected voice capture device */
-	GUID VoiceCaptureDeviceGuid;
+	FCaptureDeviceInfo DefaultVoiceCaptureDevice;
+
+	/** All available capture devices */
+	TMap<FString, FCaptureDeviceInfo> Devices;
 
 	/** Constructors */
 	FVoiceCaptureDeviceWindows();
@@ -76,9 +80,13 @@ public:
 	/**
 	 * Create a single voice capture buffer
 	 *
+	 * @param DeviceName name of device to capture with, empty for default device
+	 * @param SampleRate valid sample rate to capture audio data (8khz - 48khz)
+	 * @param NumChannel number of audio channels (1 - mono, 2 - stereo)
+	 *
 	 * @return class capable of recording voice on a given device
 	 */
-	FVoiceCaptureWindows* CreateVoiceCaptureObject();
+	FVoiceCaptureWindows* CreateVoiceCaptureObject(const FString& DeviceName, int32 SampleRate, int32 NumChannels);
 
 	/**
 	 * Singleton accessor
@@ -99,13 +107,16 @@ public:
 	~FVoiceCaptureWindows();
 
 	// IVoiceCapture
-	virtual bool Init(int32 SampleRate, int32 NumChannels) override;
+	virtual bool Init(const FString& DeviceName, int32 SampleRate, int32 NumChannels) override;
 	virtual void Shutdown() override;
 	virtual bool Start() override;
 	virtual void Stop() override;
+	virtual bool ChangeDevice(const FString& DeviceName, int32 SampleRate, int32 NumChannels) override;
 	virtual bool IsCapturing() override;
 	virtual EVoiceCaptureState::Type GetCaptureState(uint32& OutAvailableVoiceData) const override;
 	virtual EVoiceCaptureState::Type GetVoiceData(uint8* OutVoiceBuffer, uint32 InVoiceBufferSize, uint32& OutAvailableVoiceData) override;
+	virtual int32 GetBufferSize() const override;
+	virtual void DumpState() const override;
 
 	// FTickerObjectBase
 	virtual bool Tick(float DeltaTime) override;
@@ -114,13 +125,27 @@ private:
 
 	/** All windows related variables to hide windows includes */
 	FVoiceCaptureWindowsVars* CV;
-	/** Event sentinel */
-	uint32 LastEventTriggered;
+	/** Last time data was captured */
+	double LastCaptureTime;
 	/** State of capture device */
 	EVoiceCaptureState::Type VoiceCaptureState;
 	/** Uncompressed audio buffer */
 	TArray<uint8> UncompressedAudioBuffer;
 
+	/**
+	 * Create the D3D8 capture device
+	 *
+	 * @param DeviceName name of device to capture with, empty for default device
+	 * @param SampleRate valid sample rate to capture audio data (8khz - 48khz)
+	 * @param NumChannel number of audio channels (1 - mono, 2 - stereo)
+	 * 
+	 * @return true if successful, false otherwise
+	 */
+	bool CreateCaptureBuffer(const FString& DeviceName, int32 SampleRate, int32 NumChannels);
+
+	/** Clear the capture buffer and release all resources */
+	void FreeCaptureBuffer();
+
 	/**
 	 * Lock the DirectSound audio buffer and copy out the available data based on a notification
 	 * and the current audio buffer capture cursor
diff --git a/Engine/Source/Runtime/Online/Voice/Private/Windows/VoiceModuleWindows.cpp b/Engine/Source/Runtime/Online/Voice/Private/Windows/VoiceModuleWindows.cpp
--- a/Engine/Source/Runtime/Online/Voice/Private/Windows/VoiceModuleWindows.cpp
+++ b/Engine/Source/Runtime/Online/Voice/Private/Windows/VoiceModuleWindows.cpp
@@ -8,8 +8,9 @@
 #include "Voice.h"
 #include "Runtime/HeadMountedDisplay/Public/IHeadMountedDisplayModule.h"
 #include "HAL/IConsoleManager.h"
+#if PLATFORM_SUPPORTS_VOICE_CAPTURE
 
-FVoiceCaptureDeviceWindows* FVoiceCaptureDeviceWindows::Singleton = NULL;
+FVoiceCaptureDeviceWindows* FVoiceCaptureDeviceWindows::Singleton = nullptr;
 
 /** Helper for printing MS guids */
 FString PrintMSGUID(LPGUID Guid)
@@ -38,19 +39,30 @@ BOOL CALLBACK CaptureDeviceCallback(
 	FVoiceCaptureDeviceWindows* VCPtr = (FVoiceCaptureDeviceWindows*)(lpContext);
 	UE_LOG(LogVoiceCapture, Display, TEXT("Device: %s Desc: %s GUID: %s Context:0x%08x"), lpcstrDescription, lpcstrModule, *PrintMSGUID(lpGuid), lpContext);
 
-	// Allow HMD to override the voice capture device
-	if(!VCPtr->HMDAudioInputDevice.IsEmpty() && !VCPtr->HMDAudioInputDevice.Compare((LPCWSTR) lpcstrModule))
+	if (lpGuid != nullptr)
 	{
-		UE_LOG(LogVoice, Display, TEXT("VoiceCapture device overridden by HMD to use '%s' %s"), lpcstrDescription, *PrintMSGUID(lpGuid));
-		VCPtr->VoiceCaptureDeviceGuid = *lpGuid;
+		// Save the enumerated device information for later use
+		FString DeviceDescription((LPCWSTR)lpcstrDescription);
+		FVoiceCaptureDeviceWindows::FCaptureDeviceInfo DeviceDesc;
+		DeviceDesc.DeviceName = DeviceDescription;
+		DeviceDesc.DeviceId = *lpGuid;
+		VCPtr->Devices.Emplace(DeviceDescription, DeviceDesc);
+
+		// Allow HMD to override the default voice capture device
+		if (!VCPtr->HMDAudioInputDevice.IsEmpty() && !VCPtr->HMDAudioInputDevice.Compare((LPCWSTR)lpcstrModule))
+		{
+			UE_LOG(LogVoice, Display, TEXT("VoiceCapture device overridden by HMD to use '%s' %s"), lpcstrDescription, *PrintMSGUID(lpGuid));
+			VCPtr->DefaultVoiceCaptureDevice = DeviceDesc;
+			VCPtr->Devices.Add(FString(DEFAULT_DEVICE_NAME), VCPtr->DefaultVoiceCaptureDevice);
+		}
 	}
-
+	
 	return true;
 }
 
 FVoiceCaptureDeviceWindows::FVoiceCaptureDeviceWindows() :
 	bInitialized(false),
-	DirectSound(NULL)
+	DirectSound(nullptr)
 {
 }
 
@@ -61,7 +73,7 @@ FVoiceCaptureDeviceWindows::~FVoiceCaptureDeviceWindows()
 
 FVoiceCaptureDeviceWindows* FVoiceCaptureDeviceWindows::Create()
 {
-	if (Singleton == NULL)
+	if (Singleton == nullptr)
 	{
 		Singleton = new FVoiceCaptureDeviceWindows();
 		if (!Singleton->Init())
@@ -79,7 +91,7 @@ void FVoiceCaptureDeviceWindows::Destroy()
 	{
 		// calls Shutdown() above
 		delete Singleton;
-		Singleton = NULL;
+		Singleton = nullptr;
 	}
 }
 
@@ -104,25 +116,25 @@ public:
 		HRESULT hr = S_OK;
 
 		// Activate session manager.
-		IAudioSessionManager2* pSessionManager2 = NULL;
+		IAudioSessionManager2* pSessionManager2 = nullptr;
 		if (SUCCEEDED(hr))
 		{
 			hr = pEndpoint->Activate(__uuidof(IAudioSessionManager2),
 				CLSCTX_INPROC_SERVER,
-				NULL,
+				nullptr,
 				reinterpret_cast<void **>(&pSessionManager2));
 		}
 
-		IAudioSessionControl* pSessionControl = NULL;
+		IAudioSessionControl* pSessionControl = nullptr;
 		if (SUCCEEDED(hr))
 		{
-			hr = pSessionManager2->GetAudioSessionControl(NULL, 0, &pSessionControl);
+			hr = pSessionManager2->GetAudioSessionControl(nullptr, 0, &pSessionControl);
 
 			pSessionManager2->Release();
-			pSessionManager2 = NULL;
+			pSessionManager2 = nullptr;
 		}
 
-		IAudioSessionControl2* pSessionControl2 = NULL;
+		IAudioSessionControl2* pSessionControl2 = nullptr;
 		if (SUCCEEDED(hr))
 		{
 			hr = pSessionControl->QueryInterface(
@@ -130,7 +142,7 @@ public:
 				(void**)&pSessionControl2);
 
 			pSessionControl->Release();
-			pSessionControl = NULL;
+			pSessionControl = nullptr;
 		}
 
 		//  Sync the ducking state with the specified preference.
@@ -143,7 +155,7 @@ public:
 			}
 
 			pSessionControl2->Release();
-			pSessionControl2 = NULL;
+			pSessionControl2 = nullptr;
 		}
 		return hr;
 	}
@@ -153,16 +165,16 @@ public:
 		HRESULT hr = S_OK;
 
 		//  Start with the default endpoint.
-		IMMDeviceEnumerator* pDeviceEnumerator = NULL;
+		IMMDeviceEnumerator* pDeviceEnumerator = nullptr;
 		hr = CoCreateInstance(__uuidof(MMDeviceEnumerator),
-			NULL,
+			nullptr,
 			CLSCTX_INPROC_SERVER,
 			IID_PPV_ARGS(&pDeviceEnumerator));
 
 		if (SUCCEEDED(hr))
 		{
 			{
-				IMMDevice* pEndpoint = NULL;
+				IMMDevice* pEndpoint = nullptr;
 				hr = pDeviceEnumerator->GetDefaultAudioEndpoint(eRender, eConsole, &pEndpoint);
 				if (SUCCEEDED(hr))
 				{
@@ -173,18 +185,18 @@ public:
 
 					FAudioDuckingWindows::EnableDuckingOptOut(pEndpoint, bDuckingOptOutChecked);
 					pEndpoint->Release();
-					pEndpoint = NULL;
+					pEndpoint = nullptr;
 				}
 			}
 
 			if (0) // reference for enumerating all endpoints in case its necessary
 			{
-				IMMDeviceCollection* pDeviceCollection = NULL;
-				IMMDevice* pCollEndpoint = NULL;
+				IMMDeviceCollection* pDeviceCollection = nullptr;
+				IMMDevice* pCollEndpoint = nullptr;
 				hr = pDeviceEnumerator->EnumAudioEndpoints(eRender, DEVICE_STATE_ACTIVE, &pDeviceCollection);
 				if (SUCCEEDED(hr))
 				{
-					IPropertyStore *pProps = NULL;
+					IPropertyStore *pProps = nullptr;
 
 					::UINT DeviceCount = 0;
 					pDeviceCollection->GetCount(&DeviceCount);
@@ -215,28 +227,28 @@ public:
 
 									CoTaskMemFree(Desc);
 
-									Desc = NULL;
+									Desc = nullptr;
 									PropVariantClear(&varName);
 
 									pProps->Release();
-									pProps = NULL;
+									pProps = nullptr;
 								}
 							}
 
 							FAudioDuckingWindows::EnableDuckingOptOut(pCollEndpoint, bDuckingOptOutChecked);
 
 							pCollEndpoint->Release();
-							pCollEndpoint = NULL;
+							pCollEndpoint = nullptr;
 						}
 					}
 
 					pDeviceCollection->Release();
-					pDeviceCollection = NULL;
+					pDeviceCollection = nullptr;
 				}
 			}
 
 			pDeviceEnumerator->Release();
-			pDeviceEnumerator = NULL;
+			pDeviceEnumerator = nullptr;
 		}
 
 		if (FAILED(hr))
@@ -250,19 +262,23 @@ public:
 
 bool FVoiceCaptureDeviceWindows::Init()
 {
-	HRESULT hr = DirectSoundCreate8(NULL, &DirectSound, NULL);
+	HRESULT hr = DirectSoundCreate8(nullptr, &DirectSound, nullptr);
 	if (FAILED(hr))
 	{
 		UE_LOG(LogVoiceCapture, Warning, TEXT("Failed to init DirectSound %d"), hr);
 		return false;
 	}
 
-	if(IHeadMountedDisplayModule::IsAvailable())
+	if (IHeadMountedDisplayModule::IsAvailable())
 	{
 		HMDAudioInputDevice = IHeadMountedDisplayModule::Get().GetAudioInputDevice();
 	}
 
-	VoiceCaptureDeviceGuid = DSDEVID_DefaultVoiceCapture;
+	DefaultVoiceCaptureDevice.DeviceName = FString(DEFAULT_DEVICE_NAME);
+	DefaultVoiceCaptureDevice.DeviceId = DSDEVID_DefaultVoiceCapture;
+
+	Devices.Empty();
+	Devices.Add(FString(DEFAULT_DEVICE_NAME), DefaultVoiceCaptureDevice);
 
 	hr = DirectSoundCaptureEnumerate((LPDSENUMCALLBACK)CaptureDeviceCallback, this);
 	if (FAILED(hr))
@@ -300,26 +316,26 @@ void FVoiceCaptureDeviceWindows::Shutdown()
 	if (DirectSound)
 	{
 		DirectSound->Release();
-		DirectSound = NULL;
+		DirectSound = nullptr;
 	}
 
 	bInitialized = false;
 }
 
-FVoiceCaptureWindows* FVoiceCaptureDeviceWindows::CreateVoiceCaptureObject()
+FVoiceCaptureWindows* FVoiceCaptureDeviceWindows::CreateVoiceCaptureObject(const FString& DeviceName, int32 SampleRate, int32 NumChannels)
 {
-	FVoiceCaptureWindows* NewVoiceCapture = NULL;
+	FVoiceCaptureWindows* NewVoiceCapture = nullptr;
 
 	if (bInitialized)
 	{
 		NewVoiceCapture = new FVoiceCaptureWindows();
-		if (NewVoiceCapture->Init(VOICE_SAMPLE_RATE, NUM_VOICE_CHANNELS))
+		if (NewVoiceCapture->Init(DeviceName, SampleRate, NumChannels))
 		{
 			ActiveVoiceCaptures.Add(NewVoiceCapture);
 		}
 		else
 		{
-			NewVoiceCapture = NULL;
+			NewVoiceCapture = nullptr;
 		}
 	}
 
@@ -328,7 +344,7 @@ FVoiceCaptureWindows* FVoiceCaptureDeviceWindows::CreateVoiceCaptureObject()
 
 void FVoiceCaptureDeviceWindows::FreeVoiceCaptureObject(IVoiceCapture* VoiceCaptureObj)
 {
-	if (VoiceCaptureObj != NULL)
+	if (VoiceCaptureObj != nullptr)
 	{
 		int32 RemoveIdx = ActiveVoiceCaptures.Find(VoiceCaptureObj);
 		if (RemoveIdx != INDEX_NONE)
@@ -344,7 +360,7 @@ void FVoiceCaptureDeviceWindows::FreeVoiceCaptureObject(IVoiceCapture* VoiceCapt
 
 bool InitVoiceCapture()
 {
-	return FVoiceCaptureDeviceWindows::Create() != NULL;
+	return FVoiceCaptureDeviceWindows::Create() != nullptr;
 }
 
 void ShutdownVoiceCapture()
@@ -352,32 +368,34 @@ void ShutdownVoiceCapture()
 	FVoiceCaptureDeviceWindows::Destroy();
 }
 
-IVoiceCapture* CreateVoiceCaptureObject()
+IVoiceCapture* CreateVoiceCaptureObject(const FString& DeviceName, int32 SampleRate, int32 NumChannels)
 {
 	FVoiceCaptureDeviceWindows* VoiceCaptureDev = FVoiceCaptureDeviceWindows::Get();
-	return VoiceCaptureDev ? VoiceCaptureDev->CreateVoiceCaptureObject() : NULL;
+	return VoiceCaptureDev ? VoiceCaptureDev->CreateVoiceCaptureObject(DeviceName, SampleRate, NumChannels) : nullptr;
 }
 
-IVoiceEncoder* CreateVoiceEncoderObject()
+IVoiceEncoder* CreateVoiceEncoderObject(int32 SampleRate, int32 NumChannels, EAudioEncodeHint EncodeHint)
 {
 	FVoiceEncoderOpus* NewEncoder = new FVoiceEncoderOpus();
-	if (!NewEncoder->Init(VOICE_SAMPLE_RATE, NUM_VOICE_CHANNELS))
+	if (!NewEncoder->Init(SampleRate, NumChannels, EncodeHint))
 	{
 		delete NewEncoder;
-		NewEncoder = NULL;
+		NewEncoder = nullptr;
 	}
 
 	return NewEncoder; 
 }
 
-IVoiceDecoder* CreateVoiceDecoderObject()
+IVoiceDecoder* CreateVoiceDecoderObject(int32 SampleRate, int32 NumChannels)
 {
 	FVoiceDecoderOpus* NewDecoder = new FVoiceDecoderOpus();
-	if (!NewDecoder->Init(VOICE_SAMPLE_RATE, NUM_VOICE_CHANNELS))
+	if (!NewDecoder->Init(SampleRate, NumChannels))
 	{
 		delete NewDecoder;
-		NewDecoder = NULL;
+		NewDecoder = nullptr;
 	}
 
 	return NewDecoder; 
 }
+
+#endif // PLATFORM_SUPPORTS_VOICE_CAPTURE
\ No newline at end of file
diff --git a/Engine/Source/Runtime/Online/Voice/Public/Interfaces/VoiceCapture.h b/Engine/Source/Runtime/Online/Voice/Public/Interfaces/VoiceCapture.h
--- a/Engine/Source/Runtime/Online/Voice/Public/Interfaces/VoiceCapture.h
+++ b/Engine/Source/Runtime/Online/Voice/Public/Interfaces/VoiceCapture.h
@@ -78,12 +78,13 @@ public:
 	/**
 	 * Initialize the voice capture object
 	 *
+	 * @param DeviceName name of device to capture audio data with, empty for default device
 	 * @param SampleRate sampling rate of voice capture
 	 * @param NumChannels number of channels to capture
 	 *
 	 * @return true if successful, false otherwise
 	 */
-	virtual bool Init(int32 SampleRate, int32 NumChannels) = 0;
+	virtual bool Init(const FString& DeviceName, int32 SampleRate, int32 NumChannels) = 0;
 
 	/**
 	 * Shutdown the voice capture object
@@ -102,13 +103,24 @@ public:
 	 */
 	virtual void Stop() = 0;
 
+	/**
+	 * Change the associated capture device
+	 *
+	 * @param DeviceName name of device to capture audio data with, empty for default device
+	 * @param SampleRate sampling rate of voice capture
+	 * @param NumChannels number of channels to capture
+	 *
+	 * @return true if change was successful, false otherwise
+	 */
+	virtual bool ChangeDevice(const FString& DeviceName, int32 SampleRate, int32 NumChannels) = 0;
+
 	/**
 	 * Is the voice capture object actively capturing
 	 *
 	 * @return true if voice is being captured, false otherwise
 	 */
 	virtual bool IsCapturing() = 0;
-	
+
 	/**
 	 * Return the state of the voice data and its availability
 	 *
@@ -128,4 +140,12 @@ public:
 	 * @return state of the voice capture buffer
 	 */
 	virtual EVoiceCaptureState::Type GetVoiceData(uint8* OutVoiceBuffer, uint32 InVoiceBufferSize, uint32& OutAvailableVoiceData) = 0;
+
+	/**
+	 * @return number of bytes currently allocated in the capture buffer
+	 */
+	virtual int32 GetBufferSize() const = 0;
+
+	/** Dump the state of the voice capture device */
+	virtual void DumpState() const = 0;
 };
diff --git a/Engine/Source/Runtime/Online/Voice/Public/Interfaces/VoiceCodec.h b/Engine/Source/Runtime/Online/Voice/Public/Interfaces/VoiceCodec.h
--- a/Engine/Source/Runtime/Online/Voice/Public/Interfaces/VoiceCodec.h
+++ b/Engine/Source/Runtime/Online/Voice/Public/Interfaces/VoiceCodec.h
@@ -7,8 +7,27 @@
 #include "VoicePackage.h"
 
 /** Stats for voice codec */
-DECLARE_CYCLE_STAT_EXTERN(TEXT("VoiceEncode"), STAT_Voice_Encoding, STATGROUP_Net, );
-DECLARE_CYCLE_STAT_EXTERN(TEXT("VoiceDecode"), STAT_Voice_Decoding, STATGROUP_Net, );
+DECLARE_STATS_GROUP(TEXT("Voice"), STATGROUP_Voice, STATCAT_Advanced);
+DECLARE_CYCLE_STAT_EXTERN(TEXT("VoiceEncode"), STAT_Voice_Encoding, STATGROUP_Voice, );
+DECLARE_CYCLE_STAT_EXTERN(TEXT("VoiceDecode"), STAT_Voice_Decoding, STATGROUP_Voice, );
+DECLARE_DWORD_ACCUMULATOR_STAT_EXTERN(TEXT("Enc_SampleRate"), STAT_Encode_SampleRate, STATGROUP_Voice, );
+DECLARE_DWORD_ACCUMULATOR_STAT_EXTERN(TEXT("Enc_NumChannels"), STAT_Encode_NumChannels, STATGROUP_Voice, );
+DECLARE_DWORD_ACCUMULATOR_STAT_EXTERN(TEXT("Enc_Bitrate"), STAT_Encode_Bitrate, STATGROUP_Voice, );
+DECLARE_FLOAT_ACCUMULATOR_STAT_EXTERN(TEXT("Enc_CompRatio"), STAT_Encode_CompressionRatio, STATGROUP_Voice, );
+DECLARE_DWORD_ACCUMULATOR_STAT_EXTERN(TEXT("Enc_OutSize"), STAT_Encode_OutSize, STATGROUP_Voice, );
+DECLARE_DWORD_ACCUMULATOR_STAT_EXTERN(TEXT("Dec_SampleRate"), STAT_Decode_SampleRate, STATGROUP_Voice, );
+DECLARE_DWORD_ACCUMULATOR_STAT_EXTERN(TEXT("Dec_NumChannels"), STAT_Decode_NumChannels, STATGROUP_Voice, );
+DECLARE_FLOAT_ACCUMULATOR_STAT_EXTERN(TEXT("Dec_CompRatio"), STAT_Decode_CompressionRatio, STATGROUP_Voice, );
+DECLARE_DWORD_ACCUMULATOR_STAT_EXTERN(TEXT("Dec_OutSize"), STAT_Decode_OutSize, STATGROUP_Voice, );
+
+/** Encoding hints for compression */
+enum class EAudioEncodeHint : uint8
+{
+	/** Best for most VoIP applications where listening quality and intelligibility matter most */
+	VoiceEncode_Voice,
+	/** Best for broadcast/high-fidelity application where the decoded audio should be as close as possible to the input */
+	VoiceEncode_Audio
+};
 
 /**
  * Interface for encoding raw voice for transmission over the wire
@@ -29,10 +48,11 @@ public:
 	 *
 	 * @param SampleRate requested sample rate of the encoding
 	 * @param NumChannel number of channels in the raw audio stream
+	 * @param EncodeHint type of audio that will be encoded
 	 * 
 	 * @return true if initialization was successful, false otherwise
 	 */
-	virtual bool Init(int32 SampleRate, int32 NumChannels) = 0;
+	virtual bool Init(int32 SampleRate, int32 NumChannels, EAudioEncodeHint EncodeHint) = 0;
 
 	/**
 	 * Encode a raw audio stream (expects 16bit PCM audio)
@@ -46,10 +66,47 @@ public:
 	 */
 	virtual int32 Encode(const uint8* RawPCMData, uint32 RawDataSize, uint8* OutCompressedData, uint32& OutCompressedDataSize) = 0;
 
+	/**
+	 * Adjust the encoding bitrate
+	 *
+	 * @param InBitRate new bitrate value (bits/sec)
+	 *
+	 * @return true if successfully changed, false otherwise
+	 */
+	virtual bool SetBitrate(int32 InBitRate) = 0;
+
+	/**
+	 * Set the encoding to variable bitrate
+	 *
+	 * @param bEnableVBR true to use variable bit rate, false to turn it off
+	 *
+	 * @return true if successfully changed, false otherwise
+	 */
+	virtual bool SetVBR(bool bEnableVBR) = 0;
+
+	/**
+	 * Adjust the encoding complexity (platform specific)
+	 *
+	 * @param InComplexity new complexity
+	 *
+	 * @return true if successfully changed, false otherwise
+	 */
+	virtual bool SetComplexity(int32 InComplexity) = 0;
+
+	/**
+	 * Reset the encoder back to its initial state
+	 */
+	virtual void Reset() = 0;
+	
 	/**
 	 * Cleanup the encoder
 	 */
 	virtual void Destroy() = 0;
+
+	/**
+	 * Output the state of the encoder
+	 */
+	virtual void DumpState() const = 0;
 };
 
 /**
@@ -86,8 +143,18 @@ public:
 	 */
 	virtual void Decode(const uint8* CompressedData, uint32 CompressedDataSize, uint8* OutRawPCMData, uint32& OutRawDataSize) = 0;
 
+	/**
+	 * Reset the decoder back to its initial state
+	 */
+	virtual void Reset() = 0;
+
 	/**
 	 * Cleanup the decoder
 	 */
 	virtual void Destroy() = 0;
+
+	/**
+	 * Output the state of the decoder
+	 */
+	virtual void DumpState() const = 0;
 };
diff --git a/Engine/Source/Runtime/Online/Voice/Public/Voice.h b/Engine/Source/Runtime/Online/Voice/Public/Voice.h
--- a/Engine/Source/Runtime/Online/Voice/Public/Voice.h
+++ b/Engine/Source/Runtime/Online/Voice/Public/Voice.h
@@ -7,7 +7,3 @@
 #include "CoreMinimal.h"
 #include "VoiceModule.h"
 
-// Module includes
-
-#define VOICE_SAMPLE_RATE 16000
-#define NUM_VOICE_CHANNELS 1
diff --git a/Engine/Source/Runtime/Online/Voice/Public/VoiceModule.h b/Engine/Source/Runtime/Online/Voice/Public/VoiceModule.h
--- a/Engine/Source/Runtime/Online/Voice/Public/VoiceModule.h
+++ b/Engine/Source/Runtime/Online/Voice/Public/VoiceModule.h
@@ -9,6 +9,19 @@
 #include "Interfaces/VoiceCapture.h"
 #include "Interfaces/VoiceCodec.h"
 
+/** Name of default capture device */
+#define DEFAULT_DEVICE_NAME TEXT("Default Device")
+/** Default voice chat sample rate */
+#define DEFAULT_VOICE_SAMPLE_RATE 16000
+/** Deprecated value, use DEFAULT_VOICE_SAMPLE_RATE */
+#define VOICE_SAMPLE_RATE DEFAULT_VOICE_SAMPLE_RATE
+/** Default voice chat number of channels (mono) */
+#define DEFAULT_NUM_VOICE_CHANNELS 1
+
+class IVoiceCapture;
+class IVoiceEncoder;
+class IVoiceDecoder;
+
 /** Logging related to general voice chat flow (muting/registration/etc) */
 VOICE_API DECLARE_LOG_CATEGORY_EXTERN(LogVoice, Display, All);
 /** Logging related to encoding of local voice packets */
@@ -55,23 +68,34 @@ public:
 	/**
 	 * Instantiates a new voice capture object
 	 *
+	 * @param DeviceName name of device to capture audio data with, empty for default device
+	 * @param SampleRate sampling rate of voice capture
+	 * @param NumChannels number of channels to capture
+	 *
 	 * @return new voice capture object, possibly NULL
 	 */
-	virtual TSharedPtr<class IVoiceCapture> CreateVoiceCapture();
+	virtual TSharedPtr<IVoiceCapture> CreateVoiceCapture(const FString& DeviceName = DEFAULT_DEVICE_NAME, int32 SampleRate = DEFAULT_VOICE_SAMPLE_RATE, int32 NumChannels = DEFAULT_NUM_VOICE_CHANNELS);
 
 	/**
 	 * Instantiates a new voice encoder object
 	 *
+	 * @param SampleRate sampling rate of voice capture
+	 * @param NumChannels number of channels to capture
+	 * @param EncodeHint hint to describe type of audio quality desired
+	 *
 	 * @return new voice encoder object, possibly NULL
 	 */
-	virtual TSharedPtr<class IVoiceEncoder> CreateVoiceEncoder();
+	virtual TSharedPtr<IVoiceEncoder> CreateVoiceEncoder(int32 SampleRate = DEFAULT_VOICE_SAMPLE_RATE, int32 NumChannels = DEFAULT_NUM_VOICE_CHANNELS, EAudioEncodeHint EncodeHint = EAudioEncodeHint::VoiceEncode_Voice);
 
 	/**
 	 * Instantiates a new voice decoder object
 	 *
+	 * @param SampleRate sampling rate of voice capture
+	 * @param NumChannels number of channels to capture
+	 *
 	 * @return new voice decoder object, possibly NULL
 	 */
-	virtual TSharedPtr<class IVoiceDecoder> CreateVoiceDecoder();
+	virtual TSharedPtr<IVoiceDecoder> CreateVoiceDecoder(int32 SampleRate = DEFAULT_VOICE_SAMPLE_RATE, int32 NumChannels = DEFAULT_NUM_VOICE_CHANNELS);
 
 	/**
 	 * @return true if voice is enabled
