[Release-4.20] #jira UE-54812 Bring fixes over for Set Mic Threshold issue in binary builds. #rb aaron.mcleran

diff --git a/Engine/Source/Runtime/Engine/Private/VoiceConfig.cpp b/Engine/Source/Runtime/Engine/Private/VoiceConfig.cpp
--- a/Engine/Source/Runtime/Engine/Private/VoiceConfig.cpp
+++ b/Engine/Source/Runtime/Engine/Private/VoiceConfig.cpp
@@ -1,11 +1,26 @@
 // Copyright 1998-2018 Epic Games, Inc. All Rights Reserved.
 
-#include "Net/VoiceConfig.h"
+#include "VoiceConfig.h"
 #include "Engine/World.h"
 #include "GameFramework/GameStateBase.h"
 
 TMap<FUniqueNetIdWrapper, UVOIPTalker*> UVOIPStatics::VoiceTalkerMap;
 
+static TAutoConsoleVariable<float> CVarVoiceSilenceDetectionAttackTime(TEXT("voice.SilenceDetectionAttackTime"),
+	2.0f,
+	TEXT("Attack time to be set for the VOIP microphone's silence detection algorithm in milliseconds.\n"),	
+	ECVF_Default);
+
+static TAutoConsoleVariable<float> CVarVoiceSilenceDetectionReleaseTime(TEXT("voice.SilenceDetectionReleaseTime"),
+	1100.0f,
+	TEXT("Release time to be set for the VOIP microphone's silence detection algorithm in milliseconds.\n"),	
+	ECVF_Default);
+
+static TAutoConsoleVariable<float> CVarVoiceSilenceDetectionThreshold(TEXT("voice.SilenceDetectionThreshold"),
+	0.08f,
+	TEXT("Threshold to be set for the VOIP microphone's silence detection algorithm.\n"),	
+	ECVF_Default);
+
 int32 UVOIPStatics::GetVoiceSampleRate()
 {
 #if USE_DEFAULT_VOICE_SAMPLE_RATE
@@ -143,7 +158,9 @@ APlayerState* UVOIPStatics::GetPlayerStateFromUniqueNetId(UWorld* InWorld, const
 
 void UVOIPStatics::SetMicThreshold(float InThreshold)
 {
-	MicSilenceDetectionConfig::Threshold = InThreshold;
+	static IConsoleVariable* SilenceDetectionCVar = IConsoleManager::Get().FindConsoleVariable(TEXT("voice.SilenceDetectionThreshold"));
+	check(SilenceDetectionCVar);
+	SilenceDetectionCVar->Set(InThreshold, ECVF_SetByGameSetting);
 }
 
 void UVOIPStatics::SetVOIPTalkerForPlayer(const FUniqueNetIdWrapper& InPlayerId, UVOIPTalker* InTalker)
diff --git a/Engine/Source/Runtime/Engine/Public/Net/VoiceConfig.h b/Engine/Source/Runtime/Engine/Public/Net/VoiceConfig.h
--- a/Engine/Source/Runtime/Engine/Public/Net/VoiceConfig.h
+++ b/Engine/Source/Runtime/Engine/Public/Net/VoiceConfig.h
@@ -34,9 +34,6 @@ enum class EAudioEncodeHint : uint8
 
 namespace MicSilenceDetectionConfig
 {
-	static const float AttackTime = 2.0f;
-	static const float ReleaseTime = 1100.0f;
-	static float Threshold = 0.08f;
 	static const Audio::EPeakMode::Type LevelDetectionMode = Audio::EPeakMode::Type::Peak;
 	static const bool IsAnalog = false;
 
@@ -48,19 +45,12 @@ class USoundAttenuation;
 template <typename EnumType>
 static FORCEINLINE EnumType GetEnumValueFromString(const FString& EnumName, const FString& String)
 {
-	UEnum* const Enum = FindObject<UEnum>((UObject*) ANY_PACKAGE, *EnumName, true);
+	UEnum* Enum = FindObject<UEnum>((UObject*) ANY_PACKAGE, *EnumName, true);
 	if (!Enum)
 	{
 		return EnumType(0);
 	}
-
-	const int64 Result = Enum->GetValueByName(FName(*String));
-	if (Result == INDEX_NONE)
-	{
-		return EnumType(0);
-	}
-
-	return static_cast<EnumType>(Result);
+	return (EnumType)Enum->GetValueByName(FName(*String));
 }
 
 USTRUCT(BlueprintType)
diff --git a/Engine/Source/Runtime/Online/Voice/Private/Windows/VoiceCaptureWindows.cpp b/Engine/Source/Runtime/Online/Voice/Private/Windows/VoiceCaptureWindows.cpp
--- a/Engine/Source/Runtime/Online/Voice/Private/Windows/VoiceCaptureWindows.cpp
+++ b/Engine/Source/Runtime/Online/Voice/Private/Windows/VoiceCaptureWindows.cpp
@@ -112,9 +112,16 @@ bool FVoiceCaptureWindows::Init(const FString& DeviceName, int32 SampleRate, int
    	// init the sample counter to 0 on init
 	SampleCounter = 0; 
 	CachedSampleStart = 0;
-	MicLevelDetector.Init(SampleRate, MicSilenceDetectionConfig::AttackTime, MicSilenceDetectionConfig::ReleaseTime, MicSilenceDetectionConfig::LevelDetectionMode, MicSilenceDetectionConfig::IsAnalog);
+
+	// set up level detector
+	static IConsoleVariable* SilenceDetectionAttackCVar = IConsoleManager::Get().FindConsoleVariable(TEXT("voice.SilenceDetectionAttackTime"));
+	check(SilenceDetectionAttackCVar);			
+	static IConsoleVariable* SilenceDetectionReleaseCVar = IConsoleManager::Get().FindConsoleVariable(TEXT("voice.SilenceDetectionReleaseTime"));
+	check(SilenceDetectionReleaseCVar);				
+
+	MicLevelDetector.Init(SampleRate, SilenceDetectionAttackCVar->GetFloat(), SilenceDetectionReleaseCVar->GetFloat(), MicSilenceDetectionConfig::LevelDetectionMode, MicSilenceDetectionConfig::IsAnalog);
 	
-	const int32 AttackInSamples = SampleRate * MicSilenceDetectionConfig::AttackTime * 0.001f;
+	const int32 AttackInSamples = SampleRate * SilenceDetectionAttackCVar->GetFloat() * 0.001f;
 	LookaheadBuffer.Init(AttackInSamples + 1);
 	LookaheadBuffer.SetDelay(AttackInSamples);
 
@@ -412,6 +419,11 @@ void FVoiceCaptureWindows::ProcessData()
 			//First, if we have any cached audio from an onset mid-buffer, copy it in:
 			int32 SamplesPushedToUncompressedAudioBuffer = ReleaseBuffer.PopBufferedAudio(AudioBuffer, ReleaseBuffer.GetBufferCount());
 			AudioBuffer += SamplesPushedToUncompressedAudioBuffer;
+
+			// get threshold
+			static IConsoleVariable* SilenceDetectionThresholdCVar = IConsoleManager::Get().FindConsoleVariable(TEXT("voice.SilenceDetectionThreshold"));
+			check(SilenceDetectionThresholdCVar);			
+			const float MicSilenceThreshold = SilenceDetectionThresholdCVar->GetFloat();
 			
 			bool bMicReleased = false;
 
@@ -431,7 +443,7 @@ void FVoiceCaptureWindows::ProcessData()
 				MicLevelDetector.ProcessAudio(Temp);
 				LookaheadBuffer.ProcessSample(Temp, Temp);
 
-				bIsMicActive = MicLevelDetector.GetCurrentValue() > MicSilenceDetectionConfig::Threshold;
+				bIsMicActive = MicLevelDetector.GetCurrentValue() > MicSilenceThreshold;
 
 				if (bIsMicActive)
 				{
@@ -459,7 +471,7 @@ void FVoiceCaptureWindows::ProcessData()
 			}
 
 			//Set up second buffer and loop through that:
-			AudioBuffer += CaptureLength;
+			AudioBuffer += TotalNumFrames;
 			InputBuffer = (int16*)CaptureData2;
 			const int32 TotalNumFrames2 = CaptureLength2 / sizeof(int16);
 			for (int32 FrameIndex = 0; FrameIndex < TotalNumFrames2; FrameIndex += NumInputChannels)
@@ -475,7 +487,7 @@ void FVoiceCaptureWindows::ProcessData()
 				MicLevelDetector.ProcessAudio(Temp);
 				LookaheadBuffer.ProcessSample(Temp, Temp);
 
-				bIsMicActive = MicLevelDetector.GetCurrentValue() > MicSilenceDetectionConfig::Threshold;
+				bIsMicActive = MicLevelDetector.GetCurrentValue() > MicSilenceThreshold;
 
 				if (bIsMicActive)
 				{
@@ -696,4 +708,4 @@ void FVoiceCaptureWindows::DumpState() const
 
 #include "Windows/HideWindowsPlatformTypes.h"
 
-#endif // PLATFORM_SUPPORTS_VOICE_CAPTURE
+#endif // PLATFORM_SUPPORTS_VOICE_CAPTURE
\ No newline at end of file
diff --git a/Engine/Source/Runtime/Online/Voice/Private/Windows/VoiceCaptureWindows.h b/Engine/Source/Runtime/Online/Voice/Private/Windows/VoiceCaptureWindows.h
--- a/Engine/Source/Runtime/Online/Voice/Private/Windows/VoiceCaptureWindows.h
+++ b/Engine/Source/Runtime/Online/Voice/Private/Windows/VoiceCaptureWindows.h
@@ -160,7 +160,7 @@ private:
 
 	/*
 	* Envelope following DSP object.
-	* Configured using the MicSilenceDetectionConfig namespace in VoiceConfig.h.
+	* Configured using the MicSilenceDetectionConfig namespace in VoiceConfig.h and the CVars defined in VoiceConfig.cpp.
 	*/
 	Audio::FEnvelopeFollower MicLevelDetector;
 	/*
