Copying //UE4/Dev-Main to //UE4/Main (Source: //Portal/Main/Engine @ 4247640)

diff --git a/Engine/Source/Runtime/Online/BuildPatchServices/Private/BuildPatchFileConstructor.cpp b/Engine/Source/Runtime/Online/BuildPatchServices/Private/BuildPatchFileConstructor.cpp
--- a/Engine/Source/Runtime/Online/BuildPatchServices/Private/BuildPatchFileConstructor.cpp
+++ b/Engine/Source/Runtime/Online/BuildPatchServices/Private/BuildPatchFileConstructor.cpp
@@ -16,6 +16,9 @@
 #include "BuildPatchServicesPrivate.h"
 #include "Interfaces/IBuildInstaller.h"
 #include "Data/ChunkData.h"
+#include "Common/StatsCollector.h"
+#include "Common/SpeedRecorder.h"
+#include "Common/FileSystem.h"
 #include "Installer/ChunkSource.h"
 #include "Installer/ChunkReferenceTracker.h"
 #include "Installer/InstallerError.h"
@@ -40,7 +43,7 @@ namespace FileConstructorHelpers
 		}
 	}
 
-	bool CheckAndReportRemainingDiskSpaceError(IInstallerError* InstallerError, const FString& InstallDirectory, uint64 RemainingBytesRequired)
+	bool CheckAndReportRemainingDiskSpaceError(IInstallerError* InstallerError, const FString& InstallDirectory, uint64 RemainingBytesRequired, const TCHAR* SpaceErrorCode)
 	{
 		bool bContinueConstruction = true;
 		uint64 TotalSize = 0;
@@ -52,7 +55,8 @@ namespace FileConstructorHelpers
 				UE_LOG(LogBuildPatchServices, Error, TEXT("Out of HDD space. Needs %llu bytes, Free %llu bytes"), RemainingBytesRequired, AvailableSpace);
 				InstallerError->SetError(
 					EBuildPatchInstallError::OutOfDiskSpace,
-					DiskSpaceErrorCodes::InitialSpaceCheck,
+					SpaceErrorCode,
+					0,
 					BuildPatchServices::GetDiskSpaceMessage(InstallDirectory, RemainingBytesRequired, AvailableSpace));
 				bContinueConstruction = false;
 			}
@@ -158,8 +162,9 @@ public:
 
 /* FBuildPatchFileConstructor implementation
  *****************************************************************************/
-FBuildPatchFileConstructor::FBuildPatchFileConstructor(FBuildPatchAppManifestRef InBuildManifest, FString InInstallDirectory, FString InStageDirectory, TArray<FString> InConstructList, IChunkSource* InChunkSource, IChunkReferenceTracker* InChunkReferenceTracker, IInstallerError* InInstallerError, IInstallerAnalytics* InInstallerAnalytics, IFileConstructorStat* InFileConstructorStat)
-	: Thread(nullptr)
+FBuildPatchFileConstructor::FBuildPatchFileConstructor(FFileConstructorConfig InConfiguration, IFileSystem* InFileSystem, IChunkSource* InChunkSource, IChunkReferenceTracker* InChunkReferenceTracker, IInstallerError* InInstallerError, IInstallerAnalytics* InInstallerAnalytics, IFileConstructorStat* InFileConstructorStat)
+	: Configuration(MoveTemp(InConfiguration))
+	, Thread(nullptr)
 	, bIsRunning(false)
 	, bIsInited(false)
 	, bInitFailed(false)
@@ -168,10 +173,8 @@ FBuildPatchFileConstructor::FBuildPatchFileConstructor(FBuildPatchAppManifestRef
 	, bIsPaused(false)
 	, bShouldAbort(false)
 	, ThreadLock()
-	, BuildManifest(MoveTemp(InBuildManifest))
-	, InstallDirectory(MoveTemp(InInstallDirectory))
-	, StagingDirectory(MoveTemp(InStageDirectory))
-	, FilesToConstruct(MoveTemp(InConstructList))
+	, ConstructionStack()
+	, FileSystem(InFileSystem)
 	, ChunkSource(InChunkSource)
 	, ChunkReferenceTracker(InChunkReferenceTracker)
 	, InstallerError(InInstallerError)
@@ -179,12 +182,16 @@ FBuildPatchFileConstructor::FBuildPatchFileConstructor(FBuildPatchAppManifestRef
 	, FileConstructorStat(InFileConstructorStat)
 	, TotalJobSize(0)
 	, ByteProcessed(0)
-	, FilesConstructed()
 {
 	// Count initial job size
-	for (const FString& FileToConstruct : FilesToConstruct)
+	const int32 ConstructListNum = Configuration.ConstructList.Num();
+	ConstructionStack.Reserve(ConstructListNum);
+	ConstructionStack.AddDefaulted(ConstructListNum);
+	for (int32 ConstructListIdx = 0; ConstructListIdx < ConstructListNum ; ++ConstructListIdx)
 	{
-		TotalJobSize += BuildManifest->GetFileSize(*FileToConstruct);
+		const FString& ConstructListElem = Configuration.ConstructList[ConstructListIdx];
+		TotalJobSize += Configuration.BuildManifest->GetFileSize(ConstructListElem);
+		ConstructionStack[(ConstructListNum - 1) - ConstructListIdx] = ConstructListElem;
 	}
 	// Start thread!
 	const TCHAR* ThreadName = TEXT("FileConstructorThread");
@@ -205,10 +212,10 @@ FBuildPatchFileConstructor::~FBuildPatchFileConstructor()
 bool FBuildPatchFileConstructor::Init()
 {
 	// We are ready to go if our delegates are bound and directories successfully created
-	bool bStageDirExists = IFileManager::Get().DirectoryExists(*StagingDirectory);
+	bool bStageDirExists = IFileManager::Get().DirectoryExists(*Configuration.StagingDirectory);
 	if (!bStageDirExists)
 	{
-		UE_LOG(LogBuildPatchServices, Error, TEXT("FBuildPatchFileConstructor: Stage directory missing %s"), *StagingDirectory);
+		UE_LOG(LogBuildPatchServices, Error, TEXT("FBuildPatchFileConstructor: Stage directory missing %s"), *Configuration.StagingDirectory);
 		InstallerError->SetError(EBuildPatchInstallError::InitializationError, InitializationErrorCodes::MissingStageDirectory);
 	}
 	SetInitFailed(!bStageDirExists);
@@ -217,44 +224,40 @@ bool FBuildPatchFileConstructor::Init()
 
 uint32 FBuildPatchFileConstructor::Run()
 {
-	SetRunning( true );
-	SetInited( true );
-	const bool bIsFileData = BuildManifest->IsFileDataManifest();
+	SetRunning(true);
+	SetInited(true);
 	FileConstructorStat->OnTotalRequiredUpdated(TotalJobSize);
 
-	// Save the list of completed files
-	TArray< FString > ConstructedFiles;
-
 	// Check for resume data
-	FResumeData ResumeData( StagingDirectory, BuildManifest );
+	FResumeData ResumeData(Configuration.StagingDirectory, Configuration.BuildManifest);
 
-	// If we found incompatible resume data, we need to clean out the staging folder
-	// We don't delete the folder itself though as we should presume it was created with desired attributes
+	// If we found incompatible resume data, we need to clean out the staging folder.
+	// We don't delete the folder itself though as we should presume it was created with desired attributes.
 	if (ResumeData.bHasIncompatibleResumeData)
 	{
 		GLog->Logf(TEXT("BuildPatchServices: Deleting incompatible stage files"));
-		DeleteDirectoryContents(StagingDirectory);
+		DeleteDirectoryContents(Configuration.StagingDirectory);
 	}
 
-	// Save for started version
+	// Save for started version.
 	ResumeData.SaveOut();
 
-	// Start resume progress at zero or one
+	// Start resume progress at zero or one.
 	FileConstructorStat->OnResumeStarted();
 
-	// While we have files to construct, run
+	// While we have files to construct, run.
 	FString FileToConstruct;
-	while( GetFileToConstruct( FileToConstruct ) && !bShouldAbort)
+	while (GetFileToConstruct(FileToConstruct) && !bShouldAbort)
 	{
-		int64 FileSize = BuildManifest->GetFileSize(FileToConstruct);
+		int64 FileSize = Configuration.BuildManifest->GetFileSize(FileToConstruct);
 		FileConstructorStat->OnFileStarted(FileToConstruct, FileSize);
 		// Check resume status, currently we are only supporting sequential resume, so once we start downloading, we can't resume any more.
 		// this only comes up if the resume data has been changed externally.
-		ResumeData.CheckFile( FileToConstruct );
-		const bool bFilePreviouslyComplete = !bIsDownloadStarted && ResumeData.FilesCompleted.Contains( FileToConstruct );
-		const bool bFilePreviouslyStarted = !bIsDownloadStarted && ResumeData.FilesStarted.Contains( FileToConstruct );
+		ResumeData.CheckFile(FileToConstruct);
+		const bool bFilePreviouslyComplete = !bIsDownloadStarted && ResumeData.FilesCompleted.Contains(FileToConstruct);
+		const bool bFilePreviouslyStarted = !bIsDownloadStarted && ResumeData.FilesStarted.Contains(FileToConstruct);
 
-		// Construct or skip the file
+		// Construct or skip the file.
 		bool bFileSuccess;
 		if (bFilePreviouslyComplete)
 		{
@@ -262,7 +265,7 @@ uint32 FBuildPatchFileConstructor::Run()
 			CountBytesProcessed(FileSize);
 			GLog->Logf(TEXT("FBuildPatchFileConstructor::SkipFile %s"), *FileToConstruct);
 			// Get the file manifest.
-			const FFileManifest* FileManifest = BuildManifest->GetFileManifest(FileToConstruct);
+			const FFileManifest* FileManifest = Configuration.BuildManifest->GetFileManifest(FileToConstruct);
 			// Go through each chunk part, and dereference it from the reference tracker.
 			for (const FChunkPart& ChunkPart : FileManifest->FileChunkParts)
 			{
@@ -274,20 +277,32 @@ uint32 FBuildPatchFileConstructor::Run()
 			bFileSuccess = ConstructFileFromChunks(FileToConstruct, bFilePreviouslyStarted);
 		}
 
-		// If the file succeeded, add to lists
-		if( bFileSuccess )
+		if (bFileSuccess)
 		{
-			ConstructedFiles.Add( FileToConstruct );
+			// If we are destructive, remove the old file.
+			if (Configuration.InstallMode == EInstallMode::DestructiveInstall)
+			{
+				const bool bRequireExists = false;
+				const bool bEvenReadOnly = true;
+				FString FileToDelete = Configuration.InstallDirectory / FileToConstruct;
+				FPaths::NormalizeFilename(FileToDelete);
+				FPaths::CollapseRelativeDirectories(FileToDelete);
+				if (FileSystem->FileExists(*FileToDelete))
+				{
+					OnBeforeDeleteFile().Broadcast(FileToDelete);
+					IFileManager::Get().Delete(*FileToDelete, bRequireExists, bEvenReadOnly);
+				}
+			}
 		}
 		else
 		{
-			// This will only record and log if a failure was not already registered
+			// This will only record and log if a failure was not already registered.
 			bShouldAbort = true;
 			InstallerError->SetError(EBuildPatchInstallError::FileConstructionFail, ConstructionErrorCodes::UnknownFail);
 		}
 		FileConstructorStat->OnFileCompleted(FileToConstruct, bFileSuccess);
 
-		// Wait while paused
+		// Wait while paused.
 		FileConstructorHelpers::WaitWhilePaused(bIsPaused, bShouldAbort);
 	}
 
@@ -296,13 +311,9 @@ uint32 FBuildPatchFileConstructor::Run()
 	{
 		FileConstructorStat->OnResumeCompleted();
 	}
+	FileConstructorStat->OnConstructionCompleted();
 
-	// Set constructed files
-	ThreadLock.Lock();
-	FilesConstructed = MoveTemp(ConstructedFiles);
-	ThreadLock.Unlock();
-
-	SetRunning( false );
+	SetRunning(false);
 	return 0;
 }
 
@@ -320,6 +331,11 @@ bool FBuildPatchFileConstructor::IsComplete()
 	return ( !bIsRunning && bIsInited ) || bInitFailed;
 }
 
+FBuildPatchFileConstructor::FOnBeforeDeleteFile& FBuildPatchFileConstructor::OnBeforeDeleteFile()
+{
+	return BeforeDeleteFileEvent;
+}
+
 void FBuildPatchFileConstructor::SetRunning( bool bRunning )
 {
 	FScopeLock Lock( &ThreadLock );
@@ -347,11 +363,11 @@ void FBuildPatchFileConstructor::CountBytesProcessed( const int64& ByteCount )
 bool FBuildPatchFileConstructor::GetFileToConstruct(FString& Filename)
 {
 	FScopeLock Lock(&ThreadLock);
-	const bool bFileAvailable = FilesToConstruct.IsValidIndex(0);
+	const bool bFileAvailable = ConstructionStack.Num() > 0;
 	if (bFileAvailable)
 	{
-		Filename = FilesToConstruct[0];
-		FilesToConstruct.RemoveAt(0);
+		const bool bAllowShrinking = false;
+		Filename = ConstructionStack.Pop(bAllowShrinking);
 	}
 	return bFileAvailable;
 }
@@ -359,22 +375,62 @@ bool FBuildPatchFileConstructor::GetFileToConstruct(FString& Filename)
 int64 FBuildPatchFileConstructor::GetRemainingBytes()
 {
 	FScopeLock Lock(&ThreadLock);
-	return BuildManifest->GetFileSize(FilesToConstruct);
+	return Configuration.BuildManifest->GetFileSize(ConstructionStack);
+}
+
+int64 FBuildPatchFileConstructor::CalculateRequiredDiskSpace(const FString& InProgressFile, int64 InProgressFileSize)
+{
+	int64 DiskSpaceDeltaPeak = InProgressFileSize;
+	if (Configuration.InstallMode == EInstallMode::DestructiveInstall)
+	{
+		// The simplest method will be to run through each high level file operation, tracking peak disk usage delta.
+		const bool bCurrentManifestIsValid = Configuration.CurrentManifest.IsValid();
+		int64 DiskSpaceDelta = InProgressFileSize;
+
+		// Can remove old in progress file.
+		if (bCurrentManifestIsValid)
+		{
+			DiskSpaceDelta -= Configuration.CurrentManifest->GetFileSize(InProgressFile);
+		}
+
+		// Loop through all files to be made next, in order.
+		for (int32 ConstructionStackIdx = ConstructionStack.Num() - 1; ConstructionStackIdx >= 0; --ConstructionStackIdx)
+		{
+			const FString& FileToConstruct = ConstructionStack[ConstructionStackIdx];
+			// First we would need to make the new file.
+			DiskSpaceDelta += Configuration.BuildManifest->GetFileSize(FileToConstruct);
+			if (DiskSpaceDeltaPeak < DiskSpaceDelta)
+			{
+				DiskSpaceDeltaPeak = DiskSpaceDelta;
+			}
+			// Then we can remove the current existing file.
+			if (bCurrentManifestIsValid)
+			{
+				DiskSpaceDelta -= Configuration.CurrentManifest->GetFileSize(FileToConstruct);
+			}
+		}
+	}
+	else
+	{
+		// When not destructive, we always stage all new and changed files.
+		DiskSpaceDeltaPeak += Configuration.BuildManifest->GetFileSize(ConstructionStack);
+	}
+	return DiskSpaceDeltaPeak;
 }
 
 bool FBuildPatchFileConstructor::ConstructFileFromChunks( const FString& Filename, bool bResumeExisting )
 {
-	const bool bIsFileData = BuildManifest->IsFileDataManifest();
+	const bool bIsFileData = Configuration.BuildManifest->IsFileDataManifest();
 	bResumeExisting = bResumeExisting && !bIsFileData;
 	bool bSuccess = true;
-	FString NewFilename = StagingDirectory / Filename;
+	FString NewFilename = Configuration.StagingDirectory / Filename;
 
 	// Calculate the hash as we write the data
 	FSHA1 HashState;
 	FSHAHash HashValue;
 
 	// First make sure we can get the file manifest
-	const FFileManifest* FileManifest = BuildManifest->GetFileManifest(Filename);
+	const FFileManifest* FileManifest = Configuration.BuildManifest->GetFileManifest(Filename);
 	bSuccess = FileManifest != nullptr;
 	if( bSuccess )
 	{
@@ -413,9 +469,15 @@ bool FBuildPatchFileConstructor::ConstructFileFromChunks( const FString& Filenam
 					const int64 NextBytePosition = ByteCounter + ChunkPart.Size;
 					if (NextBytePosition <= StartPosition)
 					{
+						ISpeedRecorder::FRecord ActivityRecord;
 						// Read data for hash check
+						FileConstructorStat->OnBeforeRead();
+						ActivityRecord.CyclesStart = FStatsCollector::GetCycles();
 						NewFileReader->Serialize(ReadBuffer.GetData(), ChunkPart.Size);
+						ActivityRecord.CyclesEnd = FStatsCollector::GetCycles();
+						ActivityRecord.Size = ChunkPart.Size;
 						HashState.Update(ReadBuffer.GetData(), ChunkPart.Size);
+						FileConstructorStat->OnAfterRead(ActivityRecord);
 						// Count bytes read from file
 						ByteCounter = NextBytePosition;
 						// Set to resume from next chunk part
@@ -423,6 +485,7 @@ bool FBuildPatchFileConstructor::ConstructFileFromChunks( const FString& Filenam
 						// Inform the reference tracker of the chunk part skip
 						bSuccess = ChunkReferenceTracker->PopReference(ChunkPart.Guid) && bSuccess;
 						CountBytesProcessed(ChunkPart.Size);
+						FileConstructorStat->OnFileProgress(Filename, NewFileReader->Tell());
 						// Wait if paused
 						FileConstructorHelpers::WaitWhilePaused(bIsPaused, bShouldAbort);
 					}
@@ -443,8 +506,8 @@ bool FBuildPatchFileConstructor::ConstructFileFromChunks( const FString& Filenam
 		if (!bInitialDiskSizeCheck)
 		{
 			bInitialDiskSizeCheck = true;
-			const uint64 RequiredSpace = (FileManifest->GetFileSize() - StartPosition) + GetRemainingBytes();
-			if (!FileConstructorHelpers::CheckAndReportRemainingDiskSpaceError(InstallerError, InstallDirectory, RequiredSpace))
+			const uint64 RequiredSpace = CalculateRequiredDiskSpace(Filename, FileManifest->GetFileSize() - StartPosition);
+			if (!FileConstructorHelpers::CheckAndReportRemainingDiskSpaceError(InstallerError, Configuration.InstallDirectory, RequiredSpace, DiskSpaceErrorCodes::InitialSpaceCheck))
 			{
 				return false;
 			}
@@ -458,21 +521,35 @@ bool FBuildPatchFileConstructor::ConstructFileFromChunks( const FString& Filenam
 		}
 
 		// Attempt to create the file
-		FArchive* NewFile = IFileManager::Get().CreateFileWriter( *NewFilename, bResumeExisting ? ::EFileWrite::FILEWRITE_Append : 0 );
+		ISpeedRecorder::FRecord ActivityRecord;
+		FileConstructorStat->OnBeforeAdminister();
+		ActivityRecord.CyclesStart = FStatsCollector::GetCycles();
+		TUniquePtr<FArchive> NewFile = FileSystem->CreateFileWriter( *NewFilename, bResumeExisting ? EWriteFlags::Append : EWriteFlags::None );
 		uint32 LastError = FPlatformMisc::GetLastError();
+		ActivityRecord.CyclesEnd = FStatsCollector::GetCycles();
+		ActivityRecord.Size = 0;
+		FileConstructorStat->OnAfterAdminister(ActivityRecord);
 		bSuccess = NewFile != nullptr;
-		if( bSuccess )
+		if (bSuccess)
 		{
 			// Seek to file write position
-			NewFile->Seek( StartPosition );
+			if (NewFile->Tell() != StartPosition)
+			{
+				FileConstructorStat->OnBeforeAdminister();
+				ActivityRecord.CyclesStart = FStatsCollector::GetCycles();
+				NewFile->Seek(StartPosition);
+				ActivityRecord.CyclesEnd = FStatsCollector::GetCycles();
+				ActivityRecord.Size = 0;
+				FileConstructorStat->OnAfterAdminister(ActivityRecord);
+			}
 
 			// For each chunk, load it, and place it's data into the file
-			for( int32 ChunkPartIdx = StartChunkPart; ChunkPartIdx < FileManifest->FileChunkParts.Num() && bSuccess && !bShouldAbort; ++ChunkPartIdx )
+			for (int32 ChunkPartIdx = StartChunkPart; ChunkPartIdx < FileManifest->FileChunkParts.Num() && bSuccess && !bShouldAbort; ++ChunkPartIdx)
 			{
 				const FChunkPart& ChunkPart = FileManifest->FileChunkParts[ChunkPartIdx];
-				bSuccess = InsertChunkData( ChunkPart, *NewFile, HashState );
+				bSuccess = InsertChunkData(ChunkPart, *NewFile, HashState);
 				FileConstructorStat->OnFileProgress(Filename, NewFile->Tell());
-				if( bSuccess )
+				if (bSuccess)
 				{
 					CountBytesProcessed(ChunkPart.Size);
 					// Wait while paused
@@ -491,14 +568,19 @@ bool FBuildPatchFileConstructor::ConstructFileFromChunks( const FString& Filenam
 			}
 
 			// Close the file writer
+			FileConstructorStat->OnBeforeAdminister();
+			ActivityRecord.CyclesStart = FStatsCollector::GetCycles();
 			NewFile->Close();
-			delete NewFile;
+			NewFile.Reset();
+			ActivityRecord.CyclesEnd = FStatsCollector::GetCycles();
+			ActivityRecord.Size = 0;
+			FileConstructorStat->OnAfterAdminister(ActivityRecord);
 		}
 		else
 		{
 			// Check if drive space was the issue here
-			const uint64 RequiredSpace = FileManifest->GetFileSize() + GetRemainingBytes();
-			bool bError = !FileConstructorHelpers::CheckAndReportRemainingDiskSpaceError(InstallerError, InstallDirectory, RequiredSpace);
+			const uint64 RequiredSpace = CalculateRequiredDiskSpace(Filename, FileManifest->GetFileSize());
+			bool bError = !FileConstructorHelpers::CheckAndReportRemainingDiskSpaceError(InstallerError, Configuration.InstallDirectory, RequiredSpace, DiskSpaceErrorCodes::DuringInstallation);
 
 			// Otherwise we just couldn't make the file
 			if (!bError)
@@ -510,7 +592,7 @@ bool FBuildPatchFileConstructor::ConstructFileFromChunks( const FString& Filenam
 					UE_LOG(LogBuildPatchServices, Error, TEXT("FBuildPatchFileConstructor: Could not create %s"), *Filename);
 				}
 				// Always set
-				InstallerError->SetError(EBuildPatchInstallError::FileConstructionFail, ConstructionErrorCodes::FileCreateFail);
+				InstallerError->SetError(EBuildPatchInstallError::FileConstructionFail, ConstructionErrorCodes::FileCreateFail, LastError);
 			}
 		}
 	}
@@ -567,7 +649,10 @@ bool FBuildPatchFileConstructor::ConstructFileFromChunks( const FString& Filenam
 	// Delete the staging file if unsuccessful by means of construction fail (i.e. keep if canceled or download issue)
 	if( !bSuccess && InstallerError->GetErrorType() == EBuildPatchInstallError::FileConstructionFail )
 	{
-		IFileManager::Get().Delete( *NewFilename, false, true );
+		if (!FileSystem->DeleteFile(*NewFilename))
+		{
+			UE_LOG(LogBuildPatchServices, Warning, TEXT("FBuildPatchFileConstructor: Error deleting file: %s (Error Code %i)"), *NewFilename, FPlatformMisc::GetLastError());
+		}
 	}
 
 	return bSuccess;
@@ -577,14 +662,21 @@ bool FBuildPatchFileConstructor::InsertChunkData(const FChunkPart& ChunkPart, FA
 {
 	uint8* Data;
 	uint8* DataStart;
+	ISpeedRecorder::FRecord ActivityRecord;
+	FileConstructorStat->OnChunkGet(ChunkPart.Guid);
 	IChunkDataAccess* ChunkDataAccess = ChunkSource->Get(ChunkPart.Guid);
 	bool bSuccess = ChunkDataAccess != nullptr && !bShouldAbort;
 	if (bSuccess)
 	{
 		ChunkDataAccess->GetDataLock(&Data, nullptr);
+		FileConstructorStat->OnBeforeWrite();
+		ActivityRecord.CyclesStart = FStatsCollector::GetCycles();
 		DataStart = &Data[ChunkPart.Offset];
 		HashState.Update(DataStart, ChunkPart.Size);
 		DestinationFile.Serialize(DataStart, ChunkPart.Size);
+		ActivityRecord.Size = ChunkPart.Size;
+		ActivityRecord.CyclesEnd = FStatsCollector::GetCycles();
+		FileConstructorStat->OnAfterWrite(ActivityRecord);
 		ChunkDataAccess->ReleaseDataLock();
 		bSuccess = ChunkReferenceTracker->PopReference(ChunkPart.Guid);
 	}
diff --git a/Engine/Source/Runtime/Online/BuildPatchServices/Private/BuildPatchFileConstructor.h b/Engine/Source/Runtime/Online/BuildPatchServices/Private/BuildPatchFileConstructor.h
--- a/Engine/Source/Runtime/Online/BuildPatchServices/Private/BuildPatchFileConstructor.h
+++ b/Engine/Source/Runtime/Online/BuildPatchServices/Private/BuildPatchFileConstructor.h
@@ -12,12 +12,15 @@
 #include "BuildPatchProgress.h"
 #include "BuildPatchManifest.h"
 #include "Installer/Controllable.h"
+#include "Common/SpeedRecorder.h"
+#include "BuildPatchInstall.h"
 
 // Forward declarations
 class FBuildPatchAppManifest;
 struct FChunkPart;
 namespace BuildPatchServices
 {
+	class IFileSystem;
 	class IChunkSource;
 	class IChunkReferenceTracker;
 	class IInstallerError;
@@ -25,6 +28,30 @@
 	class IFileConstructorStat;
 
 	/**
+	 * A struct containing the configuration values for a file constructor.
+	 */
+	struct FFileConstructorConfig
+	{
+		// The Manifest for the build we are installing.
+		FBuildPatchAppManifestRef BuildManifest;
+
+		// The Manifest for the build we currently installed, if applicable.
+		FBuildPatchAppManifestPtr CurrentManifest;
+
+		// The root location where the installation is going.
+		FString InstallDirectory;
+
+		// The location where we will store temporary files.
+		FString StagingDirectory;
+
+		// The list of files to be constructed, filename paths should match those contained in manifest.
+		TArray<FString> ConstructList;
+
+		// The install mode used for this installation.
+		EInstallMode InstallMode;
+	};
+
+	/**
 	 * FBuildPatchFileConstructor
 	 * This class controls a thread that constructs files from a file list, given install details, and chunk availability notifications
 	 */
@@ -32,85 +59,19 @@
 		: public FRunnable
 		, public IControllable
 	{
-	private:
-		// Hold a pointer to my thread for easier deleting
-		FRunnableThread* Thread;
-
-		// A flag marking that we a running
-		bool bIsRunning;
-
-		// A flag marking that we initialized correctly
-		bool bIsInited;
-
-		// A flag marking that our init returned a failure (true means failed)
-		bool bInitFailed;
-
-		// A flag marking that we told the chunk cache to queue required downloads
-		bool bIsDownloadStarted;
-
-		// A flag marking that we have made the initial disk space check following resume logic compelte
-		bool bInitialDiskSizeCheck;
-
-		// A flag marking whether we should be paused
-		FThreadSafeBool bIsPaused;
-
-		// A flag marking whether we should abort operations and exit
-		FThreadSafeBool bShouldAbort;
-
-		// A critical section to protect the flags and variables
-		FCriticalSection ThreadLock;
-
-		// The build manifest for the app we are installing
-		FBuildPatchAppManifestRef BuildManifest;
-
-		// The root installation directory 
-		FString InstallDirectory;
-
-		// The directory for staging files
-		FString StagingDirectory;
-
-		// A list of filenames for files that need to be constructed in this build
-		TArray<FString> FilesToConstruct;
-
-		// Pointer to chunk source.
-		IChunkSource* ChunkSource;
-
-		// Pointer to the chunk reference tracker.
-		IChunkReferenceTracker* ChunkReferenceTracker;
-
-		// Pointer to the installer error class.
-		IInstallerError* InstallerError;
-
-		// Pointer to the installer analytics handler.
-		IInstallerAnalytics* InstallerAnalytics;
-
-		// Pointer to the stat class.
-		IFileConstructorStat* FileConstructorStat;
-
-		// Total job size for tracking progress
-		int64 TotalJobSize;
-
-		// Byte processed so far for tracking progress
-		int64 ByteProcessed;
-
-		// A list of filenames for files that have been constructed
-		TArray<FString> FilesConstructed;
-
 	public:
 
 		/**
 		 * Constructor
-		 * @param BuildManifest             The Manifest for the build we are installing.
-		 * @param InstallDirectory          The root location where the installation is going.
-		 * @param StageDirectory            The location where we will store temporary files.
-		 * @param ConstructList             The list of files to be constructed, filename paths should match those contained in manifest.
+		 * @param Configuration             The configuration for the constructor.
+		 * @param FileSystem                The service used to open files.
 		 * @param ChunkSource               Pointer to the chunk source.
 		 * @param ChunkReferenceTracker     Pointer to the chunk reference tracker.
 		 * @param InstallerError            Pointer to the installer error class for reporting fatal errors.
 		 * @param InstallerAnalytics        Pointer to the installer analytics handler for reporting events.
 		 * @param FileConstructorStat       Pointer to the stat class for receiving updates.
 		 */
-		FBuildPatchFileConstructor(FBuildPatchAppManifestRef BuildManifest, FString InstallDirectory, FString StageDirectory, TArray<FString> ConstructList, IChunkSource* ChunkSource, IChunkReferenceTracker* ChunkReferenceTracker, IInstallerError* InstallerError, IInstallerAnalytics* InstallerAnalytics, IFileConstructorStat* FileConstructorStat);
+		FBuildPatchFileConstructor(FFileConstructorConfig Configuration, IFileSystem* FileSystem, IChunkSource* ChunkSource, IChunkReferenceTracker* ChunkReferenceTracker, IInstallerError* InstallerError, IInstallerAnalytics* InstallerAnalytics, IFileConstructorStat* FileConstructorStat);
 
 		/**
 		 * Default Destructor, will delete the allocated Thread
@@ -145,6 +106,13 @@
 		 */
 		bool IsComplete();
 
+		/**
+		 * Broadcasts with full filepath to file that the constructor is about to delete in order to free up space.
+		 * @return	Reference to the event object.
+		 */
+		DECLARE_EVENT_OneParam(FBuildPatchFileConstructor, FOnBeforeDeleteFile, const FString& /*BuildFile*/);
+		FOnBeforeDeleteFile& OnBeforeDeleteFile();
+
 	private:
 
 		/**
@@ -184,5 +152,13 @@
 		int64 GetRemainingBytes();
 
 		/**
+		 * Calculates the minimum required disk space for the remaining work to be completed, based on a current file, and the list of files left in ConstructionStack.
+		 * @param InProgressFile		The filename for the file currently being constructed.
+		 * @param InProgressFileSize	The remaining size required for the file currently being constructed.
+		 * @return the number of bytes required on disk to complete the installation.
+		 */
+		int64 CalculateRequiredDiskSpace(const FString& InProgressFile, int64 InProgressFileSize);
+
+		/**
 		 * Constructs a particular file referenced by the given BuildManifest. The function takes an interface to a class that can provide availability information of chunks so that this
 		 * file construction process can be ran alongside chunk acquisition threads. It will Sleep while waiting for chunks that it needs.
@@ -206,6 +182,67 @@
 		 * @param RootDirectory	 	Directory to make empty
 		 */
 		void DeleteDirectoryContents(const FString& RootDirectory);
+
+	private:
+		// The configuration for the constructor.
+		const FFileConstructorConfig Configuration;
+
+		// Hold a pointer to my thread for easier deleting.
+		FRunnableThread* Thread;
+
+		// A flag marking that we a running.
+		bool bIsRunning;
+
+		// A flag marking that we initialized correctly.
+		bool bIsInited;
+
+		// A flag marking that our init returned a failure (true means failed).
+		bool bInitFailed;
+
+		// A flag marking that we told the chunk cache to queue required downloads.
+		bool bIsDownloadStarted;
+
+		// A flag marking that we have made the initial disk space check following resume logic complete.
+		bool bInitialDiskSizeCheck;
+
+		// A flag marking whether we should be paused.
+		FThreadSafeBool bIsPaused;
+
+		// A flag marking whether we should abort operations and exit.
+		FThreadSafeBool bShouldAbort;
+
+		// A critical section to protect the flags and variables.
+		FCriticalSection ThreadLock;
+
+		// A stack of filenames for files that need to be constructed.
+		TArray<FString> ConstructionStack;
+
+		// Pointer to the file system.
+		IFileSystem* FileSystem;
+
+		// Pointer to chunk source.
+		IChunkSource* ChunkSource;
+
+		// Pointer to the chunk reference tracker.
+		IChunkReferenceTracker* ChunkReferenceTracker;
+
+		// Pointer to the installer error class.
+		IInstallerError* InstallerError;
+
+		// Pointer to the installer analytics handler.
+		IInstallerAnalytics* InstallerAnalytics;
+
+		// Pointer to the stat class.
+		IFileConstructorStat* FileConstructorStat;
+
+		// Total job size for tracking progress.
+		int64 TotalJobSize;
+
+		// Byte processed so far for tracking progress.
+		int64 ByteProcessed;
+
+		// Event executed before deleting an old installation file.
+		FOnBeforeDeleteFile BeforeDeleteFileEvent;
 	};
 
 	/**
@@ -228,5 +265,11 @@
 		virtual void OnResumeCompleted() = 0;
 
 		/**
+		 * Called for each Get made to the chunk source.
+		 * @param ChunkId       The id for the chunk required.
+		 */
+		virtual void OnChunkGet(const FGuid& ChunkId) = 0;
+
+		/**
 		 * Called when a file construction has started.
 		 * @param Filename      The filename of the file.
@@ -249,5 +292,10 @@
 		virtual void OnFileCompleted(const FString& Filename, bool bSuccess) = 0;
 
 		/**
+		 * Called when the construction process completes.
+		 */
+		virtual void OnConstructionCompleted() = 0;
+
+		/**
 		 * Called to update the total amount of bytes which have been constructed.
 		 * @param TotalBytes    The number of bytes constructed so far.
@@ -259,6 +307,39 @@
 		 * @param TotalBytes    The total number of bytes to be constructed.
 		 */
 		virtual void OnTotalRequiredUpdated(int64 TotalBytes) = 0;
+
+		/**
+		 * Called when we are beginning a file administration, such as open, close, seek.
+		 */
+		virtual void OnBeforeAdminister() = 0;
+
+		/**
+		 * Called upon completing an admin operation, with activity recording.
+		 * @param Record        The activity record.
+		 */
+		virtual void OnAfterAdminister(const ISpeedRecorder::FRecord& Record) = 0;
+
+		/**
+		 * Called when we are beginning a read operation.
+		 */
+		virtual void OnBeforeRead() = 0;
+
+		/**
+		 * Called upon completing a read operation, with activity recording.
+		 * @param Record        The activity record.
+		 */
+		virtual void OnAfterRead(const ISpeedRecorder::FRecord& Record) = 0;
+
+		/**
+		 * Called when we are beginning a write operation.
+		 */
+		virtual void OnBeforeWrite() = 0;
+
+		/**
+		 * Called upon completing a write operation, with activity recording.
+		 * @param Record        The activity record.
+		 */
+		virtual void OnAfterWrite(const ISpeedRecorder::FRecord& Record) = 0;
 	};
 
 }
diff --git a/Engine/Source/Runtime/Online/BuildPatchServices/Private/BuildPatchInstaller.cpp b/Engine/Source/Runtime/Online/BuildPatchServices/Private/BuildPatchInstaller.cpp
--- a/Engine/Source/Runtime/Online/BuildPatchServices/Private/BuildPatchInstaller.cpp
+++ b/Engine/Source/Runtime/Online/BuildPatchServices/Private/BuildPatchInstaller.cpp
@@ -6,14 +6,14 @@
 =============================================================================*/
 
 #include "BuildPatchInstaller.h"
-#include "GenericPlatform/GenericPlatformFile.h"
 #include "HAL/PlatformFilemanager.h"
 #include "HAL/FileManager.h"
-#include "Misc/Paths.h"
 #include "HAL/RunnableThread.h"
+#include "Misc/Paths.h"
 #include "Misc/ScopeLock.h"
 #include "Misc/ConfigCacheIni.h"
 #include "Misc/FeedbackContext.h"
+#include "Math/UnitConversion.h"
 #include "Containers/Ticker.h"
 #include "HttpModule.h"
 #include "BuildPatchFileConstructor.h"
@@ -34,11 +34,19 @@
 #include "Installer/ChainedChunkSource.h"
 #include "Installer/Verifier.h"
 #include "Installer/FileAttribution.h"
-#include "Installer/InstallerStatistics.h"
 #include "Installer/InstallerAnalytics.h"
 #include "Installer/Prerequisites.h"
 #include "Installer/MachineConfig.h"
 #include "Installer/MessagePump.h"
+#include "Installer/Statistics/MemoryChunkStoreStatistics.h"
+#include "Installer/Statistics/DiskChunkStoreStatistics.h"
+#include "Installer/Statistics/DownloadServiceStatistics.h"
+#include "Installer/Statistics/ChunkDbChunkSourceStatistics.h"
+#include "Installer/Statistics/InstallChunkSourceStatistics.h"
+#include "Installer/Statistics/CloudChunkSourceStatistics.h"
+#include "Installer/Statistics/FileConstructorStatistics.h"
+#include "Installer/Statistics/VerifierStatistics.h"
+#include "Installer/Statistics/FileOperationTracker.h"
 
 namespace ConfigHelpers
 {
@@ -82,6 +90,174 @@ namespace ConfigHelpers
 	}
 }
 
+namespace InstallerHelpers
+{
+	void LogBuildStatInfo(const FBuildInstallStats& BuildStats)
+	{
+		using namespace BuildPatchServices;
+		UE_LOG(LogBuildPatchServices, Log, TEXT("Build Stat: AppName: %s"), *BuildStats.AppName);
+		UE_LOG(LogBuildPatchServices, Log, TEXT("Build Stat: AppInstalledVersion: %s"), *BuildStats.AppInstalledVersion);
+		UE_LOG(LogBuildPatchServices, Log, TEXT("Build Stat: AppPatchVersion: %s"), *BuildStats.AppPatchVersion);
+		UE_LOG(LogBuildPatchServices, Log, TEXT("Build Stat: CloudDirectory: %s"), *BuildStats.CloudDirectory);
+		UE_LOG(LogBuildPatchServices, Log, TEXT("Build Stat: NumFilesInBuild: %u"), BuildStats.NumFilesInBuild);
+		UE_LOG(LogBuildPatchServices, Log, TEXT("Build Stat: NumFilesOutdated: %u"), BuildStats.NumFilesOutdated);
+		UE_LOG(LogBuildPatchServices, Log, TEXT("Build Stat: NumFilesToRemove: %u"), BuildStats.NumFilesToRemove);
+		UE_LOG(LogBuildPatchServices, Log, TEXT("Build Stat: NumChunksRequired: %u"), BuildStats.NumChunksRequired);
+		UE_LOG(LogBuildPatchServices, Log, TEXT("Build Stat: ChunksQueuedForDownload: %u"), BuildStats.ChunksQueuedForDownload);
+		UE_LOG(LogBuildPatchServices, Log, TEXT("Build Stat: ChunksLocallyAvailable: %u"), BuildStats.ChunksLocallyAvailable);
+		UE_LOG(LogBuildPatchServices, Log, TEXT("Build Stat: ChunksInChunkDbs: %u"), BuildStats.ChunksInChunkDbs);
+		UE_LOG(LogBuildPatchServices, Log, TEXT("Build Stat: NumChunksDownloaded: %u"), BuildStats.NumChunksDownloaded);
+		UE_LOG(LogBuildPatchServices, Log, TEXT("Build Stat: NumChunksRecycled: %u"), BuildStats.NumChunksRecycled);
+		UE_LOG(LogBuildPatchServices, Log, TEXT("Build Stat: NumChunksReadFromChunkDbs: %u"), BuildStats.NumChunksReadFromChunkDbs);
+		UE_LOG(LogBuildPatchServices, Log, TEXT("Build Stat: NumFailedDownloads: %u"), BuildStats.NumFailedDownloads);
+		UE_LOG(LogBuildPatchServices, Log, TEXT("Build Stat: NumBadDownloads: %u"), BuildStats.NumBadDownloads);
+		UE_LOG(LogBuildPatchServices, Log, TEXT("Build Stat: NumAbortedDownloads: %u"), BuildStats.NumAbortedDownloads);
+		UE_LOG(LogBuildPatchServices, Log, TEXT("Build Stat: NumRecycleFailures: %u"), BuildStats.NumRecycleFailures);
+		UE_LOG(LogBuildPatchServices, Log, TEXT("Build Stat: NumChunksStoreBooted: %u"), BuildStats.NumChunksStoreBooted);
+		UE_LOG(LogBuildPatchServices, Log, TEXT("Build Stat: NumDriveStoreChunkLoads: %u"), BuildStats.NumDriveStoreChunkLoads);
+		UE_LOG(LogBuildPatchServices, Log, TEXT("Build Stat: NumDriveStoreLoadFailures: %u"), BuildStats.NumDriveStoreLoadFailures);
+		UE_LOG(LogBuildPatchServices, Log, TEXT("Build Stat: NumChunkDbChunksFailed: %u"), BuildStats.NumChunkDbChunksFailed);
+		UE_LOG(LogBuildPatchServices, Log, TEXT("Build Stat: TotalDownloadedData: %llu"), BuildStats.TotalDownloadedData);
+		UE_LOG(LogBuildPatchServices, Log, TEXT("Build Stat: AverageDownloadSpeed: %s bytes (%s, %s) /sec"), *FText::AsNumber(BuildStats.AverageDownloadSpeed).ToString(), *FText::AsMemory(BuildStats.AverageDownloadSpeed, EMemoryUnitStandard::SI).ToString(), *FText::AsMemory(BuildStats.AverageDownloadSpeed, EMemoryUnitStandard::IEC).ToString());
+		UE_LOG(LogBuildPatchServices, Log, TEXT("Build Stat: PeakDownloadSpeed: %s bytes (%s, %s) /sec"), *FText::AsNumber(BuildStats.PeakDownloadSpeed).ToString(), *FText::AsMemory(BuildStats.PeakDownloadSpeed, EMemoryUnitStandard::SI).ToString(), *FText::AsMemory(BuildStats.PeakDownloadSpeed, EMemoryUnitStandard::IEC).ToString());
+		UE_LOG(LogBuildPatchServices, Log, TEXT("Build Stat: TotalReadData: %llu"), BuildStats.TotalReadData);
+		UE_LOG(LogBuildPatchServices, Log, TEXT("Build Stat: AverageDiskReadSpeed: %s bytes (%s, %s) /sec"), *FText::AsNumber(BuildStats.AverageDiskReadSpeed).ToString(), *FText::AsMemory(BuildStats.AverageDiskReadSpeed, EMemoryUnitStandard::SI).ToString(), *FText::AsMemory(BuildStats.AverageDiskReadSpeed, EMemoryUnitStandard::IEC).ToString());
+		UE_LOG(LogBuildPatchServices, Log, TEXT("Build Stat: PeakDiskReadSpeed: %s bytes (%s, %s) /sec"), *FText::AsNumber(BuildStats.PeakDiskReadSpeed).ToString(), *FText::AsMemory(BuildStats.PeakDiskReadSpeed, EMemoryUnitStandard::SI).ToString(), *FText::AsMemory(BuildStats.PeakDiskReadSpeed, EMemoryUnitStandard::IEC).ToString());
+		UE_LOG(LogBuildPatchServices, Log, TEXT("Build Stat: TotalWrittenData: %llu"), BuildStats.TotalWrittenData);
+		UE_LOG(LogBuildPatchServices, Log, TEXT("Build Stat: AverageDiskWriteSpeed: %s bytes (%s, %s) /sec"), *FText::AsNumber(BuildStats.AverageDiskWriteSpeed).ToString(), *FText::AsMemory(BuildStats.AverageDiskWriteSpeed, EMemoryUnitStandard::SI).ToString(), *FText::AsMemory(BuildStats.AverageDiskWriteSpeed, EMemoryUnitStandard::IEC).ToString());
+		UE_LOG(LogBuildPatchServices, Log, TEXT("Build Stat: PeakDiskWriteSpeed: %s bytes (%s, %s) /sec"), *FText::AsNumber(BuildStats.PeakDiskWriteSpeed).ToString(), *FText::AsMemory(BuildStats.PeakDiskWriteSpeed, EMemoryUnitStandard::SI).ToString(), *FText::AsMemory(BuildStats.PeakDiskWriteSpeed, EMemoryUnitStandard::IEC).ToString());
+		UE_LOG(LogBuildPatchServices, Log, TEXT("Build Stat: NumFilesConstructed: %u"), BuildStats.NumFilesConstructed);
+		UE_LOG(LogBuildPatchServices, Log, TEXT("Build Stat: InitializeTime: %s"), *FPlatformTime::PrettyTime(BuildStats.InitializeTime));
+		UE_LOG(LogBuildPatchServices, Log, TEXT("Build Stat: TheoreticalDownloadTime: %s"), *FPlatformTime::PrettyTime(BuildStats.TheoreticalDownloadTime));
+		UE_LOG(LogBuildPatchServices, Log, TEXT("Build Stat: ConstructTime: %s"), *FPlatformTime::PrettyTime(BuildStats.ConstructTime));
+		UE_LOG(LogBuildPatchServices, Log, TEXT("Build Stat: MoveFromStageTime: %s"), *FPlatformTime::PrettyTime(BuildStats.MoveFromStageTime));
+		UE_LOG(LogBuildPatchServices, Log, TEXT("Build Stat: FileAttributesTime: %s"), *FPlatformTime::PrettyTime(BuildStats.FileAttributesTime));
+		UE_LOG(LogBuildPatchServices, Log, TEXT("Build Stat: VerifyTime: %s"), *FPlatformTime::PrettyTime(BuildStats.VerifyTime));
+		UE_LOG(LogBuildPatchServices, Log, TEXT("Build Stat: CleanUpTime: %s"), *FPlatformTime::PrettyTime(BuildStats.CleanUpTime));
+		UE_LOG(LogBuildPatchServices, Log, TEXT("Build Stat: PrereqTime: %s"), *FPlatformTime::PrettyTime(BuildStats.PrereqTime));
+		UE_LOG(LogBuildPatchServices, Log, TEXT("Build Stat: ProcessPausedTime: %s"), *FPlatformTime::PrettyTime(BuildStats.ProcessPausedTime));
+		UE_LOG(LogBuildPatchServices, Log, TEXT("Build Stat: ProcessActiveTime: %s"), *FPlatformTime::PrettyTime(BuildStats.ProcessActiveTime));
+		UE_LOG(LogBuildPatchServices, Log, TEXT("Build Stat: ProcessExecuteTime: %s"), *FPlatformTime::PrettyTime(BuildStats.ProcessExecuteTime));
+		UE_LOG(LogBuildPatchServices, Log, TEXT("Build Stat: ProcessSuccess: %s"), BuildStats.ProcessSuccess ? TEXT("TRUE") : TEXT("FALSE"));
+		UE_LOG(LogBuildPatchServices, Log, TEXT("Build Stat: ErrorCode: %s"), *BuildStats.ErrorCode);
+		UE_LOG(LogBuildPatchServices, Log, TEXT("Build Stat: FailureReasonText: %s"), *BuildStats.FailureReasonText.BuildSourceString());
+		UE_LOG(LogBuildPatchServices, Log, TEXT("Build Stat: FailureType: %s"), *EnumToString(BuildStats.FailureType));
+		UE_LOG(LogBuildPatchServices, Log, TEXT("Build Stat: NumInstallRetries: %u"), BuildStats.NumInstallRetries);
+		check(BuildStats.NumInstallRetries == BuildStats.RetryFailureTypes.Num() && BuildStats.NumInstallRetries == BuildStats.RetryErrorCodes.Num());
+		for (uint32 RetryIdx = 0; RetryIdx < BuildStats.NumInstallRetries; ++RetryIdx)
+		{
+			UE_LOG(LogBuildPatchServices, Log, TEXT("Build Stat: RetryFailureType %u: %s"), RetryIdx, *EnumToString(BuildStats.RetryFailureTypes[RetryIdx]));
+			UE_LOG(LogBuildPatchServices, Log, TEXT("Build Stat: RetryErrorCodes %u: %s"), RetryIdx, *BuildStats.RetryErrorCodes[RetryIdx]);
+		}
+		UE_LOG(LogBuildPatchServices, Log, TEXT("Build Stat: FinalProgressValue: %f"), BuildStats.FinalProgress);
+		UE_LOG(LogBuildPatchServices, Log, TEXT("Build Stat: OverallRequestSuccessRate: %f"), BuildStats.OverallRequestSuccessRate);
+		UE_LOG(LogBuildPatchServices, Log, TEXT("Build Stat: ExcellentDownloadHealthTime: %s"), *FPlatformTime::PrettyTime(BuildStats.ExcellentDownloadHealthTime));
+		UE_LOG(LogBuildPatchServices, Log, TEXT("Build Stat: GoodDownloadHealthTime: %s"), *FPlatformTime::PrettyTime(BuildStats.GoodDownloadHealthTime));
+		UE_LOG(LogBuildPatchServices, Log, TEXT("Build Stat: OkDownloadHealthTime: %s"), *FPlatformTime::PrettyTime(BuildStats.OkDownloadHealthTime));
+		UE_LOG(LogBuildPatchServices, Log, TEXT("Build Stat: PoorDownloadHealthTime: %s"), *FPlatformTime::PrettyTime(BuildStats.PoorDownloadHealthTime));
+		UE_LOG(LogBuildPatchServices, Log, TEXT("Build Stat: DisconnectedDownloadHealthTime: %s"), *FPlatformTime::PrettyTime(BuildStats.DisconnectedDownloadHealthTime));
+		UE_LOG(LogBuildPatchServices, Log, TEXT("Build Stat: AverageMemoryStoreUse: %.1f"), BuildStats.AverageMemoryStoreUse);
+		UE_LOG(LogBuildPatchServices, Log, TEXT("Build Stat: PeakMemoryStoreUse: %u"), BuildStats.PeakMemoryStoreUse);
+		UE_LOG(LogBuildPatchServices, Log, TEXT("Build Stat: AverageMemoryStoreRetained: %.1f"), BuildStats.AverageMemoryStoreRetained);
+		UE_LOG(LogBuildPatchServices, Log, TEXT("Build Stat: PeakMemoryStoreRetained: %u"), BuildStats.PeakMemoryStoreRetained);
+		UE_LOG(LogBuildPatchServices, Log, TEXT("Build Stat: MemoryStoreSize: %u"), BuildStats.MemoryStoreSize);
+	}
+
+	void LogBuildConfiguration(const BuildPatchServices::FInstallerConfiguration& InstallerConfiguration)
+	{
+		using namespace BuildPatchServices;
+		UE_LOG(LogBuildPatchServices, Log, TEXT("Build Config: InstallManifest: 0x%p"), &InstallerConfiguration.InstallManifest.Get());
+		UE_LOG(LogBuildPatchServices, Log, TEXT("Build Config: InstallManifest Version: %s"), *InstallerConfiguration.InstallManifest->GetVersionString());
+		UE_LOG(LogBuildPatchServices, Log, TEXT("Build Config: CurrentManifest: 0x%p"), InstallerConfiguration.CurrentManifest.Get());
+		UE_LOG(LogBuildPatchServices, Log, TEXT("Build Config: CurrentManifest Version: %s"), InstallerConfiguration.CurrentManifest.IsValid() ? *InstallerConfiguration.CurrentManifest->GetVersionString() : TEXT("NULL"));
+
+		UE_LOG(LogBuildPatchServices, Log, TEXT("Build Config: InstallDirectory: %s"), *InstallerConfiguration.InstallDirectory);
+		UE_LOG(LogBuildPatchServices, Log, TEXT("Build Config: StagingDirectory: %s"), *InstallerConfiguration.StagingDirectory);
+		UE_LOG(LogBuildPatchServices, Log, TEXT("Build Config: BackupDirectory: %s"), *InstallerConfiguration.BackupDirectory);
+
+		for (const FString& DatabaseFile : InstallerConfiguration.ChunkDatabaseFiles)
+		{
+			UE_LOG(LogBuildPatchServices, Log, TEXT("Build Config: ChunkDatabaseFile: %s"), *DatabaseFile);
+		}
+		
+		for (const FString& CloudDirectory : InstallerConfiguration.CloudDirectories)
+		{
+			UE_LOG(LogBuildPatchServices, Log, TEXT("Build Config: CloudDirectories: %s"), *CloudDirectory);
+		}
+
+		for (const FString& Tag : InstallerConfiguration.InstallTags)
+		{
+			UE_LOG(LogBuildPatchServices, Log, TEXT("Build Config: InstallTags: %s"), *Tag);
+		}
+
+		TSet<FString> ValidTags;
+		InstallerConfiguration.InstallManifest->GetFileTagList(ValidTags);
+		for (const FString& Tag : ValidTags)
+		{
+			UE_LOG(LogBuildPatchServices, Log, TEXT("Build Config: ValidTags: %s"), *Tag);
+		}
+
+		UE_LOG(LogBuildPatchServices, Log, TEXT("Build Config: InstallMode: %s"), *EnumToString(InstallerConfiguration.InstallMode));
+		UE_LOG(LogBuildPatchServices, Log, TEXT("Build Config: VerifyMode: %s"), *EnumToString(InstallerConfiguration.VerifyMode));
+		UE_LOG(LogBuildPatchServices, Log, TEXT("Build Config: bIsRepair: %s"), (InstallerConfiguration.bIsRepair) ? TEXT("true") : TEXT("false"));
+		UE_LOG(LogBuildPatchServices, Log, TEXT("Build Config: bRunRequiredPrereqs: %s"), (InstallerConfiguration.bRunRequiredPrereqs) ? TEXT("true") : TEXT("false"));
+		UE_LOG(LogBuildPatchServices, Log, TEXT("Build Config: bAllowConcurrentExecution: %s"), (InstallerConfiguration.bAllowConcurrentExecution) ? TEXT("true") : TEXT("false"));
+	}
+
+	TSet<FGuid> GetMultipleReferencedChunks(const FBuildPatchAppManifestRef& Manifest)
+	{
+		TSet<FGuid> MultipleReferencedChunks;
+		TSet<FGuid> AllReferencedChunks;
+		TArray<FString> AllFiles;
+		Manifest->GetFileList(AllFiles);
+		for (const FString& File : AllFiles)
+		{
+			const FFileManifest* NewFileManifest = Manifest->GetFileManifest(File);
+			if (NewFileManifest != nullptr)
+			{
+				for (const FChunkPart& ChunkPart : NewFileManifest->FileChunkParts)
+				{
+					if (AllReferencedChunks.Contains(ChunkPart.Guid))
+					{
+						MultipleReferencedChunks.Add(ChunkPart.Guid);
+					}
+					else
+					{
+						AllReferencedChunks.Add(ChunkPart.Guid);
+					}
+				}
+			}
+		}
+		return MultipleReferencedChunks;
+	}
+
+	const TCHAR* GetVerifyErrorCode(const BuildPatchServices::EVerifyResult& VerifyResult)
+	{
+		using namespace BuildPatchServices;
+		switch (VerifyResult)
+		{
+			case EVerifyResult::FileMissing: return VerifyErrorCodes::FileMissing;
+			case EVerifyResult::OpenFileFailed: return VerifyErrorCodes::OpenFileFailed;
+			case EVerifyResult::HashCheckFailed: return VerifyErrorCodes::HashCheckFailed;
+			case EVerifyResult::FileSizeFailed: return VerifyErrorCodes::FileSizeFailed;
+		}
+
+		return VerifyErrorCodes::UnknownFail;
+	}
+
+	void LogAdditionalVerifyErrors(BuildPatchServices::EVerifyError Error, int32 Count)
+	{
+		using namespace BuildPatchServices;
+		EVerifyResult VerifyResult;
+		if (TryConvertToVerifyResult(Error, VerifyResult))
+		{
+			FString Prefix = InstallErrorPrefixes::ErrorTypeStrings[static_cast<int32>(EBuildPatchInstallError::BuildVerifyFail)];
+			FString Suffix = InstallerHelpers::GetVerifyErrorCode(VerifyResult);
+			UE_LOG(LogBuildPatchServices, Log, TEXT("Build verification error encountered: %s: %d"), *(Prefix + Suffix), Count);
+		}
+	}
+}
+
 namespace BuildPatchServices
 {
 	struct FScopedControllables
@@ -145,6 +321,12 @@ namespace BuildPatchServices
 		}
 	};
 
+	enum class EMemoryStore : int32
+	{
+		Install = 0,
+		Cloud
+	};
+
 	/* FBuildPatchInstaller implementation
 	*****************************************************************************/
 	FBuildPatchInstaller::FBuildPatchInstaller(FInstallerConfiguration InConfiguration, TMap<FString, FBuildPatchAppManifestRef> InInstallationInfo, const FString& InLocalMachineConfigFile, TSharedPtr<IAnalyticsProvider> InAnalytics, TSharedPtr<FHttpServiceTracker> InHttpTracker, FBuildPatchBoolManifestDelegate InOnCompleteDelegate)
@@ -163,6 +345,7 @@ namespace BuildPatchServices
 		, bIsRunning(false)
 		, bIsInited(false)
 		, bFirstInstallIteration(true)
+		, PreviousTotalDownloadRequired(0)
 		, BuildStats()
 		, BuildProgress()
 		, bIsPaused(false)
@@ -179,11 +362,30 @@ namespace BuildPatchServices
 		, Analytics(MoveTemp(InAnalytics))
 		, HttpTracker(MoveTemp(InHttpTracker))
 		, InstallerAnalytics(FInstallerAnalyticsFactory::Create(Analytics.Get(), HttpTracker.Get()))
-		, InstallerStatistics(FInstallerStatisticsFactory::Create(InstallerAnalytics.Get(), &BuildProgress))
-		, DownloadService(FDownloadServiceFactory::Create(FTicker::GetCoreTicker(), HttpManager.Get(), FileSystem.Get(), InstallerStatistics->GetDownloadServiceStat(), InstallerAnalytics.Get()))
+		, FileOperationTracker(FFileOperationTrackerFactory::Create(FTicker::GetCoreTicker(), &NewBuildManifest.Get()))
+		, MemoryChunkStoreAggregateStatistics(FMemoryChunkStoreAggregateStatisticsFactory::Create(InstallerHelpers::GetMultipleReferencedChunks(NewBuildManifest), FileOperationTracker.Get()))
+		, DiskChunkStoreStatistics(FDiskChunkStoreStatisticsFactory::Create(InstallerAnalytics.Get(), FileOperationTracker.Get()))
+		, DownloadSpeedRecorder(FSpeedRecorderFactory::Create())
+		, DiskReadSpeedRecorder(FSpeedRecorderFactory::Create())
+		, DiskWriteSpeedRecorder(FSpeedRecorderFactory::Create())
+		, ChunkDbReadSpeedRecorder(FSpeedRecorderFactory::Create())
+		, DownloadServiceStatistics(FDownloadServiceStatisticsFactory::Create(DownloadSpeedRecorder.Get(), InstallerAnalytics.Get(), &NewBuildManifest.Get()))
+		, ChunkDbChunkSourceStatistics(FChunkDbChunkSourceStatisticsFactory::Create(ChunkDbReadSpeedRecorder.Get(), FileOperationTracker.Get()))
+		, InstallChunkSourceStatistics(FInstallChunkSourceStatisticsFactory::Create(DiskReadSpeedRecorder.Get(), InstallerAnalytics.Get(), FileOperationTracker.Get()))
+		, CloudChunkSourceStatistics(FCloudChunkSourceStatisticsFactory::Create(InstallerAnalytics.Get(), &BuildProgress, FileOperationTracker.Get()))
+		, FileConstructorStatistics(FFileConstructorStatisticsFactory::Create(DiskReadSpeedRecorder.Get(), DiskWriteSpeedRecorder.Get(), &BuildProgress, FileOperationTracker.Get()))
+		, VerifierStatistics(FVerifierStatisticsFactory::Create(DiskReadSpeedRecorder.Get(), &BuildProgress, FileOperationTracker.Get()))
+		, DownloadService(FDownloadServiceFactory::Create(FTicker::GetCoreTicker(), HttpManager.Get(), FileSystem.Get(), DownloadServiceStatistics.Get(), InstallerAnalytics.Get()))
 		, MessagePump(FMessagePumpFactory::Create())
 		, Controllables()
 	{
+		UE_LOG(LogBuildPatchServices, Log, TEXT("FBuildPatchInstaller CTOR Configuration.InstallManifest Check:   0x%p"), &Configuration.InstallManifest.Get());
+		UE_LOG(LogBuildPatchServices, Log, TEXT("FBuildPatchInstaller CTOR Configuration.CurrentManifest Check:   0x%p"), Configuration.CurrentManifest.Get());
+		UE_LOG(LogBuildPatchServices, Log, TEXT("FBuildPatchInstaller CTOR CurrentBuildManifest Check:   0x%p"), CurrentBuildManifest.Get());
+		UE_LOG(LogBuildPatchServices, Log, TEXT("FBuildPatchInstaller CTOR NewBuildManifest Check:   0x%p"), &NewBuildManifest.Get());
+
+		FPaths::NormalizeDirectoryName(Configuration.InstallDirectory);
+		FPaths::CollapseRelativeDirectories(Configuration.InstallDirectory);
 		if (InstallationInfo.Contains(Configuration.InstallDirectory) == false && CurrentBuildManifest.IsValid())
 		{
 			InstallationInfo.Add(Configuration.InstallDirectory, CurrentBuildManifest.ToSharedRef());
@@ -200,13 +402,90 @@ namespace BuildPatchServices
 	void FBuildPatchInstaller::PreExit()
 	{
 		// Set shutdown error so any running threads will exit if no error has already been set.
-		InstallerError->SetError(EBuildPatchInstallError::ApplicationClosing, ApplicationClosedErrorCodes::ApplicationClosed);
+		if (bIsRunning)
+		{
+			InstallerError->SetError(EBuildPatchInstallError::ApplicationClosing, ApplicationClosedErrorCodes::ApplicationClosed);
+		}
 		if (Thread != nullptr)
 		{
 			Thread->WaitForCompletion();
 			delete Thread;
 			Thread = nullptr;
 		}
+		if (InstallerAnalytics.IsValid())
+		{
+			InstallerAnalytics->Flush();
+		}
+	}
+
+	const IFileOperationTracker* FBuildPatchInstaller::GetFileOperationTracker() const
+	{
+		return FileOperationTracker.Get();
+	}
+
+	const ISpeedRecorder* FBuildPatchInstaller::GetDownloadSpeedRecorder() const
+	{
+		return DownloadSpeedRecorder.Get();
+	}
+
+	const ISpeedRecorder* FBuildPatchInstaller::GetDiskReadSpeedRecorder() const
+	{
+		return DiskReadSpeedRecorder.Get();
+	}
+
+	const ISpeedRecorder* FBuildPatchInstaller::GetChunkDbReadSpeedRecorder() const
+	{
+		return ChunkDbReadSpeedRecorder.Get();
+	}
+
+	const ISpeedRecorder* FBuildPatchInstaller::GetDiskWriteSpeedRecorder() const
+	{
+		return DiskWriteSpeedRecorder.Get();
+	}
+
+	const IDownloadServiceStatistics* FBuildPatchInstaller::GetDownloadServiceStatistics() const
+	{
+		return DownloadServiceStatistics.Get();
+	}
+
+	const IInstallChunkSourceStatistics* FBuildPatchInstaller::GetInstallChunkSourceStatistics() const
+	{
+		return InstallChunkSourceStatistics.Get();
+	}
+
+	const ICloudChunkSourceStatistics* FBuildPatchInstaller::GetCloudChunkSourceStatistics() const
+	{
+		return CloudChunkSourceStatistics.Get();
+	}
+
+	const IFileConstructorStatistics* FBuildPatchInstaller::GetFileConstructorStatistics() const
+	{
+		return FileConstructorStatistics.Get();
+	}
+
+	const IVerifierStatistics* FBuildPatchInstaller::GetVerifierStatistics() const
+	{
+		return VerifierStatistics.Get();
+	}
+
+	const IMemoryChunkStoreStatistics* FBuildPatchInstaller::GetCloudMemoryChunkStoreStatistics() const
+	{
+		return MemoryChunkStoreAggregateStatistics->Expose(static_cast<int32>(EMemoryStore::Cloud));
+	}
+
+	const IMemoryChunkStoreStatistics* FBuildPatchInstaller::GetInstallMemoryChunkStoreStatistics() const
+	{
+		return MemoryChunkStoreAggregateStatistics->Expose(static_cast<int32>(EMemoryStore::Install));
+	}
+
+	const IDiskChunkStoreStatistics* FBuildPatchInstaller::GetDiskChunkStoreStatistics() const
+	{
+		return DiskChunkStoreStatistics.Get();
+	}
+
+	const FInstallerConfiguration& FBuildPatchInstaller::GetConfiguration() const
+	{
+		return Configuration;
 	}
 
 	bool FBuildPatchInstaller::StartInstallation()
@@ -233,6 +512,7 @@ namespace BuildPatchServices
 	bool FBuildPatchInstaller::Initialize()
 	{
 		bool bInstallerInitSuccess = true;
+		InstallerHelpers::LogBuildConfiguration(Configuration);
 
 		// Check provided tags are all valid.
 		TSet<FString> ValidTags;
@@ -240,7 +520,7 @@ namespace BuildPatchServices
 		if (Configuration.InstallTags.Difference(ValidTags).Num() > 0)
 		{
 			UE_LOG(LogBuildPatchServices, Error, TEXT("Installer configuration: Invalid InstallTags provided."));
-			InstallerError->SetError(EBuildPatchInstallError::InitializationError, InitializationErrorCodes::InvalidInstallTags,
+			InstallerError->SetError(EBuildPatchInstallError::InitializationError, InitializationErrorCodes::InvalidInstallTags, 0,
 				NSLOCTEXT("BuildPatchInstallError", "InvalidInstallTags", "This installation could not continue due to a configuration issue. Please contact support."));
 			bInstallerInitSuccess = false;
 		}
@@ -258,7 +538,7 @@ namespace BuildPatchServices
 		if (!IFileManager::Get().DirectoryExists(*Configuration.InstallDirectory))
 		{
 			UE_LOG(LogBuildPatchServices, Error, TEXT("Installer setup: Inability to create InstallDirectory %s."), *Configuration.InstallDirectory);
-			InstallerError->SetError(EBuildPatchInstallError::InitializationError, InitializationErrorCodes::MissingInstallDirectory,
+			InstallerError->SetError(EBuildPatchInstallError::InitializationError, InitializationErrorCodes::MissingInstallDirectory, 0,
 				FText::Format(NSLOCTEXT("BuildPatchInstallError", "MissingInstallDirectory", "The installation directory could not be created.\n{0}"), FText::FromString(Configuration.InstallDirectory)));
 			bInstallerInitSuccess = false;
 		}
@@ -268,11 +548,20 @@ namespace BuildPatchServices
 		if (!IFileManager::Get().DirectoryExists(*Configuration.StagingDirectory))
 		{
 			UE_LOG(LogBuildPatchServices, Error, TEXT("Installer setup: Inability to create StagingDirectory %s."), *Configuration.StagingDirectory);
-			InstallerError->SetError(EBuildPatchInstallError::InitializationError, InitializationErrorCodes::MissingStageDirectory,
+			InstallerError->SetError(EBuildPatchInstallError::InitializationError, InitializationErrorCodes::MissingStageDirectory, 0,
 				FText::Format(NSLOCTEXT("BuildPatchInstallError", "MissingStageDirectory", "The following directory could not be created.\n{0}"), FText::FromString(Configuration.StagingDirectory)));
 			bInstallerInitSuccess = false;
 		}
 
+		// Make sure that we have a prereq if we've specified a prereq only install.
+		if (Configuration.InstallMode == EInstallMode::PrereqOnly && NewBuildManifest->GetPrereqPath().IsEmpty())
+		{
+			UE_LOG(LogBuildPatchServices, Error, TEXT("Installer setup: PrereqOnly install selected for manifest with no prereq."));
+			InstallerError->SetError(EBuildPatchInstallError::InitializationError, InitializationErrorCodes::MissingPrereqForPrereqOnlyInstall, 0,
+				NSLOCTEXT("BuildPatchInstallError", "MissingPrereqForPrereqOnlyInstall", "This installation could not continue due to a prerequisite configuration issue. Please contact support."));
+			bInstallerInitSuccess = false;
+		}
+
 		// Init build statistics that are known.
 		{
 			FScopeLock Lock(&ThreadLock);
@@ -295,7 +584,11 @@ namespace BuildPatchServices
 	{
 		// Make sure this function can never be parallelized
 		static FCriticalSection SingletonFunctionLockCS;
-		FScopeLock SingletonFunctionLock(&SingletonFunctionLockCS);
+		const bool bShouldLock = !Configuration.bAllowConcurrentExecution;
+		if (bShouldLock)
+		{
+			SingletonFunctionLockCS.Lock();
+		}
 		bIsRunning = true;
 		ProcessExecuteTimer.Start();
 		ProcessActiveTimer.Start();
@@ -350,7 +643,7 @@ namespace BuildPatchServices
 				if (bInstallSuccess)
 				{
 					CleanUpTimer.Start();
-					if (Configuration.bStageOnly)
+					if (Configuration.InstallMode == EInstallMode::StageFiles)
 					{
 						UE_LOG(LogBuildPatchServices, Log, TEXT("Deleting litter from staging area."));
 						IFileManager::Get().DeleteDirectory(*DataStagingDir, false, true);
@@ -435,65 +728,14 @@ namespace BuildPatchServices
 			BuildStats.ErrorCode = InstallerError->GetErrorCode();
 			BuildStats.FailureReasonText = InstallerError->GetErrorText();
 			BuildStats.FailureType = InstallerError->GetErrorType();
-
-			// Log stats
-			UE_LOG(LogBuildPatchServices, Log, TEXT("Build Stat: AppName: %s"), *BuildStats.AppName);
-			UE_LOG(LogBuildPatchServices, Log, TEXT("Build Stat: AppInstalledVersion: %s"), *BuildStats.AppInstalledVersion);
-			UE_LOG(LogBuildPatchServices, Log, TEXT("Build Stat: AppPatchVersion: %s"), *BuildStats.AppPatchVersion);
-			UE_LOG(LogBuildPatchServices, Log, TEXT("Build Stat: CloudDirectory: %s"), *BuildStats.CloudDirectory);
-			UE_LOG(LogBuildPatchServices, Log, TEXT("Build Stat: NumFilesInBuild: %u"), BuildStats.NumFilesInBuild);
-			UE_LOG(LogBuildPatchServices, Log, TEXT("Build Stat: NumFilesOutdated: %u"), BuildStats.NumFilesOutdated);
-			UE_LOG(LogBuildPatchServices, Log, TEXT("Build Stat: NumFilesToRemove: %u"), BuildStats.NumFilesToRemove);
-			UE_LOG(LogBuildPatchServices, Log, TEXT("Build Stat: NumChunksRequired: %u"), BuildStats.NumChunksRequired);
-			UE_LOG(LogBuildPatchServices, Log, TEXT("Build Stat: ChunksQueuedForDownload: %u"), BuildStats.ChunksQueuedForDownload);
-			UE_LOG(LogBuildPatchServices, Log, TEXT("Build Stat: ChunksLocallyAvailable: %u"), BuildStats.ChunksLocallyAvailable);
-			UE_LOG(LogBuildPatchServices, Log, TEXT("Build Stat: ChunksInChunkDbs: %u"), BuildStats.ChunksInChunkDbs);
-			UE_LOG(LogBuildPatchServices, Log, TEXT("Build Stat: NumChunksDownloaded: %u"), BuildStats.NumChunksDownloaded);
-			UE_LOG(LogBuildPatchServices, Log, TEXT("Build Stat: NumChunksRecycled: %u"), BuildStats.NumChunksRecycled);
-			UE_LOG(LogBuildPatchServices, Log, TEXT("Build Stat: NumChunksReadFromChunkDbs: %u"), BuildStats.NumChunksReadFromChunkDbs);
-			UE_LOG(LogBuildPatchServices, Log, TEXT("Build Stat: NumChunksCacheBooted: %u"), BuildStats.NumChunksCacheBooted);
-			UE_LOG(LogBuildPatchServices, Log, TEXT("Build Stat: NumDriveCacheChunkLoads: %u"), BuildStats.NumDriveCacheChunkLoads);
-			UE_LOG(LogBuildPatchServices, Log, TEXT("Build Stat: NumFailedDownloads: %u"), BuildStats.NumFailedDownloads);
-			UE_LOG(LogBuildPatchServices, Log, TEXT("Build Stat: NumBadDownloads: %u"), BuildStats.NumBadDownloads);
-			UE_LOG(LogBuildPatchServices, Log, TEXT("Build Stat: NumAbortedDownloads: %u"), BuildStats.NumAbortedDownloads);
-			UE_LOG(LogBuildPatchServices, Log, TEXT("Build Stat: NumRecycleFailures: %u"), BuildStats.NumRecycleFailures);
-			UE_LOG(LogBuildPatchServices, Log, TEXT("Build Stat: NumDriveCacheLoadFailures: %u"), BuildStats.NumDriveCacheLoadFailures);
-			UE_LOG(LogBuildPatchServices, Log, TEXT("Build Stat: NumChunkDbChunksFailed: %u"), BuildStats.NumChunkDbChunksFailed);
-			UE_LOG(LogBuildPatchServices, Log, TEXT("Build Stat: TotalDownloadedData: %lld"), BuildStats.TotalDownloadedData);
-			UE_LOG(LogBuildPatchServices, Log, TEXT("Build Stat: AverageDownloadSpeed: %.3f MB/sec"), BuildStats.AverageDownloadSpeed / 1024.0 / 1024.0);
-			UE_LOG(LogBuildPatchServices, Log, TEXT("Build Stat: TheoreticalDownloadTime: %s"), *FPlatformTime::PrettyTime(BuildStats.TheoreticalDownloadTime));
-			UE_LOG(LogBuildPatchServices, Log, TEXT("Build Stat: InitializeTime: %s"), *FPlatformTime::PrettyTime(BuildStats.InitializeTime));
-			UE_LOG(LogBuildPatchServices, Log, TEXT("Build Stat: ConstructTime: %s"), *FPlatformTime::PrettyTime(BuildStats.ConstructTime));
-			UE_LOG(LogBuildPatchServices, Log, TEXT("Build Stat: MoveFromStageTime: %s"), *FPlatformTime::PrettyTime(BuildStats.MoveFromStageTime));
-			UE_LOG(LogBuildPatchServices, Log, TEXT("Build Stat: FileAttributesTime: %s"), *FPlatformTime::PrettyTime(BuildStats.FileAttributesTime));
-			UE_LOG(LogBuildPatchServices, Log, TEXT("Build Stat: VerifyTime: %s"), *FPlatformTime::PrettyTime(BuildStats.VerifyTime));
-			UE_LOG(LogBuildPatchServices, Log, TEXT("Build Stat: CleanUpTime: %s"), *FPlatformTime::PrettyTime(BuildStats.CleanUpTime));
-			UE_LOG(LogBuildPatchServices, Log, TEXT("Build Stat: PrereqTime: %s"), *FPlatformTime::PrettyTime(BuildStats.PrereqTime));
-			UE_LOG(LogBuildPatchServices, Log, TEXT("Build Stat: ProcessPausedTime: %s"), *FPlatformTime::PrettyTime(BuildStats.ProcessPausedTime));
-			UE_LOG(LogBuildPatchServices, Log, TEXT("Build Stat: ProcessActiveTime: %s"), *FPlatformTime::PrettyTime(BuildStats.ProcessActiveTime));
-			UE_LOG(LogBuildPatchServices, Log, TEXT("Build Stat: ProcessExecuteTime: %s"), *FPlatformTime::PrettyTime(BuildStats.ProcessExecuteTime));
-			UE_LOG(LogBuildPatchServices, Log, TEXT("Build Stat: ProcessSuccess: %s"), BuildStats.ProcessSuccess ? TEXT("TRUE") : TEXT("FALSE"));
-			UE_LOG(LogBuildPatchServices, Log, TEXT("Build Stat: ErrorCode: %s"), *BuildStats.ErrorCode);
-			UE_LOG(LogBuildPatchServices, Log, TEXT("Build Stat: FailureReasonText: %s"), *BuildStats.FailureReasonText.BuildSourceString());
-			UE_LOG(LogBuildPatchServices, Log, TEXT("Build Stat: FailureType: %s"), *EnumToString(BuildStats.FailureType));
-			UE_LOG(LogBuildPatchServices, Log, TEXT("Build Stat: NumInstallRetries: %d"), BuildStats.NumInstallRetries);
-			check(BuildStats.NumInstallRetries == BuildStats.RetryFailureTypes.Num() && BuildStats.NumInstallRetries == BuildStats.RetryErrorCodes.Num());
-			for (uint32 RetryIdx = 0; RetryIdx < BuildStats.NumInstallRetries; ++RetryIdx)
-			{
-				UE_LOG(LogBuildPatchServices, Log, TEXT("Build Stat: RetryFailureType %d: %s"), RetryIdx, *EnumToString(BuildStats.RetryFailureTypes[RetryIdx]));
-				UE_LOG(LogBuildPatchServices, Log, TEXT("Build Stat: RetryErrorCodes %d: %s"), RetryIdx, *BuildStats.RetryErrorCodes[RetryIdx]);
-			}
-			UE_LOG(LogBuildPatchServices, Log, TEXT("Build Stat: FinalProgressValue: %f"), BuildStats.FinalProgress);
-			UE_LOG(LogBuildPatchServices, Log, TEXT("Build Stat: OverallRequestSuccessRate: %f"), BuildStats.OverallRequestSuccessRate);
-			UE_LOG(LogBuildPatchServices, Log, TEXT("Build Stat: ExcellentDownloadHealthTime: %f"), BuildStats.ExcellentDownloadHealthTime);
-			UE_LOG(LogBuildPatchServices, Log, TEXT("Build Stat: GoodDownloadHealthTime: %f"), BuildStats.GoodDownloadHealthTime);
-			UE_LOG(LogBuildPatchServices, Log, TEXT("Build Stat: OkDownloadHealthTime: %f"), BuildStats.OkDownloadHealthTime);
-			UE_LOG(LogBuildPatchServices, Log, TEXT("Build Stat: PoorDownloadHealthTime: %f"), BuildStats.PoorDownloadHealthTime);
-			UE_LOG(LogBuildPatchServices, Log, TEXT("Build Stat: DisconnectedDownloadHealthTime: %f"), BuildStats.DisconnectedDownloadHealthTime);
 		}
 
 		// Mark that we are done
 		bIsRunning = false;
+		if (bShouldLock)
+		{
+			SingletonFunctionLockCS.Unlock();
+		}
 		return bSuccess ? 0 : 1;
 	}
 
@@ -646,10 +888,31 @@ namespace BuildPatchServices
 		TaggedFiles.Empty();
 		NewBuildManifest->GetTaggedFileList(Configuration.InstallTags, TaggedFiles);
 
+		const bool bIsPrereqOnly = Configuration.InstallMode == EInstallMode::PrereqOnly;
+		const bool bHasCorruptFiles = CorruptFiles.Num() > 0;
+
+		// Get the list of files actually needing construction
+		FilesToConstruct.Empty();
+		if (bHasCorruptFiles)
+		{
+			FilesToConstruct.Append(CorruptFiles);
+		}
+		else if (bIsPrereqOnly)
+		{
+			FilesToConstruct.Add(NewBuildManifest->GetPrereqPath());
+		}
+		else
+		{
+			TSet<FString> OutdatedFiles;
+			NewBuildManifest->GetOutdatedFiles(CurrentBuildManifest, Configuration.InstallDirectory, OutdatedFiles);
+			FilesToConstruct = OutdatedFiles.Intersect(TaggedFiles);
+		}
+		FilesToConstruct.Sort(TLess<FString>());
+		UE_LOG(LogBuildPatchServices, Log, TEXT("Requiring %d files"), FilesToConstruct.Num());
+
 		// Check if we should skip out of this process due to existing installation,
 		// that will mean we start with the verification stage
-		bool bFirstTimeRun = CorruptFiles.Num() == 0;
-		if (bFirstTimeRun && CheckForExternallyInstalledFiles())
+		if (!bHasCorruptFiles && (bIsPrereqOnly || CheckForExternallyInstalledFiles()))
 		{
 			UE_LOG(LogBuildPatchServices, Log, TEXT("Detected previous staging completed, or existing files in target directory"));
 			// Set weights for verify only
@@ -667,32 +930,42 @@ namespace BuildPatchServices
 			return true;
 		}
 
-		// Get the list of files actually needing construction
-		FilesToConstruct.Empty();
-		if (CorruptFiles.Num() > 0)
-		{
-			FilesToConstruct.Append(CorruptFiles);
-		}
-		else
+		if (!bHasCorruptFiles)
 		{
-			TSet<FString> OutdatedFiles;
-			NewBuildManifest->GetOutdatedFiles(CurrentBuildManifest, Configuration.InstallDirectory, OutdatedFiles);
-			FilesToConstruct = OutdatedFiles.Intersect(TaggedFiles);
+			FScopeLock Lock(&ThreadLock);
+			BuildStats.NumFilesOutdated = FilesToConstruct.Num();
 		}
-		FilesToConstruct.Sort(TLess<FString>());
-		UE_LOG(LogBuildPatchServices, Log, TEXT("Requiring %d files"), FilesToConstruct.Num());
 
 		// Make sure all the files won't exceed the maximum path length
 		for (const FString& FileToConstruct : FilesToConstruct)
 		{
-			if ((InstallStagingDir / FileToConstruct).Len() >= PLATFORM_MAX_FILEPATH_LENGTH)
+			const FString InstallConstructionFile = Configuration.InstallDirectory / FileToConstruct;
+			const FString StagedConstructionFile = InstallStagingDir / FileToConstruct;
+			if (InstallConstructionFile.Len() >= PLATFORM_MAX_FILEPATH_LENGTH)
+			{
+				UE_LOG(LogBuildPatchServices, Error, TEXT("Could not create new file due to exceeding maximum path length %s"), *InstallConstructionFile);
+				InstallerError->SetError(EBuildPatchInstallError::PathLengthExceeded, PathLengthErrorCodes::InstallDirectory);
+				return false;
+			}
+			if (StagedConstructionFile.Len() >= PLATFORM_MAX_FILEPATH_LENGTH)
 			{
-				GWarn->Logf(TEXT("BuildPatchServices: ERROR: Could not create new file due to exceeding maximum path length %s"), *(InstallStagingDir / FileToConstruct));
+				UE_LOG(LogBuildPatchServices, Error, TEXT("Could not create new file due to exceeding maximum path length %s"), *StagedConstructionFile);
 				InstallerError->SetError(EBuildPatchInstallError::PathLengthExceeded, PathLengthErrorCodes::StagingDirectory);
 				return false;
 			}
 		}
 
+		// Set initial states on IO state tracker.
+		const bool bVerifyAllFiles = Configuration.VerifyMode == EVerifyMode::ShaVerifyAllFiles || Configuration.VerifyMode == EVerifyMode::FileSizeCheckAllFiles;
+		const EFileOperationState UntouchedFileState = (Configuration.bIsRepair || bVerifyAllFiles) ? EFileOperationState::Installed : EFileOperationState::Complete;
+		for (const FString& TaggedFile : TaggedFiles)
+		{
+			if (!FilesToConstruct.Contains(TaggedFile))
+			{
+				FileOperationTracker->OnFileStateUpdate(TaggedFile, UntouchedFileState);
+			}
+		}
+
 		// Default chunk store sizes to tie in with the default prefetch maxes for source configs.
 		// Cloud chunk source will share store with chunkdb source, since chunkdb is designed for standing in place of the need to download.
 		const int32 DefaultChunkDbMaxRead = FChunkDbSourceConfig({}).PreFetchMaximum;
@@ -706,6 +979,11 @@ namespace BuildPatchServices
 		// Clamp to sensible limits.
 		CloudChunkStoreMemorySize = FMath::Clamp<int32>(CloudChunkStoreMemorySize, 32, 2048);
 		InstallChunkStoreMemorySize = FMath::Clamp<int32>(InstallChunkStoreMemorySize, 32, 2048);
+		// Cache the last download requirement incase we are running a retry.
+		PreviousTotalDownloadRequired.Add(CloudChunkSourceStatistics->GetRequiredDownloadSize());
+		// Reset so that we don't double count data.
+		CloudChunkSourceStatistics->OnRequiredDataUpdated(0);
+		CloudChunkSourceStatistics->OnReceivedDataUpdated(0);
 
 		// Scoped systems composition and execution.
 		{
@@ -718,19 +996,20 @@ namespace BuildPatchServices
 			TUniquePtr<IChunkEvictionPolicy> MemoryEvictionPolicy(FChunkEvictionPolicyFactory::Create(
 				ChunkReferenceTracker.Get()));
 			TUniquePtr<IDiskChunkStore> DiskOverflowStore(FDiskChunkStoreFactory::Create(
+				FileSystem.Get(),
 				ChunkDataSerialization.Get(),
-				InstallerStatistics->GetDiskChunkStoreStat(),
-				DataStagingDir));
+				DiskChunkStoreStatistics.Get(),
+				FDiskChunkStoreConfig(DataStagingDir)));
 			TUniquePtr<IMemoryChunkStore> InstallChunkStore(FMemoryChunkStoreFactory::Create(
 				InstallChunkStoreMemorySize,
 				MemoryEvictionPolicy.Get(),
 				DiskOverflowStore.Get(),
-				InstallerStatistics->GetMemoryChunkStoreStat(EMemoryChunkStore::InstallSource)));
+				MemoryChunkStoreAggregateStatistics->Expose(static_cast<int32>(EMemoryStore::Install))));
 			TUniquePtr<IMemoryChunkStore> CloudChunkStore(FMemoryChunkStoreFactory::Create(
 				CloudChunkStoreMemorySize,
 				MemoryEvictionPolicy.Get(),
 				DiskOverflowStore.Get(),
-				InstallerStatistics->GetMemoryChunkStoreStat(EMemoryChunkStore::CloudSource)));
+				MemoryChunkStoreAggregateStatistics->Expose(static_cast<int32>(EMemoryStore::Cloud))));
 			TUniquePtr<IChunkDbChunkSource> ChunkDbChunkSource(FChunkDbChunkSourceFactory::Create(
 				BuildChunkDbSourceConfig(),
 				Platform.Get(),
@@ -740,17 +1019,20 @@ namespace BuildPatchServices
 				ChunkDataSerialization.Get(),
 				MessagePump.Get(),
 				InstallerError.Get(),
-				InstallerStatistics->GetChunkDbChunkSourceStat()));
+				ChunkDbChunkSourceStatistics.Get()));
 			TUniquePtr<IInstallChunkSource> InstallChunkSource(FInstallChunkSourceFactory::Create(
 				BuildInstallSourceConfig(ChunkDbChunkSource->GetAvailableChunks()),
 				FileSystem.Get(),
 				InstallChunkStore.Get(),
 				ChunkReferenceTracker.Get(),
 				InstallerError.Get(),
-				InstallerStatistics->GetInstallChunkSourceStat(),
+				InstallChunkSourceStatistics.Get(),
 				InstallationInfo,
 				NewBuildManifest));
 			const TSet<FGuid> InitialDownloadChunks = ReferencedChunks.Difference(InstallChunkSource->GetAvailableChunks()).Difference(ChunkDbChunkSource->GetAvailableChunks());
+			FileOperationTracker->OnDataStateUpdate(ReferencedChunks.Intersect(ChunkDbChunkSource->GetAvailableChunks()), EFileOperationState::PendingLocalChunkDbData);
+			FileOperationTracker->OnDataStateUpdate(ReferencedChunks.Intersect(InstallChunkSource->GetAvailableChunks()).Difference(ChunkDbChunkSource->GetAvailableChunks()), EFileOperationState::PendingLocalInstallData);
+			FileOperationTracker->OnDataStateUpdate(InitialDownloadChunks, EFileOperationState::PendingRemoteCloudData);
 			TUniquePtr<ICloudChunkSource> CloudChunkSource(FCloudChunkSourceFactory::Create(
 				BuildCloudSourceConfig(),
 				Platform.Get(),
@@ -760,7 +1042,7 @@ namespace BuildPatchServices
 				ChunkDataSerialization.Get(),
 				MessagePump.Get(),
 				InstallerError.Get(),
-				InstallerStatistics->GetCloudChunkSourceStat(),
+				CloudChunkSourceStatistics.Get(),
 				NewBuildManifest,
 				InitialDownloadChunks));
 			TArray<IChunkSource*> ChunkSources;
@@ -770,15 +1052,35 @@ namespace BuildPatchServices
 			TUniquePtr<IChainedChunkSource> ChainedChunkSource(FChainedChunkSourceFactory::Create(
 				ChunkSources));
 			TUniquePtr<FBuildPatchFileConstructor> FileConstructor(new FBuildPatchFileConstructor(
-				NewBuildManifest,
-				Configuration.InstallDirectory,
-				InstallStagingDir,
-				FilesToConstruct.Array(),
+				FFileConstructorConfig({
+					NewBuildManifest,
+					CurrentBuildManifest,
+					Configuration.InstallDirectory,
+					InstallStagingDir,
+					FilesToConstruct.Array(),
+					Configuration.InstallMode}),
+				FileSystem.Get(),
 				ChainedChunkSource.Get(),
 				ChunkReferenceTracker.Get(),
 				InstallerError.Get(),
 				InstallerAnalytics.Get(),
-				InstallerStatistics->GetFileConstructorStat()));
+				FileConstructorStatistics.Get()));
+			FDelegateHandle OnBeforeDeleteFileHandle = FileConstructor->OnBeforeDeleteFile().AddLambda([this, &InstallChunkSource](const FString& FilePath)
+			{
+				FString BuildRelativeFilename = FilePath;
+				BuildRelativeFilename.RemoveFromStart(Configuration.InstallDirectory);
+				BuildRelativeFilename.RemoveFromStart(TEXT("/"));
+				OldFilesRemovedBySystem.Add(BuildRelativeFilename);
+				InstallChunkSource->HarvestRemainingChunksFromFile(FilePath);
+			});
+			TFunction<void(const FGuid&)> LostChunkCallback = [&ChainedChunkSource](const FGuid& LostChunk)
+			{
+				ChainedChunkSource->AddRepeatRequirement(LostChunk);
+			};
+			DiskOverflowStore->SetLostChunkCallback(LostChunkCallback);
+			InstallChunkStore->SetLostChunkCallback(LostChunkCallback);
+			CloudChunkStore->SetLostChunkCallback(LostChunkCallback);
+
 
 			// Register controllables.
 			FScopedControllables ScopedControllables(&ThreadLock, Controllables, bIsPaused, bShouldAbort);
@@ -787,6 +1089,7 @@ namespace BuildPatchServices
 			ScopedControllables.Register(FileConstructor.Get());
 
 			// Set chunk counter stats.
+			if(!bHasCorruptFiles)
 			{
 				FScopeLock Lock(&ThreadLock);
 				BuildStats.NumChunksRequired = ReferencedChunks.Num();
@@ -823,30 +1126,26 @@ namespace BuildPatchServices
 			ConstructTimer.Start();
 			FileConstructor->Wait();
 			ConstructTimer.Stop();
+			FileConstructor->OnBeforeDeleteFile().Remove(OnBeforeDeleteFileHandle);
 			UE_LOG(LogBuildPatchServices, Log, TEXT("File construction complete"));
 		}
 
 		// Process some final stats.
 		{
 			FScopeLock Lock(&ThreadLock);
-			BuildStats.NumFilesOutdated = FilesToConstruct.Num();
-			BuildStats.TotalDownloadedData = InstallerStatistics->GetBytesDownloaded();
-			BuildStats.NumChunksDownloaded = InstallerStatistics->GetNumSuccessfulChunkDownloads();
-			BuildStats.NumFailedDownloads = InstallerStatistics->GetNumFailedChunkDownloads();
-			BuildStats.NumBadDownloads = InstallerStatistics->GetNumCorruptChunkDownloads();
-			BuildStats.NumAbortedDownloads = InstallerStatistics->GetNumAbortedChunkDownloads();
-			BuildStats.OverallRequestSuccessRate = InstallerStatistics->GetDownloadSuccessRate();
-			BuildStats.AverageDownloadSpeed = InstallerStatistics->GetDownloadSpeed(TNumericLimits<double>::Max());
-			BuildStats.FinalDownloadSpeed = GetDownloadSpeed();
-			BuildStats.TheoreticalDownloadTime = BuildStats.AverageDownloadSpeed > 0 ? BuildStats.TotalDownloadedData / BuildStats.AverageDownloadSpeed : 0;
-			BuildStats.NumChunksRecycled = InstallerStatistics->GetNumSuccessfulChunkRecycles();
-			BuildStats.NumChunksReadFromChunkDbs = InstallerStatistics->GetNumSuccessfulChunkDbLoads();
-			BuildStats.NumRecycleFailures = InstallerStatistics->GetNumFailedChunkRecycles();
-			BuildStats.NumChunksCacheBooted = InstallerStatistics->GetNumStoreBootedChunks();
-			BuildStats.NumDriveCacheChunkLoads = InstallerStatistics->GetNumSuccessfulChunkDiskCacheLoads();
-			BuildStats.NumDriveCacheLoadFailures = InstallerStatistics->GetNumFailedChunkDiskCacheLoads();
-			BuildStats.NumChunkDbChunksFailed = InstallerStatistics->GetNumFailedChunkDbLoads();
-			TArray<float> HealthTimers = InstallerStatistics->GetDownloadHealthTimers();
+			BuildStats.NumChunksDownloaded = DownloadServiceStatistics->GetNumSuccessfulChunkDownloads();
+			BuildStats.NumFailedDownloads = DownloadServiceStatistics->GetNumFailedChunkDownloads();
+			BuildStats.NumBadDownloads = CloudChunkSourceStatistics->GetNumCorruptChunkDownloads();
+			BuildStats.NumAbortedDownloads = CloudChunkSourceStatistics->GetNumAbortedChunkDownloads();
+			BuildStats.OverallRequestSuccessRate = CloudChunkSourceStatistics->GetDownloadSuccessRate();
+			BuildStats.NumChunksRecycled = InstallChunkSourceStatistics->GetNumSuccessfulChunkRecycles();
+			BuildStats.NumChunksReadFromChunkDbs = ChunkDbChunkSourceStatistics->GetNumSuccessfulLoads();
+			BuildStats.NumRecycleFailures = InstallChunkSourceStatistics->GetNumFailedChunkRecycles();
+			BuildStats.NumChunksStoreBooted = MemoryChunkStoreAggregateStatistics->GetTotalNumBooted();
+			BuildStats.NumDriveStoreChunkLoads = DiskChunkStoreStatistics->GetNumSuccessfulLoads();
+			BuildStats.NumDriveStoreLoadFailures = DiskChunkStoreStatistics->GetNumFailedLoads();
+			BuildStats.NumChunkDbChunksFailed = ChunkDbChunkSourceStatistics->GetNumFailedLoads();
+			TArray<float> HealthTimers = CloudChunkSourceStatistics->GetDownloadHealthTimers();
 			BuildStats.ExcellentDownloadHealthTime = HealthTimers[(int32)EBuildPatchDownloadHealth::Excellent];
 			BuildStats.GoodDownloadHealthTime = HealthTimers[(int32)EBuildPatchDownloadHealth::Good];
 			BuildStats.OkDownloadHealthTime = HealthTimers[(int32)EBuildPatchDownloadHealth::OK];
@@ -900,7 +1199,7 @@ namespace BuildPatchServices
 	{
 		// We skip this step if performing stage only
 		bool bMoveSuccess = true;
-		if (Configuration.bStageOnly)
+		if (Configuration.InstallMode == EInstallMode::StageFiles || (bFirstInstallIteration && Configuration.InstallMode == EInstallMode::PrereqOnly))
 		{
 			UE_LOG(LogBuildPatchServices, Log, TEXT("Skipping backup and stage relocation"));
 			BuildProgress.SetStateProgress(EBuildPatchState::MovingToInstall, 1.0f);
@@ -931,10 +1230,24 @@ namespace BuildPatchServices
 				ThreadLock.Unlock();
 				for (const FString& OldFilename : FilesToRemove)
 				{
-					BackupFileIfNecessary(OldFilename);
-					const bool bDeleteSuccess = IFileManager::Get().Delete(*(Configuration.InstallDirectory / OldFilename), false, true, true);
-					const uint32 LastError = FPlatformMisc::GetLastError();
-					UE_LOG(LogBuildPatchServices, Log, TEXT("Removed (%u,%u) %s"), bDeleteSuccess ? 1 : 0, LastError, *OldFilename);
+					const FString FullFilename = Configuration.InstallDirectory / OldFilename;
+					const bool bFileExists = FileSystem->FileExists(*FullFilename);
+					bool bWasDeleted = OldFilesRemovedBySystem.Contains(OldFilename);
+					if (bFileExists)
+					{
+						BackupFileIfNecessary(OldFilename);
+						const bool bDeleteSuccess = FileSystem->DeleteFile(*FullFilename);
+						const uint32 LastError = FPlatformMisc::GetLastError();
+						UE_LOG(LogBuildPatchServices, Log, TEXT("Removed (%u,%u) %s"), bDeleteSuccess ? 1 : 0, LastError, *OldFilename);
+						if (bDeleteSuccess)
+						{
+							bWasDeleted = true;
+						}
+					}
+					if (bWasDeleted)
+					{
+						MessagePump->SendMessage(FInstallationFileAction{FInstallationFileAction::EType::Removed, OldFilename});
+					}
 				}
 
 				// Now handle files that have been constructed
@@ -950,7 +1263,7 @@ namespace BuildPatchServices
 					const FString DestFilename = Configuration.InstallDirectory / ConstructionFile;
 					const float FileIndexFloat = ConstructionFilesIt.GetIndex();
 					// Skip files not constructed
-					if (!FPlatformFileManager::Get().GetPlatformFile().FileExists(*SrcFilename))
+					if (!FileSystem->FileExists(*SrcFilename))
 					{
 						BuildProgress.SetStateProgress(EBuildPatchState::MovingToInstall, FileIndexFloat / NumConstructionFilesFloat);
 						continue;
@@ -974,40 +1287,89 @@ namespace BuildPatchServices
 					}
 					// Backup file if need be
 					BackupFileIfNecessary(ConstructionFile);
-					// Delete existing if there
-					IFileManager::Get().Delete(*DestFilename, false, true, false);
-					// Move the file to the installation directory
-					int32 MoveRetries = ConfigHelpers::NumFileMoveRetries();
-					bMoveSuccess = IFileManager::Get().Move(*DestFilename, *SrcFilename, true, true, true, true);
-					uint32 ErrorCode = FPlatformMisc::GetLastError();
-					while (!bMoveSuccess && MoveRetries > 0)
+
+					// First delete the destination file to make room for moving.
+					const bool bOldFileExisted = OldFilesRemovedBySystem.Contains(ConstructionFile) || FileSystem->FileExists(*DestFilename);
+					bool bDeleteSuccess = false;
+					uint32 ErrorCode = 0;
+					int32 DeleteRetries = ConfigHelpers::NumFileMoveRetries();
+					do
 					{
-						--MoveRetries;
-						InstallerAnalytics->RecordConstructionError(ConstructionFile, ErrorCode, TEXT("Failed To Move"));
-						UE_LOG(LogBuildPatchServices, Error, TEXT("Failed to move file %s (%d), trying copy"), *ConstructionFile, ErrorCode);
-						bMoveSuccess = IFileManager::Get().Copy(*DestFilename, *SrcFilename, true, true, true) == COPY_OK;
+						bDeleteSuccess = FileSystem->DeleteFile(*DestFilename);
 						ErrorCode = FPlatformMisc::GetLastError();
-						if (!bMoveSuccess)
+						if (!bDeleteSuccess)
 						{
-							UE_LOG(LogBuildPatchServices, Error, TEXT("Failed to copy file %s (%d), retying after 0.5 sec"), *ConstructionFile, ErrorCode);
+							UE_LOG(LogBuildPatchServices, Error, TEXT("Failed to delete destination file %s (%d), retying after 0.5 sec"), *ConstructionFile, ErrorCode);
 							FPlatformProcess::Sleep(0.5f);
-							bMoveSuccess = IFileManager::Get().Move(*DestFilename, *SrcFilename, true, true, true, true);
+						}
+
+						--DeleteRetries;
+					} while (DeleteRetries > 0 && !bDeleteSuccess);
+					if (bDeleteSuccess)
+					{
+						uint32 CopyErrorCode = 0;
+						uint32 MoveErrorCode = 0;
+						int32 MoveRetries = ConfigHelpers::NumFileMoveRetries();
+						// Move the file to the installation directory
+						do
+						{
+							bMoveSuccess = FileSystem->MoveFile(*DestFilename, *SrcFilename);
 							ErrorCode = FPlatformMisc::GetLastError();
+							if (MoveErrorCode == 0)
+							{
+								MoveErrorCode = ErrorCode;
+							}
+
+							if (!bMoveSuccess)
+							{
+								InstallerAnalytics->RecordConstructionError(ConstructionFile, ErrorCode, TEXT("Failed To Move"));
+								UE_LOG(LogBuildPatchServices, Error, TEXT("Failed to move file %s (%d), trying copy"), *ConstructionFile, ErrorCode);
+								bMoveSuccess = FileSystem->CopyFile(*DestFilename, *SrcFilename);
+								ErrorCode = FPlatformMisc::GetLastError();
+								if (CopyErrorCode == 0)
+								{
+									CopyErrorCode = ErrorCode;
+								}
+
+								if (bMoveSuccess)
+								{
+									FileSystem->DeleteFile(*SrcFilename);
+								}
+								else
+								{
+									UE_LOG(LogBuildPatchServices, Error, TEXT("Failed to copy file %s (%d), retying after 0.5 sec"), *ConstructionFile, ErrorCode);
+									FPlatformProcess::Sleep(0.5f);
+								}
+							}
+
+							--MoveRetries;
+						} while (!bMoveSuccess && MoveRetries > 0);
+						if (!bMoveSuccess)
+						{
+							UE_LOG(LogBuildPatchServices, Error, TEXT("Failed to move file %s"), *FPaths::GetCleanFilename(ConstructionFile));
+							if (MoveErrorCode != 0 || CopyErrorCode != 0)
+							{
+								InstallerError->SetError(EBuildPatchInstallError::MoveFileToInstall, *FString::Printf(TEXT("%s-%u-%u"), MoveErrorCodes::StageToInstall, MoveErrorCode, CopyErrorCode));
+							}
+							else
+							{
+								InstallerError->SetError(EBuildPatchInstallError::MoveFileToInstall, MoveErrorCodes::StageToInstall);
+							}
+							
 						}
 						else
 						{
-							IFileManager::Get().Delete(*SrcFilename, false, true, false);
+							FilesInstalled.Add(ConstructionFile);
+							BuildProgress.SetStateProgress(EBuildPatchState::MovingToInstall, FileIndexFloat / NumConstructionFilesFloat);
+							FileOperationTracker->OnFileStateUpdate(ConstructionFile, EFileOperationState::Installed);
+							FInstallationFileAction::EType Action = bOldFileExisted ? FInstallationFileAction::EType::Updated : FInstallationFileAction::EType::Added;
+							MessagePump->SendMessage(FInstallationFileAction{Action, ConstructionFile});
 						}
 					}
-					if (!bMoveSuccess)
-					{
-						UE_LOG(LogBuildPatchServices, Error, TEXT("Failed to move file %s"), *FPaths::GetCleanFilename(ConstructionFile));
-						InstallerError->SetError(EBuildPatchInstallError::MoveFileToInstall, MoveErrorCodes::StageToInstall);
-					}
 					else
 					{
-						FilesInstalled.Add(ConstructionFile);
-						BuildProgress.SetStateProgress(EBuildPatchState::MovingToInstall, FileIndexFloat / NumConstructionFilesFloat);
+						UE_LOG(LogBuildPatchServices, Error, TEXT("Failed to delete file %s"), *FPaths::GetCleanFilename(ConstructionFile));
+						InstallerError->SetError(EBuildPatchInstallError::MoveFileToInstall, MoveErrorCodes::DeleteFileFailed, ErrorCode);
 					}
 				}
 
@@ -1030,7 +1392,7 @@ namespace BuildPatchServices
 	{
 		// Only provide stage directory if stage-only mode
 		FString EmptyString;
-		FString& OptionalStageDirectory = Configuration.bStageOnly ? InstallStagingDir : EmptyString;
+		FString& OptionalStageDirectory = Configuration.InstallMode == EInstallMode::StageFiles ? InstallStagingDir : EmptyString;
 
 		// Construct the attributes class
 		FileAttributesTimer.Start();
@@ -1048,7 +1410,11 @@ namespace BuildPatchServices
 	{
 		// Make sure this function can never be parallelized
 		static FCriticalSection SingletonFunctionLockCS;
-		FScopeLock SingletonFunctionLock(&SingletonFunctionLockCS);
+		const bool bShouldLock = !Configuration.bAllowConcurrentExecution;
+		if (bShouldLock)
+		{
+			SingletonFunctionLockCS.Lock();
+		}
 
 		VerifyTimer.Start();
 		BuildProgress.SetStateProgress(EBuildPatchState::BuildVerification, 0.0f);
@@ -1059,23 +1425,35 @@ namespace BuildPatchServices
 
 		// Only provide stage directory if stage-only mode
 		FString EmptyString;
-		FString& OptionalStageDirectory = Configuration.bStageOnly ? InstallStagingDir : EmptyString;
+		FString& OptionalStageDirectory = Configuration.InstallMode == EInstallMode::StageFiles ? InstallStagingDir : EmptyString;
 
 		// For a repair operation, the first verify must be ShaVerifyAllFiles.
 		EVerifyMode ModeToRun = Configuration.bIsRepair && bFirstInstallIteration ? EVerifyMode::ShaVerifyAllFiles : Configuration.VerifyMode;
 
 		// Construct the verifier
-		TUniquePtr<IVerifier> Verifier(FVerifierFactory::Create(FileSystem.Get(), InstallerStatistics->GetVerifierStat(), ModeToRun, FilesToConstruct, Configuration.InstallTags, NewBuildManifest, Configuration.InstallDirectory, OptionalStageDirectory));
+		TUniquePtr<IVerifier> Verifier(FVerifierFactory::Create(FileSystem.Get(), VerifierStatistics.Get(), ModeToRun, FilesToConstruct, Configuration.InstallTags, NewBuildManifest, Configuration.InstallDirectory, OptionalStageDirectory));
 		FScopedControllables ScopedControllables(&ThreadLock, Controllables, bIsPaused, bShouldAbort);
 		ScopedControllables.Register(Verifier.Get());
 
 		// Verify the build
-		bool bVerifySuccess = Verifier->Verify(CorruptFiles);
-		if (!bVerifySuccess)
+		EVerifyResult VerifyResult = Verifier->Verify(CorruptFiles);
+		const bool bVerifySuccessful = VerifyResult == EVerifyResult::Success;
+		if (!bVerifySuccessful)
 		{
 			UE_LOG(LogBuildPatchServices, Error, TEXT("Build verification failed on %u file(s)"), CorruptFiles.Num());
-			InstallerError->SetError(EBuildPatchInstallError::BuildVerifyFail, VerifyErrorCodes::FinalCheck);
+			InstallerError->SetError(EBuildPatchInstallError::BuildVerifyFail, InstallerHelpers::GetVerifyErrorCode(VerifyResult));
 		}
+		TMap<EVerifyError, int32> VerifyErrorCounts = VerifierStatistics->GetVerifyErrorCounts();
+		for (const TPair<EVerifyError, int32>& VerifyErrorCount : VerifyErrorCounts)
+		{
+			const int32 CachedCount = CachedVerifyErrorCounts.FindRef(VerifyErrorCount.Key);
+			if (CachedCount < VerifyErrorCount.Value)
+			{
+				InstallerHelpers::LogAdditionalVerifyErrors(VerifyErrorCount.Key, VerifyErrorCount.Value - CachedCount);
+			}
+		}
+		CachedVerifyErrorCounts = MoveTemp(VerifyErrorCounts);
+
 		BuildProgress.SetStateProgress(EBuildPatchState::BuildVerification, 1.0f);
 
 		// Delete/Backup any incorrect files if failure was not cancellation
@@ -1084,18 +1462,31 @@ namespace BuildPatchServices
 			for (const FString& CorruptFile : CorruptFiles)
 			{
 				BackupFileIfNecessary(CorruptFile, true);
-				if (!Configuration.bStageOnly)
+				FString StagedFile = InstallStagingDir / CorruptFile;
+				if (FileSystem->FileExists(*StagedFile))
 				{
-					IFileManager::Get().Delete(*(Configuration.InstallDirectory / CorruptFile), false, true);
+					FileSystem->DeleteFile(*StagedFile);
+				}
+				if (Configuration.InstallMode != EInstallMode::StageFiles)
+				{
+					FString InstalledFile = Configuration.InstallDirectory / CorruptFile;
+					if (FileSystem->FileExists(*InstalledFile))
+					{
+						FileSystem->DeleteFile(*InstalledFile);
+						OldFilesRemovedBySystem.Add(CorruptFile);
+					}
 				}
-				IFileManager::Get().Delete(*(InstallStagingDir / CorruptFile), false, true);
 			}
 		}
 
-		UE_LOG(LogBuildPatchServices, Log, TEXT("Verify stage complete %d"), bVerifySuccess ? 1 : 0);
+		UE_LOG(LogBuildPatchServices, Log, TEXT("Verify stage complete %d"), bVerifySuccessful ? 1 : 0);
 
 		VerifyTimer.Stop();
-		return bVerifySuccess;
+		if (bShouldLock)
+		{
+			SingletonFunctionLockCS.Unlock();
+		}
+		return bVerifySuccessful;
 	}
 
 	bool FBuildPatchInstaller::BackupFileIfNecessary(const FString& Filename, bool bDiscoveredByVerification /*= false */)
@@ -1158,46 +1549,19 @@ namespace BuildPatchServices
 		return bBackupSuccess;
 	}
 
-	//@todo this is deprecated and shouldn't be used anymore [6/4/2014 justin.sargent]
-	FText FBuildPatchInstaller::GetDownloadSpeedText() const
-	{
-		static const FText DownloadSpeedFormat = NSLOCTEXT("BuildPatchInstaller", "BuildPatchInstaller_DownloadSpeedFormat", "{Current} / {Total} ({Speed}/sec)");
-
-		FScopeLock Lock(&ThreadLock);
-		FText SpeedDisplayedText;
-		double DownloadSpeed = GetDownloadSpeed();
-		double InitialDownloadSize = GetInitialDownloadSize();
-		double TotalDownloaded = GetTotalDownloaded();
-		if (DownloadSpeed >= 0)
-		{
-			FNumberFormattingOptions FormattingOptions;
-			FormattingOptions.MaximumFractionalDigits = 1;
-			FormattingOptions.MinimumFractionalDigits = 1;
-
-			FFormatNamedArguments Args;
-			Args.Add(TEXT("Speed"), FText::AsMemory(DownloadSpeed, &FormattingOptions));
-			Args.Add(TEXT("Total"), FText::AsMemory(InitialDownloadSize, &FormattingOptions));
-			Args.Add(TEXT("Current"), FText::AsMemory(TotalDownloaded, &FormattingOptions));
-
-			return FText::Format(DownloadSpeedFormat, Args);
-		}
-
-		return FText();
-	}
-
 	double FBuildPatchInstaller::GetDownloadSpeed() const
 	{
-		return InstallerStatistics->GetDownloadSpeed(ConfigHelpers::DownloadSpeedAverageTime());
+		return DownloadSpeedRecorder->GetAverageSpeed(ConfigHelpers::DownloadSpeedAverageTime());
 	}
 
-	int64 FBuildPatchInstaller::GetInitialDownloadSize() const
+	int64 FBuildPatchInstaller::GetTotalDownloadRequired() const
 	{
-		return InstallerStatistics->GetRequiredDownloadSize();
+		return CloudChunkSourceStatistics->GetRequiredDownloadSize() + PreviousTotalDownloadRequired.GetValue();
 	}
 
 	int64 FBuildPatchInstaller::GetTotalDownloaded() const
 	{
-		return InstallerStatistics->GetBytesDownloaded();
+		return DownloadServiceStatistics->GetBytesDownloaded();
 	}
 
 	bool FBuildPatchInstaller::IsComplete() const
@@ -1227,6 +1591,11 @@ namespace BuildPatchServices
 		return !BuildStats.ProcessSuccess;
 	}
 
+	bool FBuildPatchInstaller::IsUpdate() const
+	{
+		return CurrentBuildManifest.IsValid() && CurrentBuildManifest->GetVersionString() != NewBuildManifest->GetVersionString();
+	}
+
 	bool FBuildPatchInstaller::HasError() const
 	{
 		FScopeLock Lock(&ThreadLock);
@@ -1269,15 +1638,44 @@ namespace BuildPatchServices
 		return FText::AsPercent(GetUpdateProgress(), &PercentFormattingOptions);
 	}
 
+	//@todo this is deprecated and shouldn't be used anymore [6/4/2014 justin.sargent]
+	FText FBuildPatchInstaller::GetDownloadSpeedText() const
+	{
+		static const FText DownloadSpeedFormat = NSLOCTEXT("BuildPatchInstaller", "BuildPatchInstaller_DownloadSpeedFormat", "{Current} / {Total} ({Speed}/sec)");
+
+		FScopeLock Lock(&ThreadLock);
+		FText SpeedDisplayedText;
+		double DownloadSpeed = GetDownloadSpeed();
+		double InitialDownloadSize = GetTotalDownloadRequired();
+		double TotalDownloaded = GetTotalDownloaded();
+		if (DownloadSpeed >= 0)
+		{
+			FNumberFormattingOptions FormattingOptions;
+			FormattingOptions.MaximumFractionalDigits = 1;
+			FormattingOptions.MinimumFractionalDigits = 1;
+
+			FFormatNamedArguments Args;
+			Args.Add(TEXT("Speed"), FText::AsMemory(DownloadSpeed, &FormattingOptions));
+			Args.Add(TEXT("Total"), FText::AsMemory(InitialDownloadSize, &FormattingOptions));
+			Args.Add(TEXT("Current"), FText::AsMemory(TotalDownloaded, &FormattingOptions));
+
+			return FText::Format(DownloadSpeedFormat, Args);
+		}
+
+		return FText();
+	}
+
 	EBuildPatchState FBuildPatchInstaller::GetState() const
 	{
 		return BuildProgress.GetState();
 	}
 
+	PRAGMA_DISABLE_DEPRECATION_WARNINGS
 	FText FBuildPatchInstaller::GetStatusText() const
 	{
-		return BuildProgress.GetStateText();
+		return BuildPatchServices::StateToText(GetState());
 	}
+	PRAGMA_ENABLE_DEPRECATION_WARNINGS
 
 	float FBuildPatchInstaller::GetUpdateProgress() const
 	{
@@ -1292,7 +1690,7 @@ namespace BuildPatchServices
 
 	EBuildPatchDownloadHealth FBuildPatchInstaller::GetDownloadHealth() const
 	{
-		return InstallerStatistics->GetDownloadHealth();
+		return CloudChunkSourceStatistics->GetDownloadHealth();
 	}
 
 	FText FBuildPatchInstaller::GetErrorText() const
@@ -1366,10 +1764,33 @@ namespace BuildPatchServices
 
 	void FBuildPatchInstaller::ExecuteCompleteDelegate()
 	{
-		// Should be executed in main thread, and already be complete
+		// Should be executed in main thread, and already be complete.
 		check(IsInGameThread());
 		check(IsComplete());
-		// Call the complete delegate
+		// Finish applying build statistics.
+		{
+			FScopeLock Lock(&ThreadLock);
+			BuildStats.FinalDownloadSpeed = GetDownloadSpeed();
+			BuildStats.AverageDownloadSpeed = DownloadSpeedRecorder->GetAverageSpeed(TNumericLimits<float>::Max());
+			BuildStats.PeakDownloadSpeed = DownloadSpeedRecorder->GetPeakSpeed();
+			BuildStats.AverageDiskReadSpeed = DiskReadSpeedRecorder->GetAverageSpeed(TNumericLimits<float>::Max());
+			BuildStats.PeakDiskReadSpeed = DiskReadSpeedRecorder->GetPeakSpeed();
+			BuildStats.AverageDiskWriteSpeed = DiskWriteSpeedRecorder->GetAverageSpeed(TNumericLimits<float>::Max());
+			BuildStats.PeakDiskWriteSpeed = DiskWriteSpeedRecorder->GetPeakSpeed();
+			BuildStats.TotalDownloadedData = DownloadServiceStatistics->GetBytesDownloaded();
+			BuildStats.TotalReadData = InstallChunkSourceStatistics->GetBytesRead();
+			BuildStats.TotalReadData += VerifierStatistics->GetBytesVerified();
+			BuildStats.TotalWrittenData = FileConstructorStatistics->GetBytesConstructed();
+			BuildStats.NumFilesConstructed = FileConstructorStatistics->GetFilesConstructed();
+			BuildStats.TheoreticalDownloadTime = BuildStats.AverageDownloadSpeed > 0 ? BuildStats.TotalDownloadedData / BuildStats.AverageDownloadSpeed : 0;
+			BuildStats.AverageMemoryStoreUse = MemoryChunkStoreAggregateStatistics->GetAverageStoreUse();
+			BuildStats.PeakMemoryStoreUse = MemoryChunkStoreAggregateStatistics->GetPeakStoreUse();
+			BuildStats.AverageMemoryStoreRetained = MemoryChunkStoreAggregateStatistics->GetAverageStoreRetained();
+			BuildStats.PeakMemoryStoreRetained = MemoryChunkStoreAggregateStatistics->GetPeakStoreRetained();
+			BuildStats.MemoryStoreSize = MemoryChunkStoreAggregateStatistics->GetTotalStoreSize();
+			InstallerHelpers::LogBuildStatInfo(BuildStats);
+		}
+		// Call the complete delegate.
 		OnCompleteDelegate.ExecuteIfBound(bSuccess, NewBuildManifest);
 	}
 
@@ -1386,4 +1807,4 @@ namespace BuildPatchServices
 			Thread->WaitForCompletion();
 		}
 	}
-}
+}
\ No newline at end of file
