Copying //UE4/Dev-Main to //UE4/Main (Source: //Portal/Main/Engine @ 4247640)

diff --git a/Engine/Source/Runtime/Online/BuildPatchServices/Private/BuildPatchInstaller.h b/Engine/Source/Runtime/Online/BuildPatchServices/Private/BuildPatchInstaller.h
--- a/Engine/Source/Runtime/Online/BuildPatchServices/Private/BuildPatchInstaller.h
+++ b/Engine/Source/Runtime/Online/BuildPatchServices/Private/BuildPatchInstaller.h
@@ -16,17 +16,29 @@
 #include "BuildPatchProgress.h"
 #include "Core/Platform.h"
 #include "Core/ProcessTimer.h"
+#include "Core/AsyncHelpers.h"
 #include "Common/HttpManager.h"
 
+PRAGMA_DISABLE_DEPRECATION_WARNINGS
+
 namespace BuildPatchServices
 {
 	struct FChunkDbSourceConfig;
 	struct FCloudSourceConfig;
 	struct FInstallSourceConfig;
 	class IInstallerError;
-	class IInstallerStatistics;
+	class IFileOperationTracker;
+	class IMemoryChunkStoreStatistics;
+	class IMemoryChunkStoreAggregateStatistics;
+	class IDiskChunkStoreStatistics;
+	class ISpeedRecorder;
+	class IDownloadServiceStatistics;
+	class IChunkDbChunkSourceStatistics;
+	class ICloudChunkSourceStatistics;
+	class IInstallChunkSourceStatistics;
+	class IFileConstructorStatistics;
+	class IVerifierStatistics;
 	class IInstallerAnalytics;
-	class IMachineConfig;
 	class IDownloadService;
 	class IMessagePump;
 
@@ -84,6 +96,9 @@ namespace BuildPatchServices
 		// A flag that stores whether we are on the first install iteration.
 		FThreadSafeBool bFirstInstallIteration;
 
+		// Tracks required download data between install retries.
+		FThreadSafeInt64 PreviousTotalDownloadRequired;
+
 		// Keep track of build stats.
 		FBuildInstallStats BuildStats;
 
@@ -102,9 +117,12 @@ namespace BuildPatchServices
 		// Holds the files which are all required.
 		TSet<FString> TaggedFiles;
 
-		// The files that the installation process required to construct.
+		// The files to be constructed in the current install attempt.
 		TSet<FString> FilesToConstruct;
 
+		// The files that were removed during installation due to destructive patch behavior.
+		TSet<FString> OldFilesRemovedBySystem;
+
 		// Map of registered installations.
 		TMap<FString, FBuildPatchAppManifestRef> InstallationInfo;
 
@@ -133,7 +151,19 @@ namespace BuildPatchServices
 		TUniquePtr<IInstallerAnalytics> InstallerAnalytics;
 
 		// Installer statistics tracking.
-		TUniquePtr<IInstallerStatistics> InstallerStatistics;
+		TUniquePtr<IFileOperationTracker> FileOperationTracker;
+		TUniquePtr<IMemoryChunkStoreAggregateStatistics> MemoryChunkStoreAggregateStatistics;
+		TUniquePtr<IDiskChunkStoreStatistics> DiskChunkStoreStatistics;
+		TUniquePtr<ISpeedRecorder> DownloadSpeedRecorder;
+		TUniquePtr<ISpeedRecorder> DiskReadSpeedRecorder;
+		TUniquePtr<ISpeedRecorder> DiskWriteSpeedRecorder;
+		TUniquePtr<ISpeedRecorder> ChunkDbReadSpeedRecorder;
+		TUniquePtr<IDownloadServiceStatistics> DownloadServiceStatistics;
+		TUniquePtr<IChunkDbChunkSourceStatistics> ChunkDbChunkSourceStatistics;
+		TUniquePtr<IInstallChunkSourceStatistics> InstallChunkSourceStatistics;
+		TUniquePtr<ICloudChunkSourceStatistics> CloudChunkSourceStatistics;
+		TUniquePtr<IFileConstructorStatistics> FileConstructorStatistics;
+		TUniquePtr<IVerifierStatistics> VerifierStatistics;
 
 		// Download service.
 		TUniquePtr<IDownloadService> DownloadService;
@@ -159,6 +189,9 @@ namespace BuildPatchServices
 		FProcessTimer ProcessActiveTimer;
 		FProcessTimer ProcessExecuteTimer;
 
+		// Caches verification errors encountered each run.
+		TMap<BuildPatchServices::EVerifyError, int32> CachedVerifyErrorCounts;
+
 	public:
 		/**
 		 * Constructor takes configuration and dependencies.
@@ -185,6 +218,7 @@ namespace BuildPatchServices
 		virtual bool IsCanceled() const override;
 		virtual bool IsPaused() const override;
 		virtual bool IsResumable() const override;
+		virtual bool IsUpdate() const override;
 		virtual bool HasError() const override;
 		virtual EBuildPatchInstallError GetErrorType() const override;
 		virtual FString GetErrorCode() const override;
@@ -193,7 +227,7 @@ namespace BuildPatchServices
 		//@todo this is deprecated and shouldn't be used anymore [6/4/2014 justin.sargent]
 		virtual FText GetDownloadSpeedText() const override;
 		virtual double GetDownloadSpeed() const override;
-		virtual int64 GetInitialDownloadSize() const override;
+		virtual int64 GetTotalDownloadRequired() const override;
 		virtual int64 GetTotalDownloaded() const override;
 		virtual EBuildPatchState GetState() const override;
 		virtual FText GetStatusText() const override;
@@ -234,10 +268,80 @@ namespace BuildPatchServices
 		 */
 		void PreExit();
 
+		/**
+		 * @return the file operation tracker for granular data progress states.
+		 */
+		const IFileOperationTracker* GetFileOperationTracker() const;
+
+		/**
+		 * @return the speed recorder for download speed.
+		 */
+		const ISpeedRecorder* GetDownloadSpeedRecorder() const;
+
+		/**
+		 * @return the speed recorder for disk read speed.
+		 */
+		const ISpeedRecorder* GetDiskReadSpeedRecorder() const;
+
+		/**
+		 * @return the speed recorder for chunkdb read speed.
+		 */
+		const ISpeedRecorder* GetChunkDbReadSpeedRecorder() const;
+
+		/**
+		 * @return the speed recorder for disk write speed.
+		 */
+		const ISpeedRecorder* GetDiskWriteSpeedRecorder() const;
+
+		/**
+		 * @return the download service statistics interface.
+		 */
+		const IDownloadServiceStatistics* GetDownloadServiceStatistics() const;
+
+		/**
+		 * @return the install chunk source statistics interface.
+		 */
+		const IInstallChunkSourceStatistics* GetInstallChunkSourceStatistics() const;
+
+		/**
+		 * @return the cloud chunk source statistics interface.
+		 */
+		const ICloudChunkSourceStatistics* GetCloudChunkSourceStatistics() const;
+
+		/**
+		 * @return the file constructor statistics interface.
+		 */
+		const IFileConstructorStatistics* GetFileConstructorStatistics() const;
+
+		/**
+		 * @return the verifier statistics interface.
+		 */
+		const IVerifierStatistics* GetVerifierStatistics() const;
+
+		/**
+		 * @return the cloud memory chunk store statistics interface.
+		 */
+		const IMemoryChunkStoreStatistics* GetCloudMemoryChunkStoreStatistics() const;
+
+		/**
+		 * @return the install memory chunk store statistics interface.
+		 */
+		const IMemoryChunkStoreStatistics* GetInstallMemoryChunkStoreStatistics() const;
+
+		/**
+		 * @return the disk chunk store statistics interface.
+		 */
+		const IDiskChunkStoreStatistics* GetDiskChunkStoreStatistics() const;
+
+		/**
+		 * @return the configuration being used by the installer.
+		 */
+		const FInstallerConfiguration& GetConfiguration() const;
+
 	private:
 
 		/**
-		 * Initialise the installer.
+		 * Initialize the installer.
 		 * @return Whether initialization was successful.
 		 */
 		bool Initialize();
@@ -314,6 +418,8 @@ namespace BuildPatchServices
 	};
 }
 
+PRAGMA_ENABLE_DEPRECATION_WARNINGS
+
 typedef TSharedPtr< BuildPatchServices::FBuildPatchInstaller, ESPMode::ThreadSafe > FBuildPatchInstallerPtr;
 typedef TSharedRef< BuildPatchServices::FBuildPatchInstaller, ESPMode::ThreadSafe > FBuildPatchInstallerRef;
 typedef TWeakPtr< BuildPatchServices::FBuildPatchInstaller, ESPMode::ThreadSafe > FBuildPatchInstallerWeakPtr;
diff --git a/Engine/Source/Runtime/Online/BuildPatchServices/Private/BuildPatchManifest.cpp b/Engine/Source/Runtime/Online/BuildPatchServices/Private/BuildPatchManifest.cpp
--- a/Engine/Source/Runtime/Online/BuildPatchServices/Private/BuildPatchManifest.cpp
+++ b/Engine/Source/Runtime/Online/BuildPatchServices/Private/BuildPatchManifest.cpp
@@ -1064,7 +1064,14 @@ TArray<FString> FBuildPatchAppManifest::GetBuildFileList() const
 {
 	TArray<FString> Filenames;
 	GetFileList(Filenames);
-	return MoveTemp(Filenames);
+	return Filenames;
+}
+
+TArray<FString> FBuildPatchAppManifest::GetBuildFileList(const TSet<FString>& Tags) const
+{
+	TArray<FString> Filenames;
+	GetTaggedFileList(Tags, Filenames);
+	return Filenames;
 }
 
 int64 FBuildPatchAppManifest::GetFileSize(const TArray<FString>& Filenames) const
@@ -1141,6 +1148,21 @@ void FBuildPatchAppManifest::GetFileTagList(TSet<FString>& Tags) const
 	Tags.Append(MoveTemp(TagsArray));
 }
 
+void FBuildPatchAppManifest::GetTaggedFileList(const TSet<FString>& Tags, TArray<FString>& TaggedFiles) const
+{
+	for (const FString& Tag : Tags)
+	{
+		const TArray<FFileManifest*>* Files = TaggedFilesLookup.Find(Tag);
+		if (Files != nullptr)
+		{
+			for (const FFileManifest* File : *Files)
+			{
+				TaggedFiles.Add(File->Filename);
+			}
+		}
+	}
+}
+
 void FBuildPatchAppManifest::GetTaggedFileList(const TSet<FString>& Tags, TSet<FString>& TaggedFiles) const
 {
 	for (const FString& Tag : Tags)
diff --git a/Engine/Source/Runtime/Online/BuildPatchServices/Private/BuildPatchManifest.h b/Engine/Source/Runtime/Online/BuildPatchServices/Private/BuildPatchManifest.h
--- a/Engine/Source/Runtime/Online/BuildPatchServices/Private/BuildPatchManifest.h
+++ b/Engine/Source/Runtime/Online/BuildPatchServices/Private/BuildPatchManifest.h
@@ -316,6 +316,7 @@ public:
 	virtual int64 GetBuildSize() const override;
 	virtual int64 GetBuildSize(const TSet<FString>& Tags) const override;
 	virtual TArray<FString> GetBuildFileList() const override;
+	virtual TArray<FString> GetBuildFileList(const TSet<FString>& Tags) const override;
 	virtual void GetFileTagList(TSet<FString>& Tags) const override;
 	virtual void GetRemovableFiles(const IBuildManifestRef& OldManifest, TArray< FString >& RemovableFiles) const override;
 	virtual void GetRemovableFiles(const TCHAR* InstallPath, TArray< FString >& RemovableFiles) const override;
@@ -434,12 +435,13 @@ public:
 	 * @param Tags					The tags for the required file groups.
 	 * @param TaggedFiles	OUT		Receives the tagged files.
 	 */
-	virtual void GetTaggedFileList(const TSet<FString>& Tags, TSet<FString>& TaggedFiles) const;
+	virtual void GetTaggedFileList(const TSet<FString>& Tags, TArray<FString>& TaggedFiles) const;
+	virtual void GetTaggedFileList(const TSet<FString>& Tags, TSet  <FString>& TaggedFiles) const;
 
 	/**
-	* Get the list of Guids for all files described by this manifest
-	* @param DataGuids		OUT		Receives the array of Guids.
-	*/
+	 * Get the list of Guids for all chunks referenced by this manifest
+	 * @param DataGuids		OUT		Receives the array of Guids.
+	 */
 	virtual void GetDataList(TArray<FGuid>& DataGuids) const;
 	virtual void GetDataList(TSet  <FGuid>& DataGuids) const;
 
diff --git a/Engine/Source/Runtime/Online/BuildPatchServices/Private/BuildPatchPackageChunkData.cpp b/Engine/Source/Runtime/Online/BuildPatchServices/Private/BuildPatchPackageChunkData.cpp
--- a/Engine/Source/Runtime/Online/BuildPatchServices/Private/BuildPatchPackageChunkData.cpp
+++ b/Engine/Source/Runtime/Online/BuildPatchServices/Private/BuildPatchPackageChunkData.cpp
@@ -2,12 +2,19 @@
 
 #include "BuildPatchPackageChunkData.h"
 #include "BuildPatchManifest.h"
+#include "Async/Future.h"
+#include "Async/Async.h"
 #include "Async/TaskGraphInterfaces.h"
 #include "Serialization/MemoryWriter.h"
 #include "Misc/OutputDeviceRedirector.h"
+#include "Misc/FileHelper.h"
+#include "Policies/PrettyJsonPrintPolicy.h"
+#include "Policies/CondensedJsonPrintPolicy.h"
+#include "Serialization/JsonWriter.h"
 #include "Core/Platform.h"
 #include "Common/HttpManager.h"
 #include "Common/FileSystem.h"
+#include "Common/SpeedRecorder.h"
 #include "Generation/ChunkDatabaseWriter.h"
 #include "Installer/ChunkReferenceTracker.h"
 #include "Installer/ChunkEvictionPolicy.h"
@@ -17,20 +24,144 @@
 #include "Installer/MessagePump.h"
 #include "Installer/InstallerError.h"
 #include "Installer/InstallerAnalytics.h"
-#include "Installer/InstallerStatistics.h"
+#include "Installer/Statistics/FileOperationTracker.h"
+#include "Installer/Statistics/MemoryChunkStoreStatistics.h"
+#include "Installer/Statistics/DownloadServiceStatistics.h"
+#include "Installer/Statistics/CloudChunkSourceStatistics.h"
 #include "BuildPatchProgress.h"
 
 DECLARE_LOG_CATEGORY_EXTERN(LogPackageChunkData, Log, All);
 DEFINE_LOG_CATEGORY(LogPackageChunkData);
 
-bool FBuildPackageChunkData::PackageChunkData(const FString& ManifestFilePath, const FString& OutputFile, const FString& CloudDir, uint64 MaxOutputFileSize)
+// For the output file we'll use pretty json in debug, otherwise condensed.
+#if UE_BUILD_DEBUG
+typedef  TJsonWriter<TCHAR, TPrettyJsonPrintPolicy<TCHAR>> FPackageChunksJsonWriter;
+typedef  TJsonWriterFactory<TCHAR, TPrettyJsonPrintPolicy<TCHAR>> FPackageChunksJsonWriterFactory;
+#else
+typedef  TJsonWriter<TCHAR, TCondensedJsonPrintPolicy<TCHAR>> FPackageChunksJsonWriter;
+typedef  TJsonWriterFactory<TCHAR, TCondensedJsonPrintPolicy<TCHAR>> FPackageChunksJsonWriterFactory;
+#endif //UE_BUILD_DEBUG
+
+namespace PackageChunksHelpers
+{
+	int32 GetNumDigitsRequiredForInteger(int32 Integer)
+	{
+		// Technically, there are mathematical solutions to this, however there can be floating point errors in log that cause edge cases there.
+		// We'll just use the obvious simple method.
+		return FString::Printf(TEXT("%d"), Integer).Len();
+	}
+
+	TFuture<FBuildPatchAppManifestPtr> AsyncLoadManifestFile(const FString& ManifestFilePath)
+	{
+		if (ManifestFilePath.IsEmpty())
+		{
+			TPromise<FBuildPatchAppManifestPtr> BrokenPromise;
+			BrokenPromise.SetValue(nullptr);
+			return BrokenPromise.GetFuture();
+		}
+		else
+		{
+			return Async<FBuildPatchAppManifestPtr>(EAsyncExecution::ThreadPool, [ManifestFilePath]() -> FBuildPatchAppManifestPtr
+			{
+				FBuildPatchAppManifestRef BuildManifest = MakeShareable(new FBuildPatchAppManifest());
+				if (BuildManifest->LoadFromFile(ManifestFilePath))
+				{
+					return BuildManifest;
+				}
+				return nullptr;
+			});
+		}
+	}
+
+	TArray<FGuid> GetCustomChunkReferences(const TArray<TSet<FString>>& TagSetArray, const FBuildPatchAppManifestRef& NewManifest, const FBuildPatchAppManifestRef& PrevManifest)
+	{
+		using namespace BuildPatchServices;
+		TSet<FGuid> VisitedChunks;
+		TArray<FGuid> UniqueChunkReferences;
+		for (const TSet<FString>& TagSet : TagSetArray)
+		{
+			for (const FGuid& ChunkReference : CustomChunkReferencesHelpers::OrderedUniquePatchReferencesTagged(NewManifest, PrevManifest, TagSet))
+			{
+				bool bIsAlreadyInSet = false;
+				VisitedChunks.Add(ChunkReference, &bIsAlreadyInSet);
+				if (!bIsAlreadyInSet)
+				{
+					UniqueChunkReferences.Add(ChunkReference);
+				}
+			}
+		}
+		return UniqueChunkReferences;
+	}
+
+	TArray<FGuid> GetCustomChunkReferences(const TArray<TSet<FString>>& TagSetArray, const FBuildPatchAppManifestRef& NewManifest)
+	{
+		using namespace BuildPatchServices;
+		TSet<FGuid> VisitedChunks;
+		TArray<FGuid> UniqueChunkReferences;
+		for (const TSet<FString>& TagSet : TagSetArray)
+		{
+			for (const FGuid& ChunkReference : CustomChunkReferencesHelpers::OrderedUniqueReferencesTagged(NewManifest, TagSet))
+			{
+				bool bIsAlreadyInSet = false;
+				VisitedChunks.Add(ChunkReference, &bIsAlreadyInSet);
+				if (!bIsAlreadyInSet)
+				{
+					UniqueChunkReferences.Add(ChunkReference);
+				}
+			}
+		}
+		return UniqueChunkReferences;
+	}
+}
+
+bool FBuildPackageChunkData::PackageChunkData(const FString& ManifestFilePath, const FString& PrevManifestFilePath, const TArray<TSet<FString>>& InTagSetArray, const FString& OutputFile, const FString& CloudDir, uint64 MaxOutputFileSize, const FString& ResultDataFilePath)
 {
 	const TCHAR* ChunkDbExtension = TEXT(".chunkdb");
 	using namespace BuildPatchServices;
-	bool bSuccess = false;
-	FBuildPatchAppManifestRef Manifest = MakeShareable(new FBuildPatchAppManifest());
-	if (Manifest->LoadFromFile(ManifestFilePath))
+	bool bSuccess = true;
+	TFuture<FBuildPatchAppManifestPtr> ManifestFuture = PackageChunksHelpers::AsyncLoadManifestFile(ManifestFilePath);
+	TFuture<FBuildPatchAppManifestPtr> PrevManifestFuture = PackageChunksHelpers::AsyncLoadManifestFile(PrevManifestFilePath);
+	FBuildPatchAppManifestPtr Manifest = ManifestFuture.Get();
+	FBuildPatchAppManifestPtr PrevManifest = PrevManifestFuture.Get();
+	TArray<FChunkDatabaseFile> ChunkDbFiles;
+	TArray<TSet<FString>> TagSetArray;
+	TArray<TArray<int32>> TagSetLookupTable;
+	// Check required manifest was loaded ok.
+	if (!Manifest.IsValid())
+	{
+		UE_LOG(LogPackageChunkData, Error, TEXT("Failed to load manifest %s"), *ManifestFilePath);
+		bSuccess = false;
+	}
+	// Check previous manifest was loaded ok if it was provided.
+	if (!PrevManifest.IsValid() && !PrevManifestFilePath.IsEmpty())
+	{
+		UE_LOG(LogPackageChunkData, Error, TEXT("Failed to load manifest %s"), *PrevManifestFilePath);
+		bSuccess = false;
+	}
+	if (bSuccess)
 	{
+		// If TagSetArray was not provided, we need to adjust it to contain an entry that uses all tags.
+		if (InTagSetArray.Num() == 0)
+		{
+			TagSetArray.AddDefaulted();
+			Manifest->GetFileTagList(TagSetArray.Last());
+		}
+		else
+		{
+			TagSetArray = InTagSetArray;
+		}
+		TagSetLookupTable.AddDefaulted(TagSetArray.Num());
+		// Construct the chunk reference tracker, building our list of ordered unique chunk references.
+		TUniquePtr<IChunkReferenceTracker> ChunkReferenceTracker;
+		if (PrevManifest.IsValid())
+		{
+			ChunkReferenceTracker.Reset(FChunkReferenceTrackerFactory::Create(PackageChunksHelpers::GetCustomChunkReferences(TagSetArray, Manifest.ToSharedRef(), PrevManifest.ToSharedRef())));
+		}
+		else
+		{
+			ChunkReferenceTracker.Reset(FChunkReferenceTrackerFactory::Create(PackageChunksHelpers::GetCustomChunkReferences(TagSetArray, Manifest.ToSharedRef())));
+		}
+
 		// Programmatically calculate header file size effects, so that we automatically handle any changes to the header spec.
 		TArray<uint8> HeaderData;
 		FMemoryWriter HeaderWriter(HeaderData);
@@ -44,59 +175,95 @@ bool FBuildPackageChunkData::PackageChunkData(const FString& ManifestFilePath, c
 		const uint64 PerEntryHeaderSize = HeaderData.Num() - ChunkDbHeaderSize;
 
 		// Enumerate the chunks, allocating them to chunk db files.
-		TUniquePtr<IChunkReferenceTracker> ChunkReferenceTracker(FChunkReferenceTrackerFactory::Create(Manifest));
 		TSet<FGuid> FullDataSet = ChunkReferenceTracker->GetReferencedChunks();
 		if (FullDataSet.Num() > 0)
 		{
-			TArray<FChunkDatabaseFile> ChunkDbFiles;
-			ChunkDbFiles.AddDefaulted();
-			int32 ChunkDbFileCount = 0;
-
-			// Figure out the chunks to write per chunkdb file.
-			uint64 AvailableFileSize = MaxOutputFileSize - ChunkDbHeaderSize;
-			for (const FGuid& DataId : FullDataSet)
+			// Create the data set for each tagset.
+			int32 NumSetsWithData = 0;
+			TArray<TSet<FGuid>> TaggedDataSets;
+			TSet<FGuid> VisitedChunks;
+			for (const TSet<FString>& TagSet : TagSetArray)
 			{
-				const uint64 DataSize = Manifest->GetDataSize(DataId) + PerEntryHeaderSize;
-				if (AvailableFileSize < DataSize && ChunkDbFiles[ChunkDbFileCount].DataList.Num() > 0)
-				{
-					ChunkDbFiles.AddDefaulted();
-					++ChunkDbFileCount;
-					AvailableFileSize = MaxOutputFileSize - ChunkDbHeaderSize;
-				}
-
-				ChunkDbFiles[ChunkDbFileCount].DataList.Add(DataId);
-				if (AvailableFileSize > DataSize)
-				{
-					AvailableFileSize -= DataSize;
-				}
-				else
+				TSet<FGuid> TaggedChunks;
+				TSet<FString> TaggedFiles;
+				Manifest->GetTaggedFileList(TagSet, TaggedFiles);
+				Manifest->GetChunksRequiredForFiles(TaggedFiles, TaggedChunks);
+				TaggedChunks = TaggedChunks.Intersect(FullDataSet).Difference(VisitedChunks);
+				if (TaggedChunks.Num() > 0)
 				{
-					AvailableFileSize = 0;
+					++NumSetsWithData;
+					VisitedChunks.Append(TaggedChunks);
 				}
+				TaggedDataSets.Add(MoveTemp(TaggedChunks));
 			}
+			const int32 NumDigitsForTagSets = NumSetsWithData > 1 ? PackageChunksHelpers::GetNumDigitsRequiredForInteger(TaggedDataSets.Num()) : 0;
 
-			// Figure out the filenames of each chunkdb.
-			if (ChunkDbFiles.Num() > 1)
+			for (int32 ChunkDbTagSetIdx = 0; ChunkDbTagSetIdx < TaggedDataSets.Num(); ++ChunkDbTagSetIdx)
 			{
-				// Figure out the per file filename;
-				FString FilenameFmt = OutputFile;
-				FilenameFmt.RemoveFromEnd(ChunkDbExtension);
-				// Technically, there are mathematical solutions to this, however there can be floating point errors in log that cause edge cases there.
-				// We'll just use the obvious simple method.
-				int32 NumDigitsForParts = FString::Printf(TEXT("%d"), ChunkDbFiles.Num()).Len();
-				for (int32 ChunkDbFileIdx = 0; ChunkDbFileIdx < ChunkDbFiles.Num(); ++ChunkDbFileIdx)
+				if (TaggedDataSets[ChunkDbTagSetIdx].Num() > 0)
 				{
-					ChunkDbFiles[ChunkDbFileIdx].DatabaseFilename = FString::Printf(TEXT("%s.part%0*d.chunkdb"), *FilenameFmt, NumDigitsForParts, ChunkDbFileIdx + 1);
-				}
-			}
-			else
-			{
-				// Should have figured out no data already.
-				check(ChunkDbFiles.Num() == 1);
-				ChunkDbFiles[0].DatabaseFilename = OutputFile;
-				if (!ChunkDbFiles[0].DatabaseFilename.EndsWith(ChunkDbExtension))
-				{
-					ChunkDbFiles[0].DatabaseFilename += ChunkDbExtension;
+					const int32 FirstChunkDbFileIdx = ChunkDbFiles.Num();
+					int32 ChunkDbPartCount = 0;
+					int32 CurrentChunkDbFileIdx = FirstChunkDbFileIdx-1;
+					
+					// Figure out the chunks to write per chunkdb file.
+					uint64 AvailableFileSize = 0;
+					for (const FGuid& DataId : TaggedDataSets[ChunkDbTagSetIdx])
+					{
+						const uint64 DataSize = Manifest->GetDataSize(DataId) + PerEntryHeaderSize;
+						if (AvailableFileSize < DataSize && 
+							((CurrentChunkDbFileIdx < FirstChunkDbFileIdx) || (ChunkDbFiles[CurrentChunkDbFileIdx].DataList.Num() > 0)))
+						{
+							ChunkDbFiles.AddDefaulted();
+							++ChunkDbPartCount;
+							++CurrentChunkDbFileIdx;
+							AvailableFileSize = MaxOutputFileSize - ChunkDbHeaderSize;
+							TagSetLookupTable[ChunkDbTagSetIdx].Add(CurrentChunkDbFileIdx);
+						}
+
+						ChunkDbFiles[CurrentChunkDbFileIdx].DataList.Add(DataId);
+						if (AvailableFileSize > DataSize)
+						{
+							AvailableFileSize -= DataSize;
+						}
+						else
+						{
+							AvailableFileSize = 0;
+						}
+					}
+
+					// Figure out the filenames of each chunkdb.
+					if (ChunkDbPartCount > 1)
+					{
+						// Figure out the per file filename;
+						FString FilenameFmt = OutputFile;
+						FilenameFmt.RemoveFromEnd(ChunkDbExtension);
+						if (NumDigitsForTagSets > 0)
+						{
+							FilenameFmt += FString::Printf(TEXT(".tagset%0*d"), NumDigitsForTagSets, ChunkDbTagSetIdx + 1);
+						}
+						// Technically, there are mathematical solutions to this, however there can be floating point errors in log that cause edge cases there.
+						// We'll just use the obvious simple method.
+						const int32 NumDigitsForParts = PackageChunksHelpers::GetNumDigitsRequiredForInteger(ChunkDbPartCount);
+						for (int32 ChunkDbFileIdx = FirstChunkDbFileIdx; ChunkDbFileIdx < ChunkDbFiles.Num(); ++ChunkDbFileIdx)
+						{
+							ChunkDbFiles[ChunkDbFileIdx].DatabaseFilename = FString::Printf(TEXT("%s.part%0*d.chunkdb"), *FilenameFmt, NumDigitsForParts, (ChunkDbFileIdx - FirstChunkDbFileIdx) + 1);
+						}
+					}
+					else if (ChunkDbPartCount == 1)
+					{
+						// Figure out the per file filename;
+						FString FilenameFmt = OutputFile;
+						FilenameFmt.RemoveFromEnd(ChunkDbExtension);
+						// Make sure we don't confuse any existing % characters when formatting (yes, % is a valid filename character).
+						FilenameFmt.ReplaceInline(TEXT("%"), TEXT("%%"));
+						if (NumDigitsForTagSets > 0)
+						{
+							FilenameFmt += FString::Printf(TEXT(".tagset%0*d"), NumDigitsForTagSets, ChunkDbTagSetIdx + 1);
+						}
+						FilenameFmt += ChunkDbExtension;
+						ChunkDbFiles[CurrentChunkDbFileIdx].DatabaseFilename = FilenameFmt;
+					}
 				}
 			}
 
@@ -106,14 +273,31 @@ bool FBuildPackageChunkData::PackageChunkData(const FString& ManifestFilePath, c
 			CloudSourceConfig.MaxRetryCount = 30;
 
 			// Create systems.
+			const int32 CloudStoreId = 0;
 			FBuildPatchProgress BuildProgress;
 			TUniquePtr<IHttpManager> HttpManager(FHttpManagerFactory::Create());
 			TUniquePtr<IFileSystem> FileSystem(FFileSystemFactory::Create());
 			TUniquePtr<IPlatform> Platform(FPlatformFactory::Create());
 			TUniquePtr<IMessagePump> MessagePump(FMessagePumpFactory::Create());
 			TUniquePtr<IInstallerError> InstallerError(FInstallerErrorFactory::Create());
-			TUniquePtr<IInstallerAnalytics> InstallerAnalytics(FInstallerAnalyticsFactory::Create(nullptr, nullptr));
-			TUniquePtr<IInstallerStatistics> InstallerStatistics(FInstallerStatisticsFactory::Create(InstallerAnalytics.Get(), &BuildProgress));
+			TUniquePtr<IInstallerAnalytics> InstallerAnalytics(FInstallerAnalyticsFactory::Create(
+				nullptr,
+				nullptr));
+			TUniquePtr<IFileOperationTracker> FileOperationTracker(FFileOperationTrackerFactory::Create(
+				FTicker::GetCoreTicker(),
+				Manifest.Get()));
+			TUniquePtr<IMemoryChunkStoreAggregateStatistics> MemoryChunkStoreAggregateStatistics(FMemoryChunkStoreAggregateStatisticsFactory::Create(
+				TSet<FGuid>(),
+				FileOperationTracker.Get()));
+			TUniquePtr<ISpeedRecorder> DownloadSpeedRecorder(FSpeedRecorderFactory::Create());
+			TUniquePtr<IDownloadServiceStatistics> DownloadServiceStatistics(FDownloadServiceStatisticsFactory::Create(
+				DownloadSpeedRecorder.Get(),
+				InstallerAnalytics.Get(),
+				Manifest.Get()));
+			TUniquePtr<ICloudChunkSourceStatistics> CloudChunkSourceStatistics(FCloudChunkSourceStatisticsFactory::Create(
+				InstallerAnalytics.Get(),
+				&BuildProgress,
+				FileOperationTracker.Get()));
 			TUniquePtr<IChunkDataSerialization> ChunkDataSerialization(FChunkDataSerializationFactory::Create(
 				FileSystem.Get()));
 			TUniquePtr<IChunkEvictionPolicy> MemoryEvictionPolicy(FChunkEvictionPolicyFactory::Create(
@@ -122,12 +306,12 @@ bool FBuildPackageChunkData::PackageChunkData(const FString& ManifestFilePath, c
 				512,
 				MemoryEvictionPolicy.Get(),
 				nullptr,
-				InstallerStatistics->GetMemoryChunkStoreStat(EMemoryChunkStore::CloudSource)));
+				MemoryChunkStoreAggregateStatistics->Expose(CloudStoreId)));
 			TUniquePtr<IDownloadService> DownloadService(FDownloadServiceFactory::Create(
 				FTicker::GetCoreTicker(),
 				HttpManager.Get(),
 				FileSystem.Get(),
-				InstallerStatistics->GetDownloadServiceStat(),
+				DownloadServiceStatistics.Get(),
 				InstallerAnalytics.Get()));
 			TUniquePtr<ICloudChunkSource> CloudChunkSource(FCloudChunkSourceFactory::Create(
 				CloudSourceConfig,
@@ -138,8 +322,8 @@ bool FBuildPackageChunkData::PackageChunkData(const FString& ManifestFilePath, c
 				ChunkDataSerialization.Get(),
 				MessagePump.Get(),
 				InstallerError.Get(),
-				InstallerStatistics->GetCloudChunkSourceStat(),
-				Manifest,
+				CloudChunkSourceStatistics.Get(),
+				Manifest.ToSharedRef(),
 				FullDataSet));
 
 			// Start an IO output thread which saves all the chunks to the chunkdbs.
@@ -194,9 +378,43 @@ bool FBuildPackageChunkData::PackageChunkData(const FString& ManifestFilePath, c
 			UE_LOG(LogPackageChunkData, Error, TEXT("Manifest has no data"));
 		}
 	}
-	else
+
+	// Save the output.
+	if (bSuccess && ResultDataFilePath.IsEmpty() == false)
 	{
-		UE_LOG(LogPackageChunkData, Error, TEXT("Failed to load manifest %s"), *ManifestFilePath);
+		FString JsonOutput;
+		TSharedRef<FPackageChunksJsonWriter> Writer = FPackageChunksJsonWriterFactory::Create(&JsonOutput);
+		Writer->WriteObjectStart();
+		{
+			Writer->WriteArrayStart(TEXT("ChunkDbFilePaths"));
+			for (const FChunkDatabaseFile& ChunkDbFile : ChunkDbFiles)
+			{
+				Writer->WriteValue(ChunkDbFile.DatabaseFilename);
+			}
+			Writer->WriteArrayEnd();
+			if (InTagSetArray.Num() > 0)
+			{
+				Writer->WriteArrayStart(TEXT("TagSetLookupTable"));
+				for (const TArray<int32>& TagSetLookup : TagSetLookupTable)
+				{
+					Writer->WriteArrayStart();
+					for (const int32& ChunkDbFileIdx : TagSetLookup)
+					{
+						Writer->WriteValue(ChunkDbFileIdx);
+					}
+					Writer->WriteArrayEnd();
+				}
+				Writer->WriteArrayEnd();
+			}
+		}
+		Writer->WriteObjectEnd();
+		Writer->Close();
+		bSuccess = FFileHelper::SaveStringToFile(JsonOutput, *ResultDataFilePath);
+		if (!bSuccess)
+		{
+			UE_LOG(LogPackageChunkData, Error, TEXT("Could not save output to %s"), *ResultDataFilePath);
+		}
 	}
+
 	return bSuccess;
 }
diff --git a/Engine/Source/Runtime/Online/BuildPatchServices/Private/BuildPatchPackageChunkData.h b/Engine/Source/Runtime/Online/BuildPatchServices/Private/BuildPatchPackageChunkData.h
--- a/Engine/Source/Runtime/Online/BuildPatchServices/Private/BuildPatchPackageChunkData.h
+++ b/Engine/Source/Runtime/Online/BuildPatchServices/Private/BuildPatchPackageChunkData.h
@@ -6,5 +6,5 @@
 class FBuildPackageChunkData
 {
 public:
-	static bool PackageChunkData(const FString& ManifestFilePath, const FString& OutputFile, const FString& CloudDir, uint64 MaxOutputFileSize);
+	static bool PackageChunkData(const FString& ManifestFilePath, const FString& PrevManifestFilePath, const TArray<TSet<FString>>& TagSetArray, const FString& OutputFile, const FString& CloudDir, uint64 MaxOutputFileSize, const FString& ResultDataFilePath);
 };
diff --git a/Engine/Source/Runtime/Online/BuildPatchServices/Private/BuildPatchProgress.cpp b/Engine/Source/Runtime/Online/BuildPatchServices/Private/BuildPatchProgress.cpp
--- a/Engine/Source/Runtime/Online/BuildPatchServices/Private/BuildPatchProgress.cpp
+++ b/Engine/Source/Runtime/Online/BuildPatchServices/Private/BuildPatchProgress.cpp
@@ -83,12 +83,6 @@ namespace BuildPatchServices
 		return CurrentState;
 	}
 
-	const FText& FBuildPatchProgress::GetStateText() const
-	{
-		FScopeLock ScopeLock(&ThreadLock);
-		return StateToText(CurrentState);
-	}
-
 	float FBuildPatchProgress::GetProgress() const
 	{
 		FScopeLock ScopeLock(&ThreadLock);
diff --git a/Engine/Source/Runtime/Online/BuildPatchServices/Private/BuildPatchProgress.h b/Engine/Source/Runtime/Online/BuildPatchServices/Private/BuildPatchProgress.h
--- a/Engine/Source/Runtime/Online/BuildPatchServices/Private/BuildPatchProgress.h
+++ b/Engine/Source/Runtime/Online/BuildPatchServices/Private/BuildPatchProgress.h
@@ -57,12 +57,6 @@ namespace BuildPatchServices
 		 */
 		virtual EBuildPatchState GetState() const;
 
-		/**
-		 * Gets the text for the current progress state
-		 * @return The display text for the current progress state
-		 */
-		virtual const FText& GetStateText() const;
-
 		/**
 		 * Gets the current overall progress
 		 * @return The current progress value. Range 0 to 1. -1 indicates undetermined, i.e. show a marquee style bar.
diff --git a/Engine/Source/Runtime/Online/BuildPatchServices/Private/BuildPatchServicesModule.cpp b/Engine/Source/Runtime/Online/BuildPatchServices/Private/BuildPatchServicesModule.cpp
--- a/Engine/Source/Runtime/Online/BuildPatchServices/Private/BuildPatchServicesModule.cpp
+++ b/Engine/Source/Runtime/Online/BuildPatchServices/Private/BuildPatchServicesModule.cpp
@@ -6,16 +6,20 @@
 
 #include "BuildPatchServicesModule.h"
 #include "Containers/Ticker.h"
-#include "Misc/CommandLine.h"
-#include "Misc/Paths.h"
+#include "Algo/Transform.h"
 #include "Misc/ConfigCacheIni.h"
 #include "Misc/FeedbackContext.h"
 #include "Misc/CoreDelegates.h"
+#include "Misc/CommandLine.h"
+#include "Misc/Paths.h"
 #include "Misc/App.h"
 #include "Misc/OutputDeviceRedirector.h"
 #include "Modules/ModuleManager.h"
 #include "HttpModule.h"
 #include "HttpManager.h"
+#include "Data/ManifestUObject.h"
+#include "Installer/BuildStatistics.h"
+#include "Installer/MachineConfig.h"
 #include "BuildPatchCompactifier.h"
 #include "BuildPatchDataEnumeration.h"
 #include "BuildPatchMergeManifests.h"
@@ -25,6 +29,8 @@
 #include "BuildPatchServicesPrivate.h"
 #include "BuildPatchPackageChunkData.h"
 #include "BuildPatchVerifyChunkData.h"
+#include "BuildPatchServicesSingleton.h"
+
 
 using namespace BuildPatchServices;
 
@@ -51,8 +57,9 @@ void FBuildPatchServicesModule::StartupModule()
 	// We need to initialize the lookup for our hashing functions
 	FRollingHashConst::Init();
 
+	const FBuildPatchServicesInitSettings& InitSettings = FBuildPatchServices::GetSettings();
 	// Set the local machine config filename
-	LocalMachineConfigFile = FPaths::Combine(FPlatformProcess::ApplicationSettingsDir(), FApp::GetProjectName(), TEXT("BuildPatchServicesLocal.ini"));
+	LocalMachineConfigFile = FPaths::Combine(InitSettings.ApplicationSettingsDir, InitSettings.ProjectName, InitSettings.LocalMachineConfigFileName);
 
 	// Fix up any legacy configuration data
 	FixupLegacyConfig();
@@ -82,6 +89,9 @@ void FBuildPatchServicesModule::StartupModule()
 
 	// Test the rolling hash algorithm
 	check( CheckRollingHashAlgorithm() );
+
+	// Init Manifest serialization
+	FManifestUObject::Init();
 }
 
 void FBuildPatchServicesModule::ShutdownModule()
@@ -97,6 +107,12 @@ void FBuildPatchServicesModule::ShutdownModule()
 	GLog->Log(ELogVerbosity::VeryVerbose, TEXT( "BuildPatchServicesModule: Finished shutting down" ) );
 }
 
+IBuildStatisticsRef FBuildPatchServicesModule::CreateBuildStatistics(const IBuildInstallerRef& Installer) const
+{
+	checkSlow(IsInGameThread());
+	return MakeShareable(FBuildStatisticsFactory::Create(StaticCastSharedRef<FBuildPatchInstaller>(Installer)));
+}
+
 IBuildManifestPtr FBuildPatchServicesModule::LoadManifestFromFile( const FString& Filename )
 {
 	FBuildPatchAppManifestRef Manifest = MakeShareable( new FBuildPatchAppManifest() );
@@ -120,6 +136,7 @@ IBuildManifestPtr FBuildPatchServicesModule::MakeManifestFromData(const TArray<u
 	return NULL;
 }
 
+PRAGMA_DISABLE_DEPRECATION_WARNINGS
 IBuildManifestPtr FBuildPatchServicesModule::MakeManifestFromJSON( const FString& ManifestJSON )
 {
 	FBuildPatchAppManifestRef Manifest = MakeShareable( new FBuildPatchAppManifest() );
@@ -129,12 +146,20 @@ IBuildManifestPtr FBuildPatchServicesModule::MakeManifestFromJSON( const FString
 	}
 	return NULL;
 }
+PRAGMA_ENABLE_DEPRECATION_WARNINGS
 
 bool FBuildPatchServicesModule::SaveManifestToFile(const FString& Filename, IBuildManifestRef Manifest, bool bUseBinary/* = true*/)
 {
 	return StaticCastSharedRef< FBuildPatchAppManifest >(Manifest)->SaveToFile(Filename, bUseBinary);
 }
 
+TSet<FString> FBuildPatchServicesModule::GetInstalledPrereqIds() const
+{
+	const bool bAlwaysFlushChanges = true;
+	TUniquePtr<IMachineConfig> MachineConfig(FMachineConfigFactory::Create(LocalMachineConfigFile, bAlwaysFlushChanges));
+	return MachineConfig->LoadInstalledPrereqIds();
+}
+
 IBuildInstallerPtr FBuildPatchServicesModule::StartBuildInstall(IBuildManifestPtr CurrentManifest, IBuildManifestPtr InstallManifest, const FString& InstallDirectory, FBuildPatchBoolManifestDelegate OnCompleteDelegate, bool bIsRepair, TSet<FString> InstallTags)
 {
 	if (InstallManifest.IsValid())
@@ -145,7 +170,6 @@ IBuildInstallerPtr FBuildPatchServicesModule::StartBuildInstall(IBuildManifestPt
 		InstallerConfiguration.InstallDirectory = InstallDirectory;
 		InstallerConfiguration.InstallTags = InstallTags;
 		InstallerConfiguration.bIsRepair = bIsRepair;
-		InstallerConfiguration.bStageOnly = false;
 		return StartBuildInstall(MoveTemp(InstallerConfiguration), OnCompleteDelegate);
 	}
 	else
@@ -163,8 +187,8 @@ IBuildInstallerPtr FBuildPatchServicesModule::StartBuildInstallStageOnly(IBuildM
 		InstallerConfiguration.CurrentManifest = CurrentManifest;
 		InstallerConfiguration.InstallDirectory = InstallDirectory;
 		InstallerConfiguration.InstallTags = InstallTags;
+		InstallerConfiguration.InstallMode = EInstallMode::StageFiles;
 		InstallerConfiguration.bIsRepair = bIsRepair;
-		InstallerConfiguration.bStageOnly = true;
 		return StartBuildInstall(MoveTemp(InstallerConfiguration), OnCompleteDelegate);
 	}
 	else
@@ -175,7 +199,7 @@ IBuildInstallerPtr FBuildPatchServicesModule::StartBuildInstallStageOnly(IBuildM
 
 IBuildInstallerRef FBuildPatchServicesModule::StartBuildInstall(BuildPatchServices::FInstallerConfiguration Configuration, FBuildPatchBoolManifestDelegate OnCompleteDelegate)
 {
-	checkf(IsInGameThread(), TEXT("FBuildPatchServicesModule::StartBuildInstall must be called from main thread."));
+	checkSlow(IsInGameThread());
 	// Handle any of the global module overrides, while they are not yet fully deprecated.
 	if (Configuration.StagingDirectory.IsEmpty())
 	{
@@ -198,17 +222,20 @@ IBuildInstallerRef FBuildPatchServicesModule::StartBuildInstall(BuildPatchServic
 	FBuildPatchInstallerRef Installer = MakeShareable(new FBuildPatchInstaller(MoveTemp(Configuration), AvailableInstallations, LocalMachineConfigFile, Analytics, HttpTracker, OnCompleteDelegate));
 	Installer->StartInstallation();
 	BuildPatchInstallers.Add(Installer);
+	BuildPatchInstallerInterfaces.Add(Installer);
 	return Installer;
 }
 
-bool FBuildPatchServicesModule::Tick( float Delta )
+const TArray<IBuildInstallerRef>& FBuildPatchServicesModule::GetInstallers() const
 {
-    QUICK_SCOPE_CYCLE_COUNTER(STAT_FBuildPatchServicesModule_Tick);
+	checkSlow(IsInGameThread());
+	return BuildPatchInstallerInterfaces;
+}
 
-	// Using a local bool for this check will improve the assert message that gets displayed.
-	// This one is unlikely to assert unless the FTicker's core tick is not ticked on the main thread for some reason.
-	const bool bIsCalledFromMainThread = IsInGameThread();
-	check( bIsCalledFromMainThread );
+bool FBuildPatchServicesModule::Tick(float Delta)
+{
+    QUICK_SCOPE_CYCLE_COUNTER(STAT_FBuildPatchServicesModule_Tick);
+	checkSlow(IsInGameThread());
 
 	// Pump installer messages.
 	for (FBuildPatchInstallerPtr& Installer : BuildPatchInstallers)
@@ -225,9 +252,16 @@ bool FBuildPatchServicesModule::Tick( float Delta )
 		}
 	}
 
-	// Remove completed (invalids) from the list.
+	// Remove invalids from the list.
 	BuildPatchInstallers.RemoveAll([](const FBuildPatchInstallerPtr& Installer){ return Installer.IsValid() == false; });
 
+	// Check for resetting the BuildPatchInstallerInterfaces array.
+	if (BuildPatchInstallers.Num() != BuildPatchInstallerInterfaces.Num())
+	{
+		BuildPatchInstallerInterfaces.Empty(BuildPatchInstallers.Num());
+		Algo::Transform(BuildPatchInstallers, BuildPatchInstallerInterfaces, &FBuildPatchInstallerPtr::ToSharedRef);
+	}
+
 	// More ticks.
 	return true;
 }
@@ -253,9 +287,9 @@ bool FBuildPatchServicesModule::VerifyChunkData(const FString& SearchPath, const
 	return FBuildVerifyChunkData::VerifyChunkData(SearchPath, OutputFile);
 }
 
-bool FBuildPatchServicesModule::PackageChunkData(const FString& ManifestFilePath, const FString& OutputFile, const FString& CloudDir, uint64 MaxOutputFileSize)
+bool FBuildPatchServicesModule::PackageChunkData(const FString& ManifestFilePath, const FString& PrevManifestFilePath, const TArray<TSet<FString>>& TagSetArray, const FString& OutputFile, const FString& CloudDir, uint64 MaxOutputFileSize, const FString& ResultDataFilePath)
 {
-	return FBuildPackageChunkData::PackageChunkData(ManifestFilePath, OutputFile, CloudDir, MaxOutputFileSize);
+	return FBuildPackageChunkData::PackageChunkData(ManifestFilePath, PrevManifestFilePath, TagSetArray, OutputFile, CloudDir, MaxOutputFileSize, ResultDataFilePath);
 }
 
 bool FBuildPatchServicesModule::MergeManifests(const FString& ManifestFilePathA, const FString& ManifestFilePathB, const FString& ManifestFilePathC, const FString& NewVersionString, const FString& SelectionDetailFilePath)
@@ -322,6 +356,8 @@ void FBuildPatchServicesModule::SetHttpTracker( TSharedPtr< FHttpServiceTracker
 
 void FBuildPatchServicesModule::RegisterAppInstallation(IBuildManifestRef AppManifest, FString AppInstallDirectory)
 {
+	FPaths::NormalizeDirectoryName(AppInstallDirectory);
+	FPaths::CollapseRelativeDirectories(AppInstallDirectory);
 	FBuildPatchAppManifestRef InternalRef = StaticCastSharedRef<FBuildPatchAppManifest>(MoveTemp(AppManifest));
 	AvailableInstallations.Add(MoveTemp(AppInstallDirectory), MoveTemp(InternalRef));
 }
diff --git a/Engine/Source/Runtime/Online/BuildPatchServices/Private/BuildPatchServicesModule.h b/Engine/Source/Runtime/Online/BuildPatchServices/Private/BuildPatchServicesModule.h
--- a/Engine/Source/Runtime/Online/BuildPatchServices/Private/BuildPatchServicesModule.h
+++ b/Engine/Source/Runtime/Online/BuildPatchServices/Private/BuildPatchServicesModule.h
@@ -1,9 +1,5 @@
 // Copyright 1998-2018 Epic Games, Inc. All Rights Reserved.
 
-/*=============================================================================
-	BuildPatchServicesModule.h: Declares the FBuildPatchServicesModule class.
-=============================================================================*/
-
 #pragma once
 
 #include "CoreMinimal.h"
@@ -24,6 +20,7 @@ enum
 	StreamBufferSize	= FileBufferSize*4,	// When reading from build data stream, how much to buffer.
 };
 
+PRAGMA_DISABLE_DEPRECATION_WARNINGS
 /**
  * Implements the BuildPatchServicesModule.
  */
@@ -38,12 +35,15 @@ public:
 	// IModuleInterface interface end.
 
 	// IBuildPatchServicesModule interface begin.
+	virtual BuildPatchServices::IBuildStatisticsRef CreateBuildStatistics(const IBuildInstallerRef& Installer) const override;
 	virtual IBuildManifestPtr LoadManifestFromFile(const FString& Filename) override;
 	virtual IBuildManifestPtr MakeManifestFromData(const TArray<uint8>& ManifestData) override;
 	virtual bool SaveManifestToFile(const FString& Filename, IBuildManifestRef Manifest, bool bUseBinary = true) override;
+	virtual TSet<FString> GetInstalledPrereqIds() const override;
 	virtual IBuildInstallerPtr StartBuildInstall(IBuildManifestPtr CurrentManifest, IBuildManifestPtr InstallManifest, const FString& InstallDirectory, FBuildPatchBoolManifestDelegate OnCompleteDelegate, bool bIsRepair = false, TSet<FString> InstallTags = TSet<FString>()) override;
 	virtual IBuildInstallerPtr StartBuildInstallStageOnly(IBuildManifestPtr CurrentManifest, IBuildManifestPtr InstallManifest, const FString& InstallDirectory, FBuildPatchBoolManifestDelegate OnCompleteDelegate, bool bIsRepair = false, TSet<FString> InstallTags = TSet<FString>()) override;
 	virtual IBuildInstallerRef StartBuildInstall(BuildPatchServices::FInstallerConfiguration Configuration, FBuildPatchBoolManifestDelegate OnCompleteDelegate) override;
+	virtual const TArray<IBuildInstallerRef>& GetInstallers() const override;
 	virtual void SetStagingDirectory(const FString& StagingDir) override;
 	virtual void SetCloudDirectory(FString CloudDir) override;
 	virtual void SetCloudDirectories(TArray<FString> CloudDirs) override;
@@ -56,7 +56,7 @@ public:
 	virtual bool CompactifyCloudDirectory(const FString& CloudDirectory, float DataAgeThreshold, ECompactifyMode::Type Mode, const FString& DeletedChunkLogFile) override;
 	virtual bool EnumeratePatchData(const FString& InputFile, const FString& OutputFile, bool bIncludeSizes) override;
 	virtual bool VerifyChunkData(const FString& SearchPath, const FString& OutputFile) override;
-	virtual bool PackageChunkData(const FString& ManifestFilePath, const FString& OutputFile, const FString& CloudDir, uint64 MaxOutputFileSize) override;
+	virtual bool PackageChunkData(const FString& ManifestFilePath, const FString& PrevManifestFilePath, const TArray<TSet<FString>>& TagSetArray, const FString& OutputFile, const FString& CloudDir, uint64 MaxOutputFileSize, const FString& ResultDataFilePath) override;
 	virtual bool MergeManifests(const FString& ManifestFilePathA, const FString& ManifestFilePathB, const FString& ManifestFilePathC, const FString& NewVersionString, const FString& SelectionDetailFilePath) override;
 	virtual bool DiffManifests(const FString& ManifestFilePathA, const TSet<FString>& TagSetA, const FString& ManifestFilePathB, const TSet<FString>& TagSetB, const FString& OutputFilePath) override;
 	virtual IBuildManifestPtr MakeManifestFromJSON(const FString& ManifestJSON) override;
@@ -132,8 +132,11 @@ private:
 	// A flag specifying whether prerequisites install should be skipped
 	bool bForceSkipPrereqs;
 
-	// Save the ptr to the build installer thread that we create
-	TArray< FBuildPatchInstallerPtr > BuildPatchInstallers;
+	// Array of created installers
+	TArray<FBuildPatchInstallerPtr> BuildPatchInstallers;
+
+	// Array of created installers as exposable interface refs
+	TArray<IBuildInstallerRef> BuildPatchInstallerInterfaces;
 
 	// Holds available installations used for recycling install data
 	TMap<FString, FBuildPatchAppManifestRef> AvailableInstallations;
@@ -141,3 +144,4 @@ private:
 	// Handle to the registered Tick delegate
 	FDelegateHandle TickDelegateHandle;
 };
+PRAGMA_ENABLE_DEPRECATION_WARNINGS
diff --git a/Engine/Source/Runtime/Online/BuildPatchServices/Private/BuildPatchServicesSingleton.cpp b/Engine/Source/Runtime/Online/BuildPatchServices/Private/BuildPatchServicesSingleton.cpp
--- a/Engine/Source/Runtime/Online/BuildPatchServices/Private/BuildPatchServicesSingleton.cpp
+++ b/Engine/Source/Runtime/Online/BuildPatchServices/Private/BuildPatchServicesSingleton.cpp
@@ -3,3 +3,4 @@
 #include "BuildPatchServicesSingleton.h"
 
 FName FBuildPatchServices::ModuleName = TEXT("BuildPatchServices");
+BuildPatchServices::FBuildPatchServicesInitSettings FBuildPatchServices::InitSettings;
diff --git a/Engine/Source/Runtime/Online/BuildPatchServices/Private/BuildPatchUtil.cpp b/Engine/Source/Runtime/Online/BuildPatchServices/Private/BuildPatchUtil.cpp
--- a/Engine/Source/Runtime/Online/BuildPatchServices/Private/BuildPatchUtil.cpp
+++ b/Engine/Source/Runtime/Online/BuildPatchServices/Private/BuildPatchUtil.cpp
@@ -91,34 +91,39 @@ FString FBuildPatchUtils::GetDataTypeOldFilename( EBuildPatchDataType DataType,
 
 FString FBuildPatchUtils::GetDataFilename(const FBuildPatchAppManifestRef& Manifest, const FString& RootDirectory, const FGuid& DataGUID)
 {
-	const EBuildPatchDataType DataType = Manifest->IsFileDataManifest() ? EBuildPatchDataType::FileData : EBuildPatchDataType::ChunkData;
-	if (Manifest->GetManifestVersion() < EBuildPatchAppManifestVersion::DataFileRenames)
+	return GetDataFilename(Manifest.Get(), RootDirectory, DataGUID);
+}
+
+FString FBuildPatchUtils::GetDataFilename(const FBuildPatchAppManifest& Manifest, const FString& RootDirectory, const FGuid& DataGUID)
+{
+	const EBuildPatchDataType DataType = Manifest.IsFileDataManifest() ? EBuildPatchDataType::FileData : EBuildPatchDataType::ChunkData;
+	if (Manifest.GetManifestVersion() < EBuildPatchAppManifestVersion::DataFileRenames)
 	{
 		return FBuildPatchUtils::GetDataTypeOldFilename(DataType, RootDirectory, DataGUID);
 	}
-	else if (!Manifest->IsFileDataManifest())
+	else if (!Manifest.IsFileDataManifest())
 	{
 		uint64 ChunkHash;
-		const bool bFound = Manifest->GetChunkHash(DataGUID, ChunkHash);
+		const bool bFound = Manifest.GetChunkHash(DataGUID, ChunkHash);
 		// Should be impossible to not exist
 		check(bFound);
-		return FBuildPatchUtils::GetChunkNewFilename(Manifest->GetManifestVersion(), RootDirectory, DataGUID, ChunkHash);
+		return FBuildPatchUtils::GetChunkNewFilename(Manifest.GetManifestVersion(), RootDirectory, DataGUID, ChunkHash);
 	}
-	else if (Manifest->GetManifestVersion() <= EBuildPatchAppManifestVersion::StoredAsCompressedUClass)
+	else if (Manifest.GetManifestVersion() <= EBuildPatchAppManifestVersion::StoredAsCompressedUClass)
 	{
 		FSHAHash FileHash;
-		const bool bFound = Manifest->GetFileHash(DataGUID, FileHash);
+		const bool bFound = Manifest.GetFileHash(DataGUID, FileHash);
 		// Should be impossible to not exist
 		check(bFound);
-		return FBuildPatchUtils::GetFileNewFilename(Manifest->GetManifestVersion(), RootDirectory, DataGUID, FileHash);
+		return FBuildPatchUtils::GetFileNewFilename(Manifest.GetManifestVersion(), RootDirectory, DataGUID, FileHash);
 	}
 	else
 	{
 		uint64 FileHash;
-		const bool bFound = Manifest->GetFilePartHash(DataGUID, FileHash);
+		const bool bFound = Manifest.GetFilePartHash(DataGUID, FileHash);
 		// Should be impossible to not exist
 		check(bFound);
-		return FBuildPatchUtils::GetFileNewFilename(Manifest->GetManifestVersion(), RootDirectory, DataGUID, FileHash);
+		return FBuildPatchUtils::GetFileNewFilename(Manifest.GetManifestVersion(), RootDirectory, DataGUID, FileHash);
 	}
 	return TEXT("");
 }
diff --git a/Engine/Source/Runtime/Online/BuildPatchServices/Private/BuildPatchUtil.h b/Engine/Source/Runtime/Online/BuildPatchServices/Private/BuildPatchUtil.h
--- a/Engine/Source/Runtime/Online/BuildPatchServices/Private/BuildPatchUtil.h
+++ b/Engine/Source/Runtime/Online/BuildPatchServices/Private/BuildPatchUtil.h
@@ -104,6 +104,7 @@ struct FBuildPatchUtils
 	 * @return	the data part path
 	 */
 	static FString GetDataFilename(const FBuildPatchAppManifestRef& Manifest, const FString& RootDirectory, const FGuid& DataGUID);
+	static FString GetDataFilename(const FBuildPatchAppManifest&    Manifest, const FString& RootDirectory, const FGuid& DataGUID);
 
 	/**
 	 * Gets the GUID for a data file according to it's filename (new or old)
diff --git a/Engine/Source/Runtime/Online/BuildPatchServices/Private/Common/FileSystem.cpp b/Engine/Source/Runtime/Online/BuildPatchServices/Private/Common/FileSystem.cpp
--- a/Engine/Source/Runtime/Online/BuildPatchServices/Private/Common/FileSystem.cpp
+++ b/Engine/Source/Runtime/Online/BuildPatchServices/Private/Common/FileSystem.cpp
@@ -1,9 +1,9 @@
 // Copyright 1998-2018 Epic Games, Inc. All Rights Reserved.
 
 #include "Common/FileSystem.h"
+#include "Containers/StringConv.h"
 #include "HAL/FileManager.h"
 #include "GenericPlatform/GenericPlatformFile.h"
-#include "Containers/StringConv.h"
 
 #if PLATFORM_WINDOWS
 // Start of region that uses windows types.
@@ -16,21 +16,21 @@ THIRD_PARTY_INCLUDES_START
 THIRD_PARTY_INCLUDES_END
 namespace FileSystemHelpers
 {
-	bool PlatformGetAttributes(const TCHAR* Filename, BuildPatchServices::EFileAttributes& OutFileAttributes)
+	bool PlatformGetAttributes(const TCHAR* Filename, BuildPatchServices::EAttributeFlags& OutFileAttributes)
 	{
-		OutFileAttributes = BuildPatchServices::EFileAttributes::None;
+		OutFileAttributes = BuildPatchServices::EAttributeFlags::None;
 		DWORD FileAttributes = ::GetFileAttributesW(Filename);
 		DWORD Error = ::GetLastError();
 		if (FileAttributes != INVALID_FILE_ATTRIBUTES)
 		{
-			OutFileAttributes |= BuildPatchServices::EFileAttributes::Exists;
+			OutFileAttributes |= BuildPatchServices::EAttributeFlags::Exists;
 			if ((FileAttributes & FILE_ATTRIBUTE_READONLY) != 0)
 			{
-				OutFileAttributes |= BuildPatchServices::EFileAttributes::ReadOnly;
+				OutFileAttributes |= BuildPatchServices::EAttributeFlags::ReadOnly;
 			}
 			if ((FileAttributes & FILE_ATTRIBUTE_COMPRESSED) != 0)
 			{
-				OutFileAttributes |= BuildPatchServices::EFileAttributes::Compressed;
+				OutFileAttributes |= BuildPatchServices::EAttributeFlags::Compressed;
 			}
 			return true;
 		}
@@ -92,29 +92,27 @@ namespace FileSystemHelpers
 }
 // End of region that uses windows types.
 #include "Windows/HideWindowsPlatformTypes.h"
-// Stop windows header breaking our class's function name.
-#undef GetFileAttributes
 #elif PLATFORM_MAC
 namespace FileSystemHelpers
 {
-	bool PlatformGetAttributes(const TCHAR* Filename, BuildPatchServices::EFileAttributes& OutFileAttributes)
+	bool PlatformGetAttributes(const TCHAR* Filename, BuildPatchServices::EAttributeFlags& OutFileAttributes)
 	{
 		struct stat FileInfo;
 		const FTCHARToUTF8 FilenameUtf8(Filename);
 		int32 Result = stat(FilenameUtf8.Get(), &FileInfo);
 		int32 Error = errno;
-		OutFileAttributes = BuildPatchServices::EFileAttributes::None;
+		OutFileAttributes = BuildPatchServices::EAttributeFlags::None;
 		if (Result == 0)
 		{
-			OutFileAttributes |= BuildPatchServices::EFileAttributes::Exists;
+			OutFileAttributes |= BuildPatchServices::EAttributeFlags::Exists;
 			if ((FileInfo.st_mode & S_IWUSR) == 0)
 			{
-				OutFileAttributes |= BuildPatchServices::EFileAttributes::ReadOnly;
+				OutFileAttributes |= BuildPatchServices::EAttributeFlags::ReadOnly;
 			}
 			const mode_t ExeFlags = S_IXUSR | S_IXGRP | S_IXOTH;
 			if ((FileInfo.st_mode & ExeFlags) == ExeFlags)
 			{
-				OutFileAttributes |= BuildPatchServices::EFileAttributes::Executable;
+				OutFileAttributes |= BuildPatchServices::EAttributeFlags::Executable;
 			}
 			return true;
 		}
@@ -149,7 +147,7 @@ namespace FileSystemHelpers
 #else
 namespace FileSystemHelpers
 {
-	bool PlatformGetAttributes(const TCHAR* Filename, BuildPatchServices::EFileAttributes& OutFileAttributes)
+	bool PlatformGetAttributes(const TCHAR* Filename, BuildPatchServices::EAttributeFlags& OutFileAttributes)
 	{
 		// Not implemented.
 		return true;
@@ -170,16 +168,16 @@ namespace FileSystemHelpers
 #endif
 
 // We are forwarding flags, so assert they are all equal.
-static_assert((uint32)BuildPatchServices::EFileWrite::None == (uint32)::EFileWrite::FILEWRITE_None, "Please update FileSystem.h to match BuildPatchServices::EFileWrite::None with ::EFileWrite::FILEWRITE_None");
-static_assert((uint32)BuildPatchServices::EFileWrite::NoFail == (uint32)::EFileWrite::FILEWRITE_NoFail, "Please update FileSystem.h to match BuildPatchServices::EFileWrite::NoFail with ::EFileWrite::FILEWRITE_NoFail");
-static_assert((uint32)BuildPatchServices::EFileWrite::NoReplaceExisting == (uint32)::EFileWrite::FILEWRITE_NoReplaceExisting, "Please update FileSystem.h to match BuildPatchServices::EFileWrite::NoReplaceExisting with ::EFileWrite::FILEWRITE_NoReplaceExisting");
-static_assert((uint32)BuildPatchServices::EFileWrite::EvenIfReadOnly == (uint32)::EFileWrite::FILEWRITE_EvenIfReadOnly, "Please update FileSystem.h to match BuildPatchServices::EFileWrite::EvenIfReadOnly with ::EFileWrite::FILEWRITE_EvenIfReadOnly");
-static_assert((uint32)BuildPatchServices::EFileWrite::Append == (uint32)::EFileWrite::FILEWRITE_Append, "Please update FileSystem.h to match BuildPatchServices::EFileWrite::Append with ::EFileWrite::FILEWRITE_Append");
-static_assert((uint32)BuildPatchServices::EFileWrite::AllowRead == (uint32)::EFileWrite::FILEWRITE_AllowRead, "Please update FileSystem.h to match BuildPatchServices::EFileWrite::AllowRead with ::EFileWrite::FILEWRITE_AllowRead");
-static_assert((uint32)BuildPatchServices::EFileRead::None == (uint32)::EFileRead::FILEREAD_None, "Please update FileSystem.h to match BuildPatchServices::EFileRead::None with ::EFileRead::FILEREAD_None");
-static_assert((uint32)BuildPatchServices::EFileRead::NoFail == (uint32)::EFileRead::FILEREAD_NoFail, "Please update FileSystem.h to match BuildPatchServices::EFileRead::NoFail with ::EFileRead::FILEREAD_NoFail");
-static_assert((uint32)BuildPatchServices::EFileRead::Silent == (uint32)::EFileRead::FILEREAD_Silent, "Please update FileSystem.h to match BuildPatchServices::EFileRead::Silent with ::EFileRead::FILEREAD_Silent");
-static_assert((uint32)BuildPatchServices::EFileRead::AllowWrite == (uint32)::EFileRead::FILEREAD_AllowWrite, "Please update FileSystem.h to match BuildPatchServices::EFileRead::AllowWrite with ::EFileRead::FILEREAD_AllowWrite");
+static_assert((uint32)BuildPatchServices::EWriteFlags::None == (uint32)::EFileWrite::FILEWRITE_None, "Please update FileSystem.h to match BuildPatchServices::EFileWrite::None with ::EFileWrite::FILEWRITE_None");
+static_assert((uint32)BuildPatchServices::EWriteFlags::NoFail == (uint32)::EFileWrite::FILEWRITE_NoFail, "Please update FileSystem.h to match BuildPatchServices::EFileWrite::NoFail with ::EFileWrite::FILEWRITE_NoFail");
+static_assert((uint32)BuildPatchServices::EWriteFlags::NoReplaceExisting == (uint32)::EFileWrite::FILEWRITE_NoReplaceExisting, "Please update FileSystem.h to match BuildPatchServices::EFileWrite::NoReplaceExisting with ::EFileWrite::FILEWRITE_NoReplaceExisting");
+static_assert((uint32)BuildPatchServices::EWriteFlags::EvenIfReadOnly == (uint32)::EFileWrite::FILEWRITE_EvenIfReadOnly, "Please update FileSystem.h to match BuildPatchServices::EFileWrite::EvenIfReadOnly with ::EFileWrite::FILEWRITE_EvenIfReadOnly");
+static_assert((uint32)BuildPatchServices::EWriteFlags::Append == (uint32)::EFileWrite::FILEWRITE_Append, "Please update FileSystem.h to match BuildPatchServices::EFileWrite::Append with ::EFileWrite::FILEWRITE_Append");
+static_assert((uint32)BuildPatchServices::EWriteFlags::AllowRead == (uint32)::EFileWrite::FILEWRITE_AllowRead, "Please update FileSystem.h to match BuildPatchServices::EFileWrite::AllowRead with ::EFileWrite::FILEWRITE_AllowRead");
+static_assert((uint32)BuildPatchServices::EReadFlags::None == (uint32)::EFileRead::FILEREAD_None, "Please update FileSystem.h to match BuildPatchServices::EFileRead::None with ::EFileRead::FILEREAD_None");
+static_assert((uint32)BuildPatchServices::EReadFlags::NoFail == (uint32)::EFileRead::FILEREAD_NoFail, "Please update FileSystem.h to match BuildPatchServices::EFileRead::NoFail with ::EFileRead::FILEREAD_NoFail");
+static_assert((uint32)BuildPatchServices::EReadFlags::Silent == (uint32)::EFileRead::FILEREAD_Silent, "Please update FileSystem.h to match BuildPatchServices::EFileRead::Silent with ::EFileRead::FILEREAD_Silent");
+static_assert((uint32)BuildPatchServices::EReadFlags::AllowWrite == (uint32)::EFileRead::FILEREAD_AllowWrite, "Please update FileSystem.h to match BuildPatchServices::EFileRead::AllowWrite with ::EFileRead::FILEREAD_AllowWrite");
 
 namespace BuildPatchServices
 {
@@ -192,14 +190,16 @@ namespace BuildPatchServices
 
 		// IFileSystem interface begin.
 		virtual bool GetFileSize(const TCHAR* Filename, int64& FileSize) const override;
-		virtual bool GetFileAttributes(const TCHAR* Filename, EFileAttributes& FileAttributes) const override;
+		virtual bool GetAttributes(const TCHAR* Filename, EAttributeFlags& Attributes) const override;
 		virtual bool SetReadOnly(const TCHAR* Filename, bool bIsReadOnly) const override;
 		virtual bool SetCompressed(const TCHAR* Filename, bool bIsCompressed) const override;
 		virtual bool SetExecutable(const TCHAR* Filename, bool bIsExecutable) const override;
-		virtual TUniquePtr<FArchive> CreateFileReader(const TCHAR* Filename, EFileRead ReadFlags = EFileRead::None) const override;
-		virtual TUniquePtr<FArchive> CreateFileWriter(const TCHAR* Filename, EFileWrite WriteFlags = EFileWrite::None) const override;
+		virtual TUniquePtr<FArchive> CreateFileReader(const TCHAR* Filename, EReadFlags ReadFlags = EReadFlags::None) const override;
+		virtual TUniquePtr<FArchive> CreateFileWriter(const TCHAR* Filename, EWriteFlags WriteFlags = EWriteFlags::None) const override;
 		virtual bool DeleteFile(const TCHAR* Filename) const override;
 		virtual bool MoveFile(const TCHAR* FileDest, const TCHAR* FileSource) const override;
+		virtual bool CopyFile(const TCHAR* FileDest, const TCHAR* FileSource) const override;
+		virtual bool FileExists(const TCHAR* Filename) const override;
 		// IFileSystem interface end.
 
 	private:
@@ -220,12 +220,12 @@ namespace BuildPatchServices
 	bool FFileSystem::GetFileSize(const TCHAR* Filename, int64& FileSize) const
 	{
 		FileSize = PlatformFile.FileSize(Filename);
-		return true;
+		return FileSize >= 0;
 	}
 
-	bool FFileSystem::GetFileAttributes(const TCHAR* Filename, EFileAttributes& FileAttributes) const
+	bool FFileSystem::GetAttributes(const TCHAR* Filename, EAttributeFlags& Attributes) const
 	{
-		return FileSystemHelpers::PlatformGetAttributes(Filename, FileAttributes);
+		return FileSystemHelpers::PlatformGetAttributes(Filename, Attributes);
 	}
 
 	bool FFileSystem::SetReadOnly(const TCHAR* Filename, bool bIsReadOnly) const
@@ -243,12 +243,12 @@ namespace BuildPatchServices
 		return FileSystemHelpers::PlatformSetExecutable(Filename, bIsExecutable);
 	}
 
-	TUniquePtr<FArchive> FFileSystem::CreateFileReader(const TCHAR* Filename, EFileRead ReadFlags) const
+	TUniquePtr<FArchive> FFileSystem::CreateFileReader(const TCHAR* Filename, EReadFlags ReadFlags) const
 	{
 		return TUniquePtr<FArchive>(FileManager.CreateFileReader(Filename, static_cast<uint32>(ReadFlags)));
 	}
 
-	TUniquePtr<FArchive> FFileSystem::CreateFileWriter(const TCHAR* Filename, EFileWrite WriteFlags) const
+	TUniquePtr<FArchive> FFileSystem::CreateFileWriter(const TCHAR* Filename, EWriteFlags WriteFlags) const
 	{
 		return TUniquePtr<FArchive>(FileManager.CreateFileWriter(Filename, static_cast<uint32>(WriteFlags)));
 	}
@@ -260,11 +260,21 @@ namespace BuildPatchServices
 
 	bool FFileSystem::MoveFile(const TCHAR* FileDest, const TCHAR* FileSource) const
 	{
-		return FileManager.Move(FileDest, FileSource, true, true, true, false);
+		return FileManager.Move(FileDest, FileSource, true, true, true, true);
+	}
+
+	bool FFileSystem::CopyFile(const TCHAR* FileDest, const TCHAR* FileSource) const
+	{
+		return FileManager.Copy(FileDest, FileSource, true, true, true) == COPY_OK;
+	}
+
+	bool FFileSystem::FileExists(const TCHAR* Filename) const
+	{
+		return FileManager.FileExists(Filename);
 	}
 
 	IFileSystem* FFileSystemFactory::Create()
 	{
 		return new FFileSystem();
 	}
-}
+}
\ No newline at end of file
diff --git a/Engine/Source/Runtime/Online/BuildPatchServices/Private/Common/FileSystem.h b/Engine/Source/Runtime/Online/BuildPatchServices/Private/Common/FileSystem.h
--- a/Engine/Source/Runtime/Online/BuildPatchServices/Private/Common/FileSystem.h
+++ b/Engine/Source/Runtime/Online/BuildPatchServices/Private/Common/FileSystem.h
@@ -6,7 +6,7 @@
 
 namespace BuildPatchServices
 {
-	enum class EFileAttributes : uint32
+	enum class EAttributeFlags : uint32
 	{
 		// Value for no attributes.
 		None        = 0,
@@ -19,9 +19,9 @@ namespace BuildPatchServices
 		// Whether the file is executable.
 		Executable  = 1 << 3
 	};
-	ENUM_CLASS_FLAGS(EFileAttributes);
+	ENUM_CLASS_FLAGS(EAttributeFlags);
 
-	enum class EFileWrite : uint32
+	enum class EWriteFlags : uint32
 	{
 		None                = 0,
 		NoFail              = 1 << 0,
@@ -30,16 +30,16 @@ namespace BuildPatchServices
 		Append              = 1 << 3,
 		AllowRead           = 1 << 4
 	};
-	ENUM_CLASS_FLAGS(EFileWrite);
+	ENUM_CLASS_FLAGS(EWriteFlags);
 
-	enum class EFileRead : uint32
+	enum class EReadFlags : uint32
 	{
 		None                = 0,
 		NoFail              = 1 << 0,
 		Silent              = 1 << 1,
 		AllowWrite          = 1 << 2
 	};
-	ENUM_CLASS_FLAGS(EFileRead);
+	ENUM_CLASS_FLAGS(EReadFlags);
 
 	/**
 	 * The File System class is used for classes which require file access. It wraps Core IFileManager, and IPlatformFile. Also provides additional
@@ -62,10 +62,10 @@ namespace BuildPatchServices
 		/**
 		 * Get the attributes for a file.
 		 * @param Filename          The filename for the request.
-		 * @param FileAttributes    Receives the attribute flags if successful.
+		 * @param Attributes        Receives the attribute flags if successful.
 		 * @return true if successful.
 		 */
-		virtual bool GetFileAttributes(const TCHAR* Filename, EFileAttributes& FileAttributes) const = 0;
+		virtual bool GetAttributes(const TCHAR* Filename, EAttributeFlags& Attributes) const = 0;
 
 		/**
 		 * Set whether the file is readonly.
@@ -97,7 +97,7 @@ namespace BuildPatchServices
 		 * @param ReadFlags         The file open flags.
 		 * @return unique pointer to the created archive, invalid if failed to open the file.
 		 */
-		virtual TUniquePtr<FArchive> CreateFileReader(const TCHAR* Filename, EFileRead ReadFlags = EFileRead::None) const = 0;
+		virtual TUniquePtr<FArchive> CreateFileReader(const TCHAR* Filename, EReadFlags ReadFlags = EReadFlags::None) const = 0;
 
 		/**
 		 * Creates file writer archive.
@@ -105,7 +105,7 @@ namespace BuildPatchServices
 		 * @param WriteFlags        The file open flags.
 		 * @return unique pointer to the created archive, invalid if failed to open the file.
 		 */
-		virtual TUniquePtr<FArchive> CreateFileWriter(const TCHAR* Filename, EFileWrite WriteFlags = EFileWrite::None) const = 0;
+		virtual TUniquePtr<FArchive> CreateFileWriter(const TCHAR* Filename, EWriteFlags WriteFlags = EWriteFlags::None) const = 0;
 
 		/**
 		 * Delete a file.
@@ -121,6 +121,21 @@ namespace BuildPatchServices
 		 * @return true if the file was moved/renamed successfully.
 		 */
 		virtual bool MoveFile(const TCHAR* FileDest, const TCHAR* FileSource) const = 0;
+
+		/**
+		 * Copy a file.
+		 * @param FileDest          The destination file path.
+		 * @param FileSource        The source file path.
+		 * @return true if the file was copied successfully.
+		 */
+		virtual bool CopyFile(const TCHAR* FileDest, const TCHAR* FileSource) const = 0;
+
+		/**
+		 * Checks whether a file exists.
+		 * @param Filename          The file to check.
+		 * @return true if file exists.
+		 */
+		virtual bool FileExists(const TCHAR* Filename) const = 0;
 	};
 
 	/**
diff --git a/Engine/Source/Runtime/Online/BuildPatchServices/Private/Common/HttpManager.h b/Engine/Source/Runtime/Online/BuildPatchServices/Private/Common/HttpManager.h
--- a/Engine/Source/Runtime/Online/BuildPatchServices/Private/Common/HttpManager.h
+++ b/Engine/Source/Runtime/Online/BuildPatchServices/Private/Common/HttpManager.h
@@ -1,7 +1,9 @@
 // Copyright 1998-2018 Epic Games, Inc. All Rights Reserved.
 #pragma once
 
-#include "Interfaces/IHttpRequest.h"
+#include "CoreMinimal.h"
+
+class IHttpRequest;
 
 namespace BuildPatchServices
 {
diff --git a/Engine/Source/Runtime/Online/BuildPatchServices/Private/Common/StatsCollector.cpp b/Engine/Source/Runtime/Online/BuildPatchServices/Private/Common/StatsCollector.cpp
--- a/Engine/Source/Runtime/Online/BuildPatchServices/Private/Common/StatsCollector.cpp
+++ b/Engine/Source/Runtime/Online/BuildPatchServices/Private/Common/StatsCollector.cpp
@@ -27,7 +27,16 @@ namespace BuildPatchServices
 
 	uint64 FStatsCollector::SecondsToCycles(double Seconds)
 	{
-		return Seconds / FPlatformTime::GetSecondsPerCycle64();
+		const double Result = Seconds / FPlatformTime::GetSecondsPerCycle64();
+		// We upper clamp to cover cases where Seconds is too high to represent the resulting cycles as uint64.
+		if (Result >= TNumericLimits<uint64>::Max())
+		{
+			return TNumericLimits<uint64>::Max();
+		}
+		else
+		{
+			return Result;
+		}
 	}
 
 	void FStatsCollector::AccumulateTimeBegin(uint64& TempValue)
@@ -138,10 +147,10 @@ namespace BuildPatchServices
 					GLog->Logf(TEXT("| %s%s"), *AtomicNameMap[NameLookup], *FPlatformTime::PrettyTime(FStatsCollector::CyclesToSeconds(*Stat.Value)));
 					break;
 				case EStatFormat::DataSize:
-					GLog->Logf(TEXT("| %s%s"), *AtomicNameMap[NameLookup], *FText::AsMemory(*Stat.Value).ToString());
+					GLog->Logf(TEXT("| %s%s, %s"), *AtomicNameMap[NameLookup], *FText::AsMemory(*Stat.Value, EMemoryUnitStandard::SI).ToString(), *FText::AsMemory(*Stat.Value, EMemoryUnitStandard::IEC).ToString());
 					break;
 				case EStatFormat::DataSpeed:
-					GLog->Logf(TEXT("| %s%s/s"), *AtomicNameMap[NameLookup], *FText::AsMemory(*Stat.Value).ToString());
+					GLog->Logf(TEXT("| %s%s/s, %s/s"), *AtomicNameMap[NameLookup], *FText::AsMemory(*Stat.Value, EMemoryUnitStandard::SI).ToString(), *FText::AsMemory(*Stat.Value, EMemoryUnitStandard::IEC).ToString());
 					break;
 				case EStatFormat::Percentage:
 					GLog->Logf(TEXT("| %s%s"), *AtomicNameMap[NameLookup], *FText::AsPercent(GetAsPercentage(Stat.Value), &PercentageFormattingOptions).ToString());
diff --git a/Engine/Source/Runtime/Online/BuildPatchServices/Private/Core/AsyncHelpers.h b/Engine/Source/Runtime/Online/BuildPatchServices/Private/Core/AsyncHelpers.h
--- a/Engine/Source/Runtime/Online/BuildPatchServices/Private/Core/AsyncHelpers.h
+++ b/Engine/Source/Runtime/Online/BuildPatchServices/Private/Core/AsyncHelpers.h
@@ -2,6 +2,8 @@
 #pragma once
 
 #include "CoreMinimal.h"
+#include "HAL/ThreadSafeCounter.h"
+#include "HAL/ThreadSafeCounter64.h"
 #include "Async/Async.h"
 
 namespace BuildPatchServices
@@ -62,4 +64,33 @@ namespace BuildPatchServices
 			return Promise->GetFuture();
 		}
 	}
+
+	/**
+	 * Additional atomic functionality
+	 */
+	namespace AsyncHelpers
+	{
+		/**
+		 * LockFreePeak will set the destination value, to NewSample, if NewSample is higher.
+		 * This works by spinning on InterlockedCompareExchange. For other usage examples see reference code in GenericPlatformAtomics.h
+		 * @param PeakValue     The destination variable to set.
+		 * @param NewSample     The sample to set with if higher.
+		 */
+		template<typename IntegerType>
+		void LockFreePeak(volatile IntegerType* PeakValue, IntegerType NewSample)
+		{
+			IntegerType CurrentPeak;
+			do
+			{
+				// Read the current value.
+				CurrentPeak = *PeakValue;
+			}
+			// If the value was lower than the sample, try to update it to NewSample if the current value has not been set by another thread.
+			// If the current value change since we read it, try again to see if our sample is still higher.
+			while (CurrentPeak < NewSample && FPlatformAtomics::InterlockedCompareExchange(PeakValue, NewSample, CurrentPeak) != CurrentPeak);
+		}
+	}
+
+	typedef FThreadSafeCounter64 FThreadSafeInt64;
+	typedef FThreadSafeCounter FThreadSafeInt32;
 }
diff --git a/Engine/Source/Runtime/Online/BuildPatchServices/Private/Data/ChunkData.cpp b/Engine/Source/Runtime/Online/BuildPatchServices/Private/Data/ChunkData.cpp
--- a/Engine/Source/Runtime/Online/BuildPatchServices/Private/Data/ChunkData.cpp
+++ b/Engine/Source/Runtime/Online/BuildPatchServices/Private/Data/ChunkData.cpp
@@ -15,6 +15,54 @@
 
 namespace BuildPatchServices
 {
+	const TCHAR* ToString(const EChunkLoadResult& ChunkLoadResult)
+	{
+		switch(ChunkLoadResult)
+		{
+			case EChunkLoadResult::Success:
+				return TEXT("Success");
+			case EChunkLoadResult::OpenFileFail:
+				return TEXT("OpenFileFail");
+			case EChunkLoadResult::BadArchive:
+				return TEXT("BadArchive");
+			case EChunkLoadResult::CorruptHeader:
+				return TEXT("CorruptHeader");
+			case EChunkLoadResult::IncorrectFileSize:
+				return TEXT("IncorrectFileSize");
+			case EChunkLoadResult::UnsupportedStorage:
+				return TEXT("UnsupportedStorage");
+			case EChunkLoadResult::MissingHashInfo:
+				return TEXT("MissingHashInfo");
+			case EChunkLoadResult::SerializationError:
+				return TEXT("SerializationError");
+			case EChunkLoadResult::DecompressFailure:
+				return TEXT("DecompressFailure");
+			case EChunkLoadResult::HashCheckFailed:
+				return TEXT("HashCheckFailed");
+			case EChunkLoadResult::Aborted:
+				return TEXT("Aborted");
+			default:
+				return TEXT("Unknown");
+		}
+	}
+
+	const TCHAR* ToString(const EChunkSaveResult& ChunkSaveResult)
+	{
+		switch(ChunkSaveResult)
+		{
+			case EChunkSaveResult::Success:
+				return TEXT("Success");
+			case EChunkSaveResult::FileCreateFail:
+				return TEXT("FileCreateFail");
+			case EChunkSaveResult::BadArchive:
+				return TEXT("BadArchive");
+			case EChunkSaveResult::SerializationError:
+				return TEXT("SerializationError");
+			default:
+				return TEXT("Unknown");
+		}
+	}
+
 	namespace HeaderHelpers
 	{
 		void ZeroHeader(FChunkHeader& Header)
@@ -566,7 +614,7 @@ namespace BuildPatchServices
 			// Setup Header.
 			FChunkHeader Header;
 			FMemory::Memcpy(Header, *ChunkAccessHeader);
-			const int32 StartPos = Writer.Tell();
+			const int64 StartPos = Writer.Tell();
 			Writer << Header;
 			Header.HeaderSize = Writer.Tell() - StartPos;
 			Header.StoredAs = bDataIsCompressed ? EChunkStorageFlags::Compressed : EChunkStorageFlags::None;
diff --git a/Engine/Source/Runtime/Online/BuildPatchServices/Private/Data/ChunkData.h b/Engine/Source/Runtime/Online/BuildPatchServices/Private/Data/ChunkData.h
--- a/Engine/Source/Runtime/Online/BuildPatchServices/Private/Data/ChunkData.h
+++ b/Engine/Source/Runtime/Online/BuildPatchServices/Private/Data/ChunkData.h
@@ -74,9 +74,17 @@ namespace BuildPatchServices
 		DecompressFailure,
 
 		// The expected data hash in the header did not match the hash of the data.
-		HashCheckFailed
+		HashCheckFailed,
+
+		// The operation was aborted.
+		Aborted
 	};
 
+	/**
+	 * A ToString implementation for EChunkLoadResult.
+	 */
+	const TCHAR* ToString(const EChunkLoadResult& ChunkLoadResult);
+
 	/**
 	 * Enum which describes success, or the reason for failure when saving a chunk.
 	 */
@@ -94,6 +102,11 @@ namespace BuildPatchServices
 		SerializationError
 	};
 
+	/**
+	 * A ToString implementation for EChunkSaveResult.
+	 */
+	const TCHAR* ToString(const EChunkSaveResult& ChunkSaveResult);
+
 	/**
 	 * Declares a struct to store the info for a chunk header.
 	 */
diff --git a/Engine/Source/Runtime/Online/BuildPatchServices/Private/Data/ManifestUObject.cpp b/Engine/Source/Runtime/Online/BuildPatchServices/Private/Data/ManifestUObject.cpp
--- a/Engine/Source/Runtime/Online/BuildPatchServices/Private/Data/ManifestUObject.cpp
+++ b/Engine/Source/Runtime/Online/BuildPatchServices/Private/Data/ManifestUObject.cpp
@@ -388,6 +388,17 @@ private:
 /* FManifestUObject implementation
 *****************************************************************************/
 
+void FManifestUObject::Init()
+{
+#if !BUILDPATCHSERVICES_NOUOBJECT
+	// This fixes a potential crash if async loading manifests.
+	// We make sure that NewObject<UBuildPatchManifest>() has been called from main thread before it can be called for the 'first time' concurrently
+	// on multiple threads, otherwise a race condition can hit unprotected Emplace on UPackage::ClassUniqueNameIndexMap.
+	// The object will be GC'd on next GC run.
+	NewObject<UBuildPatchManifest>();
+#endif // !BUILDPATCHSERVICES_NOUOBJECT
+}
+
 bool FManifestUObject::LoadFromMemory(const TArray<uint8>& DataInput, FBuildPatchAppManifest& AppManifest)
 {
 #if !BUILDPATCHSERVICES_NOUOBJECT
diff --git a/Engine/Source/Runtime/Online/BuildPatchServices/Private/Data/ManifestUObject.h b/Engine/Source/Runtime/Online/BuildPatchServices/Private/Data/ManifestUObject.h
--- a/Engine/Source/Runtime/Online/BuildPatchServices/Private/Data/ManifestUObject.h
+++ b/Engine/Source/Runtime/Online/BuildPatchServices/Private/Data/ManifestUObject.h
@@ -178,6 +178,7 @@ public:
 class FManifestUObject
 {
 public:
+	static void Init();
 	static bool LoadFromMemory(const TArray<uint8>& DataInput, FBuildPatchAppManifest& AppManifest);
 	static bool SaveToArchive(FArchive& Ar, const FBuildPatchAppManifest& AppManifest);
 
diff --git a/Engine/Source/Runtime/Online/BuildPatchServices/Private/Generation/BuildStreamer.cpp b/Engine/Source/Runtime/Online/BuildPatchServices/Private/Generation/BuildStreamer.cpp
--- a/Engine/Source/Runtime/Online/BuildPatchServices/Private/Generation/BuildStreamer.cpp
+++ b/Engine/Source/Runtime/Online/BuildPatchServices/Private/Generation/BuildStreamer.cpp
@@ -421,7 +421,7 @@ namespace BuildPatchServices
 		FString IgnoreFileList = TEXT("");
 		FFileHelper::LoadFileToString(IgnoreFileList, *IgnoreListFile);
 		TArray< FString > IgnoreFiles;
-		IgnoreFileList.ParseIntoArray(IgnoreFiles, TEXT("\r\n"), true);
+		IgnoreFileList.ParseIntoArrayLines(IgnoreFiles, true);
 
 		// Normalize all paths first
 		for (FString& Filename : AllFiles)
diff --git a/Engine/Source/Runtime/Online/BuildPatchServices/Private/Generation/ChunkDatabaseWriter.cpp b/Engine/Source/Runtime/Online/BuildPatchServices/Private/Generation/ChunkDatabaseWriter.cpp
--- a/Engine/Source/Runtime/Online/BuildPatchServices/Private/Generation/ChunkDatabaseWriter.cpp
+++ b/Engine/Source/Runtime/Online/BuildPatchServices/Private/Generation/ChunkDatabaseWriter.cpp
@@ -29,12 +29,20 @@ namespace BuildPatchServices
 			: Filename(InFilename)
 			, Pos(INDEX_NONE)
 		{
+			if ( Filename.IsEmpty() )
+			{
+				UE_LOG(LogChunkDatabaseWriter, Error, TEXT("Created data message with no filename but has INDEX_NONE, WILL TRY AND CREATE FILE"));
+			}
 		}
 
 		FDataMessage(int64 InPos)
 			: Filename()
 			, Pos(InPos)
 		{
+			if ( InPos == INDEX_NONE )
+			{
+				UE_LOG(LogChunkDatabaseWriter, Error, TEXT("Created data message with no filename but has INDEX_NONE, WILL TRY AND CREATE FILE"));
+			}
 			Memory.Reset(DataMessageBufferSize);
 		}
 
@@ -218,7 +226,12 @@ namespace BuildPatchServices
 					UE_LOG(LogChunkDatabaseWriter, Log, TEXT("Writing chunk database %s"), *DataMessage->Filename);
 					CurrentFile = FileSystem->CreateFileWriter(*DataMessage->Filename);
 					FilesCreated.Add(DataMessage->Filename);
-					bSuccess = CurrentFile.IsValid();
+					if ( CurrentFile.IsValid() == false )
+					{
+						bSuccess = false;
+						UE_LOG(LogChunkDatabaseWriter, Error, TEXT("Failed to create file with name \"%s\""));
+						InstallerError->SetError(EBuildPatchInstallError::FileConstructionFail, ConstructionErrorCodes::FileCreateFail);
+					}
 				}
 				// Process a data serialize.
 				else if (CurrentFile.IsValid())
diff --git a/Engine/Source/Runtime/Online/BuildPatchServices/Private/Installer/BuildStatistics.cpp b/Engine/Source/Runtime/Online/BuildPatchServices/Private/Installer/BuildStatistics.cpp
new file mode 100644
--- /dev/null
+++ b/Engine/Source/Runtime/Online/BuildPatchServices/Private/Installer/BuildStatistics.cpp
@@ -0,0 +1,203 @@
+// Copyright 1998-2018 Epic Games, Inc. All Rights Reserved.
+
+#include "Installer/BuildStatistics.h"
+#include "Installer/Statistics/FileOperationTracker.h"
+#include "Installer/Statistics/CloudChunkSourceStatistics.h"
+#include "Installer/Statistics/FileConstructorStatistics.h"
+#include "Installer/Statistics/MemoryChunkStoreStatistics.h"
+#include "Installer/Statistics/DownloadServiceStatistics.h"
+#include "Installer/Statistics/InstallChunkSourceStatistics.h"
+#include "Installer/Statistics/VerifierStatistics.h"
+
+namespace BuildPatchServices
+{
+	FFileOperation::FFileOperation(FString InFilename, const FGuid& InDataId, uint64 InOffest, uint64 InSize, EFileOperationState InCurrentState)
+		: Filename(MoveTemp(InFilename))
+		, DataId(InDataId)
+		, Offest(InOffest)
+		, Size(InSize)
+		, CurrentState(InCurrentState)
+	{
+	}
+
+	FFileOperation::~FFileOperation()
+	{
+	}
+
+	class FBuildStatistics
+		: public IBuildStatistics
+	{
+	public:
+		FBuildStatistics(FBuildPatchInstallerRef Installer);
+		~FBuildStatistics() {}
+
+	public:
+		// IBuildStatistics interface begin.
+		virtual const FInstallerConfiguration& GetConfiguration() const override;
+		virtual int64 GetDownloadSize() const override;
+		virtual int64 GetBuildSize() const override;
+		virtual int32 GetInstallMemoryChunkStoreSize() const override;
+		virtual int32 GetInstallMemoryChunksInStore() const override;
+		virtual int32 GetInstallMemoryChunksBooted() const override;
+		virtual int32 GetInstallMemoryChunksRetained() const override;
+		virtual int32 GetCloudMemoryChunkStoreSize() const override;
+		virtual int32 GetCloudMemoryChunksInStore() const override;
+		virtual int32 GetCloudMemoryChunksBooted() const override;
+		virtual int32 GetCloudMemoryChunksRetained() const override;
+		virtual FString GetCurrentWorkingFileName() const override;
+		virtual float GetCurrentWorkingFileProgress() const override;
+		virtual FGuid GetCurrentWorkingData() const override;
+		virtual TArray<FDownload> GetCurrentDownloads() const override;
+		virtual bool IsDownloadActive() const override;
+		virtual bool IsHardDiskActiveAdministering() const override;
+		virtual bool IsHardDiskActiveWrite() const override;
+		virtual bool IsHardDiskActiveRead() const override;
+		virtual const TArray<FFileOperation>& GetFileOperationStates() const override;
+		virtual double GetDownloadByteSpeed() const override;
+		virtual double GetDiskReadByteSpeed() const override;
+		virtual double GetChunkDbReadByteSpeed() const override;
+		virtual double GetDiskWriteByteSpeed() const override;
+		virtual TMap<EVerifyError, int32> GetVerifyErrorCounts() const override;
+		// IBuildStatistics interface end.
+
+	private:
+		FBuildPatchInstallerRef Installer;
+	};
+
+	FBuildStatistics::FBuildStatistics(FBuildPatchInstallerRef InInstaller)
+		: Installer(MoveTemp(InInstaller))
+	{
+	}
+
+	const FInstallerConfiguration& FBuildStatistics::GetConfiguration() const
+	{
+		return Installer->GetConfiguration();
+	}
+
+	int64 FBuildStatistics::GetDownloadSize() const
+	{
+		return Installer->GetTotalDownloadRequired();
+	}
+
+	int64 FBuildStatistics::GetBuildSize() const
+	{
+		return Installer->GetFileConstructorStatistics()->GetRequiredConstructSize();
+	}
+
+	int32 FBuildStatistics::GetInstallMemoryChunkStoreSize() const
+	{
+		return Installer->GetInstallMemoryChunkStoreStatistics()->GetStoreSize();
+	}
+
+	int32 FBuildStatistics::GetInstallMemoryChunksInStore() const
+	{
+		return Installer->GetInstallMemoryChunkStoreStatistics()->GetStoreUse();
+	}
+
+	int32 FBuildStatistics::GetInstallMemoryChunksBooted() const
+	{
+		return Installer->GetInstallMemoryChunkStoreStatistics()->GetNumBooted();
+	}
+
+	int32 FBuildStatistics::GetInstallMemoryChunksRetained() const
+	{
+		return Installer->GetInstallMemoryChunkStoreStatistics()->GetStoreRetained();
+	}
+
+	int32 FBuildStatistics::GetCloudMemoryChunkStoreSize() const
+	{
+		return Installer->GetCloudMemoryChunkStoreStatistics()->GetStoreSize();
+	}
+
+	int32 FBuildStatistics::GetCloudMemoryChunksInStore() const
+	{
+		return Installer->GetCloudMemoryChunkStoreStatistics()->GetStoreUse();
+	}
+
+	int32 FBuildStatistics::GetCloudMemoryChunksBooted() const
+	{
+		return Installer->GetCloudMemoryChunkStoreStatistics()->GetNumBooted();
+	}
+
+	int32 FBuildStatistics::GetCloudMemoryChunksRetained() const
+	{
+		return Installer->GetCloudMemoryChunkStoreStatistics()->GetStoreRetained();
+	}
+
+	FString FBuildStatistics::GetCurrentWorkingFileName() const
+	{
+		return Installer->GetFileConstructorStatistics()->GetCurrentFile();
+	}
+
+	float FBuildStatistics::GetCurrentWorkingFileProgress() const
+	{
+		return Installer->GetFileConstructorStatistics()->GetCurrentFileProgress();
+	}
+
+	FGuid FBuildStatistics::GetCurrentWorkingData() const
+	{
+		return Installer->GetFileConstructorStatistics()->GetCurrentChunk();
+	}
+
+	TArray<FDownload> FBuildStatistics::GetCurrentDownloads() const
+	{
+		return Installer->GetDownloadServiceStatistics()->GetCurrentDownloads();
+	}
+
+	bool FBuildStatistics::IsDownloadActive() const
+	{
+		return Installer->GetDownloadServiceStatistics()->GetNumCurrentDownloads() > 0;
+	}
+
+	bool FBuildStatistics::IsHardDiskActiveAdministering() const
+	{
+		return Installer->GetFileConstructorStatistics()->IsCurrentlyAdministering();
+	}
+
+	bool FBuildStatistics::IsHardDiskActiveWrite() const
+	{
+		return Installer->GetFileConstructorStatistics()->IsCurrentlyWriting();
+	}
+
+	bool FBuildStatistics::IsHardDiskActiveRead() const
+	{
+		return Installer->GetFileConstructorStatistics()->IsCurrentlyReading()
+			|| Installer->GetInstallChunkSourceStatistics()->IsCurrentlyReading()
+			|| Installer->GetVerifierStatistics()->IsCurrentlyReading();
+	}
+
+	const TArray<FFileOperation>& FBuildStatistics::GetFileOperationStates() const
+	{
+		return Installer->GetFileOperationTracker()->GetStates();
+	}
+
+	double FBuildStatistics::GetDownloadByteSpeed() const
+	{
+		return Installer->GetDownloadSpeedRecorder()->GetAverageSpeed(10.0f);
+	}
+
+	double FBuildStatistics::GetDiskReadByteSpeed() const
+	{
+		return Installer->GetDiskReadSpeedRecorder()->GetAverageSpeed(10.0f);
+	}
+
+	double FBuildStatistics::GetChunkDbReadByteSpeed() const
+	{
+		return Installer->GetChunkDbReadSpeedRecorder()->GetAverageSpeed(10.0f);
+	}
+
+	double FBuildStatistics::GetDiskWriteByteSpeed() const
+	{
+		return Installer->GetDiskWriteSpeedRecorder()->GetAverageSpeed(10.0f);
+	}
+
+	TMap<EVerifyError, int32> FBuildStatistics::GetVerifyErrorCounts() const
+	{
+		return Installer->GetVerifierStatistics()->GetVerifyErrorCounts();
+	}
+
+	IBuildStatistics* FBuildStatisticsFactory::Create(FBuildPatchInstallerRef Installer)
+	{
+		return new FBuildStatistics(MoveTemp(Installer));
+	}
+}
diff --git a/Engine/Source/Runtime/Online/BuildPatchServices/Private/Installer/ChainedChunkSource.cpp b/Engine/Source/Runtime/Online/BuildPatchServices/Private/Installer/ChainedChunkSource.cpp
--- a/Engine/Source/Runtime/Online/BuildPatchServices/Private/Installer/ChainedChunkSource.cpp
+++ b/Engine/Source/Runtime/Online/BuildPatchServices/Private/Installer/ChainedChunkSource.cpp
@@ -14,6 +14,7 @@ namespace BuildPatchServices
 		virtual IChunkDataAccess* Get(const FGuid& DataId) override;
 		virtual TSet<FGuid> AddRuntimeRequirements(TSet<FGuid> NewRequirements) override;
 		virtual void SetUnavailableChunksCallback(TFunction<void(TSet<FGuid>)> Callback) override;
+		virtual bool AddRepeatRequirement(const FGuid& RepeatRequirement) override;
 		// IChunkSource interface end.
 
 	private:
@@ -65,6 +66,16 @@ namespace BuildPatchServices
 		return CascadeRuntimeRequirements(MoveTemp(NewRequirements));
 	}
 
+	bool FChainedChunkSource::AddRepeatRequirement(const FGuid& RepeatRequirement) 
+	{
+		bool bAtLeastOneAccepted = false;
+		for (IChunkSource* ChunkSource : ChunkSources)
+		{
+			bAtLeastOneAccepted = ChunkSource->AddRepeatRequirement(RepeatRequirement) || bAtLeastOneAccepted;
+		}
+		return bAtLeastOneAccepted;
+	}
+
 	void FChainedChunkSource::SetUnavailableChunksCallback(TFunction<void(TSet<FGuid>)> Callback)
 	{
 		UnavailableChunksCallback = Callback;
diff --git a/Engine/Source/Runtime/Online/BuildPatchServices/Private/Installer/ChunkDbChunkSource.cpp b/Engine/Source/Runtime/Online/BuildPatchServices/Private/Installer/ChunkDbChunkSource.cpp
--- a/Engine/Source/Runtime/Online/BuildPatchServices/Private/Installer/ChunkDbChunkSource.cpp
+++ b/Engine/Source/Runtime/Online/BuildPatchServices/Private/Installer/ChunkDbChunkSource.cpp
@@ -54,16 +54,13 @@ namespace ChunkDbSourceHelpers
 		switch (LoadResult)
 		{
 			case EChunkLoadResult::Success:
-			return IChunkDbChunkSourceStat::ELoadResult::Success;
-
+				return IChunkDbChunkSourceStat::ELoadResult::Success;
 			case EChunkLoadResult::MissingHashInfo:
-			return IChunkDbChunkSourceStat::ELoadResult::MissingHashInfo;
-
+				return IChunkDbChunkSourceStat::ELoadResult::MissingHashInfo;
 			case EChunkLoadResult::HashCheckFailed:
-			return IChunkDbChunkSourceStat::ELoadResult::HashCheckFailed;
-
+				return IChunkDbChunkSourceStat::ELoadResult::HashCheckFailed;
 			default:
-			return IChunkDbChunkSourceStat::ELoadResult::SerializationError;
+				return IChunkDbChunkSourceStat::ELoadResult::SerializationError;
 		}
 	}
 }
@@ -115,6 +112,7 @@ namespace BuildPatchServices
 		// IChunkSource interface begin.
 		virtual IChunkDataAccess* Get(const FGuid& DataId) override;
 		virtual TSet<FGuid> AddRuntimeRequirements(TSet<FGuid> NewRequirements) override;
+		virtual bool AddRepeatRequirement(const FGuid& RepeatRequirement) override;
 		virtual void SetUnavailableChunksCallback(TFunction<void(TSet<FGuid>)> Callback) override;
 		// IChunkSource interface end.
 
@@ -161,6 +159,8 @@ namespace BuildPatchServices
 		// Communication between the incoming Get calls, and our worker thread.
 		TQueue<FGuid, EQueueMode::Spsc> FailedToLoadMessages;
 		TSet<FGuid> FailedToLoad;
+		// Communication between the incoming AddRepeatRequirement calls, and our worker thread.
+		TQueue<FGuid, EQueueMode::Mpsc> RepeatRequirementMessages;
 	};
 
 	FChunkDbChunkSource::FChunkDbChunkSource(FChunkDbSourceConfig InConfiguration, IPlatform* InPlatform, IFileSystem* InFileSystem, IChunkStore* InChunkStore, IChunkReferenceTracker* InChunkReferenceTracker, IChunkDataSerialization* InChunkDataSerialization, IMessagePump* InMessagePump, IInstallerError* InInstallerError, IChunkDbChunkSourceStat* InChunkDbChunkSourceStat)
@@ -273,6 +273,16 @@ namespace BuildPatchServices
 		return NewRequirements.Difference(AvailableChunks);
 	}
 
+	bool FChunkDbChunkSource::AddRepeatRequirement(const FGuid& RepeatRequirement)
+	{
+		if (AvailableChunks.Contains(RepeatRequirement))
+		{
+			RepeatRequirementMessages.Enqueue(RepeatRequirement);
+			return true;
+		}
+		return false;
+	}
+
 	void FChunkDbChunkSource::SetUnavailableChunksCallback(TFunction<void(TSet<FGuid>)> Callback)
 	{
 		UnavailableChunksCallback = Callback;
@@ -290,6 +300,12 @@ namespace BuildPatchServices
 			bool bWorkPerformed = false;
 			if (bStartedLoading)
 			{
+				// 'Forget' any repeat requirements.
+				FGuid RepeatRequirement;
+				while (RepeatRequirementMessages.Dequeue(RepeatRequirement))
+				{
+					PlacedInStore.Remove(RepeatRequirement);
+				}
 				// Select chunks that are contained in our chunk db files.
 				TFunction<bool(const FGuid&)> SelectPredicate = [this](const FGuid& ChunkId)
 				{
@@ -308,6 +324,7 @@ namespace BuildPatchServices
 				BatchLoadChunks.RemoveAll(RemovePredicate);
 
 				// Load this batch.
+				ChunkDbChunkSourceStat->OnBatchStarted(BatchLoadChunks);
 				for (int32 ChunkIdx = 0; ChunkIdx < BatchLoadChunks.Num() && !bShouldAbort; ++ChunkIdx)
 				{
 					const FGuid& BatchLoadChunk = BatchLoadChunks[ChunkIdx];
@@ -415,15 +432,19 @@ namespace BuildPatchServices
 				int64 DataEndPoint = ChunkLocation.ByteStart + ChunkLocation.ByteSize;
 				if (TotalFileSize >= DataEndPoint)
 				{
+					ISpeedRecorder::FRecord ActivityRecord;
+					ActivityRecord.CyclesStart = FStatsCollector::GetCycles();
 					if (ChunkDbFile->Tell() != ChunkLocation.ByteStart)
 					{
 						ChunkDbFile->Seek(ChunkLocation.ByteStart);
 					}
 					EChunkLoadResult LoadResult;
 					TUniquePtr<IChunkDataAccess> ChunkDataAccess(ChunkDataSerialization->LoadFromArchive(*ChunkDbFile, LoadResult));
+					ActivityRecord.CyclesEnd = FStatsCollector::GetCycles();
 					const uint64 End = ChunkDbFile->Tell();
+					ActivityRecord.Size = End - ChunkLocation.ByteStart;
 					bChunkGood = LoadResult == EChunkLoadResult::Success && ChunkDataAccess.IsValid() && (End == DataEndPoint);
-					ChunkDbChunkSourceStat->OnLoadComplete(DataId, ChunkDbSourceHelpers::FromSerializer(LoadResult));
+					ChunkDbChunkSourceStat->OnLoadComplete(DataId, ChunkDbSourceHelpers::FromSerializer(LoadResult), ActivityRecord);
 					if (bChunkGood)
 					{
 						// Add it to our cache.
@@ -433,7 +454,7 @@ namespace BuildPatchServices
 				}
 				else
 				{
-					ChunkDbChunkSourceStat->OnLoadComplete(DataId, IChunkDbChunkSourceStat::ELoadResult::LocationOutOfBounds);
+					ChunkDbChunkSourceStat->OnLoadComplete(DataId, IChunkDbChunkSourceStat::ELoadResult::LocationOutOfBounds, ISpeedRecorder::FRecord());
 				}
 			}
 		}
@@ -490,4 +511,23 @@ namespace BuildPatchServices
 		check(ChunkDbChunkSourceStat != nullptr);
 		return new FChunkDbChunkSource(MoveTemp(Configuration), Platform, FileSystem, ChunkStore, ChunkReferenceTracker, ChunkDataSerialization, MessagePump, InstallerError, ChunkDbChunkSourceStat);
 	}
+
+	const TCHAR* ToString(const IChunkDbChunkSourceStat::ELoadResult& LoadResult)
+	{
+		switch(LoadResult)
+		{
+			case IChunkDbChunkSourceStat::ELoadResult::Success:
+				return TEXT("Success");
+			case IChunkDbChunkSourceStat::ELoadResult::MissingHashInfo:
+				return TEXT("MissingHashInfo");
+			case IChunkDbChunkSourceStat::ELoadResult::HashCheckFailed:
+				return TEXT("HashCheckFailed");
+			case IChunkDbChunkSourceStat::ELoadResult::LocationOutOfBounds:
+				return TEXT("LocationOutOfBounds");
+			case IChunkDbChunkSourceStat::ELoadResult::SerializationError:
+				return TEXT("SerializationError");
+			default:
+				return TEXT("Unknown");
+		}
+	}
 }
diff --git a/Engine/Source/Runtime/Online/BuildPatchServices/Private/Installer/ChunkDbChunkSource.h b/Engine/Source/Runtime/Online/BuildPatchServices/Private/Installer/ChunkDbChunkSource.h
--- a/Engine/Source/Runtime/Online/BuildPatchServices/Private/Installer/ChunkDbChunkSource.h
+++ b/Engine/Source/Runtime/Online/BuildPatchServices/Private/Installer/ChunkDbChunkSource.h
@@ -4,6 +4,7 @@
 #include "Installer/ChunkSource.h"
 #include "Installer/Controllable.h"
 #include "Containers/Set.h"
+#include "Common/SpeedRecorder.h"
 
 namespace BuildPatchServices
 {
@@ -117,6 +118,12 @@ namespace BuildPatchServices
 	public:
 		virtual ~IChunkDbChunkSourceStat() {}
 
+		/**
+		 * Called when a batch of chunks are going to be loaded.
+		 * @param ChunkIds  The ids of each chunk.
+		 */
+		virtual void OnBatchStarted(const TArray<FGuid>& ChunkIds) = 0;
+
 		/**
 		 * Called each time a chunk load begins.
 		 * @param ChunkId   The id of the chunk.
@@ -127,7 +134,13 @@ namespace BuildPatchServices
 		 * Called each time a chunk load completes.
 		 * @param ChunkId   The id of the chunk.
 		 * @param Result    The load result.
+		 * @param Record    The recording for the received data.
 		 */
-		virtual void OnLoadComplete(const FGuid& ChunkId, ELoadResult Result) = 0;
+		virtual void OnLoadComplete(const FGuid& ChunkId, ELoadResult Result, const ISpeedRecorder::FRecord& Record) = 0;
 	};
+	
+	/**
+	 * A ToString implementation for IChunkDbChunkSourceStat::ELoadResult.
+	 */
+	const TCHAR* ToString(const IChunkDbChunkSourceStat::ELoadResult& LoadResult);
 }
\ No newline at end of file
diff --git a/Engine/Source/Runtime/Online/BuildPatchServices/Private/Installer/ChunkEvictionPolicy.h b/Engine/Source/Runtime/Online/BuildPatchServices/Private/Installer/ChunkEvictionPolicy.h
--- a/Engine/Source/Runtime/Online/BuildPatchServices/Private/Installer/ChunkEvictionPolicy.h
+++ b/Engine/Source/Runtime/Online/BuildPatchServices/Private/Installer/ChunkEvictionPolicy.h
@@ -1,8 +1,9 @@
 // Copyright 1998-2018 Epic Games, Inc. All Rights Reserved.
 #pragma once
 
-#include "Data/ChunkData.h"
+#include "CoreMinimal.h"
 #include "Templates/UniquePtr.h"
+#include "Data/ChunkData.h"
 
 namespace BuildPatchServices
 {
diff --git a/Engine/Source/Runtime/Online/BuildPatchServices/Private/Installer/ChunkReferenceTracker.cpp b/Engine/Source/Runtime/Online/BuildPatchServices/Private/Installer/ChunkReferenceTracker.cpp
--- a/Engine/Source/Runtime/Online/BuildPatchServices/Private/Installer/ChunkReferenceTracker.cpp
+++ b/Engine/Source/Runtime/Online/BuildPatchServices/Private/Installer/ChunkReferenceTracker.cpp
@@ -14,7 +14,7 @@ namespace BuildPatchServices
 	{
 	public:
 		FChunkReferenceTracker(const FBuildPatchAppManifestRef& InstallManifest, const TSet<FString>& FilesToConstruct);
-		FChunkReferenceTracker(const FBuildPatchAppManifestRef& InstallManifest);
+		FChunkReferenceTracker(TArray<FGuid> CustomUseStack);
 
 		~FChunkReferenceTracker();
 
@@ -27,15 +27,15 @@ namespace BuildPatchServices
 		// IChunkReferenceTracker interface end.
 
 	private:
-		TMap<FGuid, int32> ReferenceCount;
+		TMap<FGuid, FThreadSafeCounter> ReferenceCount;
 		TArray<FGuid> UseStack;
-		mutable FCriticalSection ThreadLockCs;
+		mutable FCriticalSection UseStackCs;
 	};
 
 	FChunkReferenceTracker::FChunkReferenceTracker(const FBuildPatchAppManifestRef& InstallManifest, const TSet<FString>& FilesToConstruct)
 		: ReferenceCount()
 		, UseStack()
-		, ThreadLockCs()
+		, UseStackCs()
 	{
 		// Create our full list of chunks, including dupe references, and track the reference count of each chunk.
 		for (const FString& File : FilesToConstruct)
@@ -45,7 +45,7 @@ namespace BuildPatchServices
 			{
 				for (const FChunkPart& ChunkPart : NewFileManifest->FileChunkParts)
 				{
-					ReferenceCount.FindOrAdd(ChunkPart.Guid)++;
+					ReferenceCount.FindOrAdd(ChunkPart.Guid).Increment();
 					UseStack.Add(ChunkPart.Guid);
 				}
 			}
@@ -55,31 +55,15 @@ namespace BuildPatchServices
 		UE_LOG(LogChunkReferenceTracker, VeryVerbose, TEXT("Created. Total references:%d. Unique chunks:%d"), UseStack.Num(), ReferenceCount.Num());
 	}
 
-	FChunkReferenceTracker::FChunkReferenceTracker(const FBuildPatchAppManifestRef& InstallManifest)
+	FChunkReferenceTracker::FChunkReferenceTracker(TArray<FGuid> CustomChunkReferences)
 		: ReferenceCount()
-		, UseStack()
-		, ThreadLockCs()
+		, UseStack(MoveTemp(CustomChunkReferences))
+		, UseStackCs()
 	{
-		// Create our full list of chunks, no dupes, just one reference per chunk in the correct order.
-		TArray<FString> AllFiles;
-		TSet<FGuid> AllChunks;
-		InstallManifest->GetFileList(AllFiles);
-		for (const FString& File : AllFiles)
+		ReferenceCount.Reserve(UseStack.Num());
+		for (const FGuid& Chunk : UseStack)
 		{
-			const FFileManifest* NewFileManifest = InstallManifest->GetFileManifest(File);
-			if (NewFileManifest != nullptr)
-			{
-				for (const FChunkPart& ChunkPart : NewFileManifest->FileChunkParts)
-				{
-					bool bWasAlreadyInSet = false;
-					AllChunks.Add(ChunkPart.Guid, &bWasAlreadyInSet);
-					if (!bWasAlreadyInSet)
-					{
-						ReferenceCount.Add(ChunkPart.Guid, 1);
-						UseStack.Add(ChunkPart.Guid);
-					}
-				}
-			}
+			ReferenceCount.FindOrAdd(Chunk).Increment();
 		}
 		// Reverse the order of UseStack so it can be used as a stack.
 		Algo::Reverse(UseStack);
@@ -92,12 +76,10 @@ namespace BuildPatchServices
 
 	TSet<FGuid> FChunkReferenceTracker::GetReferencedChunks() const
 	{
-		// Thread lock to protect access to ReferenceCount.
-		FScopeLock ThreadLock(&ThreadLockCs);
 		TSet<FGuid> ReferencedChunks;
-		for (const TPair<FGuid, int32>& Pair : ReferenceCount)
+		for (const TPair<FGuid, FThreadSafeCounter>& Pair : ReferenceCount)
 		{
-			if (Pair.Value > 0)
+			if (Pair.Value.GetValue() > 0)
 			{
 				ReferencedChunks.Add(Pair.Key);
 			}
@@ -107,15 +89,13 @@ namespace BuildPatchServices
 
 	int32 FChunkReferenceTracker::GetReferenceCount(const FGuid& ChunkId) const
 	{
-		// Thread lock to protect access to ReferenceCount.
-		FScopeLock ThreadLock(&ThreadLockCs);
-		return ReferenceCount.Contains(ChunkId) ? ReferenceCount[ChunkId] : 0;
+		return ReferenceCount.Contains(ChunkId) ? ReferenceCount[ChunkId].GetValue() : 0;
 	}
 
 	void FChunkReferenceTracker::SortByUseOrder(TArray<FGuid>& ChunkList, ESortDirection Direction) const
 	{
 		// Thread lock to protect access to UseStack.
-		FScopeLock ThreadLock(&ThreadLockCs);
+		FScopeLock ThreadLock(&UseStackCs);
 		struct FIndexCache
 		{
 			FIndexCache(const TArray<FGuid>& InArray)
@@ -137,19 +117,19 @@ namespace BuildPatchServices
 		FIndexCache ChunkUseIndexes(UseStack);
 		switch (Direction)
 		{
-		case ESortDirection::Ascending:
-			Algo::SortBy(ChunkList, [&ChunkUseIndexes](const FGuid& Id) { return ChunkUseIndexes.GetIndex(Id); }, TGreater<int32>());
-			break;
-		case ESortDirection::Descending:
-			Algo::SortBy(ChunkList, [&ChunkUseIndexes](const FGuid& Id) { return ChunkUseIndexes.GetIndex(Id); }, TLess<int32>());
-			break;
+			case ESortDirection::Ascending:
+				Algo::SortBy(ChunkList, [&ChunkUseIndexes](const FGuid& Id) { return ChunkUseIndexes.GetIndex(Id); }, TGreater<int32>());
+				break;
+			case ESortDirection::Descending:
+				Algo::SortBy(ChunkList, [&ChunkUseIndexes](const FGuid& Id) { return ChunkUseIndexes.GetIndex(Id); }, TLess<int32>());
+				break;
 		}
 	}
 
 	TArray<FGuid> FChunkReferenceTracker::GetNextReferences(int32 Count, const TFunction<bool(const FGuid&)>& SelectPredicate) const
 	{
 		// Thread lock to protect access to UseStack.
-		FScopeLock ThreadLock(&ThreadLockCs);
+		FScopeLock ThreadLock(&UseStackCs);
 		TSet<FGuid> AddedIds;
 		TArray<FGuid> NextReferences;
 		for (int32 UseStackIdx = UseStack.Num() - 1; UseStackIdx >= 0 && Count > NextReferences.Num(); --UseStackIdx)
@@ -166,11 +146,11 @@ namespace BuildPatchServices
 
 	bool FChunkReferenceTracker::PopReference(const FGuid& ChunkId)
 	{
-		// Thread lock to protect access to ReferenceCount and UseStack.
-		FScopeLock ThreadLock(&ThreadLockCs);
+		// Thread lock to protect access to UseStack.
+		FScopeLock ThreadLock(&UseStackCs);
 		if (UseStack.Last() == ChunkId)
 		{
-			ReferenceCount[ChunkId]--;
+			ReferenceCount[ChunkId].Decrement();
 			UseStack.Pop();
 			return true;
 		}
@@ -182,8 +162,8 @@ namespace BuildPatchServices
 		return new FChunkReferenceTracker(InstallManifest, FilesToConstruct);
 	}
 
-	IChunkReferenceTracker* FChunkReferenceTrackerFactory::Create(const FBuildPatchAppManifestRef& InstallManifest)
+	IChunkReferenceTracker* FChunkReferenceTrackerFactory::Create(TArray<FGuid> CustomChunkReferences)
 	{
-		return new FChunkReferenceTracker(InstallManifest);
+		return new FChunkReferenceTracker(MoveTemp(CustomChunkReferences));
 	}
 }
\ No newline at end of file
