Copying //UE4/Dev-Main to //UE4/Main (Source: //Portal/Main/Engine @ 4247640)

diff --git a/Engine/Source/Runtime/Online/BuildPatchServices/Private/Installer/ChunkReferenceTracker.h b/Engine/Source/Runtime/Online/BuildPatchServices/Private/Installer/ChunkReferenceTracker.h
--- a/Engine/Source/Runtime/Online/BuildPatchServices/Private/Installer/ChunkReferenceTracker.h
+++ b/Engine/Source/Runtime/Online/BuildPatchServices/Private/Installer/ChunkReferenceTracker.h
@@ -74,13 +74,166 @@ namespace BuildPatchServices
 		 * @return the new IChunkReferenceTracker instance created.
 		 */
 		static IChunkReferenceTracker* Create(const FBuildPatchAppManifestRef& InstallManifest, const TSet<FString>& FilesToConstruct);
+		
+		/**
+		 * This implementation takes custom chunk references to track. The array should be every chunk reference, including duplicates, in order of use.
+		 * See namespace CustomChunkReferencesHelpers for common setup examples to use.
+		 * @param CustomChunkReferences    The custom chunk references to track.
+		 * @return the new IChunkReferenceTracker instance created.
+		 */
+		static IChunkReferenceTracker* Create(TArray<FGuid> CustomChunkReferences);
+	};
+
 
+	/**
+	 * Helpers for creating a custom chunk use stack for use with the equivalent FChunkReferenceTrackerFactory.
+	 */
+	namespace CustomChunkReferencesHelpers
+	{
 		/**
-		 * This implementation takes the install manifest and generates the internal data and chunk reference tracking based
+		 * This implementation takes the install manifest and generates the chunk use stack needed for a chunk reference tracker based
 		 * on caching data and so using each chunk once in the order that would be required to install the build.
 		 * @param InstallManifest   The install manifest to enumerate references from.
-		 * @return the new IChunkReferenceTracker instance created.
+		 * @return the chunk use references for FChunkReferenceTrackerFactory::Create.
 		 */
-		static IChunkReferenceTracker* Create(const FBuildPatchAppManifestRef& InstallManifest);
-	};
+		FORCEINLINE TArray<FGuid> OrderedUniqueReferences(const FBuildPatchAppManifestRef& InstallManifest)
+		{
+			TArray<FGuid> ChunkReferences;
+			// Create our full list of chunks, no dupes, just one reference per chunk in the correct order.
+			TArray<FString> AllFiles;
+			TSet<FGuid> AllChunks;
+			InstallManifest->GetFileList(AllFiles);
+			for (const FString& File : AllFiles)
+			{
+				const FFileManifest* NewFileManifest = InstallManifest->GetFileManifest(File);
+				if (NewFileManifest != nullptr)
+				{
+					for (const FChunkPart& ChunkPart : NewFileManifest->FileChunkParts)
+					{
+						bool bWasAlreadyInSet = false;
+						AllChunks.Add(ChunkPart.Guid, &bWasAlreadyInSet);
+						if (!bWasAlreadyInSet)
+						{
+							ChunkReferences.Add(ChunkPart.Guid);
+						}
+					}
+				}
+			}
+			return ChunkReferences;
+		}
+		
+		/**
+		 * This implementation takes a new install manifest and a current manifest. It generates the chunk use stack needed for a chunk reference tracker based
+		 * on caching data for a patch only, and so using the chunks in InstallManifest, which are not in CurrentManifest, once each in the order that they
+		 * would be required to patch the build.
+		 * @param InstallManifest   The install manifest to enumerate chunk references from.
+		 * @param CurrentManifest   The current manifest to exclude chunk references from.
+		 * @return the chunk use references for FChunkReferenceTrackerFactory::Create.
+		 */
+		FORCEINLINE TArray<FGuid> OrderedUniquePatchReferences(const FBuildPatchAppManifestRef& InstallManifest, const FBuildPatchAppManifestRef& CurrentManifest)
+		{
+			TArray<FGuid> ChunkReferences;
+			// Create our list of chunks, no dupes, just one reference per chunk which appears only in InstallManifest, and in the correct order of use.
+			TArray<FString> AllFiles;
+			TSet<FGuid> OldChunks;
+			TSet<FGuid> NewChunks;
+			CurrentManifest->GetDataList(OldChunks);
+			InstallManifest->GetFileList(AllFiles);
+			for (const FString& File : AllFiles)
+			{
+				const FFileManifest* NewFileManifest = InstallManifest->GetFileManifest(File);
+				if (NewFileManifest != nullptr)
+				{
+					for (const FChunkPart& ChunkPart : NewFileManifest->FileChunkParts)
+					{
+						const bool bIsNewChunk = !OldChunks.Contains(ChunkPart.Guid);
+						if (bIsNewChunk)
+						{
+							bool bWasAlreadyInSet = false;
+							NewChunks.Add(ChunkPart.Guid, &bWasAlreadyInSet);
+							if (!bWasAlreadyInSet)
+							{
+								ChunkReferences.Add(ChunkPart.Guid);
+							}
+						}
+					}
+				}
+			}
+			return ChunkReferences;
+		}
+		
+		/**
+		 * This implementation takes the install manifest and a tagset. It generates the chunk use stack needed for a chunk reference tracker based
+		 * on caching data and so using each chunk once in the order that would be required to install the build when using the same tagset provided.
+		 * @param InstallManifest   The install manifest to enumerate references from.
+		 * @param TagSet            The tagset that would be used to install the build, which will filter down required file list and thus required chunk list.
+		 * @return the chunk use references for FChunkReferenceTrackerFactory::Create.
+		 */
+		FORCEINLINE TArray<FGuid> OrderedUniqueReferencesTagged(const FBuildPatchAppManifestRef& InstallManifest, const TSet<FString>& TagSet)
+		{
+			TArray<FGuid> ChunkReferences;
+			// Create our full list of chunks, no dupes, just one reference per chunk in the correct order.
+			TArray<FString> TaggedFiles;
+			TSet<FGuid> TaggedChunks;
+			InstallManifest->GetTaggedFileList(TagSet, TaggedFiles);
+			for (const FString& File : TaggedFiles)
+			{
+				const FFileManifest* NewFileManifest = InstallManifest->GetFileManifest(File);
+				if (NewFileManifest != nullptr)
+				{
+					for (const FChunkPart& ChunkPart : NewFileManifest->FileChunkParts)
+					{
+						bool bWasAlreadyInSet = false;
+						TaggedChunks.Add(ChunkPart.Guid, &bWasAlreadyInSet);
+						if (!bWasAlreadyInSet)
+						{
+							ChunkReferences.Add(ChunkPart.Guid);
+						}
+					}
+				}
+			}
+			return ChunkReferences;
+		}
+		
+		/**
+		 * This implementation takes a new install manifest, a current manifest, and a tagset. It generates the chunk use stack needed for a chunk reference tracker based
+		 * on caching data for a patch only, and so using the chunks in InstallManifest, which are not in CurrentManifest, once each in the order that would be required
+         * to patch the build when using the same tagset provided.
+		 * @param InstallManifest   The install manifest to enumerate references from.
+		 * @param CurrentManifest   The current manifest to exclude chunk references from.
+		 * @param TagSet            The tagset that would be used to patch the build, which will filter down required file list and thus required chunk list.
+		 * @return the chunk use references for FChunkReferenceTrackerFactory::Create.
+		 */
+		FORCEINLINE TArray<FGuid> OrderedUniquePatchReferencesTagged(const FBuildPatchAppManifestRef& InstallManifest, const FBuildPatchAppManifestRef& CurrentManifest, const TSet<FString>& TagSet)
+		{
+			TArray<FGuid> ChunkReferences;
+			// Create our list of chunks, no dupes, just one reference per chunk which appears only in InstallManifest, and in the correct order of use.
+			TSet<FGuid> OldChunks;
+			CurrentManifest->GetDataList(OldChunks);
+			TArray<FString> TaggedFiles;
+			TSet<FGuid> TaggedNewChunks;
+			InstallManifest->GetTaggedFileList(TagSet, TaggedFiles);
+			for (const FString& File : TaggedFiles)
+			{
+				const FFileManifest* NewFileManifest = InstallManifest->GetFileManifest(File);
+				if (NewFileManifest != nullptr)
+				{
+					for (const FChunkPart& ChunkPart : NewFileManifest->FileChunkParts)
+					{
+						const bool bIsNewChunk = !OldChunks.Contains(ChunkPart.Guid);
+						if (bIsNewChunk)
+						{
+							bool bWasAlreadyInSet = false;
+							TaggedNewChunks.Add(ChunkPart.Guid, &bWasAlreadyInSet);
+							if (!bWasAlreadyInSet)
+							{
+								ChunkReferences.Add(ChunkPart.Guid);
+							}
+						}
+					}
+				}
+			}
+			return ChunkReferences;
+		}
+	}
 }
\ No newline at end of file
diff --git a/Engine/Source/Runtime/Online/BuildPatchServices/Private/Installer/ChunkSource.h b/Engine/Source/Runtime/Online/BuildPatchServices/Private/Installer/ChunkSource.h
--- a/Engine/Source/Runtime/Online/BuildPatchServices/Private/Installer/ChunkSource.h
+++ b/Engine/Source/Runtime/Online/BuildPatchServices/Private/Installer/ChunkSource.h
@@ -32,6 +32,14 @@ namespace BuildPatchServices
 		 */
 		virtual TSet<FGuid> AddRuntimeRequirements(TSet<FGuid> NewRequirements) = 0;
 
+		/**
+		 * Adds a requirement to reacquire a chunk that may have already been acquired before by this source. This allows the source implementation to
+		 * to support forward reading of chunks, and track which it would not need to request.
+		 * @param RepeatRequirement     The chunk that needs reacquiring.
+		 * @return true if the chunk provided can be acquired by this source.
+		 */
+		virtual bool AddRepeatRequirement(const FGuid& RepeatRequirement) = 0;
+
 		/**
 		 * Sets a callback to be used when chunks that are being fetched by this source are no longer available.
 		 * @param Callback          The function to call with the set of chunks no longer still available.
diff --git a/Engine/Source/Runtime/Online/BuildPatchServices/Private/Installer/ChunkStore.h b/Engine/Source/Runtime/Online/BuildPatchServices/Private/Installer/ChunkStore.h
--- a/Engine/Source/Runtime/Online/BuildPatchServices/Private/Installer/ChunkStore.h
+++ b/Engine/Source/Runtime/Online/BuildPatchServices/Private/Installer/ChunkStore.h
@@ -1,12 +1,13 @@
 // Copyright 1998-2018 Epic Games, Inc. All Rights Reserved.
 #pragma once
 
-#include "Data/ChunkData.h"
-#include "Misc/Guid.h"
-#include "Templates/UniquePtr.h"
+#include "CoreMinimal.h"
+#include "Templates/Function.h"
 
 namespace BuildPatchServices
 {
+	class IChunkDataAccess;
+
 	/**
 	 * An interface providing access to storage of chunk data instances.
 	 */
@@ -49,5 +50,15 @@ namespace BuildPatchServices
 		 * @return the slack space for the store.
 		 */
 		virtual int32 GetSlack() const = 0;
+
+		/**
+		 * Sets a callback to be used when chunks which have been Put, are lost.
+		 * Examples of why this may occur:
+		 *     An eviction policy instructs the store to boot a chunk, but this store has no overflow store provided. (see IChunkEvictionPolicy::Query).
+		 *     The system backing this store (e.g. a file on disk storage) experiences a failure and the chunk could not be held.
+		 * NB: The callback is not executed for a standard Clean instruction from an eviction policy.
+		 * @param Callback  The function to call with the chunk that is no longer available.
+		 */
+		virtual void SetLostChunkCallback(TFunction<void(const FGuid&)> Callback) = 0;
 	};
 }
diff --git a/Engine/Source/Runtime/Online/BuildPatchServices/Private/Installer/CloudChunkSource.cpp b/Engine/Source/Runtime/Online/BuildPatchServices/Private/Installer/CloudChunkSource.cpp
--- a/Engine/Source/Runtime/Online/BuildPatchServices/Private/Installer/CloudChunkSource.cpp
+++ b/Engine/Source/Runtime/Online/BuildPatchServices/Private/Installer/CloudChunkSource.cpp
@@ -2,6 +2,9 @@
 
 #include "Installer/CloudChunkSource.h"
 #include "CoreMinimal.h"
+#include "HAL/ThreadSafeBool.h"
+#include "Containers/Queue.h"
+#include "Misc/ScopeLock.h"
 #include "Async/Async.h"
 #include "Async/Future.h"
 #include "Core/Platform.h"
@@ -12,9 +15,6 @@
 #include "Common/StatsCollector.h"
 #include "Interfaces/IBuildInstaller.h"
 #include "BuildPatchUtil.h"
-#include "HAL/ThreadSafeBool.h"
-#include "Containers/Queue.h"
-#include "Misc/ScopeLock.h"
 
 DECLARE_LOG_CATEGORY_EXTERN(LogCloudChunkSource, Warning, All);
 DEFINE_LOG_CATEGORY(LogCloudChunkSource);
@@ -184,6 +184,7 @@ namespace BuildPatchServices
 		// IChunkSource interface begin.
 		virtual IChunkDataAccess* Get(const FGuid& DataId) override;
 		virtual TSet<FGuid> AddRuntimeRequirements(TSet<FGuid> NewRequirements) override;
+		virtual bool AddRepeatRequirement(const FGuid& RepeatRequirement) override;
 		virtual void SetUnavailableChunksCallback(TFunction<void(TSet<FGuid>)> Callback) override;
 		// IChunkSource interface end.
 
@@ -192,7 +193,7 @@ namespace BuildPatchServices
 		const FString& GetCloudRoot(int32 RetryNum) const;
 		float GetRetryDelay(int32 RetryNum);
 		EBuildPatchDownloadHealth GetDownloadHealth(bool bIsDisconnected, float ChunkSuccessRate);
-		FGuid GetNextTask(const TMap<FGuid, FTaskInfo>& TaskInfos, const TMap<int32, FGuid>& InFlightDownloads, const TSet<FGuid>& PriorityRequests, const TSet<FGuid>& FailedDownloads, const TSet<FGuid>& Stored, TArray<FGuid>& DownloadQueue);
+		FGuid GetNextTask(const TMap<FGuid, FTaskInfo>& TaskInfos, const TMap<int32, FGuid>& InFlightDownloads, const TSet<FGuid>& TotalRequiredChunks, const TSet<FGuid>& PriorityRequests, const TSet<FGuid>& FailedDownloads, const TSet<FGuid>& Stored, TArray<FGuid>& DownloadQueue);
 		void ThreadRun();
 		void OnDownloadProgress(int32 RequestId, int32 BytesSoFar);
 		void OnDownloadComplete(int32 RequestId, const FDownloadRef& Download);
@@ -230,8 +231,10 @@ namespace BuildPatchServices
 		TArray<FGuid> RequestedDownloads;
 
 		// Communication and storage of incoming additional requirements.
-		TQueue<TSet<FGuid>, EQueueMode::Spsc> RuntimeRequestMessages;
-		TSet<FGuid> RuntimeRequests;
+		TQueue<TSet<FGuid>, EQueueMode::Mpsc> RuntimeRequestMessages;
+
+		// Communication and storage of incoming repeat requirements.
+		TQueue<FGuid, EQueueMode::Mpsc> RepeatRequirementMessages;
 	};
 
 	FCloudChunkSource::FDownloadDelegates::FDownloadDelegates(FCloudChunkSource& InCloudChunkSource)
@@ -319,11 +322,19 @@ namespace BuildPatchServices
 
 	TSet<FGuid> FCloudChunkSource::AddRuntimeRequirements(TSet<FGuid> NewRequirements)
 	{
+		CloudChunkSourceStat->OnAcceptedNewRequirements(NewRequirements);
 		RuntimeRequestMessages.Enqueue(MoveTemp(NewRequirements));
 		// We don't have a concept of being unavailable yet.
 		return TSet<FGuid>();
 	}
 
+	bool FCloudChunkSource::AddRepeatRequirement(const FGuid& RepeatRequirement)
+	{
+		RepeatRequirementMessages.Enqueue(RepeatRequirement);
+		// We don't have a concept of being unavailable yet.
+		return true;
+	}
+
 	void FCloudChunkSource::SetUnavailableChunksCallback(TFunction<void(TSet<FGuid>)> Callback)
 	{
 		// We don't have a concept of being unavailable yet.
@@ -372,7 +383,7 @@ namespace BuildPatchServices
 		return DownloadHealth;
 	}
 
-	FGuid FCloudChunkSource::GetNextTask(const TMap<FGuid, FTaskInfo>& TaskInfos, const TMap<int32, FGuid>& InFlightDownloads, const TSet<FGuid>& PriorityRequests, const TSet<FGuid>& FailedDownloads, const TSet<FGuid>& Stored, TArray<FGuid>& DownloadQueue)
+	FGuid FCloudChunkSource::GetNextTask(const TMap<FGuid, FTaskInfo>& TaskInfos, const TMap<int32, FGuid>& InFlightDownloads, const TSet<FGuid>& TotalRequiredChunks, const TSet<FGuid>& PriorityRequests, const TSet<FGuid>& FailedDownloads, const TSet<FGuid>& Stored, TArray<FGuid>& DownloadQueue)
 	{
 		// Check for aborting.
 		if (bShouldAbort)
@@ -410,17 +421,10 @@ namespace BuildPatchServices
 			// Find the next chunks to get if we completed the last batch.
 			if (DownloadQueue.Num() == 0)
 			{
-				// Process new runtime requests.
-				TSet<FGuid> Temp;
-				while (RuntimeRequestMessages.Dequeue(Temp))
-				{
-					RuntimeRequests.Append(MoveTemp(Temp));
-				}
-
 				// Select the next X chunks that we initially instructed to download.
-				TFunction<bool(const FGuid&)> SelectPredicate = [this](const FGuid& ChunkId)
+				TFunction<bool(const FGuid&)> SelectPredicate = [&TotalRequiredChunks](const FGuid& ChunkId)
 				{
-					return InitialDownloadSet.Contains(ChunkId) || RuntimeRequests.Contains(ChunkId);
+					return TotalRequiredChunks.Contains(ChunkId);
 				};
 				// Clamp fetch count between min and max according to current space in the store.
 				int32 StoreSlack = ChunkStore->GetSlack();
@@ -463,10 +467,12 @@ namespace BuildPatchServices
 		EBuildPatchDownloadHealth TrackedDownloadHealth = EBuildPatchDownloadHealth::Excellent;
 		int32 TrackedActiveRequestCount = 0;
 		TSet<FGuid> TotalRequiredChunks = InitialDownloadSet;
+		int64 TotalRequiredChunkSize = InstallManifest->GetDataSize(TotalRequiredChunks);
 		int64 TotalReceivedData = 0;
+		int64 RepeatRequirementSize = 0;
 
 		// Provide initial stat values.
-		CloudChunkSourceStat->OnRequiredDataUpdated(InstallManifest->GetDataSize(TotalRequiredChunks));
+		CloudChunkSourceStat->OnRequiredDataUpdated(TotalRequiredChunkSize + RepeatRequirementSize);
 		CloudChunkSourceStat->OnReceivedDataUpdated(TotalReceivedData);
 		CloudChunkSourceStat->OnDownloadHealthUpdated(TrackedDownloadHealth);
 		CloudChunkSourceStat->OnSuccessRateUpdated(ChunkSuccessRate.GetOverall());
@@ -474,6 +480,30 @@ namespace BuildPatchServices
 
 		while (!bShouldAbort)
 		{
+			bool bRequiredDataUpdated = false;
+			// 'Forget' any repeat requirements.
+			FGuid RepeatRequirement;
+			while (RepeatRequirementMessages.Dequeue(RepeatRequirement))
+			{
+				if (PlacedInStore.Remove(RepeatRequirement) > 0)
+				{
+					RepeatRequirementSize += InstallManifest->GetDataSize(RepeatRequirement);
+					bRequiredDataUpdated = true;
+				}
+			}
+			// Process new runtime requests.
+			TSet<FGuid> Temp;
+			while (RuntimeRequestMessages.Dequeue(Temp))
+			{
+				Temp = Temp.Intersect(ChunkReferenceTracker->GetReferencedChunks());
+				Temp = Temp.Difference(TotalRequiredChunks);
+				if (Temp.Num() > 0)
+				{
+					TotalRequiredChunkSize += InstallManifest->GetDataSize(Temp);
+					TotalRequiredChunks.Append(MoveTemp(Temp));
+					bRequiredDataUpdated = true;
+				}
+			}
 			// Grab incoming requests as a priority.
 			TArray<FGuid> FrameRequestedDownloads;
 			RequestedDownloadsCS.Lock();
@@ -482,23 +512,33 @@ namespace BuildPatchServices
 			for (const FGuid& FrameRequestedDownload : FrameRequestedDownloads)
 			{
 				bDownloadsStarted = true;
-				if (PlacedInStore.Contains(FrameRequestedDownload) == false && TaskInfos.Contains(FrameRequestedDownload) == false)
+				if (!TaskInfos.Contains(FrameRequestedDownload) && !PlacedInStore.Contains(FrameRequestedDownload))
 				{
 					PriorityRequests.Add(FrameRequestedDownload);
 					if (!TotalRequiredChunks.Contains(FrameRequestedDownload))
 					{
 						TotalRequiredChunks.Add(FrameRequestedDownload);
-						CloudChunkSourceStat->OnRequiredDataUpdated(InstallManifest->GetDataSize(TotalRequiredChunks));
+						TotalRequiredChunkSize += InstallManifest->GetDataSize(FrameRequestedDownload);
+						bRequiredDataUpdated = true;
 					}
 				}
 			}
-
 			// Trim our initial download list on first begin.
 			if (!bTotalRequiredTrimmed && bDownloadsStarted)
 			{
 				bTotalRequiredTrimmed = true;
 				TotalRequiredChunks = TotalRequiredChunks.Intersect(ChunkReferenceTracker->GetReferencedChunks());
-				CloudChunkSourceStat->OnRequiredDataUpdated(InstallManifest->GetDataSize(TotalRequiredChunks));
+				const int64 NewChunkSize = InstallManifest->GetDataSize(TotalRequiredChunks);
+				if (NewChunkSize != TotalRequiredChunkSize)
+				{
+					TotalRequiredChunkSize = NewChunkSize;
+					bRequiredDataUpdated = true;
+				}
+			}
+			// Update required data spec.
+			if (bRequiredDataUpdated)
+			{
+				CloudChunkSourceStat->OnRequiredDataUpdated(TotalRequiredChunkSize + RepeatRequirementSize);
 			}
 
 			// Process completed downloads.
@@ -534,10 +574,11 @@ namespace BuildPatchServices
 					if (bDownloadSuccess)
 					{
 						TotalReceivedData += TaskInfo.ExpectedSize;
-						CloudChunkSourceStat->OnReceivedDataUpdated(TotalReceivedData);
 						TaskInfos.Remove(DownloadId);
 						PlacedInStore.Add(DownloadId);
 						ChunkStore->Put(DownloadId, MoveTemp(ChunkDataAccess));
+						CloudChunkSourceStat->OnDownloadSuccess(DownloadId);
+						CloudChunkSourceStat->OnReceivedDataUpdated(TotalReceivedData);
 					}
 					else
 					{
@@ -598,7 +639,7 @@ namespace BuildPatchServices
 			if (bDownloadsStarted)
 			{
 				FGuid NextTask;
-				while ((NextTask = GetNextTask(TaskInfos, InFlightDownloads, PriorityRequests, FailedDownloads, PlacedInStore, DownloadQueue)).IsValid())
+				while ((NextTask = GetNextTask(TaskInfos, InFlightDownloads, TotalRequiredChunks, PriorityRequests, FailedDownloads, PlacedInStore, DownloadQueue)).IsValid())
 				{
 					FTaskInfo& TaskInfo = TaskInfos.FindOrAdd(NextTask);
 					TaskInfo.UrlUsed = FBuildPatchUtils::GetDataFilename(InstallManifest, GetCloudRoot(TaskInfo.RetryNum), NextTask);
diff --git a/Engine/Source/Runtime/Online/BuildPatchServices/Private/Installer/CloudChunkSource.h b/Engine/Source/Runtime/Online/BuildPatchServices/Private/Installer/CloudChunkSource.h
--- a/Engine/Source/Runtime/Online/BuildPatchServices/Private/Installer/CloudChunkSource.h
+++ b/Engine/Source/Runtime/Online/BuildPatchServices/Private/Installer/CloudChunkSource.h
@@ -54,7 +54,8 @@ namespace BuildPatchServices
 		TArray<float> HealthPercentages;
 		// When all requests are failing, how many seconds before a success until we determine the state as disconnected.
 		float DisconnectedDelay;
-		// If true, the downloads will not begin until the first Get request is made.
+		// If true, the downloads will not begin until the first Get request is made. It is fairly fundamental to stop downloads of chunks until resume
+		// data is processed, but can be special case disabled.
 		bool bBeginDownloadsOnFirstGet;
 		// The minimum time to allow a http download before assessing it as affected by TCP zero window issue.
 		float TcpZeroWindowMinimumSeconds;
@@ -65,7 +66,7 @@ namespace BuildPatchServices
 		 */
 		FCloudSourceConfig(TArray<FString> InCloudRoots)
 			: CloudRoots(MoveTemp(InCloudRoots))
-			, NumSimultaneousDownloads(8)
+			, NumSimultaneousDownloads(16)
 			, MaxRetryCount(6)
 			, PreFetchMinimum(16)
 			, PreFetchMaximum(256)
@@ -125,6 +126,12 @@ namespace BuildPatchServices
 		 */
 		virtual void OnDownloadRequested(const FGuid& ChunkId) = 0;
 
+		/**
+		 * Called whenever a chunk download request has succeeded.
+		 * @param ChunkId           The id of the chunk.
+		 */
+		virtual void OnDownloadSuccess(const FGuid& ChunkId) = 0;
+
 		/**
 		 * Called whenever a chunk download request has failed.
 		 * @param ChunkId           The id of the chunk.
@@ -182,5 +189,11 @@ namespace BuildPatchServices
 		 * @param RequestCount      The number of currently active requests, this will range between 0 and NumSimultaneousDownloads config.
 		 */
 		virtual void OnActiveRequestCountUpdated(int32 RequestCount) = 0;
+
+		/**
+		 * Called when a batch of chunks are added and accepted via IChunkSource::AddRuntimeRequirements.
+		 * @param ChunkIds  The ids of each chunk.
+		 */
+		virtual void OnAcceptedNewRequirements(const TSet<FGuid>& ChunkIds) = 0;
 	};
 }
\ No newline at end of file
diff --git a/Engine/Source/Runtime/Online/BuildPatchServices/Private/Installer/DiskChunkStore.cpp b/Engine/Source/Runtime/Online/BuildPatchServices/Private/Installer/DiskChunkStore.cpp
--- a/Engine/Source/Runtime/Online/BuildPatchServices/Private/Installer/DiskChunkStore.cpp
+++ b/Engine/Source/Runtime/Online/BuildPatchServices/Private/Installer/DiskChunkStore.cpp
@@ -1,16 +1,31 @@
 // Copyright 1998-2018 Epic Games, Inc. All Rights Reserved.
 
 #include "Installer/DiskChunkStore.h"
-#include "Data/ChunkData.h"
+#include "HAL/ThreadSafeBool.h"
+#include "Containers/Union.h"
+#include "Containers/Queue.h"
 #include "Misc/ScopeLock.h"
+#include "Async/Async.h"
+#include "Math/UnitConversion.h"
+#include "Core/AsyncHelpers.h"
+#include "Data/ChunkData.h"
+#include "Common/FileSystem.h"
 
 namespace BuildPatchServices
 {
+	typedef TTuple<EChunkLoadResult, IChunkDataAccess*> FLoadResult;
+	typedef TFuture<FLoadResult> FLoadFuture;
+	typedef TPromise<FLoadResult> FLoadPromise;
+	typedef TTuple<FGuid, FLoadPromise*> FQueuedChunkLoad;
+	typedef TTuple<FGuid, IChunkDataAccess*> FQueuedChunkSave;
+	typedef TUnion<FQueuedChunkLoad, FQueuedChunkSave> FQueuedChunk;
+	typedef TQueue<FQueuedChunk, EQueueMode::Mpsc> FChunkQueue;
+
 	class FDiskChunkStore
 		: public IDiskChunkStore
 	{
 	public:
-		FDiskChunkStore(IChunkDataSerialization* InSerializer, IDiskChunkStoreStat* InDiskChunkStoreStat, FString InStoreRootPath);
+		FDiskChunkStore(IFileSystem* FileSystem, IChunkDataSerialization* InSerializer, IDiskChunkStoreStat* InDiskChunkStoreStat, FDiskChunkStoreConfig InConfiguration);
 		~FDiskChunkStore();
 
 		// IChunkStore interface begin.
@@ -18,51 +33,89 @@ namespace BuildPatchServices
 		virtual IChunkDataAccess* Get(const FGuid& DataId) override;
 		virtual TUniquePtr<IChunkDataAccess> Remove(const FGuid& DataId) override;
 		virtual int32 GetSlack() const override;
+		virtual void SetLostChunkCallback(TFunction<void(const FGuid&)> Callback) override;
 		// IChunkStore interface end.
 
 	private:
-		FString GetChunkFilename(const FGuid& DataId) const;
+		void IoThread();
+		FLoadFuture QueueLoadRequest(const FGuid& DataId);
+		void QueueSaveRequest(const FGuid& DataId, TUniquePtr<IChunkDataAccess>&& ChunkData);
+		void ExecLostChunkCallback(const FGuid& LostChunk);
 
 	private:
-		IChunkDataSerialization* Serializer;
-		IDiskChunkStoreStat* DiskChunkStoreStat;
-		FString StoreRootPath;
+		IFileSystem* const FileSystem;
+		IChunkDataSerialization* const Serializer;
+		IDiskChunkStoreStat* const DiskChunkStoreStat;
+		const FDiskChunkStoreConfig Configuration;
+		mutable FCriticalSection LostChunkCallbackCs;
+		TFunction<void(const FGuid&)> LostChunkCallback;
 		mutable FCriticalSection ThreadLockCs;
 		FGuid LastGetId;
 		TUniquePtr<IChunkDataAccess> LastGetData;
 		TSet<FGuid> PlacedInStore;
+		FThreadSafeBool bShouldRun;
+		FChunkQueue Queue;
+		FThreadSafeInt32 QueueNum;
+		FEvent* QueueTrigger;
+		TFuture<void> IoThreadFuture;
+		FEvent* IoThreadTrigger;
 	};
 
-	FDiskChunkStore::FDiskChunkStore(IChunkDataSerialization* InSerializer, IDiskChunkStoreStat* InDiskChunkStoreStat, FString InStoreRootPath)
-		: Serializer(InSerializer)
+	FDiskChunkStore::FDiskChunkStore(IFileSystem* InFileSystem, IChunkDataSerialization* InSerializer, IDiskChunkStoreStat* InDiskChunkStoreStat, FDiskChunkStoreConfig InConfiguration)
+		: FileSystem(InFileSystem)
+		, Serializer(InSerializer)
 		, DiskChunkStoreStat(InDiskChunkStoreStat)
-		, StoreRootPath(InStoreRootPath)
+		, Configuration(MoveTemp(InConfiguration))
+		, LostChunkCallbackCs()
+		, LostChunkCallback(nullptr)
 		, ThreadLockCs()
 		, LastGetId()
 		, LastGetData(nullptr)
 		, PlacedInStore()
+		, bShouldRun(true)
+		, QueueNum(0)
 	{
+		QueueTrigger = FPlatformProcess::GetSynchEventFromPool(true);
+		IoThreadTrigger = FPlatformProcess::GetSynchEventFromPool(true);
+		IoThreadFuture = Async<void>(EAsyncExecution::Thread, [this]()
+		{
+			return IoThread();
+		});
 	}
 
 	FDiskChunkStore::~FDiskChunkStore()
 	{
+		// Signal running loops to exit.
+		bShouldRun = false;
+		QueueTrigger->Trigger();
+		IoThreadTrigger->Trigger();
+		IoThreadFuture.Wait();
+		// Return events.
+		FPlatformProcess::ReturnSynchEventToPool(IoThreadTrigger);
+		FPlatformProcess::ReturnSynchEventToPool(QueueTrigger);
+		// Clean up allocations left in the queue.
+		FQueuedChunk QueuedChunk;
+		while (Queue.Dequeue(QueuedChunk))
+		{
+			if (QueuedChunk.HasSubtype<FQueuedChunkLoad>())
+			{
+				TUniquePtr<FLoadPromise> LoadPromisePtr(QueuedChunk.GetSubtype<FQueuedChunkLoad>().Get<1>());
+				LoadPromisePtr->SetValue(FLoadResult(EChunkLoadResult::Aborted, nullptr));
+			}
+			else if (QueuedChunk.HasSubtype<FQueuedChunkSave>())
+			{
+				TUniquePtr<IChunkDataAccess>(QueuedChunk.GetSubtype<FQueuedChunkSave>().Get<1>());
+			}
+		}
 	}
 
 	void FDiskChunkStore::Put(const FGuid& DataId, TUniquePtr<IChunkDataAccess> ChunkData)
 	{
 		// Thread lock to protect access to PlacedInStore.
 		FScopeLock ThreadLock(&ThreadLockCs);
-		if (!PlacedInStore.Contains(DataId))
-		{
-			const FString ChunkFilename = GetChunkFilename(DataId);
-			const EChunkSaveResult SaveResult = Serializer->SaveToFile(ChunkFilename, ChunkData.Get());
-			if (SaveResult == EChunkSaveResult::Success)
-			{
-				PlacedInStore.Add(DataId);
-				DiskChunkStoreStat->OnCacheUseUpdated(PlacedInStore.Num());
-			}
-			DiskChunkStoreStat->OnChunkStored(DataId, ChunkFilename, SaveResult);
-		}
+		QueueSaveRequest(DataId, MoveTemp(ChunkData));
+		PlacedInStore.Add(DataId);
+		DiskChunkStoreStat->OnCacheUseUpdated(PlacedInStore.Num());
 	}
 
 	IChunkDataAccess* FDiskChunkStore::Get(const FGuid& DataId)
@@ -75,20 +128,19 @@ namespace BuildPatchServices
 		{
 			if (PlacedInStore.Contains(DataId))
 			{
-				const FString ChunkFilename = GetChunkFilename(DataId);
-				EChunkLoadResult LoadResult;
-				ChunkData = Serializer->LoadFromFile(ChunkFilename, LoadResult);
-				if (LoadResult == EChunkLoadResult::Success)
+				FLoadFuture ChunkLoadFuture = QueueLoadRequest(DataId);
+				FLoadResult LoadResult = ChunkLoadFuture.Get();
+				if (LoadResult.Get<0>() == EChunkLoadResult::Success)
 				{
 					LastGetId = DataId;
-					LastGetData.Reset(ChunkData);
+					LastGetData.Reset(LoadResult.Get<1>());
+					ChunkData = LastGetData.Get();
 				}
 				else
 				{
 					PlacedInStore.Remove(DataId);
 					DiskChunkStoreStat->OnCacheUseUpdated(PlacedInStore.Num());
 				}
-				DiskChunkStoreStat->OnChunkLoaded(DataId, ChunkFilename, LoadResult);
 			}
 		}
 		else
@@ -107,15 +159,14 @@ namespace BuildPatchServices
 		{
 			if (PlacedInStore.Contains(DataId))
 			{
-				const FString ChunkFilename = GetChunkFilename(DataId);
-				EChunkLoadResult LoadResult;
-				ReturnValue.Reset(Serializer->LoadFromFile(ChunkFilename, LoadResult));
-				if (LoadResult != EChunkLoadResult::Success)
+				FLoadFuture ChunkLoadFuture = QueueLoadRequest(DataId);
+				FLoadResult LoadResult = ChunkLoadFuture.Get();
+				ReturnValue.Reset(LoadResult.Get<1>());
+				if (LoadResult.Get<0>() != EChunkLoadResult::Success)
 				{
 					PlacedInStore.Remove(DataId);
 					DiskChunkStoreStat->OnCacheUseUpdated(PlacedInStore.Num());
 				}
-				DiskChunkStoreStat->OnChunkLoaded(DataId, ChunkFilename, LoadResult);
 			}
 		}
 		else
@@ -132,16 +183,186 @@ namespace BuildPatchServices
 		return MAX_int32;
 	}
 
-	FString FDiskChunkStore::GetChunkFilename(const FGuid& DataId) const
+	void FDiskChunkStore::SetLostChunkCallback(TFunction<void(const FGuid&)> Callback)
+	{
+		// Thread lock to protect access to LostChunkCallback.
+		FScopeLock ThreadLock(&LostChunkCallbackCs);
+		LostChunkCallback = MoveTemp(Callback);
+	}
+
+	void FDiskChunkStore::IoThread()
+	{
+		FString DumpFilename = Configuration.StoreRootPath / TEXT("chunkdump");
+		TMap<FGuid, int64> ChunkIdToDumpLocation;
+
+		TUniquePtr<FArchive> DumpWriter = FileSystem->CreateFileWriter(*DumpFilename, EWriteFlags::AllowRead);
+		TUniquePtr<FArchive> DumpReader = FileSystem->CreateFileReader(*DumpFilename, EReadFlags::AllowWrite);
+		bool bValidFileHandle = DumpReader.IsValid() && DumpWriter.IsValid();
+		if (bValidFileHandle)
+		{
+			DumpWriter->Seek(0);
+		}
+
+		FQueuedChunk QueuedChunk;
+		while (bShouldRun)
+		{
+			if (Queue.Dequeue(QueuedChunk))
+			{
+				QueueNum.Decrement();
+				QueueTrigger->Trigger();
+				if (QueuedChunk.HasSubtype<FQueuedChunkLoad>())
+				{
+					FQueuedChunkLoad& QueuedChunkLoad = QueuedChunk.GetSubtype<FQueuedChunkLoad>();
+					FGuid& DataId = QueuedChunkLoad.Get<0>();
+					TUniquePtr<FLoadPromise> LoadPromisePtr(QueuedChunkLoad.Get<1>());
+					DiskChunkStoreStat->OnBeforeChunkLoad(DataId);
+					FLoadResult LoadResult(EChunkLoadResult::BadArchive, nullptr);
+					EChunkLoadResult& ChunkLoadResult = LoadResult.Get<0>();
+					IChunkDataAccess*& ChunkDataAccess = LoadResult.Get<1>();
+					if (ChunkIdToDumpLocation.Contains(DataId))
+					{
+						if (bValidFileHandle && DumpReader->TotalSize() != DumpWriter->TotalSize())
+						{
+							// We have to re-open the reader when the size has not been refreshed.
+							DumpReader = FileSystem->CreateFileReader(*DumpFilename, EReadFlags::AllowWrite);
+							bValidFileHandle = DumpReader.IsValid();
+						}
+						if (bValidFileHandle)
+						{
+							DumpWriter->Flush();
+							const int64& ChunkStartPos = ChunkIdToDumpLocation[DataId];
+							DumpReader->Seek(ChunkStartPos);
+							ChunkDataAccess = Serializer->LoadFromArchive(*DumpReader, ChunkLoadResult);
+							bValidFileHandle = !DumpReader->IsError() && !DumpWriter->IsError();
+						}
+						if (ChunkLoadResult != EChunkLoadResult::Success)
+						{
+							ChunkIdToDumpLocation.Remove(DataId);
+							ExecLostChunkCallback(DataId);
+						}
+					}
+					LoadPromisePtr->SetValue(LoadResult);
+					DiskChunkStoreStat->OnChunkLoaded(DataId, DumpFilename, ChunkLoadResult);
+				}
+				else if (QueuedChunk.HasSubtype<FQueuedChunkSave>())
+				{
+					FQueuedChunkSave& QueuedChunkSave = QueuedChunk.GetSubtype<FQueuedChunkSave>();
+					FGuid& DataId = QueuedChunkSave.Get<0>();
+					TUniquePtr<IChunkDataAccess> ChunkDataPtr(QueuedChunkSave.Get<1>());
+					if (!ChunkIdToDumpLocation.Contains(DataId))
+					{
+						EChunkSaveResult SaveResult = EChunkSaveResult::FileCreateFail;
+						if (bValidFileHandle)
+						{
+							DumpWriter->Seek(DumpWriter->TotalSize());
+							const int64 ChunkStartPos = DumpWriter->Tell();
+							SaveResult = Serializer->SaveToArchive(*DumpWriter, ChunkDataPtr.Get());
+							const int64 ChunkEndPos = DumpWriter->Tell();
+							bValidFileHandle = !DumpWriter->IsError();
+							if (SaveResult == EChunkSaveResult::Success)
+							{
+								ChunkIdToDumpLocation.Add(DataId, ChunkStartPos);
+							}
+						}
+						if (SaveResult != EChunkSaveResult::Success)
+						{
+							ExecLostChunkCallback(DataId);
+						}
+						DiskChunkStoreStat->OnChunkStored(DataId, DumpFilename, SaveResult);
+					}
+				}
+				else
+				{
+					checkf(0, TEXT("Union has no recognized type"));
+				}
+			}
+			else if (bValidFileHandle && DumpReader->TotalSize() != DumpWriter->TotalSize())
+			{
+				// We have to re-open the reader when the size has not been refreshed.
+				DumpReader = FileSystem->CreateFileReader(*DumpFilename, EReadFlags::AllowWrite);
+				bValidFileHandle = DumpReader.IsValid();
+			}
+			else
+			{
+				// Flush when waiting for jobs.
+				if (bValidFileHandle)
+				{
+					DumpWriter->Flush();
+					bValidFileHandle = !DumpWriter->IsError();
+				}
+				// Or try to fix our file access issue if we have one.
+				else
+				{
+					DumpWriter = FileSystem->CreateFileWriter(*DumpFilename, EWriteFlags::AllowRead);
+					DumpReader = FileSystem->CreateFileReader(*DumpFilename, EReadFlags::AllowWrite);
+					bValidFileHandle = DumpReader.IsValid() && DumpWriter.IsValid();
+				}
+				// Max wait time allows us to periodically retry inaccessible file.
+				static const uint32 WaitTime = Configuration.MaxRetryTime * 1000.0;
+				IoThreadTrigger->Wait(WaitTime);
+				IoThreadTrigger->Reset();
+			}
+		}
+
+		DumpWriter.Reset();
+		DumpReader.Reset();
+		FileSystem->DeleteFile(*DumpFilename);
+	}
+
+	BuildPatchServices::FLoadFuture FDiskChunkStore::QueueLoadRequest(const FGuid& DataId)
 	{
-		return StoreRootPath / DataId.ToString() + TEXT(".chunk");
+		while (QueueNum.GetValue() > Configuration.QueueSize && bShouldRun)
+		{
+			// Wait 1 second max in case of abort.
+			static const uint32 WaitTime = 1000;
+			QueueTrigger->Wait(WaitTime);
+			QueueTrigger->Reset();
+		}
+		QueueNum.Increment();
+		TUniquePtr<FLoadPromise> PromisePtr(new FLoadPromise());
+		FLoadFuture Future = PromisePtr->GetFuture();
+		if (bShouldRun)
+		{
+			Queue.Enqueue(FQueuedChunk(FQueuedChunkLoad(DataId, PromisePtr.Release())));
+			IoThreadTrigger->Trigger();
+		}
+		else
+		{
+			PromisePtr->SetValue(FLoadResult(EChunkLoadResult::Aborted, nullptr));
+		}
+		return Future;
+	}
+
+	void FDiskChunkStore::QueueSaveRequest(const FGuid& DataId, TUniquePtr<IChunkDataAccess>&& ChunkData)
+	{
+		while (QueueNum.GetValue() > Configuration.QueueSize && bShouldRun)
+		{
+			// Wait 1 second max in case of abort.
+			static const uint32 WaitTime = 1000;
+			QueueTrigger->Wait(WaitTime);
+			QueueTrigger->Reset();
+		}
+		QueueNum.Increment();
+		Queue.Enqueue(FQueuedChunk(FQueuedChunkSave(DataId, ChunkData.Release())));
+		IoThreadTrigger->Trigger();
+	}
+
+	void FDiskChunkStore::ExecLostChunkCallback(const FGuid& LostChunk)
+	{
+		// Thread lock to protect access to LostChunkCallback.
+		FScopeLock ThreadLock(&LostChunkCallbackCs);
+		if (LostChunkCallback)
+		{
+			LostChunkCallback(LostChunk);
+		}
 	}
 
-	IDiskChunkStore* FDiskChunkStoreFactory::Create(IChunkDataSerialization* Serializer, IDiskChunkStoreStat* DiskChunkStoreStat, FString StoreRootPath)
+	IDiskChunkStore* FDiskChunkStoreFactory::Create(IFileSystem* FileSystem, IChunkDataSerialization* Serializer, IDiskChunkStoreStat* DiskChunkStoreStat, FDiskChunkStoreConfig Configuration)
 	{
+		check(FileSystem != nullptr);
 		check(Serializer != nullptr);
 		check(DiskChunkStoreStat != nullptr);
-		check(StoreRootPath.IsEmpty() == false);
-		return new FDiskChunkStore(Serializer, DiskChunkStoreStat, MoveTemp(StoreRootPath));
+		check(Configuration.StoreRootPath.IsEmpty() == false);
+		return new FDiskChunkStore(FileSystem, Serializer, DiskChunkStoreStat, MoveTemp(Configuration));
 	}
 }
\ No newline at end of file
diff --git a/Engine/Source/Runtime/Online/BuildPatchServices/Private/Installer/DiskChunkStore.h b/Engine/Source/Runtime/Online/BuildPatchServices/Private/Installer/DiskChunkStore.h
--- a/Engine/Source/Runtime/Online/BuildPatchServices/Private/Installer/DiskChunkStore.h
+++ b/Engine/Source/Runtime/Online/BuildPatchServices/Private/Installer/DiskChunkStore.h
@@ -5,6 +5,7 @@
 
 namespace BuildPatchServices
 {
+	class IFileSystem;
 	class IChunkDataSerialization;
 	class IDiskChunkStoreStat;
 	enum class EChunkSaveResult : uint8;
@@ -20,6 +21,26 @@ namespace BuildPatchServices
 		virtual ~IDiskChunkStore() {}
 	};
 
+	struct FDiskChunkStoreConfig
+	{
+		// The directory path for where to store chunk files.
+		FString StoreRootPath;
+		// The number of requests to allow to queue up before blocking Put/Get calls.
+		int32 QueueSize;
+		// The max time in seconds to wait before retrying file handle for dump file.
+		double MaxRetryTime;
+
+		/**
+		 * Constructor which sets usual defaults.
+		 */
+		FDiskChunkStoreConfig(FString InStoreRootPath)
+			: StoreRootPath(MoveTemp(InStoreRootPath))
+			, QueueSize(50)
+			, MaxRetryTime(2.0)
+		{
+		}
+	};
+
 	/**
 	 * A factory for creating an IDiskChunkStore instance.
 	 */
@@ -31,13 +52,14 @@ namespace BuildPatchServices
 		 * As per the IChunkStore contract, the ptr returned by Get() will be valid at least until another Get() call is made.
 		 * A Remove() call will not actually delete the data from disk.
 		 * A Get(), Put(), or Remove() call may block on the file IO.
+		 * @param FileSystem            Required ptr to file system for file IO.
 		 * @param Serializer            Required ptr to the serialization implementation to use. If existing chunks used a different serializer,
 		 *                              then Get() and Remove() calls for those could fail.
 		 * @param DiskChunkStoreStat    Pointer to the statistics receiver.
-		 * @param StoreRootPath         The directory path for where to store chunk files.
+		 * @param Configuration         The configuration struct.
 		 * @return the new IDiskChunkStore instance created.
 		 */
-		static IDiskChunkStore* Create(IChunkDataSerialization* Serializer, IDiskChunkStoreStat* DiskChunkStoreStat, FString StoreRootPath);
+		static IDiskChunkStore* Create(IFileSystem* FileSystem, IChunkDataSerialization* Serializer, IDiskChunkStoreStat* DiskChunkStoreStat, FDiskChunkStoreConfig Configuration);
 	};
 
 	/**
@@ -57,6 +79,12 @@ namespace BuildPatchServices
 		 */
 		virtual void OnChunkStored(const FGuid& ChunkId, const FString& ChunkFilename, EChunkSaveResult SaveResult) = 0;
 
+		/**
+		 * Called whenever chunk is going to be loaded from the store.
+		 * @param ChunkId           The id of the chunk.
+		 */
+		virtual void OnBeforeChunkLoad(const FGuid& ChunkId) = 0;
+
 		/**
 		 * Called whenever a new chunk has been loaded from the store.
 		 * @param ChunkId           The id of the chunk.
diff --git a/Engine/Source/Runtime/Online/BuildPatchServices/Private/Installer/DownloadService.cpp b/Engine/Source/Runtime/Online/BuildPatchServices/Private/Installer/DownloadService.cpp
--- a/Engine/Source/Runtime/Online/BuildPatchServices/Private/Installer/DownloadService.cpp
+++ b/Engine/Source/Runtime/Online/BuildPatchServices/Private/Installer/DownloadService.cpp
@@ -1,6 +1,8 @@
 // Copyright 1998-2018 Epic Games, Inc. All Rights Reserved.
 
 #include "Installer/DownloadService.h"
+#include "HAL/ThreadSafeBool.h"
+#include "Misc/ScopeLock.h"
 #include "Core/AsyncHelpers.h"
 #include "Interfaces/IHttpRequest.h"
 #include "Interfaces/IHttpResponse.h"
@@ -8,8 +10,6 @@
 #include "Installer/InstallerAnalytics.h"
 #include "Common/HttpManager.h"
 #include "Common/FileSystem.h"
-#include "HAL/ThreadSafeBool.h"
-#include "Misc/ScopeLock.h"
 
 DECLARE_LOG_CATEGORY_EXTERN(LogDownloadService, Warning, All);
 DEFINE_LOG_CATEGORY(LogDownloadService);
@@ -37,11 +37,27 @@ namespace BuildPatchServices
 	{
 	}
 
-	class FHttpDownload
+	class FDownloadBase
 		: public IDownload
 	{
 	public:
-		FHttpDownload(const FHttpResponsePtr& InHttpResponse, bool bInSuccess);
+		FDownloadBase(IDownloadServiceStat::FDownloadRecord&& InDownloadRecord);
+
+	public:
+		IDownloadServiceStat::FDownloadRecord DownloadRecord;
+	};
+	typedef TSharedRef<FDownloadBase, ESPMode::ThreadSafe> FDownloadBaseRef;
+
+	FDownloadBase::FDownloadBase(IDownloadServiceStat::FDownloadRecord&& InDownloadRecord)
+		: DownloadRecord(MoveTemp(InDownloadRecord))
+	{
+	}
+
+	class FHttpDownload
+		: public FDownloadBase
+	{
+	public:
+		FHttpDownload(const FHttpResponsePtr& HttpResponse, bool bSuccess, IDownloadServiceStat::FDownloadRecord&& DownloadRecord);
 		~FHttpDownload();
 
 		// IDownload interface begin.
@@ -55,8 +71,9 @@ namespace BuildPatchServices
 		bool bSuccess;
 	};
 
-	FHttpDownload::FHttpDownload(const FHttpResponsePtr& InHttpResponse, bool bInSuccess)
-		: HttpResponse(InHttpResponse)
+	FHttpDownload::FHttpDownload(const FHttpResponsePtr& InHttpResponse, bool bInSuccess, IDownloadServiceStat::FDownloadRecord&& InDownloadRecord)
+		: FDownloadBase(MoveTemp(InDownloadRecord))
+		, HttpResponse(InHttpResponse)
 		, bSuccess(bInSuccess)
 	{
 	}
@@ -81,10 +98,10 @@ namespace BuildPatchServices
 	}
 
 	class FFileDownload
-		: public IDownload
+		: public FDownloadBase
 	{
 	public:
-		FFileDownload(TArray<uint8> InDataArray, bool bInSuccess);
+		FFileDownload(TArray<uint8> DataArray, bool bSuccess, IDownloadServiceStat::FDownloadRecord&& DownloadRecord);
 		~FFileDownload();
 
 		// IDownload interface begin.
@@ -98,8 +115,9 @@ namespace BuildPatchServices
 		bool bSuccess;
 	};
 
-	FFileDownload::FFileDownload(TArray<uint8> InDataArray, bool bInSuccess)
-		: DataArray(MoveTemp(InDataArray))
+	FFileDownload::FFileDownload(TArray<uint8> InDataArray, bool bInSuccess, IDownloadServiceStat::FDownloadRecord&& InDownloadRecord)
+		: FDownloadBase(MoveTemp(InDownloadRecord))
+		, DataArray(MoveTemp(InDataArray))
 		, bSuccess(bInSuccess)
 	{
 	}
@@ -150,7 +168,7 @@ namespace BuildPatchServices
 	{
 		void ExecuteCancelled(int32 RequestId, const FDownloadDelegates& DownloadDelegates)
 		{
-			DownloadDelegates.OnCompleteDelegate.ExecuteIfBound(RequestId, MakeShareable(new FFileDownload(TArray<uint8>(), false)));
+			DownloadDelegates.OnCompleteDelegate.ExecuteIfBound(RequestId, MakeShareable(new FFileDownload(TArray<uint8>(), false, IDownloadServiceStat::FDownloadRecord())));
 		}
 
 		void ExecuteCancelledAndReset(int32 RequestId, FDownloadDelegates& DownloadDelegates)
@@ -206,8 +224,8 @@ namespace BuildPatchServices
 		void HttpRequestProgress(FHttpRequestPtr Request, int32 BytesSent, int32 BytesReceived, int32 RequestId);
 		void HttpRequestComplete(FHttpRequestPtr Request, FHttpResponsePtr Response, bool bSuccess, IDownloadServiceStat::FDownloadRecord DownloadRecord);
 		void SetRequestProgress(int32 RequestId, int32 BytesSoFar);
-		void SetFileRequestComplete(int32 RequestId, bool bSuccess, TArray<uint8> FileDataArray);
-		void SetHttpRequestComplete(int32 RequestId, bool bSuccess, FHttpResponsePtr Response);
+		void SetFileRequestComplete(int32 RequestId, bool bSuccess, TArray<uint8> FileDataArray, IDownloadServiceStat::FDownloadRecord&& DownloadRecord);
+		void SetHttpRequestComplete(int32 RequestId, bool bSuccess, FHttpResponsePtr Response, IDownloadServiceStat::FDownloadRecord&& DownloadRecord);
 
 	private:
 		TSharedRef<FHttpDelegates, ESPMode::ThreadSafe> HttpDelegates;
@@ -236,7 +254,7 @@ namespace BuildPatchServices
 		TMap<int32, int32> ProgressUpdates;
 
 		FCriticalSection CompletedRequestsCS;
-		TMap<int32, FDownloadRef> CompletedRequests;
+		TMap<int32, FDownloadBaseRef> CompletedRequests;
 
 		FDelegateHandle TickerHandle;
 	};
@@ -439,6 +457,7 @@ namespace BuildPatchServices
 				FileRequest->Future = Async(EAsyncExecution::ThreadPool, MakeFileLoadTask(NewRequest.Key, NewRequest.Value, FileRequestPtr));
 				RegisterRequest(NewRequest.Key, MoveTemp(FileRequest));
 			}
+			DownloadServiceStat->OnDownloadStarted(NewRequest.Key, NewRequest.Value);
 		}
 	}
 
@@ -460,19 +479,23 @@ namespace BuildPatchServices
 			}
 		}
 		RequestDelegatesCS.Unlock();
+		for (const TPair<int32, int32>& FrameProgressUpdate : FrameProgressUpdates)
+		{
+			DownloadServiceStat->OnDownloadProgress(FrameProgressUpdate.Key, FrameProgressUpdate.Value);
+		}
 	}
 
 	void FDownloadService::ProcessCompletedRequests()
 	{
 		// Grab the completed requests for this frame.
-		TMap<int32, FDownloadRef> FrameCompletedRequests;
+		TMap<int32, FDownloadBaseRef> FrameCompletedRequests;
 		CompletedRequestsCS.Lock();
 		FrameCompletedRequests = MoveTemp(CompletedRequests);
 		CompletedRequestsCS.Unlock();
 
 		// Process completed requests.
 		RequestDelegatesCS.Lock();
-		for (const TPair<int32, FDownloadRef>& FrameCompletedRequest : FrameCompletedRequests)
+		for (const TPair<int32, FDownloadBaseRef>& FrameCompletedRequest : FrameCompletedRequests)
 		{
 			if (ensureMsgf(RequestDelegates.Contains(FrameCompletedRequest.Key), TEXT("Missing request delegates for %d"), FrameCompletedRequest.Key))
 			{
@@ -482,6 +505,10 @@ namespace BuildPatchServices
 			}
 		}
 		RequestDelegatesCS.Unlock();
+		for (const TPair<int32, FDownloadBaseRef>& FrameCompletedRequest : FrameCompletedRequests)
+		{
+			DownloadServiceStat->OnDownloadComplete(MoveTemp(FrameCompletedRequest.Value->DownloadRecord));
+		}
 	}
 
 	int32 FDownloadService::MakeRequestId()
@@ -496,9 +523,9 @@ namespace BuildPatchServices
 		DownloadRecord.Uri = MoveTemp(Uri);
 		DownloadRecord.bSuccess = false;
 		DownloadRecord.ResponseCode = INDEX_NONE;
-		DownloadRecord.StartedAt = FStatsCollector::GetSeconds();
-		DownloadRecord.CompletedAt = DownloadRecord.StartedAt;
-		DownloadRecord.BytesReceived = 0;
+		DownloadRecord.SpeedRecord.CyclesStart = FStatsCollector::GetCycles();
+		DownloadRecord.SpeedRecord.CyclesEnd = DownloadRecord.SpeedRecord.CyclesStart;
+		DownloadRecord.SpeedRecord.Size = 0;
 		return DownloadRecord;
 	}
 
@@ -510,9 +537,9 @@ namespace BuildPatchServices
 			TSharedPtr<FThreadSafeBool, ESPMode::ThreadSafe> LambdaShouldRun = WeakShouldRun.Pin();
 			TArray<uint8> FileDataArray;
 			bool bSuccess = LambdaShouldRun.IsValid() && *LambdaShouldRun.Get() && !FileRequest->ShouldCancel;
+			IDownloadServiceStat::FDownloadRecord DownloadRecord = MakeDownloadRecord(RequestId, FileUri);
 			if (bSuccess)
 			{
-				IDownloadServiceStat::FDownloadRecord DownloadRecord = MakeDownloadRecord(RequestId, FileUri);
 				TUniquePtr<FArchive> Reader = FileSystem->CreateFileReader(*FileUri);
 				bSuccess = Reader.IsValid();
 				if (bSuccess)
@@ -528,14 +555,13 @@ namespace BuildPatchServices
 						BytesRead += ReadLen;
 						SetRequestProgress(RequestId, BytesRead);
 					}
-					DownloadRecord.BytesReceived = BytesRead;
+					DownloadRecord.SpeedRecord.Size = BytesRead;
 					bSuccess = Reader->Close() && BytesRead == FileSize;
 				}
-				DownloadRecord.CompletedAt = FStatsCollector::GetSeconds();
+				DownloadRecord.SpeedRecord.CyclesEnd = FStatsCollector::GetCycles();
 				DownloadRecord.bSuccess = bSuccess;
-				DownloadServiceStat->OnDownloadComplete(MoveTemp(DownloadRecord));
 			}
-			SetFileRequestComplete(RequestId, bSuccess, MoveTemp(FileDataArray));
+			SetFileRequestComplete(RequestId, bSuccess, MoveTemp(FileDataArray), MoveTemp(DownloadRecord));
 		};
 	}
 
@@ -568,10 +594,9 @@ namespace BuildPatchServices
 		InstallerAnalytics->TrackRequest(Request);
 		DownloadRecord.bSuccess = bSuccess;
 		DownloadRecord.ResponseCode = Response.IsValid() ? Response->GetResponseCode() : INDEX_NONE;
-		DownloadRecord.CompletedAt = FStatsCollector::GetSeconds();
-		DownloadRecord.BytesReceived = Response.IsValid() ? Response->GetContent().Num() : 0;
-		SetHttpRequestComplete(DownloadRecord.RequestId, bSuccess, MoveTemp(Response));
-		DownloadServiceStat->OnDownloadComplete(MoveTemp(DownloadRecord));
+		DownloadRecord.SpeedRecord.CyclesEnd = FStatsCollector::GetCycles();
+		DownloadRecord.SpeedRecord.Size = Response.IsValid() ? Response->GetContent().Num() : 0;
+		SetHttpRequestComplete(DownloadRecord.RequestId, bSuccess, MoveTemp(Response), MoveTemp(DownloadRecord));
 	}
 
 	void FDownloadService::SetRequestProgress(int32 RequestId, int32 BytesSoFar)
@@ -580,16 +605,16 @@ namespace BuildPatchServices
 		ProgressUpdates.Add(RequestId, BytesSoFar);
 	}
 
-	void FDownloadService::SetFileRequestComplete(int32 RequestId, bool bSuccess, TArray<uint8> FileDataArray)
+	void FDownloadService::SetFileRequestComplete(int32 RequestId, bool bSuccess, TArray<uint8> FileDataArray, IDownloadServiceStat::FDownloadRecord&& DownloadRecord)
 	{
 		FScopeLock ScopeLock(&CompletedRequestsCS);
-		CompletedRequests.Emplace(RequestId, MakeShareable(new FFileDownload(MoveTemp(FileDataArray), bSuccess)));
+		CompletedRequests.Emplace(RequestId, MakeShareable(new FFileDownload(MoveTemp(FileDataArray), bSuccess, MoveTemp(DownloadRecord))));
 	}
 
-	void FDownloadService::SetHttpRequestComplete(int32 RequestId, bool bSuccess, FHttpResponsePtr Response)
+	void FDownloadService::SetHttpRequestComplete(int32 RequestId, bool bSuccess, FHttpResponsePtr Response, IDownloadServiceStat::FDownloadRecord&& DownloadRecord)
 	{
 		FScopeLock ScopeLock(&CompletedRequestsCS);
-		CompletedRequests.Emplace(RequestId, MakeShareable(new FHttpDownload(MoveTemp(Response), bSuccess)));
+		CompletedRequests.Emplace(RequestId, MakeShareable(new FHttpDownload(MoveTemp(Response), bSuccess, MoveTemp(DownloadRecord))));
 	}
 
 	IDownloadService* FDownloadServiceFactory::Create(FTicker& Ticker, IHttpManager* HttpManager, IFileSystem* FileSystem, IDownloadServiceStat* DownloadServiceStat, IInstallerAnalytics* InstallerAnalytics)
diff --git a/Engine/Source/Runtime/Online/BuildPatchServices/Private/Installer/DownloadService.h b/Engine/Source/Runtime/Online/BuildPatchServices/Private/Installer/DownloadService.h
--- a/Engine/Source/Runtime/Online/BuildPatchServices/Private/Installer/DownloadService.h
+++ b/Engine/Source/Runtime/Online/BuildPatchServices/Private/Installer/DownloadService.h
@@ -3,6 +3,7 @@
 
 #include "CoreMinimal.h"
 #include "Containers/Ticker.h"
+#include "Common/SpeedRecorder.h"
 
 namespace BuildPatchServices
 {
@@ -117,18 +118,28 @@ namespace BuildPatchServices
 			bool bSuccess;
 			// The response code for the request.
 			int32 ResponseCode;
-			// The time in seconds when the request was started.
-			double StartedAt;
-			// The time in seconds when the request was completed.
-			double CompletedAt;
-			// The number of bytes received.
-			int32 BytesReceived;
+			// The timing record.
+			ISpeedRecorder::FRecord SpeedRecord;
 		};
 
+		/**
+		 * Called as each request has started.
+		 * @param RequestId             The id for the request.
+		 * @param Uri                   The uri for the request.
+		 */
+		virtual void OnDownloadStarted(int32 RequestId, const FString& Uri) = 0;
+
+		/**
+		 * Called for each request completion.
+		 * @param RequestId             The id for the request.
+		 * @param BytesReceived         The bytes received so far.
+		 */
+		virtual void OnDownloadProgress(int32 RequestId, int32 BytesReceived) = 0;
+
 		/**
 		 * Called for each request completion.
 		 * @param DownloadRecord        The struct containing the stats for this request.
 		 */
-		virtual void OnDownloadComplete(FDownloadRecord DownloadRecord) = 0;
+		virtual void OnDownloadComplete(const FDownloadRecord& DownloadRecord) = 0;
 	};
 }
\ No newline at end of file
diff --git a/Engine/Source/Runtime/Online/BuildPatchServices/Private/Installer/FileAttribution.cpp b/Engine/Source/Runtime/Online/BuildPatchServices/Private/Installer/FileAttribution.cpp
--- a/Engine/Source/Runtime/Online/BuildPatchServices/Private/Installer/FileAttribution.cpp
+++ b/Engine/Source/Runtime/Online/BuildPatchServices/Private/Installer/FileAttribution.cpp
@@ -100,7 +100,7 @@ namespace BuildPatchServices
 		{
 			InstallFilename = StagedFileDirectory / BuildFile;
 			int64 FileSize;
-			if (FileSystem->GetFileSize(*InstallFilename, FileSize) && FileSize != INDEX_NONE)
+			if (FileSystem->GetFileSize(*InstallFilename, FileSize))
 			{
 				return InstallFilename;
 			}
@@ -123,14 +123,14 @@ namespace BuildPatchServices
 
 	void FFileAttribution::SetupFileAttributes(const FString& FilePath, const FFileManifest& FileManifest, bool bForce)
 	{
-		EFileAttributes FileAttributes = EFileAttributes::None;
+		EAttributeFlags FileAttributes = EAttributeFlags::None;
 
 		// First check file attributes as it's much faster to read and do nothing
-		bool bKnownAttributes = FileSystem->GetFileAttributes(*FilePath, FileAttributes);
-		bool bIsReadOnly = (FileAttributes & EFileAttributes::ReadOnly) != EFileAttributes::None;
-		const bool bFileExists = (FileAttributes & EFileAttributes::Exists) != EFileAttributes::None;
-		const bool bIsCompressed = (FileAttributes & EFileAttributes::Compressed) != EFileAttributes::None;
-		const bool bIsUnixExecutable = (FileAttributes & EFileAttributes::Executable) != EFileAttributes::None;
+		bool bKnownAttributes = FileSystem->GetAttributes(*FilePath, FileAttributes);
+		bool bIsReadOnly = (FileAttributes & EAttributeFlags::ReadOnly) != EAttributeFlags::None;
+		const bool bFileExists = (FileAttributes & EAttributeFlags::Exists) != EAttributeFlags::None;
+		const bool bIsCompressed = (FileAttributes & EAttributeFlags::Compressed) != EAttributeFlags::None;
+		const bool bIsUnixExecutable = (FileAttributes & EAttributeFlags::Executable) != EAttributeFlags::None;
 
 		// If we know the file is missing, skip out
 		if (bKnownAttributes && !bFileExists)
diff --git a/Engine/Source/Runtime/Online/BuildPatchServices/Private/Installer/InstallChunkSource.cpp b/Engine/Source/Runtime/Online/BuildPatchServices/Private/Installer/InstallChunkSource.cpp
--- a/Engine/Source/Runtime/Online/BuildPatchServices/Private/Installer/InstallChunkSource.cpp
+++ b/Engine/Source/Runtime/Online/BuildPatchServices/Private/Installer/InstallChunkSource.cpp
@@ -1,13 +1,17 @@
 // Copyright 1998-2018 Epic Games, Inc. All Rights Reserved.
 
 #include "Installer/InstallChunkSource.h"
+#include "HAL/ThreadSafeBool.h"
+#include "Serialization/MemoryReader.h"
+#include "Containers/Queue.h"
 #include "Installer/ChunkReferenceTracker.h"
 #include "Installer/ChunkStore.h"
 #include "Installer/InstallerError.h"
 #include "Common/FileSystem.h"
+#include "Common/StatsCollector.h"
 #include "BuildPatchHash.h"
-#include "HAL/ThreadSafeBool.h"
-#include "Serialization/MemoryReader.h"
+#include "Algo/Transform.h"
+#include "Misc/Paths.h"
 
 DECLARE_LOG_CATEGORY_EXTERN(LogInstallChunkSource, Warning, All);
 DEFINE_LOG_CATEGORY(LogInstallChunkSource);
@@ -28,11 +32,13 @@ namespace BuildPatchServices
 		// IChunkSource interface begin.
 		virtual IChunkDataAccess* Get(const FGuid& DataId) override;
 		virtual TSet<FGuid> AddRuntimeRequirements(TSet<FGuid> NewRequirements) override;
+		virtual bool AddRepeatRequirement(const FGuid& RepeatRequirement) override;
 		virtual void SetUnavailableChunksCallback(TFunction<void(TSet<FGuid>)> Callback) override;
 		// IChunkSource interface end.
 
 		// IInstallChunkSource interface begin.
 		virtual const TSet<FGuid>& GetAvailableChunks() const override;
+		virtual void HarvestRemainingChunksFromFile(const FString& BuildFile) override;
 		// IInstallChunkSource interface end.
 
 	private:
@@ -56,12 +62,16 @@ namespace BuildPatchServices
 		// Handling of chunks we lose access to.
 		TFunction<void(TSet<FGuid>)> UnavailableChunksCallback;
 		TSet<FGuid> UnavailableChunks;
+		// Keep a separate copy of failed chunks internally.
+		TSet<FGuid> FailedChunks;
 		// Storage of enumerated chunks.
 		TSet<FGuid> AvailableInBuilds;
 		TArray<TPair<FString, FBuildPatchAppManifestRef>> InstallationSources;
 		TSet<FGuid> PlacedInStore;
 		// Track additional requests.
 		TSet<FGuid> RuntimeRequests;
+		// Communication and storage of incoming repeat requirements.
+		TQueue<FGuid, EQueueMode::Mpsc> RepeatRequirementMessages;
 	};
 
 	FInstallChunkSource::FInstallChunkSource(FInstallSourceConfig InConfiguration, IFileSystem* InFileSystem, IChunkStore* InChunkStore, IChunkReferenceTracker* InChunkReferenceTracker, IInstallerError* InInstallerError, IInstallChunkSourceStat* InInstallChunkSourceStat, const TMap<FString, FBuildPatchAppManifestRef>& InInstallationSources, const FBuildPatchAppManifestRef& InstallManifest)
@@ -111,6 +121,12 @@ namespace BuildPatchServices
 			// If available, load next batch into store.
 			if (AvailableInBuilds.Contains(DataId))
 			{
+				// 'Forget' any repeat requirements.
+				FGuid RepeatRequirement;
+				while (RepeatRequirementMessages.Dequeue(RepeatRequirement))
+				{
+					PlacedInStore.Remove(RepeatRequirement);
+				}
 				// Select the next X chunks that are locally available.
 				TFunction<bool(const FGuid&)> SelectPredicate = [this](const FGuid& ChunkId)
 				{
@@ -123,11 +139,13 @@ namespace BuildPatchServices
 				// Remove already loaded chunks.
 				TFunction<bool(const FGuid&)> RemovePredicate = [this](const FGuid& ChunkId)
 				{
-					return PlacedInStore.Contains(ChunkId);
+					return PlacedInStore.Contains(ChunkId) || FailedChunks.Contains(ChunkId);
 				};
 				BatchLoadChunks.RemoveAll(RemovePredicate);
 				// Ensure requested chunk is in the array.
 				BatchLoadChunks.AddUnique(DataId);
+				// Call to stat.
+				InstallChunkSourceStat->OnBatchStarted(BatchLoadChunks);
 				// Load this batch
 				for (int32 ChunkIdx = 0; ChunkIdx < BatchLoadChunks.Num() && !bShouldAbort; ++ChunkIdx)
 				{
@@ -148,8 +166,21 @@ namespace BuildPatchServices
 
 	TSet<FGuid> FInstallChunkSource::AddRuntimeRequirements(TSet<FGuid> NewRequirements)
 	{
-		RuntimeRequests.Append(NewRequirements.Intersect(AvailableInBuilds));
-		return NewRequirements.Difference(AvailableInBuilds);
+		TSet<FGuid> Unhandled = NewRequirements.Difference(AvailableInBuilds);
+		TSet<FGuid> Accepted = NewRequirements.Intersect(AvailableInBuilds);
+		RuntimeRequests.Append(Accepted);
+		InstallChunkSourceStat->OnAcceptedNewRequirements(Accepted);
+		return Unhandled;
+	}
+
+	bool FInstallChunkSource::AddRepeatRequirement(const FGuid& RepeatRequirement)
+	{
+		if (AvailableInBuilds.Contains(RepeatRequirement))
+		{
+			RepeatRequirementMessages.Enqueue(RepeatRequirement);
+			return true;
+		}
+		return false;
 	}
 
 	void FInstallChunkSource::SetUnavailableChunksCallback(TFunction<void(TSet<FGuid>)> Callback)
@@ -162,6 +193,44 @@ namespace BuildPatchServices
 		return AvailableInBuilds;
 	}
 
+	void FInstallChunkSource::HarvestRemainingChunksFromFile(const FString& FilePath)
+	{
+		const FFileManifest* FileManifest = nullptr;
+		for (const TPair<FString, FBuildPatchAppManifestRef>& Pair : InstallationSources)
+		{
+			if (FilePath.StartsWith(Pair.Key))
+			{
+				FString BuildRelativeFilePath = FilePath;
+				FPaths::MakePathRelativeTo(BuildRelativeFilePath, *(Pair.Key / TEXT("")));
+				FileManifest = Pair.Value->GetFileManifest(BuildRelativeFilePath);
+				break;
+			}
+		}
+		if (FileManifest != nullptr)
+		{
+			// Collect all chunks in this file.
+			TSet<FGuid> FileManifestChunks;
+			Algo::Transform(FileManifest->FileChunkParts, FileManifestChunks, &FChunkPart::Guid);
+			// Select all chunks still required from this file.
+			TFunction<bool(const FGuid&)> SelectPredicate = [this, &FileManifestChunks](const FGuid& ChunkId)
+			{
+				return !PlacedInStore.Contains(ChunkId) && (FileManifestChunks.Contains(ChunkId) && (!Configuration.ChunkIgnoreSet.Contains(ChunkId) || RuntimeRequests.Contains(ChunkId)));
+			};
+			TArray<FGuid> BatchLoadChunks = ChunkReferenceTracker->GetNextReferences(TNumericLimits<int32>::Max(), SelectPredicate);
+			if (BatchLoadChunks.Num() > 0)
+			{
+				// Call to stat.
+				InstallChunkSourceStat->OnBatchStarted(BatchLoadChunks);
+				// Load the batch.
+				for (int32 ChunkIdx = 0; ChunkIdx < BatchLoadChunks.Num() && !bShouldAbort; ++ChunkIdx)
+				{
+					const FGuid& BatchLoadChunk = BatchLoadChunks[ChunkIdx];
+					LoadFromBuild(BatchLoadChunk);
+				}
+			}
+		}
+	}
+
 	void FInstallChunkSource::FindChunkLocation(const FGuid& DataId, const FString** FoundInstallDirectory, const FBuildPatchAppManifest** FoundInstallManifest) const
 	{
 		uint64 ChunkHash;
@@ -191,6 +260,8 @@ namespace BuildPatchServices
 		}
 
 		// Attempt construction of the chunk from the parts.
+		ISpeedRecorder::FRecord LoadRecord;
+		LoadRecord.Size = 0;
 		InstallChunkSourceStat->OnLoadStarted(DataId);
 		TUniquePtr<FArchive> FileArchive;
 		FString FileOpened;
@@ -208,6 +279,7 @@ namespace BuildPatchServices
 		{
 			HashType |= EChunkHashFlags::RollingPoly64;
 		}
+		LoadRecord.CyclesStart = FStatsCollector::GetCycles();
 		IInstallChunkSourceStat::ELoadResult LoadResult = HashType == EChunkHashFlags::None ? IInstallChunkSourceStat::ELoadResult::MissingHashInfo : IInstallChunkSourceStat::ELoadResult::Success;
 		if (LoadResult == IInstallChunkSourceStat::ELoadResult::Success)
 		{
@@ -253,6 +325,7 @@ namespace BuildPatchServices
 						{
 							FileArchive->Seek(FileChunkPart.FileOffset);
 							FileArchive->Serialize(TempChunkConstruction + FileChunkPart.ChunkPart.Offset, FileChunkPart.ChunkPart.Size);
+							LoadRecord.Size += FileChunkPart.ChunkPart.Size;
 						}
 					}
 					// Wait while paused
@@ -282,9 +355,6 @@ namespace BuildPatchServices
 				// Save the chunk to the store if all went well.
 				if (LoadResult == IInstallChunkSourceStat::ELoadResult::Success)
 				{
-					// It was added asynchronously!!
-					check(PlacedInStore.Contains(DataId) == false);
-
 					// Create the ChunkFile data structure
 					IChunkDataAccess* NewChunkFile = FChunkDataAccessFactory::Create(BuildPatchServices::ChunkDataSize);
 
@@ -321,7 +391,8 @@ namespace BuildPatchServices
 				}
 			}
 		}
-		InstallChunkSourceStat->OnLoadComplete(DataId, LoadResult);
+		LoadRecord.CyclesEnd = FStatsCollector::GetCycles();
+		InstallChunkSourceStat->OnLoadComplete(DataId, LoadResult, LoadRecord);
 		if (LoadResult == IInstallChunkSourceStat::ELoadResult::Success)
 		{
 			return true;
@@ -329,6 +400,7 @@ namespace BuildPatchServices
 		else
 		{
 			UnavailableChunks.Add(DataId);
+			FailedChunks.Add(DataId);
 			return false;
 		}
 	}
@@ -349,4 +421,27 @@ namespace BuildPatchServices
 		check(InstallChunkSourceStat != nullptr);
 		return new FInstallChunkSource(MoveTemp(Configuration), FileSystem, ChunkStore, ChunkReferenceTracker, InstallerError, InstallChunkSourceStat, InstallationSources, InstallManifest);
 	}
+
+	const TCHAR* ToString(const IInstallChunkSourceStat::ELoadResult& LoadResult)
+	{
+		switch(LoadResult)
+		{
+			case IInstallChunkSourceStat::ELoadResult::Success:
+				return TEXT("Success");
+			case IInstallChunkSourceStat::ELoadResult::MissingHashInfo:
+				return TEXT("MissingHashInfo");
+			case IInstallChunkSourceStat::ELoadResult::MissingPartInfo:
+				return TEXT("MissingPartInfo");
+			case IInstallChunkSourceStat::ELoadResult::OpenFileFail:
+				return TEXT("OpenFileFail");
+			case IInstallChunkSourceStat::ELoadResult::IncorrectFileSize:
+				return TEXT("IncorrectFileSize");
+			case IInstallChunkSourceStat::ELoadResult::HashCheckFailed:
+				return TEXT("HashCheckFailed");
+			case IInstallChunkSourceStat::ELoadResult::Aborted:
+				return TEXT("Aborted");
+			default:
+				return TEXT("Unknown");
+		}
+	}
 }
diff --git a/Engine/Source/Runtime/Online/BuildPatchServices/Private/Installer/InstallChunkSource.h b/Engine/Source/Runtime/Online/BuildPatchServices/Private/Installer/InstallChunkSource.h
--- a/Engine/Source/Runtime/Online/BuildPatchServices/Private/Installer/InstallChunkSource.h
+++ b/Engine/Source/Runtime/Online/BuildPatchServices/Private/Installer/InstallChunkSource.h
@@ -3,6 +3,7 @@
 
 #include "Installer/ChunkSource.h"
 #include "Installer/Controllable.h"
+#include "Common/SpeedRecorder.h"
 #include "BuildPatchManifest.h"
 
 namespace BuildPatchServices
@@ -28,6 +29,13 @@ namespace BuildPatchServices
 		 * @return the set of chunks that should be available locally.
 		 */
 		virtual const TSet<FGuid>& GetAvailableChunks() const = 0;
+
+		/**
+		 * Given an installation file path, harvest all remaining required chunks from the file immediately.
+		 * NB: The filepath must match a file contained in one of the installation sources provided in order to load chunks.
+		 * @param FilePath  The full filepath to the file to harvest.
+		 */
+		virtual void HarvestRemainingChunksFromFile(const FString& FilePath) = 0;
 	};
 
 	/**
@@ -112,6 +120,12 @@ namespace BuildPatchServices
 	public:
 		virtual ~IInstallChunkSourceStat() {}
 
+		/**
+		 * Called when a batch of chunks are going to be loaded.
+		 * @param ChunkIds  The ids of each chunk.
+		 */
+		virtual void OnBatchStarted(const TArray<FGuid>& ChunkIds) = 0;
+
 		/**
 		 * Called each time a chunk load begins.
 		 * @param ChunkId   The id of the chunk.
@@ -122,7 +136,19 @@ namespace BuildPatchServices
 		 * Called each time a chunk load completes.
 		 * @param ChunkId   The id of the chunk.
 		 * @param Result    The load result.
+		 * @param Record    The recording for the received data.
 		 */
-		virtual void OnLoadComplete(const FGuid& ChunkId, ELoadResult Result) = 0;
+		virtual void OnLoadComplete(const FGuid& ChunkId, const ELoadResult& Result, const ISpeedRecorder::FRecord& Record) = 0;
+
+		/**
+		 * Called when a batch of chunks are added and accepted via IChunkSource::AddRuntimeRequirements.
+		 * @param ChunkIds  The ids of each chunk.
+		 */
+		virtual void OnAcceptedNewRequirements(const TSet<FGuid>& ChunkIds) = 0;
 	};
+	
+	/**
+	 * A ToString implementation for IInstallChunkSourceStat::ELoadResult.
+	 */
+	const TCHAR* ToString(const IInstallChunkSourceStat::ELoadResult& LoadResult);
 }
\ No newline at end of file
diff --git a/Engine/Source/Runtime/Online/BuildPatchServices/Private/Installer/InstallerAnalytics.cpp b/Engine/Source/Runtime/Online/BuildPatchServices/Private/Installer/InstallerAnalytics.cpp
--- a/Engine/Source/Runtime/Online/BuildPatchServices/Private/Installer/InstallerAnalytics.cpp
+++ b/Engine/Source/Runtime/Online/BuildPatchServices/Private/Installer/InstallerAnalytics.cpp
@@ -47,6 +47,7 @@ namespace BuildPatchServices
 		virtual void RecordConstructionError(const FString& Filename, int32 LastError, const FString& ErrorString) override;
 		virtual void RecordPrereqInstallationError(const FString& AppName, const FString& AppVersion, const FString& Filename, const FString& CommandLine, int32 ErrorCode, const FString& ErrorString) override;
 		virtual void TrackRequest(const FHttpRequestPtr& Request) override;
+		virtual void Flush() override;
 		// IInstallerAnalytics interface end.
 
 	private:
@@ -152,14 +153,10 @@ namespace BuildPatchServices
 		}
 	}
 
-	void FInstallerAnalytics::QueueAnalyticsEvent(FString EventName, TArray<FAnalyticsEventAttribute> Attributes)
+	void FInstallerAnalytics::Flush()
 	{
-		FScopeLock ScopeLock(&AnalyticsEventQueueCS);
-		AnalyticsEventQueue.Emplace(MoveTemp(EventName), MoveTemp(Attributes));
-	}
+		check(IsInGameThread());
 
-	bool FInstallerAnalytics::Tick(float Delta)
-	{
         QUICK_SCOPE_CYCLE_COUNTER(STAT_FInstallerAnalytics_Tick);
 
 		if (Analytics != nullptr)
@@ -172,6 +169,17 @@ namespace BuildPatchServices
 			}
 			AnalyticsEventQueue.Reset();
 		}
+	}
+
+	void FInstallerAnalytics::QueueAnalyticsEvent(FString EventName, TArray<FAnalyticsEventAttribute> Attributes)
+	{
+		FScopeLock ScopeLock(&AnalyticsEventQueueCS);
+		AnalyticsEventQueue.Emplace(MoveTemp(EventName), MoveTemp(Attributes));
+	}
+
+	bool FInstallerAnalytics::Tick(float Delta)
+	{
+		Flush();
 		return true;
 	}
 
diff --git a/Engine/Source/Runtime/Online/BuildPatchServices/Private/Installer/InstallerAnalytics.h b/Engine/Source/Runtime/Online/BuildPatchServices/Private/Installer/InstallerAnalytics.h
--- a/Engine/Source/Runtime/Online/BuildPatchServices/Private/Installer/InstallerAnalytics.h
+++ b/Engine/Source/Runtime/Online/BuildPatchServices/Private/Installer/InstallerAnalytics.h
@@ -22,7 +22,7 @@ namespace BuildPatchServices
 		/**
 		 * @EventName Patcher.Error.Download
 		 *
-		 * @Trigger Fires for the first 20 chunk download errors that occur. Deprecated: Not recommenced for use.
+		 * @Trigger Fires for the first 20 chunk download errors that occur. Deprecated: Not recommended for use.
 		 *
 		 * @Type Static
 		 *
@@ -37,7 +37,7 @@ namespace BuildPatchServices
 		/**
 		 * @EventName Patcher.Warning.ChunkAborted
 		 *
-		 * @Trigger Fires for chunks that were aborted as deemed failing. This is a prediction that the chunk is suffering from 0-byte TCP window issue. Deprecated: Not recommenced for use.
+		 * @Trigger Fires for chunks that were aborted as deemed failing. This is a prediction that the chunk is suffering from 0-byte TCP window issue. Deprecated: Not recommended for use.
 		 *
 		 * @Type Static
 		 *
@@ -54,7 +54,7 @@ namespace BuildPatchServices
 		/**
 		 * @EventName Patcher.Error.Cache
 		 *
-		 * @Trigger Fires for the first 20 chunks which experienced a cache error. Deprecated: Not recommenced for use.
+		 * @Trigger Fires for the first 20 chunks which experienced a cache error. Deprecated: Not recommended for use.
 		 *
 		 * @Type Static
 		 *
@@ -71,7 +71,7 @@ namespace BuildPatchServices
 		/**
 		 * @EventName Patcher.Error.Construction
 		 *
-		 * @Trigger Fires for the first 20 errors experienced while constructing an installation file. Deprecated: Not recommenced for use.
+		 * @Trigger Fires for the first 20 errors experienced while constructing an installation file. Deprecated: Not recommended for use.
 		 *
 		 * @Type Static
 		 *
@@ -86,7 +86,7 @@ namespace BuildPatchServices
 		/**
 		 * @EventName Patcher.Error.Prerequisites
 		 *
-		 * @Trigger Fires for errors experienced while installing prerequisites. Deprecated: Not recommenced for use.
+		 * @Trigger Fires for errors experienced while installing prerequisites. Deprecated: Not recommended for use.
 		 *
 		 * @Type Static
 		 *
@@ -124,6 +124,11 @@ namespace BuildPatchServices
 		 * @Owner Portal Team
 		 */
 		virtual void TrackRequest(const FHttpRequestPtr& Request) = 0;
+
+		/**
+		 * Ensures that any queued analytics are sent immediately.
+		 */
+		virtual void Flush() = 0;
 	};
 
 	/**
diff --git a/Engine/Source/Runtime/Online/BuildPatchServices/Private/Installer/InstallerError.cpp b/Engine/Source/Runtime/Online/BuildPatchServices/Private/Installer/InstallerError.cpp
--- a/Engine/Source/Runtime/Online/BuildPatchServices/Private/Installer/InstallerError.cpp
+++ b/Engine/Source/Runtime/Online/BuildPatchServices/Private/Installer/InstallerError.cpp
@@ -2,9 +2,9 @@
 
 #include "Installer/InstallerError.h"
 #include "Misc/ScopeLock.h"
+#include "HAL/ThreadSafeCounter.h"
 #include "Interfaces/IBuildInstaller.h"
 #include "BuildPatchServicesPrivate.h"
-#include "HAL/ThreadSafeCounter.h"
 
 namespace BuildPatchServices
 {
@@ -56,9 +56,9 @@ namespace BuildPatchServices
 		FormatOptions = FormatOptions == nullptr ? &DefaultOptions : FormatOptions;
 		FFormatNamedArguments Arguments;
 		Arguments.Emplace(TEXT("Location"), FText::FromString(Location));
-		Arguments.Emplace(TEXT("RequiredBytes"), FText::AsMemory(RequiredBytes, FormatOptions));
-		Arguments.Emplace(TEXT("AvailableBytes"), FText::AsMemory(AvailableBytes, FormatOptions));
-		Arguments.Emplace(TEXT("SpaceAdditional"), FText::AsMemory(RequiredBytes - AvailableBytes, FormatOptions));
+		Arguments.Emplace(TEXT("RequiredBytes"), FText::AsMemory(RequiredBytes, FormatOptions, nullptr, EMemoryUnitStandard::SI));
+		Arguments.Emplace(TEXT("AvailableBytes"), FText::AsMemory(AvailableBytes, FormatOptions, nullptr, EMemoryUnitStandard::SI));
+		Arguments.Emplace(TEXT("SpaceAdditional"), FText::AsMemory(RequiredBytes - AvailableBytes, FormatOptions, nullptr, EMemoryUnitStandard::SI));
 		return FText::Format(OutOfDiskSpace, Arguments);
 	}
 
@@ -76,7 +76,7 @@ namespace BuildPatchServices
 		virtual EBuildPatchInstallError GetErrorType() const override;
 		virtual FString GetErrorCode() const override;
 		virtual FText GetErrorText() const override;
-		virtual void SetError(EBuildPatchInstallError InErrorType, const TCHAR* InErrorCode, FText InErrorText) override;
+		virtual void SetError(EBuildPatchInstallError InErrorType, const TCHAR* InErrorSubType, uint32 InErrorCode, FText InErrorText) override;
 		virtual int32 RegisterForErrors(FOnErrorDelegate Delegate) override;
 		virtual void UnregisterForErrors(int32 Handle) override;
 		// IInstallerError interface end.
@@ -145,7 +145,7 @@ namespace BuildPatchServices
 		return ErrorText;
 	}
 
-	void FInstallerError::SetError(EBuildPatchInstallError InErrorType, const TCHAR* InErrorCode, FText InErrorText)
+	void FInstallerError::SetError(EBuildPatchInstallError InErrorType, const TCHAR* InErrorSubType, uint32 InErrorCode, FText InErrorText)
 	{
 		// Only accept the first error
 		TArray<FOnErrorDelegate> DelegatesToCall;
@@ -154,7 +154,8 @@ namespace BuildPatchServices
 		{
 			ErrorType = InErrorType;
 			ErrorCode = InstallErrorPrefixes::ErrorTypeStrings[static_cast<int32>(InErrorType)];
-			ErrorCode += InErrorCode;
+			ErrorCode += InErrorSubType;
+			ErrorCode += (InErrorCode > 0) ? FString::Printf(TEXT("-%u"), InErrorCode): TEXT("");
 			ErrorText = InErrorText.IsEmpty() ? GetStandardErrorText(ErrorType) : MoveTemp(InErrorText);
 			RegisteredDelegates.GenerateValueArray(DelegatesToCall);
 			if (IsCancelled())
diff --git a/Engine/Source/Runtime/Online/BuildPatchServices/Private/Installer/InstallerError.h b/Engine/Source/Runtime/Online/BuildPatchServices/Private/Installer/InstallerError.h
--- a/Engine/Source/Runtime/Online/BuildPatchServices/Private/Installer/InstallerError.h
+++ b/Engine/Source/Runtime/Online/BuildPatchServices/Private/Installer/InstallerError.h
@@ -17,6 +17,7 @@ namespace BuildPatchServices
 		static const TCHAR* MissingCompleteDelegate = TEXT("03");
 		static const TCHAR* InvalidInstallTags = TEXT("04");
 		static const TCHAR* ChunkReferenceTracking = TEXT("05");
+		static const TCHAR* MissingPrereqForPrereqOnlyInstall = TEXT("06");
 	}
 
 	/**
@@ -34,6 +35,7 @@ namespace BuildPatchServices
 	namespace PathLengthErrorCodes
 	{
 		static const TCHAR* StagingDirectory = TEXT("01");
+		static const TCHAR* InstallDirectory = TEXT("02");
 	}
 
 	/**
@@ -63,6 +65,7 @@ namespace BuildPatchServices
 	namespace MoveErrorCodes
 	{
 		static const TCHAR* StageToInstall = TEXT("01");
+		static const TCHAR* DeleteFileFailed = TEXT("02");
 	}
 
 	/**
@@ -70,7 +73,12 @@ namespace BuildPatchServices
 	 */
 	namespace VerifyErrorCodes
 	{
+		static const TCHAR* UnknownFail = TEXT("00");
 		static const TCHAR* FinalCheck = TEXT("01");
+		static const TCHAR* FileMissing = TEXT("02");
+		static const TCHAR* OpenFileFailed = TEXT("03");
+		static const TCHAR* HashCheckFailed = TEXT("04");
+		static const TCHAR* FileSizeFailed = TEXT("05");
 	}
 
 	/**
@@ -166,10 +174,11 @@ namespace BuildPatchServices
 		/**
 		 * Report a fatal error that has occurred which should cause other systems to abort.
 		 * @param ErrorType     The enum value for the error.
-		 * @param ErrorCode     The error code string for the error.
+		 * @param ErrorSubType  The error sub type string for the error.
+		 * @param ErrorCode     The error code for the error. 0 means no additional information.
 		 * @param ErrorText     Optional specific error display text to use instead of the standard generic one.
 		 */
-		virtual void SetError(EBuildPatchInstallError ErrorType, const TCHAR* ErrorCode, FText ErrorText = FText::GetEmpty()) = 0;
+		virtual void SetError(EBuildPatchInstallError ErrorType, const TCHAR* ErrorSubType, uint32 ErrorCode = 0, FText ErrorText = FText::GetEmpty()) = 0;
 
 		/**
 		 * Register a delegate to be called upon an error occurring.
@@ -205,4 +214,4 @@ namespace BuildPatchServices
 	 * @return A string value for the EBuildPatchInstallError enum.
 	 */
 	const FString& EnumToString(const EBuildPatchInstallError& ErrorType);
-};
+}
diff --git a/Engine/Source/Runtime/Online/BuildPatchServices/Private/Installer/MachineConfig.h b/Engine/Source/Runtime/Online/BuildPatchServices/Private/Installer/MachineConfig.h
--- a/Engine/Source/Runtime/Online/BuildPatchServices/Private/Installer/MachineConfig.h
+++ b/Engine/Source/Runtime/Online/BuildPatchServices/Private/Installer/MachineConfig.h
@@ -14,14 +14,14 @@ namespace BuildPatchServices
 		virtual ~IMachineConfig() { }
 
 		/**
-		 * Loads the list of installed prereq ids whose corresponding prerequisistes are installed on this computer.
+		 * Loads the list of installed prereq ids whose corresponding prerequisites are installed on this computer.
 		 * @return         A set containing the prerequisite ids which are installed on this computer.
 		 */
 		virtual TSet<FString> LoadInstalledPrereqIds() = 0;
 
 		/**
 		 * Saves updated list of installed prereqs to the configuration file.
-		 * @param InstalledPrereqIds       A set of all installed prerequisistes to save to configuration.
+		 * @param InstalledPrereqIds       A set of all installed prerequisites to save to configuration.
 		 */
 		virtual void SaveInstalledPrereqIds(const TSet<FString>& InstalledPrereqIds) = 0;
 	};
diff --git a/Engine/Source/Runtime/Online/BuildPatchServices/Private/Installer/MemoryChunkStore.cpp b/Engine/Source/Runtime/Online/BuildPatchServices/Private/Installer/MemoryChunkStore.cpp
--- a/Engine/Source/Runtime/Online/BuildPatchServices/Private/Installer/MemoryChunkStore.cpp
+++ b/Engine/Source/Runtime/Online/BuildPatchServices/Private/Installer/MemoryChunkStore.cpp
@@ -1,9 +1,8 @@
 // Copyright 1998-2018 Epic Games, Inc. All Rights Reserved.
 
 #include "Installer/MemoryChunkStore.h"
-#include "Installer/ChunkEvictionPolicy.h"
-#include "Data/ChunkData.h"
 #include "Misc/ScopeLock.h"
+#include "Installer/ChunkEvictionPolicy.h"
 
 namespace BuildPatchServices
 {
@@ -12,13 +11,14 @@ namespace BuildPatchServices
 	{
 	public:
 		FMemoryChunkStore(int32 InStoreSize, IChunkEvictionPolicy* InEvictionPolicy, IChunkStore* InOverflowStore, IMemoryChunkStoreStat* InMemoryChunkStoreStat);
-		~FMemoryChunkStore() {}
+		~FMemoryChunkStore();
 
 		// IChunkStore interface begin.
 		virtual void Put(const FGuid& DataId, TUniquePtr<IChunkDataAccess> ChunkData) override;
 		virtual IChunkDataAccess* Get(const FGuid& DataId) override;
 		virtual TUniquePtr<IChunkDataAccess> Remove(const FGuid& DataId) override;
 		virtual int32 GetSlack() const override;
+		virtual void SetLostChunkCallback(TFunction<void(const FGuid&)> Callback) override;
 		// IChunkStore interface end.
 
 		// IMemoryChunkStore interface begin.
@@ -26,8 +26,9 @@ namespace BuildPatchServices
 		// IMemoryChunkStore interface end.
 
 	private:
-		void PutInternal(const FGuid& DataId, TUniquePtr<IChunkDataAccess> ChunkData);
+		void PutInternal(const FGuid& DataId, TUniquePtr<IChunkDataAccess> ChunkData, bool bIsNewChunk);
 		void UpdateStoreUsage() const;
+		void ExecLostChunkCallback(const FGuid& LostChunk);
 
 	private:
 		int32 StoreSize;
@@ -35,6 +36,8 @@ namespace BuildPatchServices
 		IChunkEvictionPolicy* EvictionPolicy;
 		IChunkStore* OverflowStore;
 		IMemoryChunkStoreStat* MemoryChunkStoreStat;
+		mutable FCriticalSection LostChunkCallbackCs;
+		TFunction<void(const FGuid&)> LostChunkCallback;
 		FGuid LastGetId;
 		TUniquePtr<IChunkDataAccess> LastGetData;
 		mutable FCriticalSection ThreadLockCs;
@@ -46,16 +49,31 @@ namespace BuildPatchServices
 		, EvictionPolicy(InEvictionPolicy)
 		, OverflowStore(InOverflowStore)
 		, MemoryChunkStoreStat(InMemoryChunkStoreStat)
+		, LostChunkCallbackCs()
+		, LostChunkCallback(nullptr)
 		, LastGetId()
 		, LastGetData(nullptr)
 		, ThreadLockCs()
 	{
+		MemoryChunkStoreStat->OnStoreSizeUpdated(StoreSize);
+	}
+
+	FMemoryChunkStore::~FMemoryChunkStore()
+	{
+		for (const TPair<FGuid, TUniquePtr<IChunkDataAccess>>& Entry : Store)
+		{
+			MemoryChunkStoreStat->OnChunkReleased(Entry.Key);
+		}
+		if (LastGetData.IsValid())
+		{
+			MemoryChunkStoreStat->OnChunkReleased(LastGetId);
+		}
+		MemoryChunkStoreStat->OnStoreUseUpdated(0);
 	}
 
 	void FMemoryChunkStore::Put(const FGuid& DataId, TUniquePtr<IChunkDataAccess> ChunkData)
 	{
-		PutInternal(DataId, MoveTemp(ChunkData));
-		MemoryChunkStoreStat->OnChunkStored(DataId);
+		PutInternal(DataId, MoveTemp(ChunkData), true);
 	}
 
 	IChunkDataAccess* FMemoryChunkStore::Get(const FGuid& DataId)
@@ -69,7 +87,7 @@ namespace BuildPatchServices
 			{
 				if (Store.Contains(LastGetId) == false)
 				{
-					PutInternal(LastGetId, MoveTemp(LastGetData));
+					PutInternal(LastGetId, MoveTemp(LastGetData), false);
 				}
 			}
 			// Invalidate last get.
@@ -84,6 +102,11 @@ namespace BuildPatchServices
 			else if (OverflowStore != nullptr)
 			{
 				LastGetData = OverflowStore->Remove(DataId);
+				if (LastGetData.IsValid())
+				{
+					MemoryChunkStoreStat->OnChunkStored(DataId);
+					UpdateStoreUsage();
+				}
 			}
 			// Save ID if successful.
 			if (LastGetData.IsValid())
@@ -123,6 +146,13 @@ namespace BuildPatchServices
 		return StoreSize - (Store.Num() - Cleanable.Num());
 	}
 
+	void FMemoryChunkStore::SetLostChunkCallback(TFunction<void(const FGuid&)> Callback)
+	{
+		// Thread lock to protect access to LostChunkCallback.
+		FScopeLock ThreadLock(&LostChunkCallbackCs);
+		LostChunkCallback = MoveTemp(Callback);
+	}
+
 	void FMemoryChunkStore::DumpToOverflow()
 	{
 		// Thread lock to protect access to Store.
@@ -144,12 +174,17 @@ namespace BuildPatchServices
 		UpdateStoreUsage();
 	}
 
-	void FMemoryChunkStore::PutInternal(const FGuid& DataId, TUniquePtr<IChunkDataAccess> ChunkData)
+	void FMemoryChunkStore::PutInternal(const FGuid& DataId, TUniquePtr<IChunkDataAccess> ChunkData, bool bIsNewChunk)
 	{
 		// Thread lock to protect access to Store, LastGetId, and LastGetData.
 		FScopeLock ThreadLock(&ThreadLockCs);
 		// Add this new chunk.
 		Store.Add(DataId, MoveTemp(ChunkData));
+		if (bIsNewChunk)
+		{
+			MemoryChunkStoreStat->OnChunkStored(DataId);
+			UpdateStoreUsage();
+		}
 		// Clean out our store.
 		TSet<FGuid> Cleanable;
 		TSet<FGuid> Bootable;
@@ -167,16 +202,30 @@ namespace BuildPatchServices
 			{
 				OverflowStore->Put(BootId, MoveTemp(Store[BootId]));
 			}
+			else
+			{
+				ExecLostChunkCallback(BootId);
+			}
 			Store.Remove(BootId);
 			MemoryChunkStoreStat->OnChunkBooted(BootId);
 		}
-		UpdateStoreUsage();
 	}
 
 	void FMemoryChunkStore::UpdateStoreUsage() const
 	{
 		const int32 LastGetCount = LastGetId.IsValid() && !Store.Contains(LastGetId);
-		MemoryChunkStoreStat->OnStoreUseUpdated(Store.Num() + LastGetCount);
+		const int32 StoreNum = Store.Num();
+		MemoryChunkStoreStat->OnStoreUseUpdated(StoreNum + LastGetCount);
+	}
+
+	void FMemoryChunkStore::ExecLostChunkCallback(const FGuid& LostChunk)
+	{
+		// Thread lock to protect access to LostChunkCallback.
+		FScopeLock ThreadLock(&LostChunkCallbackCs);
+		if (LostChunkCallback)
+		{
+			LostChunkCallback(LostChunk);
+		}
 	}
 
 	IMemoryChunkStore* FMemoryChunkStoreFactory::Create(int32 StoreSize, IChunkEvictionPolicy* EvictionPolicy, IChunkStore* OverflowStore, IMemoryChunkStoreStat* MemoryChunkStoreStat)
diff --git a/Engine/Source/Runtime/Online/BuildPatchServices/Private/Installer/MemoryChunkStore.h b/Engine/Source/Runtime/Online/BuildPatchServices/Private/Installer/MemoryChunkStore.h
--- a/Engine/Source/Runtime/Online/BuildPatchServices/Private/Installer/MemoryChunkStore.h
+++ b/Engine/Source/Runtime/Online/BuildPatchServices/Private/Installer/MemoryChunkStore.h
@@ -78,5 +78,11 @@ namespace BuildPatchServices
 		 * @param ChunkCount        The number of chunks now contained in the store.
 		 */
 		virtual void OnStoreUseUpdated(int32 ChunkCount) = 0;
+
+		/**
+		 * Called whenever the total size allotted for the store is updated.
+		 * @param Size              The total number of chunks that can be stored.
+		 */
+		virtual void OnStoreSizeUpdated(int32 Size) = 0;
 	};
 }
\ No newline at end of file
diff --git a/Engine/Source/Runtime/Online/BuildPatchServices/Private/Installer/MessagePump.cpp b/Engine/Source/Runtime/Online/BuildPatchServices/Private/Installer/MessagePump.cpp
--- a/Engine/Source/Runtime/Online/BuildPatchServices/Private/Installer/MessagePump.cpp
+++ b/Engine/Source/Runtime/Online/BuildPatchServices/Private/Installer/MessagePump.cpp
@@ -7,7 +7,7 @@
 namespace BuildPatchServices
 {
 	// Union of all possible message types.
-	typedef TUnion<FChunkSourceEvent> FMessageUnion;
+	typedef TUnion<FChunkSourceEvent, FInstallationFileAction> FMessageUnion;
 
 	// Queue type for messages.
 	typedef TQueue<FMessageUnion, EQueueMode::Mpsc> FMessageQueue;
@@ -41,6 +41,7 @@ namespace BuildPatchServices
 
 		// IMessagePump interface begin.
 		virtual void SendMessage(FChunkSourceEvent Message) override;
+		virtual void SendMessage(FInstallationFileAction Message) override;
 		virtual void PumpMessages(const TArray<FMessageHandler*>& Handlers) override;
 		// IMessagePump interface end.
 
@@ -61,6 +62,11 @@ namespace BuildPatchServices
 		MessageQueue.Enqueue(FMessageUnion(MoveTemp(Message)));
 	}
 
+	void FMessagePump::SendMessage(FInstallationFileAction Message)
+	{
+		MessageQueue.Enqueue(FMessageUnion(MoveTemp(Message)));
+	}
+
 	void FMessagePump::PumpMessages(const TArray<FMessageHandler*>& Handlers)
 	{
 		FMessageUnion MessageUnion;
@@ -70,6 +76,10 @@ namespace BuildPatchServices
 			{
 				continue;
 			}
+			else if (MessagePumpHelpers::TryPump<FInstallationFileAction>(Handlers, MessageUnion))
+			{
+				continue;
+			}
 		}
 	}
 
diff --git a/Engine/Source/Runtime/Online/BuildPatchServices/Private/Installer/MessagePump.h b/Engine/Source/Runtime/Online/BuildPatchServices/Private/Installer/MessagePump.h
--- a/Engine/Source/Runtime/Online/BuildPatchServices/Private/Installer/MessagePump.h
+++ b/Engine/Source/Runtime/Online/BuildPatchServices/Private/Installer/MessagePump.h
@@ -23,6 +23,12 @@ namespace BuildPatchServices
 		 */
 		virtual void SendMessage(FChunkSourceEvent Message) = 0;
 
+		/**
+		 * Sends an installation file action message.
+		 * @param Message   The message to be sent.
+		 */
+		virtual void SendMessage(FInstallationFileAction Message) = 0;
+
 		/**
 		 * Dequeues received messages, pushing them to the provided handlers.
 		 * @param Handlers      The array of handlers.
diff --git a/Engine/Source/Runtime/Online/BuildPatchServices/Private/Installer/Prerequisites.cpp b/Engine/Source/Runtime/Online/BuildPatchServices/Private/Installer/Prerequisites.cpp
--- a/Engine/Source/Runtime/Online/BuildPatchServices/Private/Installer/Prerequisites.cpp
+++ b/Engine/Source/Runtime/Online/BuildPatchServices/Private/Installer/Prerequisites.cpp
@@ -1,6 +1,7 @@
 // Copyright 1998-2018 Epic Games, Inc. All Rights Reserved.
 
 #include "Installer/Prerequisites.h"
+#include "Misc/Paths.h"
 #include "BuildPatchManifest.h"
 #include "BuildPatchProgress.h"
 #include "BuildPatchSettings.h"
@@ -12,7 +13,6 @@
 #include "Installer/InstallerAnalytics.h"
 #include "Installer/InstallerError.h"
 #include "Installer/MachineConfig.h"
-#include "Misc/Paths.h"
 
 DECLARE_LOG_CATEGORY_EXTERN(LogPrerequisites, Log, All);
 DEFINE_LOG_CATEGORY(LogPrerequisites);
@@ -105,14 +105,14 @@ namespace BuildPatchServices
 
 		bool bUsingInstallRoot = false;
 		bool bUsingStageRoot = false;
-		if (Configuration.bStageOnly)
+		if (Configuration.InstallMode == EInstallMode::StageFiles)
 		{
 			if (PrereqPath.ReplaceInline(RootDirectoryVariable, *StageDirWithSlash) == 0)
 			{
 				PrereqPath = StageDirWithSlash + BuildManifest->GetPrereqPath();
 			}
 			int64 FileSize;
-			if (FileSystem->GetFileSize(*PrereqPath, FileSize) && FileSize > 0)
+			if (FileSystem->GetFileSize(*PrereqPath, FileSize))
 			{
 				bUsingStageRoot = true;
 			}
@@ -124,7 +124,7 @@ namespace BuildPatchServices
 				PrereqPath = InstallDirWithSlash + BuildManifest->GetPrereqPath();
 			}
 			int64 FileSize;
-			if (FileSystem->GetFileSize(*PrereqPath, FileSize) && FileSize > 0)
+			if (FileSystem->GetFileSize(*PrereqPath, FileSize))
 			{
 				bUsingInstallRoot = true;
 			}
@@ -168,6 +168,7 @@ namespace BuildPatchServices
 
 		if (bPrereqInstallSuccessful)
 		{
+			UE_LOG(LogPrerequisites, Log, TEXT("Prerequisite installation successful"));
 			BuildProgress.SetStateProgress(EBuildPatchState::PrerequisitesInstall, 1.0f);
 			InstalledPrereqs.Append(PrereqIds);
 			MachineConfig->SaveInstalledPrereqIds(InstalledPrereqs);
diff --git a/Engine/Source/Runtime/Online/BuildPatchServices/Private/Installer/Statistics/DiskChunkStoreStatistics.cpp b/Engine/Source/Runtime/Online/BuildPatchServices/Private/Installer/Statistics/DiskChunkStoreStatistics.cpp
new file mode 100644
--- /dev/null
+++ b/Engine/Source/Runtime/Online/BuildPatchServices/Private/Installer/Statistics/DiskChunkStoreStatistics.cpp
@@ -0,0 +1,119 @@
+// Copyright 1998-2018 Epic Games, Inc. All Rights Reserved.
+
+#include "Installer/Statistics/DiskChunkStoreStatistics.h"
+#include "CoreMinimal.h"
+#include "Core/AsyncHelpers.h"
+#include "Data/ChunkData.h"
+#include "Installer/InstallerAnalytics.h"
+#include "Installer/Statistics/FileOperationTracker.h"
+
+namespace BuildPatchServices
+{
+	class FDiskChunkStoreStatistics
+		: public IDiskChunkStoreStatistics
+	{
+	public:
+		FDiskChunkStoreStatistics(IInstallerAnalytics* InInstallerAnalytics, IFileOperationTracker* FileOperationTracker);
+		~FDiskChunkStoreStatistics();
+
+	public:
+		// IDiskChunkStoreStat interface begin.
+		virtual void OnChunkStored(const FGuid& ChunkId, const FString& ChunkFilename, EChunkSaveResult SaveResult) override;
+		virtual void OnBeforeChunkLoad(const FGuid& ChunkId) override;
+		virtual void OnChunkLoaded(const FGuid& ChunkId, const FString& ChunkFilename, EChunkLoadResult LoadResult) override;
+		virtual void OnCacheUseUpdated(int32 ChunkCount) override;
+		// IDiskChunkStoreStat interface end.
+
+		// IDiskChunkStoreStatistics interface begin.
+		virtual int32 GetNumSuccessfulLoads() const override;
+		virtual int32 GetNumFailedLoads() const override;
+		virtual int32 GetNumSuccessfulSaves() const override;
+		virtual int32 GetNumFailedSaves() const override;
+		// IDiskChunkStoreStatistics interface end.
+
+	private:
+		IInstallerAnalytics* const InstallerAnalytics;
+		IFileOperationTracker* const FileOperationTracker;
+		FThreadSafeInt32 NumSuccessfulLoads;
+		FThreadSafeInt32 NumSuccessfulSaves;
+		FThreadSafeInt32 NumFailedLoads;
+		FThreadSafeInt32 NumFailedSaves;
+	};
+
+	FDiskChunkStoreStatistics::FDiskChunkStoreStatistics(IInstallerAnalytics* InInstallerAnalytics, IFileOperationTracker* InFileOperationTracker)
+		: InstallerAnalytics(InInstallerAnalytics)
+		, FileOperationTracker(InFileOperationTracker)
+		, NumSuccessfulLoads(0)
+		, NumSuccessfulSaves(0)
+		, NumFailedLoads(0)
+		, NumFailedSaves(0)
+	{
+	}
+
+	FDiskChunkStoreStatistics::~FDiskChunkStoreStatistics()
+	{
+	}
+
+	void FDiskChunkStoreStatistics::OnChunkStored(const FGuid& ChunkId, const FString& ChunkFilename, EChunkSaveResult SaveResult)
+	{
+		if (SaveResult == EChunkSaveResult::Success)
+		{
+			NumSuccessfulSaves.Increment();
+			FileOperationTracker->OnDataStateUpdate(ChunkId, EFileOperationState::PendingLocalDataStore);
+		}
+		else
+		{
+			InstallerAnalytics->RecordChunkCacheError(ChunkId, ChunkFilename, FPlatformMisc::GetLastError(), TEXT("DiskChunkStoreSave"), ToString(SaveResult));
+			NumFailedSaves.Increment();
+		}
+	}
+
+	void FDiskChunkStoreStatistics::OnBeforeChunkLoad(const FGuid& ChunkId)
+	{
+		FileOperationTracker->OnDataStateUpdate(ChunkId, EFileOperationState::RetrievingLocalDataStore);
+	}
+
+	void FDiskChunkStoreStatistics::OnChunkLoaded(const FGuid& ChunkId, const FString& ChunkFilename, EChunkLoadResult LoadResult)
+	{
+		if (LoadResult == EChunkLoadResult::Success)
+		{
+			NumSuccessfulLoads.Increment();
+		}
+		else
+		{
+			InstallerAnalytics->RecordChunkCacheError(ChunkId, ChunkFilename, FPlatformMisc::GetLastError(), TEXT("DiskChunkStoreLoad"), ToString(LoadResult));
+			NumFailedLoads.Increment();
+		}
+	}
+
+	void FDiskChunkStoreStatistics::OnCacheUseUpdated(int32 ChunkCount)
+	{
+	}
+
+	int32 FDiskChunkStoreStatistics::GetNumSuccessfulLoads() const
+	{
+		return NumSuccessfulLoads.GetValue();
+	}
+
+	int32 FDiskChunkStoreStatistics::GetNumFailedLoads() const
+	{
+		return NumFailedLoads.GetValue();
+	}
+
+	int32 FDiskChunkStoreStatistics::GetNumSuccessfulSaves() const
+	{
+		return NumSuccessfulSaves.GetValue();
+	}
+
+	int32 FDiskChunkStoreStatistics::GetNumFailedSaves() const
+	{
+		return NumFailedSaves.GetValue();
+	}
+
+	IDiskChunkStoreStatistics* FDiskChunkStoreStatisticsFactory::Create(IInstallerAnalytics* InstallerAnalytics, IFileOperationTracker* FileOperationTracker)
+	{
+		check(InstallerAnalytics != nullptr);
+		check(FileOperationTracker != nullptr);
+		return new FDiskChunkStoreStatistics(InstallerAnalytics, FileOperationTracker);
+	}
+};
\ No newline at end of file
diff --git a/Engine/Source/Runtime/Online/BuildPatchServices/Private/Installer/Statistics/DiskChunkStoreStatistics.h b/Engine/Source/Runtime/Online/BuildPatchServices/Private/Installer/Statistics/DiskChunkStoreStatistics.h
new file mode 100644
--- /dev/null
+++ b/Engine/Source/Runtime/Online/BuildPatchServices/Private/Installer/Statistics/DiskChunkStoreStatistics.h
@@ -0,0 +1,53 @@
+// Copyright 1998-2018 Epic Games, Inc. All Rights Reserved.
+#pragma once
+
+#include "Installer/DiskChunkStore.h"
+
+namespace BuildPatchServices
+{
+	class IInstallerAnalytics;
+	class IFileOperationTracker;
+
+	/**
+	 * Interface to the statistics class which provides access to tracked values from a disk chunk source stat.
+	 */
+	class IDiskChunkStoreStatistics
+		: public IDiskChunkStoreStat
+	{
+	public:
+		/**
+		 * @return the number of chunks which were loaded from the disk store.
+		 */
+		virtual int32 GetNumSuccessfulLoads() const = 0;
+
+		/**
+		 * @return the number of chunks which failed to load from the disk store.
+		 */
+		virtual int32 GetNumFailedLoads() const = 0;
+
+		/**
+		 * @return the number of chunks which were saved to the disk store.
+		 */
+		virtual int32 GetNumSuccessfulSaves() const = 0;
+
+		/**
+		 * @return the number of chunks which failed to save to the disk store.
+		 */
+		virtual int32 GetNumFailedSaves() const = 0;
+	};
+
+	/**
+	 * A factory for creating an IDiskChunkStoreStatistics instance.
+	 */
+	class FDiskChunkStoreStatisticsFactory
+	{
+	public:
+		/**
+		 * Creates the disk chunk store's dependency interface and exposes additional information.
+		 * @param InstallerAnalytics        The analytics implementation for reporting the installer events.
+		 * @param FileOperationTracker      The file operation tracker which will be used to update data states.
+		 * @return the new IDiskChunkStoreStatistics instance created.
+		 */
+		static IDiskChunkStoreStatistics* Create(IInstallerAnalytics* InstallerAnalytics, IFileOperationTracker* FileOperationTracker);
+	};
+}
\ No newline at end of file
