Copying //UE4/Dev-Main to //UE4/Main (Source: //Portal/Main/Engine @ 4247640)

diff --git a/Engine/Source/Runtime/Online/BuildPatchServices/Private/Installer/Statistics/DownloadServiceStatistics.cpp b/Engine/Source/Runtime/Online/BuildPatchServices/Private/Installer/Statistics/DownloadServiceStatistics.cpp
new file mode 100644
--- /dev/null
+++ b/Engine/Source/Runtime/Online/BuildPatchServices/Private/Installer/Statistics/DownloadServiceStatistics.cpp
@@ -0,0 +1,161 @@
+// Copyright 1998-2018 Epic Games, Inc. All Rights Reserved.
+
+#include "Installer/Statistics/DownloadServiceStatistics.h"
+#include "Misc/ScopeLock.h"
+#include "Misc/Paths.h"
+#include "Core/AsyncHelpers.h"
+#include "Common/StatsCollector.h"
+#include "Installer/InstallerAnalytics.h"
+#include "BuildPatchProgress.h"
+#include "BuildPatchManifest.h"
+#include "BuildPatchUtil.h"
+
+namespace BuildPatchServices
+{
+	class FDownloadServiceStatistics
+		: public IDownloadServiceStatistics
+	{
+	public:
+		FDownloadServiceStatistics(ISpeedRecorder* SpeedRecorder, IInstallerAnalytics* InstallerAnalytics, FBuildPatchAppManifest* Manifest);
+		~FDownloadServiceStatistics();
+
+		// IDownloadServiceStat interface start.
+		virtual void OnDownloadStarted(int32 RequestId, const FString& Uri) override;
+		virtual void OnDownloadProgress(int32 RequestId, int32 BytesReceived) override;
+		virtual void OnDownloadComplete(const FDownloadRecord& DownloadRecord) override;
+		// IDownloadServiceStat interface end.
+
+		// IDownloadServiceStatistics interface start.
+		virtual uint64 GetBytesDownloaded() const override;
+		virtual int32 GetNumSuccessfulChunkDownloads() const override;
+		virtual int32 GetNumFailedChunkDownloads() const override;
+		virtual int32 GetNumCurrentDownloads() const override;
+		virtual TArray<FDownload> GetCurrentDownloads() const override;
+		// IDownloadServiceStatistics interface end.
+
+	private:
+		int64 GetDownloadSize(const FString& Uri) const;
+
+	private:
+		ISpeedRecorder* SpeedRecorder;
+		IInstallerAnalytics* InstallerAnalytics;
+		FBuildPatchAppManifest* Manifest;
+		TMap<FString, int64> DownloadSizes;
+		FThreadSafeInt64 TotalBytesReceived;
+		FThreadSafeInt32 NumSuccessfulDownloads;
+		FThreadSafeInt32 NumFailedDownloads;
+
+		typedef TTuple<FString, int32> FDownloadTuple;
+		TMap<int32, FDownloadTuple> Downloads;
+	};
+
+	FDownloadServiceStatistics::FDownloadServiceStatistics(ISpeedRecorder* InSpeedRecorder, IInstallerAnalytics* InInstallerAnalytics, FBuildPatchAppManifest* InManifest)
+		: SpeedRecorder(InSpeedRecorder)
+		, InstallerAnalytics(InInstallerAnalytics)
+		, Manifest(InManifest)
+		, TotalBytesReceived(0)
+		, NumSuccessfulDownloads(0)
+		, NumFailedDownloads(0)
+	{
+		// Fill out download size lookup.
+		TSet<FGuid> DataList;
+		Manifest->GetDataList(DataList);
+		for (const FGuid& DataId : DataList)
+		{
+			FString CleanFilename = FPaths::GetCleanFilename(FBuildPatchUtils::GetDataFilename(*Manifest, TEXT(""), DataId));
+			DownloadSizes.Add(MoveTemp(CleanFilename), Manifest->GetDataSize(DataId));
+		}
+	}
+
+	FDownloadServiceStatistics::~FDownloadServiceStatistics()
+	{
+	}
+
+	void FDownloadServiceStatistics::OnDownloadStarted(int32 RequestId, const FString& Uri)
+	{
+		checkSlow(IsInGameThread());
+		FDownloadTuple& DownloadTuple = Downloads.FindOrAdd(RequestId);
+		DownloadTuple.Get<0>() = Uri;
+		DownloadTuple.Get<1>() = 0;
+	}
+
+	void FDownloadServiceStatistics::OnDownloadProgress(int32 RequestId, int32 BytesReceived)
+	{
+		checkSlow(IsInGameThread());
+		FDownloadTuple& DownloadTuple = Downloads.FindOrAdd(RequestId);
+		DownloadTuple.Get<1>() = BytesReceived;
+	}
+
+	void FDownloadServiceStatistics::OnDownloadComplete(const FDownloadRecord& DownloadRecord)
+	{
+		checkSlow(IsInGameThread());
+		Downloads.Remove(DownloadRecord.RequestId);
+		TotalBytesReceived.Add(DownloadRecord.SpeedRecord.Size);
+		if (DownloadRecord.bSuccess)
+		{
+			NumSuccessfulDownloads.Increment();
+			SpeedRecorder->AddRecord(DownloadRecord.SpeedRecord);
+		}
+		else
+		{
+			NumFailedDownloads.Increment();
+			InstallerAnalytics->RecordChunkDownloadError(DownloadRecord.Uri, DownloadRecord.ResponseCode, TEXT("DownloadFail"));
+		}
+	}
+
+	uint64 FDownloadServiceStatistics::GetBytesDownloaded() const
+	{
+		return TotalBytesReceived.GetValue();
+	}
+
+	int32 FDownloadServiceStatistics::GetNumSuccessfulChunkDownloads() const
+	{
+		return NumSuccessfulDownloads.GetValue();
+	}
+
+	int32 FDownloadServiceStatistics::GetNumFailedChunkDownloads() const
+	{
+		return NumFailedDownloads.GetValue();
+	}
+
+	int32 FDownloadServiceStatistics::GetNumCurrentDownloads() const
+	{
+		return Downloads.Num();
+	}
+
+	TArray<FDownload> FDownloadServiceStatistics::GetCurrentDownloads() const
+	{
+		checkSlow(IsInGameThread());
+		TArray<FDownload> Result;
+		Result.Empty(Downloads.Num());
+		for (const TPair<int32, FDownloadTuple>& Download : Downloads)
+		{
+			Result.AddDefaulted();
+			FDownload& Element = Result.Last();
+			Element.Data = FPaths::GetCleanFilename(Download.Value.Get<0>());
+			Element.Size = GetDownloadSize(Element.Data);
+			Element.Received = Download.Value.Get<1>();
+		}
+		return Result;
+	}
+
+	int64 FDownloadServiceStatistics::GetDownloadSize(const FString& Filename) const
+	{
+		checkSlow(IsInGameThread());
+		int64 DownloadSize = BuildPatchServices::ChunkDataSize;
+		const int64* DownloadSizePtr = DownloadSizes.Find(Filename);
+		if (DownloadSizePtr != nullptr)
+		{
+			DownloadSize = *DownloadSizePtr;
+		}
+		return DownloadSize;
+	}
+
+	IDownloadServiceStatistics* FDownloadServiceStatisticsFactory::Create(ISpeedRecorder* SpeedRecorder, IInstallerAnalytics* InstallerAnalytics, FBuildPatchAppManifest* Manifest)
+	{
+		check(SpeedRecorder != nullptr);
+		check(InstallerAnalytics != nullptr);
+		check(Manifest != nullptr);
+		return new FDownloadServiceStatistics(SpeedRecorder, InstallerAnalytics, Manifest);
+	}
+};
\ No newline at end of file
diff --git a/Engine/Source/Runtime/Online/BuildPatchServices/Private/Installer/Statistics/MemoryChunkStoreStatistics.cpp b/Engine/Source/Runtime/Online/BuildPatchServices/Private/Installer/Statistics/MemoryChunkStoreStatistics.cpp
new file mode 100644
--- /dev/null
+++ b/Engine/Source/Runtime/Online/BuildPatchServices/Private/Installer/Statistics/MemoryChunkStoreStatistics.cpp
@@ -0,0 +1,291 @@
+// Copyright 1998-2018 Epic Games, Inc. All Rights Reserved.
+
+#include "Installer/Statistics/MemoryChunkStoreStatistics.h"
+#include "Misc/Guid.h"
+#include "Core/AsyncHelpers.h"
+#include "Installer/Statistics/FileOperationTracker.h"
+
+namespace BuildPatchServices
+{
+	class FMeanUseTracker
+	{
+	public:
+		FMeanUseTracker();
+
+		void Increment();
+		void Decrement();
+		void CountSame();
+		float GetMean() const;
+		int32 GetPeak() const;
+		int32 GetCurrent() const;
+
+	private:
+		FThreadSafeInt32 UsesCount;
+		FThreadSafeInt64 SumCurrentUses;
+		FThreadSafeInt32 CurrentUse;
+		volatile int32 PeakUse;
+	};
+
+	FMeanUseTracker::FMeanUseTracker()
+		: UsesCount(0)
+		, SumCurrentUses(0)
+		, CurrentUse(0)
+		, PeakUse(0)
+	{
+	}
+
+	void FMeanUseTracker::Increment()
+	{
+		int32 NewValue = CurrentUse.Increment();
+		SumCurrentUses.Add(NewValue);
+		UsesCount.Increment();
+		AsyncHelpers::LockFreePeak(&PeakUse, NewValue);
+	}
+
+	void FMeanUseTracker::Decrement()
+	{
+		int32 NewValue = CurrentUse.Decrement();
+		SumCurrentUses.Add(NewValue);
+		UsesCount.Increment();
+	}
+
+	void FMeanUseTracker::CountSame()
+	{
+		SumCurrentUses.Add(CurrentUse.GetValue());
+		UsesCount.Increment();
+	}
+
+	float FMeanUseTracker::GetMean() const
+	{
+		const double SumCurrentUsesFloat = SumCurrentUses.GetValue();
+		const double UseChangeCountFloat = UsesCount.GetValue();
+		return UseChangeCountFloat > 0.f ? SumCurrentUsesFloat / UseChangeCountFloat : 0.f;
+	}
+
+	int32 FMeanUseTracker::GetPeak() const
+	{
+		return PeakUse;
+	}
+
+	int32 FMeanUseTracker::GetCurrent() const
+	{
+		return CurrentUse.GetValue();
+	}
+
+	class FMemoryChunkStoreStatistics
+		: public IMemoryChunkStoreStatistics
+	{
+	public:
+		FMemoryChunkStoreStatistics(const TSet<FGuid>& MultipleReferencedChunks, IFileOperationTracker* FileOperationTracker, FMeanUseTracker* AggregateStoreUseTracker, FMeanUseTracker* AggregateStoreRetainTracker, FThreadSafeInt32* AggregateNumChunksBooted, FThreadSafeInt32* AggregateStoreSize);
+		~FMemoryChunkStoreStatistics();
+
+		// IMemoryChunkStoreStat interface begin.
+		virtual void OnChunkStored(const FGuid& ChunkId) override;
+		virtual void OnChunkReleased(const FGuid& ChunkId) override;
+		virtual void OnChunkBooted(const FGuid& ChunkId) override;
+		virtual void OnStoreUseUpdated(int32 ChunkCount) override;
+		virtual void OnStoreSizeUpdated(int32 Size) override;
+		// IMemoryChunkStoreStat interface end.
+
+		// IMemoryChunkStoreStatistics interface begin.
+		virtual int32 GetStoreUse() const override;
+		virtual int32 GetStoreRetained() const override;
+		virtual int32 GetNumBooted() const override;
+		virtual int32 GetStoreSize() const override;
+		// IMemoryChunkStoreStatistics interface end.
+
+	private:
+		const TSet<FGuid>& MultipleReferencedChunks;
+		IFileOperationTracker* const FileOperationTracker;
+		FMeanUseTracker StoreUseTracker;
+		FMeanUseTracker StoreRetainTracker;
+		FThreadSafeInt32 NumChunksBooted;
+		FThreadSafeInt32 StoreSize;
+		FMeanUseTracker* AggregateStoreUseTracker;
+		FMeanUseTracker* AggregateStoreRetainTracker;
+		FThreadSafeInt32* AggregateNumChunksBooted;
+		FThreadSafeInt32* AggregateStoreSize;
+	};
+
+	FMemoryChunkStoreStatistics::FMemoryChunkStoreStatistics(const TSet<FGuid>& InMultipleReferencedChunks, IFileOperationTracker* InFileOperationTracker, FMeanUseTracker* InStoreUseTracker, FMeanUseTracker* InStoreRetainTracker, FThreadSafeInt32* InNumChunksBooted, FThreadSafeInt32* InTotalStoreSize)
+		: MultipleReferencedChunks(InMultipleReferencedChunks)
+		, FileOperationTracker(InFileOperationTracker)
+		, AggregateStoreUseTracker(InStoreUseTracker)
+		, AggregateStoreRetainTracker(InStoreRetainTracker)
+		, AggregateNumChunksBooted(InNumChunksBooted)
+		, AggregateStoreSize(InTotalStoreSize)
+	{
+	}
+
+	FMemoryChunkStoreStatistics::~FMemoryChunkStoreStatistics()
+	{
+	}
+
+	void FMemoryChunkStoreStatistics::OnChunkStored(const FGuid& ChunkId)
+	{
+		StoreUseTracker.Increment();
+		AggregateStoreUseTracker->Increment();
+		FileOperationTracker->OnDataStateUpdate(ChunkId, EFileOperationState::DataInMemoryStore);
+		if (MultipleReferencedChunks.Contains(ChunkId))
+		{
+			StoreRetainTracker.Increment();
+			AggregateStoreRetainTracker->Increment();
+		}
+		else
+		{
+			StoreRetainTracker.CountSame();
+			AggregateStoreRetainTracker->CountSame();
+		}
+	}
+
+	void FMemoryChunkStoreStatistics::OnChunkReleased(const FGuid& ChunkId)
+	{
+		StoreUseTracker.Decrement();
+		AggregateStoreUseTracker->Decrement();
+		if (MultipleReferencedChunks.Contains(ChunkId))
+		{
+			StoreRetainTracker.Decrement();
+			AggregateStoreRetainTracker->Decrement();
+		}
+		else
+		{
+			StoreRetainTracker.CountSame();
+			AggregateStoreRetainTracker->CountSame();
+		}
+	}
+
+	void FMemoryChunkStoreStatistics::OnChunkBooted(const FGuid& ChunkId)
+	{
+		NumChunksBooted.Increment();
+		AggregateNumChunksBooted->Increment();
+		// We desire a release count too, as it tracks just numbers in the store.
+		OnChunkReleased(ChunkId);
+	}
+
+	void FMemoryChunkStoreStatistics::OnStoreUseUpdated(int32 ChunkCount)
+	{
+	}
+
+	void FMemoryChunkStoreStatistics::OnStoreSizeUpdated(int32 Size)
+	{
+		int32 Diff = Size - StoreSize.Set(Size);
+		if (Diff > 0)
+		{
+			AggregateStoreSize->Add(Diff);
+		}
+		else
+		{
+			AggregateStoreSize->Subtract(FMath::Abs(Diff));
+		}
+	}
+
+	int32 FMemoryChunkStoreStatistics::GetStoreUse() const
+	{
+		return StoreUseTracker.GetCurrent();
+	}
+
+	int32 FMemoryChunkStoreStatistics::GetStoreRetained() const
+	{
+		return StoreRetainTracker.GetCurrent();
+	}
+
+	int32 FMemoryChunkStoreStatistics::GetNumBooted() const
+	{
+		return NumChunksBooted.GetValue();
+	}
+
+	int32 FMemoryChunkStoreStatistics::GetStoreSize() const
+	{
+		return StoreSize.GetValue();
+	}
+
+	class FMemoryChunkStoreAggregateStatistics
+		: public IMemoryChunkStoreAggregateStatistics
+	{
+	public:
+		FMemoryChunkStoreAggregateStatistics(const TSet<FGuid>& MultipleReferencedChunks, IFileOperationTracker* FileOperationTracker);
+		~FMemoryChunkStoreAggregateStatistics();
+
+		// IMemoryChunkStoreStatistics interface begin.
+		virtual IMemoryChunkStoreStatistics* Expose(int32 Index) override;
+		virtual float GetAverageStoreUse() const override;
+		virtual int32 GetPeakStoreUse() const override;
+		virtual float GetAverageStoreRetained() const override;
+		virtual int32 GetPeakStoreRetained() const override;
+		virtual int32 GetTotalStoreSize() const override;
+		virtual int32 GetTotalNumBooted() const override;
+		// IMemoryChunkStoreStatistics interface end.
+
+	private:
+		const TSet<FGuid> MultipleReferencedChunks;
+		IFileOperationTracker* const FileOperationTracker;
+		FMeanUseTracker AggregateStoreUseTracker;
+		FMeanUseTracker AggregateStoreRetainTracker;
+		FThreadSafeInt32 AggregateNumChunksBooted;
+		FThreadSafeInt32 AggregateStoreSize;
+		TMap<int32, TUniquePtr<FMemoryChunkStoreStatistics>> MemoryChunkStoreStatistics;
+	};
+
+	FMemoryChunkStoreAggregateStatistics::FMemoryChunkStoreAggregateStatistics(const TSet<FGuid>& InMultipleReferencedChunks, IFileOperationTracker* InFileOperationTracker)
+		: MultipleReferencedChunks(InMultipleReferencedChunks)
+		, FileOperationTracker(InFileOperationTracker)
+		, AggregateNumChunksBooted(0)
+		, AggregateStoreSize(0)
+	{
+	}
+
+	FMemoryChunkStoreAggregateStatistics::~FMemoryChunkStoreAggregateStatistics()
+	{
+	}
+
+	IMemoryChunkStoreStatistics* FMemoryChunkStoreAggregateStatistics::Expose(int32 Index)
+	{
+		if (MemoryChunkStoreStatistics.Contains(Index) == false)
+		{
+			MemoryChunkStoreStatistics.Emplace(Index, new FMemoryChunkStoreStatistics(
+				MultipleReferencedChunks,
+				FileOperationTracker,
+				&AggregateStoreUseTracker,
+				&AggregateStoreRetainTracker,
+				&AggregateNumChunksBooted,
+				&AggregateStoreSize));
+		}
+		return MemoryChunkStoreStatistics[Index].Get();
+	}
+
+	float FMemoryChunkStoreAggregateStatistics::GetAverageStoreUse() const
+	{
+		return AggregateStoreUseTracker.GetMean();
+	}
+
+	int32 FMemoryChunkStoreAggregateStatistics::GetPeakStoreUse() const
+	{
+		return AggregateStoreUseTracker.GetPeak();
+	}
+
+	float FMemoryChunkStoreAggregateStatistics::GetAverageStoreRetained() const
+	{
+		return AggregateStoreRetainTracker.GetMean();
+	}
+
+	int32 FMemoryChunkStoreAggregateStatistics::GetPeakStoreRetained() const
+	{
+		return AggregateStoreRetainTracker.GetPeak();
+	}
+
+	int32 FMemoryChunkStoreAggregateStatistics::GetTotalStoreSize() const
+	{
+		return AggregateStoreSize.GetValue();
+	}
+
+	int32 FMemoryChunkStoreAggregateStatistics::GetTotalNumBooted() const
+	{
+		return AggregateNumChunksBooted.GetValue();
+	}
+
+	IMemoryChunkStoreAggregateStatistics* FMemoryChunkStoreAggregateStatisticsFactory::Create(const TSet<FGuid>& MultipleReferencedChunks, IFileOperationTracker* FileOperationTracker)
+	{
+		check(FileOperationTracker != nullptr);
+		return new FMemoryChunkStoreAggregateStatistics(MultipleReferencedChunks, FileOperationTracker);
+	}
+};
\ No newline at end of file
diff --git a/Engine/Source/Runtime/Online/BuildPatchServices/Private/Installer/Statistics/MemoryChunkStoreStatistics.h b/Engine/Source/Runtime/Online/BuildPatchServices/Private/Installer/Statistics/MemoryChunkStoreStatistics.h
new file mode 100644
--- /dev/null
+++ b/Engine/Source/Runtime/Online/BuildPatchServices/Private/Installer/Statistics/MemoryChunkStoreStatistics.h
@@ -0,0 +1,100 @@
+// Copyright 1998-2018 Epic Games, Inc. All Rights Reserved.
+#pragma once
+
+#include "Installer/MemoryChunkStore.h"
+
+namespace BuildPatchServices
+{
+	class IMemoryChunkStoreStat;
+	class IFileOperationTracker;
+
+	/**
+	 * Interface to the statistics class which provides access to tracked values from a memory chunk store stat.
+	 */
+	class IMemoryChunkStoreStatistics
+		: public IMemoryChunkStoreStat
+	{
+	public:
+		/**
+		 * @return the number of chunks held.
+		 */
+		virtual int32 GetStoreUse() const = 0;
+
+		/**
+		 * @return the number of chunks held which are retained due to multiple references.
+		 */
+		virtual int32 GetStoreRetained() const = 0;
+
+		/**
+		 * @return the number of chunks which have been booted.
+		 */
+		virtual int32 GetNumBooted() const = 0;
+
+		/**
+		 * @return the maximum number of chunks which can be held.
+		 */
+		virtual int32 GetStoreSize() const = 0;
+	};
+
+	/**
+	 * Interface to the statistics class which aggregates tracked values from multiple memory chunk store stats.
+	 */
+	class IMemoryChunkStoreAggregateStatistics
+	{
+	public:
+		virtual ~IMemoryChunkStoreAggregateStatistics() { }
+
+		/**
+		 * Exposes an IMemoryChunkStoreStatistics interface which can be given to a memory chunk store and used for individual stats.
+		 * @param Index     The index for the interface to get. Repeated calls with the same value will get the same instance.
+		 * @return the interface for the stat instance assigned to the given index.
+		 */
+		virtual IMemoryChunkStoreStatistics* Expose(int32 Index) = 0;
+
+		/**
+		 * @return the average number of chunks held in the memory stores.
+		 */
+		virtual float GetAverageStoreUse() const = 0;
+
+		/**
+		 * @return the peak number of chunks held in the memory stores.
+		 */
+		virtual int32 GetPeakStoreUse() const = 0;
+
+		/**
+		 * @return the average number of chunks held in the memory stores which are retained due to multiple references.
+		 */
+		virtual float GetAverageStoreRetained() const = 0;
+
+		/**
+		 * @return the peak number of chunks held in the memory stores which are retained due to multiple references.
+		 */
+		virtual int32 GetPeakStoreRetained() const = 0;
+
+		/**
+		 * @return the total number of chunks which can be held in memory stores.
+		 */
+		virtual int32 GetTotalStoreSize() const = 0;
+
+		/**
+		 * @return the number of chunks which were booted from memory stores.
+		 */
+		virtual int32 GetTotalNumBooted() const = 0;
+	};
+
+	/**
+	 * A factory for creating an IMemoryChunkStoreAggregateStatistics instance.
+	 */
+	class FMemoryChunkStoreAggregateStatisticsFactory
+	{
+	public:
+		/**
+		 * Creates a statistics interface which exposes individual memory chunk stores stat dependencies, and collates
+		 * calls received by these into additional aggregated information.
+		 * @param MultipleReferencedChunks      For retained stats, the set of chunks which have multiple references.
+		 * @param FileOperationTracker          The file operation tracker which will be used to update data states.
+		 * @return the new IMemoryChunkStoreAggregateStatistics instance created.
+		 */
+		static IMemoryChunkStoreAggregateStatistics* Create(const TSet<FGuid>& MultipleReferencedChunks, IFileOperationTracker* FileOperationTracker);
+	};
+}
\ No newline at end of file
diff --git a/Engine/Source/Runtime/Online/BuildPatchServices/Private/Installer/Verifier.cpp b/Engine/Source/Runtime/Online/BuildPatchServices/Private/Installer/Verifier.cpp
--- a/Engine/Source/Runtime/Online/BuildPatchServices/Private/Installer/Verifier.cpp
+++ b/Engine/Source/Runtime/Online/BuildPatchServices/Private/Installer/Verifier.cpp
@@ -1,19 +1,49 @@
 // Copyright 1998-2018 Epic Games, Inc. All Rights Reserved.
 
 #include "Installer/Verifier.h"
+#include "HAL/ThreadSafeBool.h"
 #include "Common/StatsCollector.h"
 #include "Common/FileSystem.h"
 #include "BuildPatchVerify.h"
 #include "BuildPatchUtil.h"
-#include "HAL/ThreadSafeBool.h"
+
+DECLARE_LOG_CATEGORY_EXTERN(LogVerifier, Warning, All);
+DEFINE_LOG_CATEGORY(LogVerifier);
 
 namespace BuildPatchServices
 {
-	class FVerification : public IVerifier
+	// 4MB Read buffer.
+	const int32 VerifyBufferSize = 1024 * 1024 * 4;
+
+	bool TryConvertToVerifyResult(EVerifyError InVerifyError, EVerifyResult& OutVerifyResult)
+	{
+		switch (InVerifyError)
+		{
+			case EVerifyError::FileMissing: OutVerifyResult = EVerifyResult::FileMissing; return true;
+			case EVerifyError::OpenFileFailed: OutVerifyResult = EVerifyResult::OpenFileFailed; return true;
+			case EVerifyError::HashCheckFailed: OutVerifyResult = EVerifyResult::HashCheckFailed; return true;
+			case EVerifyError::FileSizeFailed: OutVerifyResult = EVerifyResult::FileSizeFailed; return true;
+		}
+		return false;
+	}
+
+	bool TryConvertToVerifyError(EVerifyResult InVerifyResult, EVerifyError& OutVerifyError)
+	{
+		switch (InVerifyResult)
+		{
+			case EVerifyResult::FileMissing: OutVerifyError = EVerifyError::FileMissing; return true;
+			case EVerifyResult::OpenFileFailed: OutVerifyError = EVerifyError::OpenFileFailed; return true;
+			case EVerifyResult::HashCheckFailed: OutVerifyError = EVerifyError::HashCheckFailed; return true;
+			case EVerifyResult::FileSizeFailed: OutVerifyError = EVerifyError::FileSizeFailed; return true;
+		}
+		return false;
+	}
+
+	class FVerifier : public IVerifier
 	{
 	public:
-		FVerification(IFileSystem* FileSystem, IVerifierStat* InVerificationStat, EVerifyMode InVerifyMode, TSet<FString> InTouchedFiles, TSet<FString> InInstallTags, FBuildPatchAppManifestRef InManifest, FString InVerifyDirectory, FString InStagedFileDirectory);
-		~FVerification() {}
+		FVerifier(IFileSystem* FileSystem, IVerifierStat* InVerificationStat, EVerifyMode InVerifyMode, TSet<FString> InTouchedFiles, TSet<FString> InInstallTags, FBuildPatchAppManifestRef InManifest, FString InVerifyDirectory, FString InStagedFileDirectory);
+		~FVerifier() {}
 
 		// IControllable interface begin.
 		virtual void SetPaused(bool bInIsPaused) override;
@@ -21,15 +51,16 @@ namespace BuildPatchServices
 		// IControllable interface end.
 
 		// IVerifier interface begin.
-		virtual bool Verify(TArray<FString>& OutDatedFiles) override;
+		virtual EVerifyResult Verify(TArray<FString>& OutDatedFiles) override;
 		// IVerifier interface end.
 
 	private:
 		FString SelectFullFilePath(const FString& BuildFile);
-		bool VerfiyFileSha(const FString& BuildFile, int64 BuildFileSize);
-		bool VerfiyFileSize(const FString& BuildFile, int64 BuildFileSize);
+		EVerifyResult VerfiyFileSha(const FString& BuildFile, int64 BuildFileSize);
+		EVerifyResult VerfiyFileSize(const FString& BuildFile, int64 BuildFileSize);
 
 	private:
+		TArray<uint8> FileReadBuffer;
 		IFileSystem* FileSystem;
 		IVerifierStat* VerifierStat;
 		EVerifyMode VerifyMode;
@@ -43,7 +74,7 @@ namespace BuildPatchServices
 		int64 ProcessedBytes;
 	};
 
-	FVerification::FVerification(IFileSystem* InFileSystem, IVerifierStat* InVerificationStat, EVerifyMode InVerifyMode, TSet<FString> InTouchedFiles, TSet<FString> InInstallTags, FBuildPatchAppManifestRef InManifest, FString InVerifyDirectory, FString InStagedFileDirectory)
+	FVerifier::FVerifier(IFileSystem* InFileSystem, IVerifierStat* InVerificationStat, EVerifyMode InVerifyMode, TSet<FString> InTouchedFiles, TSet<FString> InInstallTags, FBuildPatchAppManifestRef InManifest, FString InVerifyDirectory, FString InStagedFileDirectory)
 		: FileSystem(InFileSystem)
 		, VerifierStat(InVerificationStat)
 		, VerifyMode(InVerifyMode)
@@ -55,21 +86,24 @@ namespace BuildPatchServices
 		, bIsPaused(false)
 		, bShouldAbort(false)
 		, ProcessedBytes(0)
-	{}
+	{
+		FileReadBuffer.Reserve(VerifyBufferSize);
+		FileReadBuffer.AddUninitialized(VerifyBufferSize);
+	}
 
-	void FVerification::SetPaused(bool bInIsPaused)
+	void FVerifier::SetPaused(bool bInIsPaused)
 	{
 		bIsPaused = bInIsPaused;
 	}
 
-	void FVerification::Abort()
+	void FVerifier::Abort()
 	{
 		bShouldAbort = true;
 	}
 
-	bool FVerification::Verify(TArray<FString>& OutDatedFiles)
+	EVerifyResult FVerifier::Verify(TArray<FString>& OutDatedFiles)
 	{
-		bool bAllCorrect = true;
+		EVerifyResult VerifyResult = EVerifyResult::Success;
 		OutDatedFiles.Empty();
 		if (VerifyMode == EVerifyMode::FileSizeCheckAllFiles || VerifyMode == EVerifyMode::ShaVerifyAllFiles)
 		{
@@ -77,6 +111,7 @@ namespace BuildPatchServices
 		}
 
 		// Setup progress tracking.
+		VerifierStat->OnProcessedDataUpdated(0);
 		VerifierStat->OnTotalRequiredUpdated(Manifest->GetFileSize(RequiredFiles));
 
 		// Select verify function.
@@ -97,28 +132,34 @@ namespace BuildPatchServices
 
 			// Verify the file.
 			VerifierStat->OnFileStarted(BuildFile, BuildFileSize);
-			bool bFileOk = bVerifySha ? VerfiyFileSha(BuildFile, BuildFileSize) : VerfiyFileSize(BuildFile, BuildFileSize);
-			VerifierStat->OnFileCompleted(BuildFile, bFileOk);
-			if (bFileOk == false)
+			EVerifyResult FileVerifyResult = bVerifySha ? VerfiyFileSha(BuildFile, BuildFileSize) : VerfiyFileSize(BuildFile, BuildFileSize);
+			VerifierStat->OnFileCompleted(BuildFile, FileVerifyResult);
+			if (FileVerifyResult != EVerifyResult::Success)
 			{
-				bAllCorrect = false;
 				OutDatedFiles.Add(BuildFile);
+				if (VerifyResult == EVerifyResult::Success)
+				{
+					VerifyResult = FileVerifyResult;
+				}
 			}
-			ProcessedBytes += BuildFileSize;
-			VerifierStat->OnProcessedDataUpdated(ProcessedBytes);
 		}
 
-		return bAllCorrect && !bShouldAbort;
+		if (bShouldAbort && VerifyResult == EVerifyResult::Success)
+		{
+			VerifyResult = EVerifyResult::Aborted;
+		}
+
+		return VerifyResult;
 	}
 
-	FString FVerification::SelectFullFilePath(const FString& BuildFile)
+	FString FVerifier::SelectFullFilePath(const FString& BuildFile)
 	{
 		FString FullFilePath;
 		if (StagedFileDirectory.IsEmpty() == false)
 		{
 			FullFilePath = StagedFileDirectory / BuildFile;
 			int64 FileSize;
-			if (FileSystem->GetFileSize(*FullFilePath, FileSize) && FileSize != INDEX_NONE)
+			if (FileSystem->GetFileSize(*FullFilePath, FileSize))
 			{
 				return FullFilePath;
 			}
@@ -127,38 +168,84 @@ namespace BuildPatchServices
 		return FullFilePath;
 	}
 
-	bool FVerification::VerfiyFileSha(const FString& BuildFile, int64 BuildFileSize)
+	EVerifyResult FVerifier::VerfiyFileSha(const FString& BuildFile, int64 BuildFileSize)
 	{
 		FSHAHash BuildFileHash;
+		ISpeedRecorder::FRecord ActivityRecord;
+		const int64 PrevProcessedBytes = ProcessedBytes;
 		bool bFoundHash = Manifest->GetFileHash(BuildFile, BuildFileHash);
 		checkf(bFoundHash, TEXT("Missing file hash from manifest."));
-		TFunction<void(float)> FileProgress = [this, &BuildFileSize, &BuildFile](float Progress)
+		EVerifyResult VerifyResult;
+		const FString FileToVerify = SelectFullFilePath(BuildFile);
+		uint8 ReturnValue = 0;
+		TUniquePtr<FArchive> FileReader = FileSystem->CreateFileReader(*FileToVerify);
+		VerifierStat->OnFileProgress(BuildFile, 0);
+		if (FileReader.IsValid())
+		{
+			const int64 FileSize = FileReader->TotalSize();
+			if (FileSize != Manifest->GetFileSize(BuildFile))
+			{
+				VerifyResult = EVerifyResult::FileSizeFailed;
+			}
+			else
+			{
+				FSHA1 HashState;
+				FSHAHash HashValue;
+				while (!FileReader->AtEnd() && !bShouldAbort)
+				{
+					// Pause if necessary
+					while (bIsPaused && !bShouldAbort)
+					{
+						FPlatformProcess::Sleep(0.1f);
+					}
+					ActivityRecord.CyclesStart = FStatsCollector::GetCycles();
+					// Read file and update hash state
+					const int64 SizeLeft = FileSize - FileReader->Tell();
+					ActivityRecord.Size = FMath::Min<int64>(VerifyBufferSize, SizeLeft);
+					FileReader->Serialize(FileReadBuffer.GetData(), ActivityRecord.Size);
+					HashState.Update(FileReadBuffer.GetData(), ActivityRecord.Size);
+					ProcessedBytes = PrevProcessedBytes + FileReader->Tell();
+					ActivityRecord.CyclesEnd = FStatsCollector::GetCycles();
+					VerifierStat->OnFileRead(ActivityRecord);
+					VerifierStat->OnFileProgress(BuildFile, FileReader->Tell());
+					VerifierStat->OnProcessedDataUpdated(ProcessedBytes);
+				}
+				HashState.Final();
+				HashState.GetHash(HashValue.Hash);
+				if (HashValue == BuildFileHash)
+				{
+					VerifyResult = EVerifyResult::Success;
+				}
+				else if (!bShouldAbort)
+				{
+					VerifyResult = EVerifyResult::HashCheckFailed;
+				}
+				else
+				{
+					VerifyResult = EVerifyResult::Aborted;
+				}
+			}
+			FileReader->Close();
+		}
+		else if (FileSystem->FileExists(*FileToVerify))
 		{
-			int64 FileProcessed = FMath::Min<int64>(BuildFileSize * Progress, BuildFileSize);
-			VerifierStat->OnFileProgress(BuildFile, FileProcessed);
-			VerifierStat->OnProcessedDataUpdated(ProcessedBytes + FileProcessed);
-		};
-		TFunction<bool()> IsPaused = [this]() -> bool
+			VerifyResult = EVerifyResult::OpenFileFailed;
+		}
+		else
 		{
-			return bIsPaused;
-		};
-		TFunction<bool()> ShouldAbort = [this]() -> bool
+			VerifyResult = EVerifyResult::FileMissing;
+		}
+		ProcessedBytes = PrevProcessedBytes + BuildFileSize;
+		if (VerifyResult != EVerifyResult::Success)
 		{
-			return bShouldAbort;
-		};
-		bool bSuccess = FBuildPatchUtils::VerifyFile(
-			FileSystem,
-			SelectFullFilePath(BuildFile),
-			BuildFileHash,
-			BuildFileHash,
-			FBuildPatchFloatDelegate::CreateLambda(MoveTemp(FileProgress)),
-			FBuildPatchBoolRetDelegate::CreateLambda(MoveTemp(IsPaused)),
-			FBuildPatchBoolRetDelegate::CreateLambda(MoveTemp(ShouldAbort))) != 0;
-		VerifierStat->OnFileProgress(BuildFile, BuildFileSize);
-		return bSuccess;
+			VerifierStat->OnFileProgress(BuildFile, BuildFileSize);
+			VerifierStat->OnProcessedDataUpdated(ProcessedBytes);
+		}
+
+		return VerifyResult;
 	}
 
-	bool FVerification::VerfiyFileSize(const FString& BuildFile, int64 BuildFileSize)
+	EVerifyResult FVerifier::VerfiyFileSize(const FString& BuildFile, int64 BuildFileSize)
 	{
 		// Pause if necessary.
 		const double PrePauseTime = FStatsCollector::GetSeconds();
@@ -170,19 +257,32 @@ namespace BuildPatchServices
 		}
 		VerifierStat->OnFileProgress(BuildFile, 0);
 		int64 FileSize;
-		bool bSuccess = false;
+		EVerifyResult VerifyResult;
 		if (FileSystem->GetFileSize(*SelectFullFilePath(BuildFile), FileSize))
 		{
-			bSuccess = FileSize == BuildFileSize;
+			if (FileSize == BuildFileSize)
+			{
+				VerifyResult = EVerifyResult::Success;
+			}
+			else
+			{
+				VerifyResult = EVerifyResult::FileSizeFailed;
+			}
+		}
+		else
+		{
+			VerifyResult = EVerifyResult::FileMissing;
 		}
 		VerifierStat->OnFileProgress(BuildFile, BuildFileSize);
-		return bSuccess;
+		ProcessedBytes += BuildFileSize;
+		VerifierStat->OnProcessedDataUpdated(ProcessedBytes);
+		return VerifyResult;
 	}
 
 	IVerifier* FVerifierFactory::Create(IFileSystem* FileSystem, IVerifierStat* VerifierStat, EVerifyMode VerifyMode, TSet<FString> TouchedFiles, TSet<FString> InstallTags, FBuildPatchAppManifestRef Manifest, FString VerifyDirectory, FString StagedFileDirectory)
 	{
 		check(FileSystem != nullptr);
 		check(VerifierStat != nullptr);
-		return new FVerification(FileSystem, VerifierStat, VerifyMode, MoveTemp(TouchedFiles), MoveTemp(InstallTags), MoveTemp(Manifest), MoveTemp(VerifyDirectory), MoveTemp(StagedFileDirectory));
+		return new FVerifier(FileSystem, VerifierStat, VerifyMode, MoveTemp(TouchedFiles), MoveTemp(InstallTags), MoveTemp(Manifest), MoveTemp(VerifyDirectory), MoveTemp(StagedFileDirectory));
 	}
 }
diff --git a/Engine/Source/Runtime/Online/BuildPatchServices/Private/Installer/Verifier.h b/Engine/Source/Runtime/Online/BuildPatchServices/Private/Installer/Verifier.h
--- a/Engine/Source/Runtime/Online/BuildPatchServices/Private/Installer/Verifier.h
+++ b/Engine/Source/Runtime/Online/BuildPatchServices/Private/Installer/Verifier.h
@@ -2,12 +2,53 @@
 #pragma once
 
 #include "Installer/Controllable.h"
+#include "Common/SpeedRecorder.h"
 #include "BuildPatchManifest.h"
 
 namespace BuildPatchServices
 {
+	class IFileSystem;
 	class IVerifierStat;
 	enum class EVerifyMode : uint32;
+	enum class EVerifyError : uint32;
+	/**
+	 * An enum defining the result of a verification process.
+	 */
+	enum class EVerifyResult : uint32
+	{
+		Success = 0,
+
+		// If the process was ended due to an external cancel.
+		Aborted,
+
+		// The verify failed due to a missing file.
+		FileMissing,
+
+		// The verify failed due to a file failing to open.
+		OpenFileFailed,
+
+		// The expected data hash for a file did not match.
+		HashCheckFailed,
+
+		// A file did not match the expected size.
+		FileSizeFailed,
+	};
+
+	/**
+	 * Used to convert a EVerifyError to a EVerifyResult.
+	 * @param InVerifyError   - Error to convert.
+	 * @param OutVerifyResult - Reference to completed.
+	 * @return true if successful.
+	 */
+	bool TryConvertToVerifyResult(EVerifyError InVerifyError, EVerifyResult& OutVerifyResult);
+
+	/**
+	 * Used to convert a EVerifyResult to a EVerifyError.
+	 * @param InVerifyResult - Result to convert.
+	 * @param OutVerifyError - Reference to complete.
+	 * @return true if successful.
+	 */
+	bool TryConvertToVerifyError(EVerifyResult InVerifyResult, EVerifyError& OutVerifyError);
 
 	/**
 	 * An interface providing the functionality to verify a local installation.
@@ -20,9 +61,10 @@ namespace BuildPatchServices
 		 * Verifies a local directory structure against a given manifest.
 		 * NOTE: This function is blocking and will not return until finished. Don't run on main thread.
 		 * @param OutDatedFiles    OUT  The array of files that do not match or are locally missing.
-		 * @return    true if no file errors occurred AND the verification was successful
+		 * @return    EVerifiyResult::Success if no file errors occurred AND the verification was successful.
+		 *            Otherwise it will return the first error encountered during verification.
 		 */
-		virtual bool Verify(TArray<FString>& OutDatedFiles) = 0;
+		virtual EVerifyResult Verify(TArray<FString>& OutDatedFiles) = 0;
 	};
 
 	class FVerifierFactory
@@ -71,9 +113,15 @@ namespace BuildPatchServices
 		/**
 		 * Called each time a file has finished being verified.
 		 * @param Filename      The filename of the file.
-		 * @param bSuccess      True if the file data was correct.
+		 * @param VerifyResult  The result of the file's verify test.
+		 */
+		virtual void OnFileCompleted(const FString& Filename, EVerifyResult VerifyResult) = 0;
+
+		/**
+		 * Called each time a read operation is made.
+		 * @param Record        The details for the operation.
 		 */
-		virtual void OnFileCompleted(const FString& Filename, bool bSuccess) = 0;
+		virtual void OnFileRead(const ISpeedRecorder::FRecord& Record) = 0;
 
 		/**
 		 * Called to update the total amount of bytes which have been processed.
diff --git a/Engine/Source/Runtime/Online/BuildPatchServices/Private/Tests/Fake/ChunkDataAccess.fake.h b/Engine/Source/Runtime/Online/BuildPatchServices/Private/Tests/Fake/ChunkDataAccess.fake.h
--- a/Engine/Source/Runtime/Online/BuildPatchServices/Private/Tests/Fake/ChunkDataAccess.fake.h
+++ b/Engine/Source/Runtime/Online/BuildPatchServices/Private/Tests/Fake/ChunkDataAccess.fake.h
@@ -11,6 +11,14 @@ namespace BuildPatchServices
 		: public IChunkDataAccess
 	{
 	public:
+		virtual ~FFakeChunkDataAccess()
+		{
+			if (OnDeleted)
+			{
+				OnDeleted();
+			}
+		}
+
 		virtual void GetDataLock(const uint8** OutChunkData, const FChunkHeader** OutChunkHeader) const override
 		{
 			(*OutChunkData) = ChunkData;
@@ -35,6 +43,7 @@ namespace BuildPatchServices
 	public:
 		FChunkHeader ChunkHeader;
 		uint8* ChunkData;
+		TFunction<void()> OnDeleted;
 	};
 }
 
diff --git a/Engine/Source/Runtime/Online/BuildPatchServices/Private/Tests/Fake/ChunkSource.fake.h b/Engine/Source/Runtime/Online/BuildPatchServices/Private/Tests/Fake/ChunkSource.fake.h
--- a/Engine/Source/Runtime/Online/BuildPatchServices/Private/Tests/Fake/ChunkSource.fake.h
+++ b/Engine/Source/Runtime/Online/BuildPatchServices/Private/Tests/Fake/ChunkSource.fake.h
@@ -24,6 +24,11 @@ namespace BuildPatchServices
 			return MoveTemp(NewRequirements);
 		}
 
+		bool AddRepeatRequirement(const FGuid& RepeatRequirement)
+		{
+			return false;
+		}
+
 		virtual void SetUnavailableChunksCallback(TFunction<void(TSet<FGuid>)> Callback) override
 		{
 		}
diff --git a/Engine/Source/Runtime/Online/BuildPatchServices/Private/Tests/Fake/FileSystem.fake.h b/Engine/Source/Runtime/Online/BuildPatchServices/Private/Tests/Fake/FileSystem.fake.h
--- a/Engine/Source/Runtime/Online/BuildPatchServices/Private/Tests/Fake/FileSystem.fake.h
+++ b/Engine/Source/Runtime/Online/BuildPatchServices/Private/Tests/Fake/FileSystem.fake.h
@@ -2,50 +2,137 @@
 #pragma once
 
 #include "Tests/Mock/FileSystem.mock.h"
+#include "Misc/Paths.h"
+#include "Misc/ScopeLock.h"
 
 #if WITH_DEV_AUTOMATION_TESTS
 
 namespace BuildPatchServices
 {
+	/**
+	 * Using a fake file reader allows us to simulate UE4's file reader behavior, where if the file is written to
+	 * after the handle is opened, you will cause an assert if you try to read the new data at the end as the total size is cached.
+	 */
+	class FFakeFileReader
+		: public FMemoryArchive
+	{
+	public:
+		FFakeFileReader(const TArray<uint8>& InBytes)
+			: Bytes(InBytes)
+			, FakeTotalSize(Bytes.Num())
+		{
+			this->SetIsLoading(true);
+		}
+
+		virtual FString GetArchiveName() const
+		{
+			return TEXT("FFakeFileReader");
+		}
+
+		virtual int64 TotalSize() override
+		{
+			return FakeTotalSize;
+		}
+
+		void Serialize(void* Data, int64 Num)
+		{
+			if (Num && !ArIsError)
+			{
+				if (Offset + Num <= TotalSize())
+				{
+					FMemory::Memcpy(Data, &Bytes[Offset], Num);
+					Offset += Num;
+				}
+				else
+				{
+					ArIsError = true;
+				}
+			}
+		}
+
+	public:
+		const TArray<uint8>& Bytes;
+		int64 FakeTotalSize;
+	};
+
 	class FFakeFileSystem
 		: public FMockFileSystem
 	{
 	public:
-		virtual TUniquePtr<FArchive> CreateFileReader(const TCHAR* Filename, EFileRead ReadFlags = EFileRead::None) const override
+		virtual TUniquePtr<FArchive> CreateFileReader(const TCHAR* Filename, EReadFlags ReadFlags = EReadFlags::None) const override
 		{
-			FString NormalizedFilename = Filename;
-			FPaths::NormalizeFilename(NormalizedFilename);
-			NormalizedFilename = FPaths::ConvertRelativePathToFull(TEXT(""), MoveTemp(NormalizedFilename));
-			FScopeLock ScopeLock(&ThreadLock);
-			TUniquePtr<FArchive> Reader(new FMemoryReader(DiskData.FindOrAdd(NormalizedFilename)));
+			TUniquePtr<FArchive> Reader;
+			if (CreateFileReaderFunc)
+			{
+				Reader = CreateFileReaderFunc(Filename, ReadFlags);
+			}
+			else
+			{
+				FString NormalizedFilename = Filename;
+				FPaths::NormalizeFilename(NormalizedFilename);
+				NormalizedFilename = FPaths::ConvertRelativePathToFull(TEXT(""), MoveTemp(NormalizedFilename));
+				FScopeLock ScopeLock(&ThreadLock);
+				if (DiskData.Contains(NormalizedFilename) && !DiskDataOpenFailure.Contains(NormalizedFilename))
+				{
+					Reader.Reset(new FFakeFileReader(DiskData[NormalizedFilename]));
+				}
+			}
 			RxCreateFileReader.Emplace(FStatsCollector::GetSeconds(), Reader.Get(), Filename, ReadFlags);
 			return Reader;
 		}
 
-		virtual TUniquePtr<FArchive> CreateFileWriter(const TCHAR* Filename, EFileWrite WriteFlags = EFileWrite::None) const override
+		virtual TUniquePtr<FArchive> CreateFileWriter(const TCHAR* Filename, EWriteFlags WriteFlags = EWriteFlags::None) const override
+		{
+			TUniquePtr<FArchive> Writer;
+			if (CreateFileWriterFunc)
+			{
+				Writer = CreateFileWriterFunc(Filename, WriteFlags);
+			}
+			else
+			{
+				FString NormalizedFilename = Filename;
+				FPaths::NormalizeFilename(NormalizedFilename);
+				NormalizedFilename = FPaths::ConvertRelativePathToFull(TEXT(""), MoveTemp(NormalizedFilename));
+				FScopeLock ScopeLock(&ThreadLock);
+				Writer.Reset(new FMemoryWriter(DiskData.FindOrAdd(NormalizedFilename)));
+			}
+			RxCreateFileWriter.Emplace(FStatsCollector::GetSeconds(), Writer.Get(), Filename, WriteFlags);
+			return Writer;
+		}
+
+		bool DeleteFile(const TCHAR* Filename) const
 		{
 			FString NormalizedFilename = Filename;
 			FPaths::NormalizeFilename(NormalizedFilename);
 			NormalizedFilename = FPaths::ConvertRelativePathToFull(TEXT(""), MoveTemp(NormalizedFilename));
 			FScopeLock ScopeLock(&ThreadLock);
-			TUniquePtr<FArchive> Writer(new FMemoryWriter(DiskData.FindOrAdd(NormalizedFilename)));
-			RxCreateFileWriter.Emplace(FStatsCollector::GetSeconds(), Writer.Get(), Filename, WriteFlags);
-			return Writer;
+			DiskData.Remove(NormalizedFilename);
+			return true;
 		}
 
 		virtual bool GetFileSize(const TCHAR* Filename, int64& OutFileSize) const override
 		{
-			OutFileSize = INDEX_NONE;
+			FScopeLock ScopeLock(&ThreadLock);
+			OutFileSize = -1;
 			if (DiskData.Contains(Filename))
 			{
 				OutFileSize = DiskData[Filename].Num();
 			}
 			RxGetFileSize.Emplace(FStatsCollector::GetSeconds(), Filename, OutFileSize);
-			return true;
+			return OutFileSize >= 0;
+		}
+
+		virtual bool FileExists(const TCHAR* Filename) const override
+		{
+			return DiskData.Contains(Filename) || DiskDataOpenFailure.Contains(Filename);
 		}
 
 	public:
 		mutable TMap<FString, TArray<uint8>> DiskData;
+		mutable TArray<FString> DiskDataOpenFailure;
+
+		TFunction<TUniquePtr<FArchive>(const TCHAR*, EReadFlags)> CreateFileReaderFunc;
+		TFunction<TUniquePtr<FArchive>(const TCHAR*, EWriteFlags)> CreateFileWriterFunc;
 	};
 }
 
diff --git a/Engine/Source/Runtime/Online/BuildPatchServices/Private/Tests/Fake/InstallerError.fake.h b/Engine/Source/Runtime/Online/BuildPatchServices/Private/Tests/Fake/InstallerError.fake.h
--- a/Engine/Source/Runtime/Online/BuildPatchServices/Private/Tests/Fake/InstallerError.fake.h
+++ b/Engine/Source/Runtime/Online/BuildPatchServices/Private/Tests/Fake/InstallerError.fake.h
@@ -41,10 +41,10 @@ namespace BuildPatchServices
 			return ErrorText;
 		}
 
-		virtual void SetError(EBuildPatchInstallError InErrorType, const TCHAR* InErrorCode, FText InErrorText) override
+		virtual void SetError(EBuildPatchInstallError InErrorType, const TCHAR* InErrorSubType, uint32 InErrorCode, FText InErrorText) override
 		{
 			ErrorType = InErrorType;
-			ErrorCode = InErrorCode;
+			ErrorCode = InErrorSubType + ((InErrorCode > 0) ? FString::Printf(TEXT("-%u"), InErrorCode) : TEXT(""));
 			ErrorText = MoveTemp(InErrorText);
 		}
 
diff --git a/Engine/Source/Runtime/Online/BuildPatchServices/Private/Tests/Mock/BuildPatchProgress.mock.h b/Engine/Source/Runtime/Online/BuildPatchServices/Private/Tests/Mock/BuildPatchProgress.mock.h
--- a/Engine/Source/Runtime/Online/BuildPatchServices/Private/Tests/Mock/BuildPatchProgress.mock.h
+++ b/Engine/Source/Runtime/Online/BuildPatchServices/Private/Tests/Mock/BuildPatchProgress.mock.h
@@ -52,13 +52,6 @@ namespace BuildPatchServices
 			return EBuildPatchState();
 		}
 
-		virtual const FText& GetStateText() const override
-		{
-			static FText NotImp;
-			MOCK_FUNC_NOT_IMPLEMENTED("FMockBuildPatchProgress::GetStateText");
-			return NotImp;
-		}
-
 		virtual float GetProgress() const override
 		{
 			MOCK_FUNC_NOT_IMPLEMENTED("FMockBuildPatchProgress::GetProgress");
@@ -101,7 +94,7 @@ namespace BuildPatchServices
 			return bool();
 		}
 
-		virtual void SetIsDownloading(bool IsDownloading) override
+		virtual void SetIsDownloading(bool bInIsDownloading) override
 		{
 			MOCK_FUNC_NOT_IMPLEMENTED("FMockBuildPatchProgress::SetIsDownloading");
 		}
diff --git a/Engine/Source/Runtime/Online/BuildPatchServices/Private/Tests/Mock/ChunkDataSerialization.mock.h b/Engine/Source/Runtime/Online/BuildPatchServices/Private/Tests/Mock/ChunkDataSerialization.mock.h
--- a/Engine/Source/Runtime/Online/BuildPatchServices/Private/Tests/Mock/ChunkDataSerialization.mock.h
+++ b/Engine/Source/Runtime/Online/BuildPatchServices/Private/Tests/Mock/ChunkDataSerialization.mock.h
@@ -16,10 +16,13 @@ namespace BuildPatchServices
 		typedef TTuple<FString, EChunkLoadResult> FRxLoadFromFile;
 		typedef TTuple<FString, const IChunkDataAccess*> FRxSaveToFile;
 		typedef TTuple<TArray<uint8>, EChunkLoadResult> FRxLoadFromMemory;
+		typedef TTuple<EChunkLoadResult> FRxLoadFromArchive;
+		typedef TTuple<const IChunkDataAccess*> FRxSaveToArchive;
 		typedef TTuple<TArray<uint8>, FSHAHash> FRxInjectShaToChunkData;
 
 		typedef TTuple<IChunkDataAccess*, EChunkLoadResult> FTxLoadFromFile;
 		typedef TTuple<IChunkDataAccess*, EChunkLoadResult> FTxLoadFromMemory;
+		typedef TTuple<IChunkDataAccess*, EChunkLoadResult> FTxLoadFromArchive;
 
 	public:
 		virtual IChunkDataAccess* LoadFromFile(const FString& Filename, EChunkLoadResult& OutLoadResult) const override
@@ -62,15 +65,26 @@ namespace BuildPatchServices
 
 		virtual IChunkDataAccess* LoadFromArchive(FArchive& Archive, EChunkLoadResult& OutLoadResult) const override
 		{
-			MOCK_FUNC_NOT_IMPLEMENTED("FMockChunkDataSerialization::LoadFromArchive");
-			OutLoadResult = EChunkLoadResult::SerializationError;
-			return nullptr;
+			IChunkDataAccess* Result = nullptr;
+			if (TxLoadFromArchive.Num())
+			{
+				FTxLoadFromFile LoadFromArchiveResult = TxLoadFromArchive.Pop();
+				Result = LoadFromArchiveResult.Get<0>();
+				OutLoadResult = LoadFromArchiveResult.Get<1>();
+			}
+			RxLoadFromArchive.Emplace(OutLoadResult);
+			return Result;
 		}
 
 		virtual EChunkSaveResult SaveToArchive(FArchive& Archive, const IChunkDataAccess* ChunkDataAccess) const override
 		{
-			MOCK_FUNC_NOT_IMPLEMENTED("FMockChunkDataSerialization::SaveToArchive");
-			return EChunkSaveResult::SerializationError;
+			EChunkSaveResult Result = EChunkSaveResult::Success;
+			if (SaveToArchiveFunc)
+			{
+				Result = SaveToArchiveFunc(Archive, ChunkDataAccess);
+			}
+			RxSaveToArchive.Emplace(ChunkDataAccess);
+			return Result;
 		}
 
 		virtual void InjectShaToChunkData(TArray<uint8>& Memory, const FSHAHash& ShaHashData) const override
@@ -82,10 +96,15 @@ namespace BuildPatchServices
 		mutable TArray<FRxLoadFromFile> RxLoadFromFile;
 		mutable TArray<FRxSaveToFile> RxSaveToFile;
 		mutable TArray<FRxLoadFromMemory> RxLoadFromMemory;
+		mutable TArray<FRxLoadFromArchive> RxLoadFromArchive;
+		mutable TArray<FRxSaveToArchive> RxSaveToArchive;
 		mutable TArray<FRxInjectShaToChunkData> RxInjectShaToChunkData;
 
 		mutable TArray<FTxLoadFromFile> TxLoadFromFile;
 		mutable TArray<FTxLoadFromMemory> TxLoadFromMemory;
+		mutable TArray<FTxLoadFromArchive> TxLoadFromArchive;
+
+		TFunction<EChunkSaveResult(FArchive&, const IChunkDataAccess*)> SaveToArchiveFunc;
 	};
 }
 
diff --git a/Engine/Source/Runtime/Online/BuildPatchServices/Private/Tests/Mock/ChunkStore.mock.h b/Engine/Source/Runtime/Online/BuildPatchServices/Private/Tests/Mock/ChunkStore.mock.h
--- a/Engine/Source/Runtime/Online/BuildPatchServices/Private/Tests/Mock/ChunkStore.mock.h
+++ b/Engine/Source/Runtime/Online/BuildPatchServices/Private/Tests/Mock/ChunkStore.mock.h
@@ -41,6 +41,11 @@ namespace BuildPatchServices
 			return INDEX_NONE;
 		}
 
+		virtual void SetLostChunkCallback(TFunction<void(const FGuid&)> Callback) override
+		{
+			MOCK_FUNC_NOT_IMPLEMENTED("FMockChunkStore::SetLostChunkCallback");
+		}
+
 	public:
 		mutable TArray<FPut> RxPut;
 		mutable TArray<FGet> RxGet;
diff --git a/Engine/Source/Runtime/Online/BuildPatchServices/Private/Tests/Mock/CloudChunkSourceStat.mock.h b/Engine/Source/Runtime/Online/BuildPatchServices/Private/Tests/Mock/CloudChunkSourceStat.mock.h
--- a/Engine/Source/Runtime/Online/BuildPatchServices/Private/Tests/Mock/CloudChunkSourceStat.mock.h
+++ b/Engine/Source/Runtime/Online/BuildPatchServices/Private/Tests/Mock/CloudChunkSourceStat.mock.h
@@ -13,6 +13,7 @@ namespace BuildPatchServices
 	{
 	public:
 		typedef TTuple<double, FGuid> FDownloadRequested;
+		typedef TTuple<double, FGuid> FDownloadSuccess;
 		typedef TTuple<double, FGuid, FString> FDownloadFailed;
 		typedef TTuple<double, FGuid, FString, EChunkLoadResult> FDownloadCorrupt;
 		typedef TTuple<double, FGuid, FString, double, double, double, double> FDownloadAborted;
@@ -28,6 +29,11 @@ namespace BuildPatchServices
 			RxDownloadRequested.Emplace(FStatsCollector::GetSeconds(), ChunkId);
 		}
 
+		virtual void OnDownloadSuccess(const FGuid& ChunkId) override
+		{
+			RxDownloadSuccess.Emplace(FStatsCollector::GetSeconds(), ChunkId);
+		}
+
 		virtual void OnDownloadFailed(const FGuid& ChunkId, const FString& Url) override
 		{
 			RxDownloadFailed.Emplace(FStatsCollector::GetSeconds(), ChunkId, Url);
@@ -68,8 +74,13 @@ namespace BuildPatchServices
 			RxActiveRequestCountUpdated.Emplace(FStatsCollector::GetSeconds(), RequestCount);
 		}
 
+		virtual void OnAcceptedNewRequirements(const TSet<FGuid>& ChunkIds) override
+		{
+		}
+
 	public:
 		TArray<FDownloadRequested> RxDownloadRequested;
+		TArray<FDownloadSuccess> RxDownloadSuccess;
 		TArray<FDownloadFailed> RxDownloadFailed;
 		TArray<FDownloadCorrupt> RxDownloadCorrupt;
 		TArray<FDownloadAborted> RxDownloadAborted;
diff --git a/Engine/Source/Runtime/Online/BuildPatchServices/Private/Tests/Mock/DiskChunkStoreStat.mock.h b/Engine/Source/Runtime/Online/BuildPatchServices/Private/Tests/Mock/DiskChunkStoreStat.mock.h
--- a/Engine/Source/Runtime/Online/BuildPatchServices/Private/Tests/Mock/DiskChunkStoreStat.mock.h
+++ b/Engine/Source/Runtime/Online/BuildPatchServices/Private/Tests/Mock/DiskChunkStoreStat.mock.h
@@ -12,6 +12,7 @@ namespace BuildPatchServices
 	{
 	public:
 		typedef TTuple<double, FGuid, FString, EChunkSaveResult> FChunkStored;
+		typedef TTuple<double, FGuid> FBeforeChunkLoad;
 		typedef TTuple<double, FGuid, FString, EChunkLoadResult> FChunkLoaded;
 		typedef TTuple<double, int32> FCacheUseUpdated;
 
@@ -21,6 +22,11 @@ namespace BuildPatchServices
 			RxChunkStored.Emplace(FStatsCollector::GetSeconds(), ChunkId, ChunkFilename, SaveResult);
 		}
 
+		virtual void OnBeforeChunkLoad(const FGuid& ChunkId) override
+		{
+			RxBeforeChunkLoad.Emplace(FStatsCollector::GetSeconds(), ChunkId);
+		}
+
 		virtual void OnChunkLoaded(const FGuid& ChunkId, const FString& ChunkFilename, EChunkLoadResult LoadResult) override
 		{
 			RxChunkLoaded.Emplace(FStatsCollector::GetSeconds(), ChunkId, ChunkFilename, LoadResult);
@@ -33,6 +39,7 @@ namespace BuildPatchServices
 
 	public:
 		TArray<FChunkStored> RxChunkStored;
+		TArray<FBeforeChunkLoad> RxBeforeChunkLoad;
 		TArray<FChunkLoaded> RxChunkLoaded;
 		TArray<FCacheUseUpdated> RxCacheUseUpdated;
 	};
diff --git a/Engine/Source/Runtime/Online/BuildPatchServices/Private/Tests/Mock/DownloadServiceStat.mock.h b/Engine/Source/Runtime/Online/BuildPatchServices/Private/Tests/Mock/DownloadServiceStat.mock.h
--- a/Engine/Source/Runtime/Online/BuildPatchServices/Private/Tests/Mock/DownloadServiceStat.mock.h
+++ b/Engine/Source/Runtime/Online/BuildPatchServices/Private/Tests/Mock/DownloadServiceStat.mock.h
@@ -11,15 +11,29 @@ namespace BuildPatchServices
 		: public IDownloadServiceStat
 	{
 	public:
+		typedef TTuple<double, int32, FString> FDownloadStarted;
+		typedef TTuple<double, int32, int32> FDownloadProgress;
 		typedef TTuple<double, FDownloadRecord> FDownloadComplete;
 
 	public:
-		virtual void OnDownloadComplete(FDownloadRecord DownloadRecord) override
+		virtual void OnDownloadStarted(int32 RequestId, const FString& Uri) override
+		{
+			RxDownloadStarted.Emplace(FStatsCollector::GetSeconds(), RequestId, Uri);
+		}
+
+		virtual void OnDownloadProgress(int32 RequestId, int32 BytesReceived) override
+		{
+			RxDownloadProgress.Emplace(FStatsCollector::GetSeconds(), RequestId, BytesReceived);
+		}
+
+		virtual void OnDownloadComplete(const FDownloadRecord& DownloadRecord) override
 		{
 			RxDownloadComplete.Emplace(FStatsCollector::GetSeconds(), DownloadRecord);
 		}
 
 	public:
+		TArray<FDownloadStarted> RxDownloadStarted;
+		TArray<FDownloadProgress> RxDownloadProgress;
 		TArray<FDownloadComplete> RxDownloadComplete;
 	};
 }
diff --git a/Engine/Source/Runtime/Online/BuildPatchServices/Private/Tests/Mock/FileSystem.mock.h b/Engine/Source/Runtime/Online/BuildPatchServices/Private/Tests/Mock/FileSystem.mock.h
--- a/Engine/Source/Runtime/Online/BuildPatchServices/Private/Tests/Mock/FileSystem.mock.h
+++ b/Engine/Source/Runtime/Online/BuildPatchServices/Private/Tests/Mock/FileSystem.mock.h
@@ -16,16 +16,16 @@ namespace BuildPatchServices
 		: public IFileSystem
 	{
 	public:
-		typedef TTuple<double, FArchive*, FString, EFileRead> FCreateFileReader;
-		typedef TTuple<double, FArchive*, FString, EFileWrite> FCreateFileWriter;
+		typedef TTuple<double, FArchive*, FString, EReadFlags> FCreateFileReader;
+		typedef TTuple<double, FArchive*, FString, EWriteFlags> FCreateFileWriter;
 		typedef TTuple<double, FString, int64> FGetFileSize;
-		typedef TTuple<double, FString, EFileAttributes> FGetFileAttributes;
+		typedef TTuple<double, FString, EAttributeFlags> FGetAttributes;
 		typedef TTuple<double, FString, bool> FSetReadOnly;
 		typedef TTuple<double, FString, bool> FSetCompressed;
 		typedef TTuple<double, FString, bool> FSetExecutable;
 
 	public:
-		virtual TUniquePtr<FArchive> CreateFileReader(const TCHAR* Filename, EFileRead ReadFlags = EFileRead::None) const override
+		virtual TUniquePtr<FArchive> CreateFileReader(const TCHAR* Filename, EReadFlags ReadFlags = EReadFlags::None) const override
 		{
 			FScopeLock ScopeLock(&ThreadLock);
 			TUniquePtr<FArchive> Reader(new FMemoryReader(ReadFile));
@@ -33,7 +33,7 @@ namespace BuildPatchServices
 			return Reader;
 		}
 
-		virtual TUniquePtr<FArchive> CreateFileWriter(const TCHAR* Filename, EFileWrite WriteFlags = EFileWrite::None) const override
+		virtual TUniquePtr<FArchive> CreateFileWriter(const TCHAR* Filename, EWriteFlags WriteFlags = EWriteFlags::None) const override
 		{
 			FScopeLock ScopeLock(&ThreadLock);
 			TUniquePtr<FArchive> Writer(new FMemoryWriter(WriteFile));
@@ -53,25 +53,31 @@ namespace BuildPatchServices
 			return false;
 		}
 
+		bool CopyFile(const TCHAR* FileDest, const TCHAR* FileSource) const
+		{
+			MOCK_FUNC_NOT_IMPLEMENTED("FMockFileSystem::CopyFile");
+			return false;
+		}
+
 		virtual bool GetFileSize(const TCHAR* Filename, int64& OutFileSize) const override
 		{
-			OutFileSize = INDEX_NONE;
+			OutFileSize = -1;
 			if (FileSizes.Contains(Filename))
 			{
 				OutFileSize = FileSizes[Filename];
 			}
 			RxGetFileSize.Emplace(FStatsCollector::GetSeconds(), Filename, OutFileSize);
-			return true;
+			return OutFileSize >= 0;
 		}
 
-		virtual bool GetFileAttributes(const TCHAR* Filename, EFileAttributes& OutFileAttributes) const override
+		virtual bool GetAttributes(const TCHAR* Filename, EAttributeFlags& OutFileAttributes) const override
 		{
-			OutFileAttributes = EFileAttributes::None;
+			OutFileAttributes = EAttributeFlags::None;
 			if (FileAttributes.Contains(Filename))
 			{
 				OutFileAttributes = FileAttributes[Filename];
 			}
-			RxGetFileAttributes.Emplace(FStatsCollector::GetSeconds(), Filename, OutFileAttributes);
+			RxGetAttributes.Emplace(FStatsCollector::GetSeconds(), Filename, OutFileAttributes);
 			return true;
 		}
 
@@ -93,11 +99,17 @@ namespace BuildPatchServices
 			return true;
 		}
 
+		virtual bool FileExists(const TCHAR* Filename) const override
+		{
+			MOCK_FUNC_NOT_IMPLEMENTED("FMockFileSystem::FileExists");
+			return true;
+		}
+
 	public:
 		mutable FCriticalSection ThreadLock;
 		mutable TArray<FCreateFileReader> RxCreateFileReader;
 		mutable TArray<FCreateFileWriter> RxCreateFileWriter;
-		mutable TArray<FGetFileAttributes> RxGetFileAttributes;
+		mutable TArray<FGetAttributes> RxGetAttributes;
 		mutable TArray<FGetFileSize> RxGetFileSize;
 		mutable TArray<FSetReadOnly> RxSetReadOnly;
 		mutable TArray<FSetCompressed> RxSetCompressed;
@@ -105,7 +117,7 @@ namespace BuildPatchServices
 		mutable TArray<uint8> ReadFile;
 		mutable TArray<uint8> WriteFile;
 		mutable TMap<FString, int64> FileSizes;
-		mutable TMap<FString, EFileAttributes> FileAttributes;
+		mutable TMap<FString, EAttributeFlags> FileAttributes;
 	};
 }
 
diff --git a/Engine/Source/Runtime/Online/BuildPatchServices/Private/Tests/Mock/InstallChunkSourceStat.mock.h b/Engine/Source/Runtime/Online/BuildPatchServices/Private/Tests/Mock/InstallChunkSourceStat.mock.h
--- a/Engine/Source/Runtime/Online/BuildPatchServices/Private/Tests/Mock/InstallChunkSourceStat.mock.h
+++ b/Engine/Source/Runtime/Online/BuildPatchServices/Private/Tests/Mock/InstallChunkSourceStat.mock.h
@@ -12,33 +12,49 @@ namespace BuildPatchServices
 		: public IInstallChunkSourceStat
 	{
 	public:
+		typedef TTuple<double, TArray<FGuid>> FBatchStarted;
 		typedef TTuple<double, FGuid> FLoadStarted;
-		typedef TTuple<double, FGuid, ELoadResult> FLoadComplete;
+		typedef TTuple<double, FGuid, ELoadResult, ISpeedRecorder::FRecord> FLoadComplete;
 
 	public:
+		virtual void OnBatchStarted(const TArray<FGuid>& ChunkIds) override
+		{
+			if (OnBatchStartedFunc != nullptr)
+			{
+				OnBatchStartedFunc(ChunkIds);
+			}
+			RxBatchStarted.Emplace(FStatsCollector::GetSeconds(), ChunkIds);
+		}
+
 		virtual void OnLoadStarted(const FGuid& ChunkId) override
 		{
-			if (OnLoadStartedFunc)
+			if (OnLoadStartedFunc != nullptr)
 			{
 				OnLoadStartedFunc(ChunkId);
 			}
 			RxLoadStarted.Emplace(FStatsCollector::GetSeconds(), ChunkId);
 		}
 
-		virtual void OnLoadComplete(const FGuid& ChunkId, ELoadResult Result) override
+		virtual void OnLoadComplete(const FGuid& ChunkId, const ELoadResult& Result, const ISpeedRecorder::FRecord& Record) override
 		{
-			if (OnLoadCompleteFunc)
+			if (OnLoadCompleteFunc != nullptr)
 			{
-				OnLoadCompleteFunc(ChunkId, Result);
+				OnLoadCompleteFunc(ChunkId, Result, Record);
 			}
-			RxLoadComplete.Emplace(FStatsCollector::GetSeconds(), ChunkId, Result);
+			RxLoadComplete.Emplace(FStatsCollector::GetSeconds(), ChunkId, Result, Record);
+		}
+
+		virtual void OnAcceptedNewRequirements(const TSet<FGuid>& ChunkIds) override
+		{
 		}
 
 	public:
+		TArray<FBatchStarted> RxBatchStarted;
 		TArray<FLoadStarted> RxLoadStarted;
 		TArray<FLoadComplete> RxLoadComplete;
+		TFunction<void(const TArray<FGuid>&)> OnBatchStartedFunc;
 		TFunction<void(const FGuid&)> OnLoadStartedFunc;
-		TFunction<void(const FGuid&, ELoadResult)> OnLoadCompleteFunc;
+		TFunction<void(const FGuid&, const ELoadResult&, const ISpeedRecorder::FRecord&)> OnLoadCompleteFunc;
 	};
 }
 
diff --git a/Engine/Source/Runtime/Online/BuildPatchServices/Private/Tests/Mock/InstallerAnalytics.mock.h b/Engine/Source/Runtime/Online/BuildPatchServices/Private/Tests/Mock/InstallerAnalytics.mock.h
--- a/Engine/Source/Runtime/Online/BuildPatchServices/Private/Tests/Mock/InstallerAnalytics.mock.h
+++ b/Engine/Source/Runtime/Online/BuildPatchServices/Private/Tests/Mock/InstallerAnalytics.mock.h
@@ -49,6 +49,11 @@ namespace BuildPatchServices
 			RxTrackRequest.Emplace(FStatsCollector::GetSeconds(), Request);
 		}
 
+		virtual void Flush() override
+		{
+			++RxFlush;
+		}
+
 	public:
 		TArray<FRecordChunkDownloadError> RxRecordChunkDownloadError;
 		TArray<FRecordChunkDownloadAborted> RxRecordChunkDownloadAborted;
@@ -56,6 +61,7 @@ namespace BuildPatchServices
 		TArray<FRecordConstructionError> RxRecordConstructionError;
 		TArray<FRecordPrereqInstallationError> RxRecordPrereqInstallationError;
 		TArray<FTrackRequest> RxTrackRequest;
+		uint64 RxFlush;
 	};
 }
 
diff --git a/Engine/Source/Runtime/Online/BuildPatchServices/Private/Tests/Mock/Manifest.mock.h b/Engine/Source/Runtime/Online/BuildPatchServices/Private/Tests/Mock/Manifest.mock.h
--- a/Engine/Source/Runtime/Online/BuildPatchServices/Private/Tests/Mock/Manifest.mock.h
+++ b/Engine/Source/Runtime/Online/BuildPatchServices/Private/Tests/Mock/Manifest.mock.h
@@ -80,7 +80,16 @@ namespace BuildPatchServices
 
 		virtual TArray<FString> GetBuildFileList() const override
 		{
-			return BuildFileList;
+			TArray<FString> Filenames;
+			GetFileList(Filenames);
+			return Filenames;
+		}
+
+		virtual TArray<FString> GetBuildFileList(const TSet<FString>& Tags) const override
+		{
+			TArray<FString> Filenames;
+			GetTaggedFileList(Tags, Filenames);
+			return Filenames;
 		}
 
 		virtual void GetFileTagList(TSet<FString>& Tags) const override
@@ -226,6 +235,11 @@ namespace BuildPatchServices
 			TaggedFiles = TaggedFileList;
 		}
 
+		virtual void GetTaggedFileList(const TSet<FString>& Tags, TArray<FString>& TaggedFiles) const override
+		{
+			TaggedFiles = TaggedFileList.Array();
+		}
+
 		virtual void GetDataList(TArray<FGuid>& DataGuids) const override
 		{
 			DataGuids = DataList;
diff --git a/Engine/Source/Runtime/Online/BuildPatchServices/Private/Tests/Mock/MemoryChunkStoreStat.mock.h b/Engine/Source/Runtime/Online/BuildPatchServices/Private/Tests/Mock/MemoryChunkStoreStat.mock.h
--- a/Engine/Source/Runtime/Online/BuildPatchServices/Private/Tests/Mock/MemoryChunkStoreStat.mock.h
+++ b/Engine/Source/Runtime/Online/BuildPatchServices/Private/Tests/Mock/MemoryChunkStoreStat.mock.h
@@ -15,6 +15,7 @@ namespace BuildPatchServices
 		typedef TTuple<double, FGuid> FChunkReleased;
 		typedef TTuple<double, FGuid> FChunkBooted;
 		typedef TTuple<double, int32> FStoreUseUpdated;
+		typedef TTuple<double, int32> FStoreSizeUpdated;
 
 	public:
 		virtual void OnChunkStored(const FGuid& ChunkId) override
@@ -36,11 +37,18 @@ namespace BuildPatchServices
 		{
 			RxStoreUseUpdated.Emplace(FStatsCollector::GetSeconds(), ChunkCount);
 		}
+
+		virtual void OnStoreSizeUpdated(int32 Size) override
+		{
+			RxStoreSizeUpdated.Emplace(FStatsCollector::GetSeconds(), Size);
+		}
+
 	public:
 		TArray<FChunkStored> RxChunkStored;
 		TArray<FChunkReleased> RxChunkReleased;
 		TArray<FChunkBooted> RxChunkBooted;
 		TArray<FStoreUseUpdated> RxStoreUseUpdated;
+		TArray<FStoreSizeUpdated> RxStoreSizeUpdated;
 	};
 }
 
diff --git a/Engine/Source/Runtime/Online/BuildPatchServices/Private/Tests/Mock/MessagePump.mock.h b/Engine/Source/Runtime/Online/BuildPatchServices/Private/Tests/Mock/MessagePump.mock.h
--- a/Engine/Source/Runtime/Online/BuildPatchServices/Private/Tests/Mock/MessagePump.mock.h
+++ b/Engine/Source/Runtime/Online/BuildPatchServices/Private/Tests/Mock/MessagePump.mock.h
@@ -15,6 +15,10 @@ namespace BuildPatchServices
 		{
 		}
 
+		virtual void SendMessage(FInstallationFileAction Message) override
+		{
+		}
+
 		virtual void PumpMessages(const TArray<FMessageHandler*>& Handlers) override
 		{
 		}
diff --git a/Engine/Source/Runtime/Online/BuildPatchServices/Private/Tests/Mock/VerifierStat.mock.h b/Engine/Source/Runtime/Online/BuildPatchServices/Private/Tests/Mock/VerifierStat.mock.h
--- a/Engine/Source/Runtime/Online/BuildPatchServices/Private/Tests/Mock/VerifierStat.mock.h
+++ b/Engine/Source/Runtime/Online/BuildPatchServices/Private/Tests/Mock/VerifierStat.mock.h
@@ -14,7 +14,8 @@ namespace BuildPatchServices
 	public:
 		typedef TTuple<double, FString, int64> FOnFileStarted;
 		typedef TTuple<double, FString, int64> FOnFileProgress;
-		typedef TTuple<double, FString, bool> FOnFileCompleted;
+		typedef TTuple<double, FString, EVerifyResult> FOnFileCompleted;
+		typedef TTuple<double, ISpeedRecorder::FRecord> FOnFileRead;
 		typedef TTuple<double, int64> FOnProcessedDataUpdated;
 		typedef TTuple<double, int64> FOnTotalRequiredUpdated;
 
@@ -33,13 +34,18 @@ namespace BuildPatchServices
 			RxOnFileProgress.Emplace(FStatsCollector::GetSeconds(), Filename, TotalBytes);
 		}
 
-		virtual void OnFileCompleted(const FString& Filename, bool bSuccess) override
+		virtual void OnFileCompleted(const FString& Filename, EVerifyResult VerifyResult) override
 		{
 			if (OnFileCompletedFunc)
 			{
-				OnFileCompletedFunc(Filename, bSuccess);
+				OnFileCompletedFunc(Filename, VerifyResult);
 			}
-			RxOnFileCompleted.Emplace(FStatsCollector::GetSeconds(), Filename, bSuccess);
+			RxOnFileCompleted.Emplace(FStatsCollector::GetSeconds(), Filename, VerifyResult);
+		}
+
+		virtual void OnFileRead(const ISpeedRecorder::FRecord& Record) override
+		{
+			RxOnFileRead.Emplace(FStatsCollector::GetSeconds(), Record);
 		}
 
 		virtual void OnProcessedDataUpdated(int64 TotalBytes) override
@@ -56,10 +62,11 @@ namespace BuildPatchServices
 		TArray<FOnFileStarted> RxOnFileStarted;
 		TArray<FOnFileProgress> RxOnFileProgress;
 		TArray<FOnFileCompleted> RxOnFileCompleted;
+		TArray<FOnFileRead> RxOnFileRead;
 		TArray<FOnProcessedDataUpdated> RxOnProcessedDataUpdated;
 		TArray<FOnTotalRequiredUpdated> RxOnTotalRequiredUpdated;
 		TFunction<void(const FString&, int64)> OnFileProgressFunc;
-		TFunction<void(const FString&, bool)> OnFileCompletedFunc;
+		TFunction<void(const FString&, EVerifyResult)> OnFileCompletedFunc;
 	};
 }
 
diff --git a/Engine/Source/Runtime/Online/BuildPatchServices/Private/Tests/TestHelpers.h b/Engine/Source/Runtime/Online/BuildPatchServices/Private/Tests/TestHelpers.h
--- a/Engine/Source/Runtime/Online/BuildPatchServices/Private/Tests/TestHelpers.h
+++ b/Engine/Source/Runtime/Online/BuildPatchServices/Private/Tests/TestHelpers.h
@@ -2,20 +2,63 @@
 #pragma once
 
 #include "BuildPatchServicesPrivate.h"
+#include "HAL/PlatformTime.h"
+#include "HAL/PlatformProcess.h"
+
+namespace BuildPatchTestHelpers
+{
+	template<typename T1, typename T2>
+	inline bool TestEqual(const FString& Description, T1 Expression, T2 Expected, FAutomationSpecBase* This)
+	{
+		This->TestEqual(Description, Expression, Expected);
+		return Expression == Expected;
+	}
+	
+	template<typename T1, typename T2>
+	inline bool TestNotEqual(const FString& Description, T1 Expression, T2 Expected, FAutomationSpecBase* This)
+	{
+		This->TestNotEqual(Description, Expression, Expected);
+		return Expression != Expected;
+	}
+
+	template<typename T>
+	inline bool TestNull(const FString& Description, T Expression, FAutomationSpecBase* This)
+	{
+		This->TestNull(Description, Expression);
+		return Expression == nullptr;
+	}
+
+	template<typename T>
+	inline bool TestNotNull(const FString& Description, T Expression, FAutomationSpecBase* This)
+	{
+		This->TestNotNull(Description, Expression);
+		return Expression != nullptr;
+	}
+
+	inline bool WaitUntilTrue(const TFunction<bool()>& Pred, const TFunction<bool()>& TestResult, double TimeoutSeconds)
+	{
+		const double TimeoutEnd = FPlatformTime::Seconds() + TimeoutSeconds;
+		while (!Pred() && FPlatformTime::Seconds() < TimeoutEnd)
+		{
+			FPlatformProcess::Sleep(0);
+		}
+		return TestResult();
+	}
+}
 
 #if WITH_DEV_AUTOMATION_TESTS
 
 #define _TEST_EQUAL(text, expression, expected) \
-	TestEqual(text, expression, expected)
+	BuildPatchTestHelpers::TestEqual(text, expression, expected, this)
 
 #define _TEST_NOT_EQUAL(text, expression, expected) \
-	TestNotEqual(text, expression, expected)
+	BuildPatchTestHelpers::TestNotEqual(text, expression, expected, this)
 
 #define _TEST_NULL(text, expression) \
-	TestNull(text, expression)
+	BuildPatchTestHelpers::TestNull(text, expression, this)
 
 #define _TEST_NOT_NULL(text, expression) \
-	TestNotNull(text, expression)
+	BuildPatchTestHelpers::TestNotNull(text, expression, this)
 
 #define TEST_EQUAL(expression, expected) \
 	_TEST_EQUAL(TEXT(#expression), expression, expected)
@@ -35,6 +78,9 @@
 #define TEST_NOT_NULL(expression) \
 	_TEST_NOT_NULL(TEXT(#expression), expression)
 
+#define TEST_BECOMES_TRUE(expression, timeout) \
+	BuildPatchTestHelpers::WaitUntilTrue([this](){ return expression; }, [this](){ return TEST_EQUAL(expression, true); }, timeout)
+
 #define MOCK_FUNC_NOT_IMPLEMENTED(funcname) \
 	UE_LOG(LogBuildPatchServices, Error, TEXT(funcname) TEXT(": Called but there is no implementation."))
 
diff --git a/Engine/Source/Runtime/Online/BuildPatchServices/Private/Tests/Unit/DiskChunkStore.spec.cpp b/Engine/Source/Runtime/Online/BuildPatchServices/Private/Tests/Unit/DiskChunkStore.spec.cpp
--- a/Engine/Source/Runtime/Online/BuildPatchServices/Private/Tests/Unit/DiskChunkStore.spec.cpp
+++ b/Engine/Source/Runtime/Online/BuildPatchServices/Private/Tests/Unit/DiskChunkStore.spec.cpp
@@ -2,6 +2,7 @@
 
 #include "Misc/AutomationTest.h"
 #include "Tests/TestHelpers.h"
+#include "Tests/Fake/FileSystem.fake.h"
 #include "Tests/Fake/ChunkDataAccess.fake.h"
 #include "Tests/Mock/ChunkDataSerialization.mock.h"
 #include "Tests/Mock/DiskChunkStoreStat.mock.h"
@@ -14,14 +15,25 @@ BEGIN_DEFINE_SPEC(FDiskChunkStoreSpec, "BuildPatchServices.Unit", EAutomationTes
 // Unit
 TUniquePtr<BuildPatchServices::IDiskChunkStore> DiskChunkStore;
 // Mock
+TUniquePtr<BuildPatchServices::FFakeFileSystem> FakeFileSystem;
 TUniquePtr<BuildPatchServices::FMockChunkDataSerialization> MockChunkDataSerialization;
 TUniquePtr<BuildPatchServices::FMockDiskChunkStoreStat> MockDiskChunkStoreStat;
-TUniquePtr<BuildPatchServices::FFakeChunkDataAccess> MockChunkDataAccess;
+TUniquePtr<BuildPatchServices::FFakeChunkDataAccess> FakeChunkDataAccessOne;
+TUniquePtr<BuildPatchServices::FFakeChunkDataAccess> FakeChunkDataAccessTwo;
 // Data
 FString StoreRootPath;
 FGuid SomeChunk;
+bool bChunkOneWasDeleted;
+bool bChunkTwoWasDeleted;
+TArray<uint8> SomeData;
+// Helpers
+void MakeChunkData();
+void MakeUnit();
+int32 DiskDataNum();
 END_DEFINE_SPEC(FDiskChunkStoreSpec)
 
+#define DISK_STORE_TEST_TIMEOUT 1.0
+
 void FDiskChunkStoreSpec::Define()
 {
 	using namespace BuildPatchServices;
@@ -30,38 +42,140 @@ void FDiskChunkStoreSpec::Define()
 	FRollingHashConst::Init();
 	StoreRootPath = TEXT("RootPath");
 	SomeChunk = FGuid::NewGuid();
+	SomeData.AddUninitialized(64);
 
 	// Specs.
 	BeforeEach([this]()
 	{
+		FakeFileSystem.Reset(new FFakeFileSystem());
 		MockChunkDataSerialization.Reset(new FMockChunkDataSerialization());
 		MockDiskChunkStoreStat.Reset(new FMockDiskChunkStoreStat());
-		MockChunkDataAccess.Reset(new FFakeChunkDataAccess());
-		DiskChunkStore.Reset(FDiskChunkStoreFactory::Create(
-			MockChunkDataSerialization.Get(),
-			MockDiskChunkStoreStat.Get(),
-			StoreRootPath));
+		FakeChunkDataAccessOne.Reset(new FFakeChunkDataAccess());
+		FakeChunkDataAccessTwo.Reset(new FFakeChunkDataAccess());
+		bChunkOneWasDeleted = false;
+		bChunkTwoWasDeleted = false;
+		FakeChunkDataAccessOne->OnDeleted = [this]()
+		{
+			bChunkOneWasDeleted = true;
+		};
+		FakeChunkDataAccessTwo->OnDeleted = [this]()
+		{
+			bChunkTwoWasDeleted = true;
+		};
+		MockChunkDataSerialization->SaveToArchiveFunc = [](FArchive& Ar, const IChunkDataAccess* ChunkPtr)
+		{
+			const FFakeChunkDataAccess* FakeChunkDataAccess = static_cast<const FFakeChunkDataAccess*>(ChunkPtr);
+			Ar.Serialize(FakeChunkDataAccess->ChunkData, FakeChunkDataAccess->ChunkHeader.DataSize);
+			return BuildPatchServices::EChunkSaveResult::Success;
+		};
+		MakeChunkData();
+		MakeUnit();
 	});
 
 	Describe("DiskChunkStore", [this]()
 	{
-		Describe("Put", [this]()
+		Describe("Construction", [this]()
 		{
-			It("should save some chunk to the provided store root path.", [this]()
+			It("should create a chunkdump file at provided path.", [this]()
 			{
-				DiskChunkStore->Put(SomeChunk, TUniquePtr<IChunkDataAccess>(new FFakeChunkDataAccess()));
-				TEST_EQUAL(MockChunkDataSerialization->RxSaveToFile.Num(), 1);
-				if (MockChunkDataSerialization->RxSaveToFile.Num() == 1)
+				if(TEST_BECOMES_TRUE(DiskDataNum() == 1, DISK_STORE_TEST_TIMEOUT))
 				{
-					TEST_TRUE(MockChunkDataSerialization->RxSaveToFile[0].Get<0>().StartsWith(StoreRootPath + TEXT("/")));
+					FScopeLock ScopeLock(&FakeFileSystem->ThreadLock);
+					TEST_TRUE(FakeFileSystem->DiskData.CreateConstIterator().Key().StartsWith(StoreRootPath / TEXT("")));
 				}
 			});
 
+			Describe("when there are errors opening the chunkdump", [this]()
+			{
+				BeforeEach([this]()
+				{
+					DiskChunkStore.Reset();
+					FakeFileSystem->CreateFileReaderFunc = [this](const TCHAR*, EReadFlags) -> TUniquePtr<FArchive>
+					{
+						if (FakeFileSystem->RxCreateFileReader.Num() == 10)
+						{
+							FakeFileSystem->CreateFileReaderFunc = nullptr;
+						}
+						return nullptr;
+					};
+					FakeFileSystem->CreateFileWriterFunc = [this](const TCHAR*, EWriteFlags) -> TUniquePtr<FArchive>
+					{
+						if (FakeFileSystem->RxCreateFileWriter.Num() == 10)
+						{
+							FakeFileSystem->CreateFileWriterFunc = nullptr;
+						}
+						return nullptr;
+					};
+					FakeFileSystem->RxCreateFileReader.Reset();
+					FakeFileSystem->RxCreateFileWriter.Reset();
+					MakeUnit();
+				});
+
+				It("should retry until successful.", [this]()
+				{
+					TEST_BECOMES_TRUE(FakeFileSystem->RxCreateFileWriter.Num() == 12, DISK_STORE_TEST_TIMEOUT);
+					TEST_BECOMES_TRUE(FakeFileSystem->RxCreateFileReader.Num() == 12, DISK_STORE_TEST_TIMEOUT);
+				});
+			});
+		});
+
+		Describe("Destruction", [this]()
+		{
+			It("should delete the chunkdump file created.", [this]()
+			{
+				DiskChunkStore.Reset();
+				TEST_TRUE(DiskDataNum() == 0);
+			});
+
+			Describe("when there are still queued requests", [this]()
+			{
+				BeforeEach([this]()
+				{
+					MockChunkDataSerialization->SaveToArchiveFunc = [](FArchive&, const IChunkDataAccess*)
+					{
+						FPlatformProcess::Sleep(0.5);
+						return BuildPatchServices::EChunkSaveResult::Success;
+					};
+					DiskChunkStore->Put(FGuid::NewGuid(), MoveTemp(FakeChunkDataAccessOne));
+					DiskChunkStore->Put(FGuid::NewGuid(), MoveTemp(FakeChunkDataAccessTwo));
+				});
+
+				It("should clean up all queued put memory.", [this]()
+				{
+					DiskChunkStore.Reset();
+					TEST_BECOMES_TRUE(bChunkOneWasDeleted, DISK_STORE_TEST_TIMEOUT);
+					TEST_BECOMES_TRUE(bChunkTwoWasDeleted, DISK_STORE_TEST_TIMEOUT);
+				});
+			});
+		});
+
+		Describe("Put", [this]()
+		{
+			It("should release chunk data once saved.", [this]()
+			{
+				DiskChunkStore->Put(SomeChunk, MoveTemp(FakeChunkDataAccessOne));
+				TEST_BECOMES_TRUE(bChunkOneWasDeleted, DISK_STORE_TEST_TIMEOUT);
+			});
+
+			It("should save some chunk to the chunkdump.", [this]()
+			{
+				DiskChunkStore->Put(SomeChunk, MoveTemp(FakeChunkDataAccessOne));
+				TEST_BECOMES_TRUE(MockChunkDataSerialization->RxSaveToArchive.Num() == 1, DISK_STORE_TEST_TIMEOUT);
+			});
+
 			It("should not save some chunk that was previously saved.", [this]()
 			{
-				DiskChunkStore->Put(SomeChunk, TUniquePtr<IChunkDataAccess>(new FFakeChunkDataAccess()));
-				DiskChunkStore->Put(SomeChunk, TUniquePtr<IChunkDataAccess>(new FFakeChunkDataAccess()));
-				TEST_EQUAL(MockChunkDataSerialization->RxSaveToFile.Num(), 1);
+				DiskChunkStore->Put(SomeChunk, MoveTemp(FakeChunkDataAccessOne));
+				DiskChunkStore->Put(SomeChunk, MoveTemp(FakeChunkDataAccessTwo));
+				TEST_BECOMES_TRUE(bChunkOneWasDeleted, DISK_STORE_TEST_TIMEOUT);
+				TEST_BECOMES_TRUE(bChunkTwoWasDeleted, DISK_STORE_TEST_TIMEOUT);
+				TEST_EQUAL(MockChunkDataSerialization->RxSaveToArchive.Num(), 1);
+			});
+
+			It("should cause the reader to be reopened ready for a Get.", [this]()
+			{
+				DiskChunkStore->Put(SomeChunk, MoveTemp(FakeChunkDataAccessOne));
+				TEST_BECOMES_TRUE(FakeFileSystem->RxCreateFileReader.Num() == 2, DISK_STORE_TEST_TIMEOUT);
 			});
 		});
 
@@ -72,7 +186,7 @@ void FDiskChunkStoreSpec::Define()
 				It("should not attempt to load some chunk.", [this]()
 				{
 					TEST_NULL(DiskChunkStore->Get(SomeChunk));
-					TEST_EQUAL(MockChunkDataSerialization->RxLoadFromFile.Num(), 0);
+					TEST_EQUAL(MockChunkDataSerialization->RxLoadFromArchive.Num(), 0);
 				});
 			});
 
@@ -80,38 +194,40 @@ void FDiskChunkStoreSpec::Define()
 			{
 				BeforeEach([this]()
 				{
-					DiskChunkStore->Put(SomeChunk, TUniquePtr<IChunkDataAccess>(new FFakeChunkDataAccess()));
+					DiskChunkStore->Put(SomeChunk, MoveTemp(FakeChunkDataAccessOne));
 				});
 
-				It("should load some chunk from the provided store root path.", [this]()
+				It("should load some chunk from the chunkdump.", [this]()
 				{
 					DiskChunkStore->Get(SomeChunk);
-					TEST_EQUAL(MockChunkDataSerialization->RxLoadFromFile.Num(), 1);
-					if (MockChunkDataSerialization->RxLoadFromFile.Num() == 1)
-					{
-						TEST_TRUE(MockChunkDataSerialization->RxLoadFromFile[0].Get<0>().StartsWith(StoreRootPath + TEXT("/")));
-					}
+					TEST_EQUAL(MockChunkDataSerialization->RxLoadFromArchive.Num(), 1);
 				});
 
-				Describe("and LoadFromFile will be successful", [this]()
+				It("should enforce the reader to have been reopened.", [this]()
+				{
+					DiskChunkStore->Get(SomeChunk);
+					TEST_EQUAL(FakeFileSystem->RxCreateFileReader.Num(), 2);
+				});
+
+				Describe("and LoadFromArchive will be successful", [this]()
 				{
 					BeforeEach([this]()
 					{
-						MockChunkDataSerialization->TxLoadFromFile.Emplace(MockChunkDataAccess.Release(), EChunkLoadResult::Success);
+						MockChunkDataSerialization->TxLoadFromArchive.Emplace(FakeChunkDataAccessOne.Release(), EChunkLoadResult::Success);
 					});
 
 					It("should not load some chunk twice in a row.", [this]()
 					{
 						TEST_EQUAL(DiskChunkStore->Get(SomeChunk), DiskChunkStore->Get(SomeChunk));
-						TEST_EQUAL(MockChunkDataSerialization->RxLoadFromFile.Num(), 1);
+						TEST_EQUAL(MockChunkDataSerialization->RxLoadFromArchive.Num(), 1);
 					});
 				});
 
-				Describe("and LoadFromFile will not be successful", [this]()
+				Describe("and LoadFromArchive will not be successful", [this]()
 				{
 					BeforeEach([this]()
 					{
-						MockChunkDataSerialization->TxLoadFromFile.Emplace(nullptr, EChunkLoadResult::SerializationError);
+						MockChunkDataSerialization->TxLoadFromArchive.Emplace(nullptr, EChunkLoadResult::SerializationError);
 					});
 
 					It("should return nullptr.", [this]()
@@ -123,7 +239,7 @@ void FDiskChunkStoreSpec::Define()
 					{
 						DiskChunkStore->Get(SomeChunk);
 						DiskChunkStore->Get(SomeChunk);
-						TEST_EQUAL(MockChunkDataSerialization->RxLoadFromFile.Num(), 1);
+						TEST_EQUAL(MockChunkDataSerialization->RxLoadFromArchive.Num(), 1);
 					});
 				});
 			});
@@ -137,7 +253,7 @@ void FDiskChunkStoreSpec::Define()
 				{
 					TUniquePtr<IChunkDataAccess> Removed = DiskChunkStore->Remove(SomeChunk);
 					TEST_FALSE(Removed.IsValid());
-					TEST_EQUAL(MockChunkDataSerialization->RxLoadFromFile.Num(), 0);
+					TEST_EQUAL(MockChunkDataSerialization->RxLoadFromArchive.Num(), 0);
 				});
 			});
 
@@ -148,22 +264,18 @@ void FDiskChunkStoreSpec::Define()
 					DiskChunkStore->Put(SomeChunk, TUniquePtr<IChunkDataAccess>(new FFakeChunkDataAccess()));
 				});
 
-				It("should load the chunk from the provided store root path.", [this]()
+				It("should load some chunk from the chunkdump.", [this]()
 				{
 					DiskChunkStore->Remove(SomeChunk);
-					TEST_EQUAL(MockChunkDataSerialization->RxLoadFromFile.Num(), 1);
-					if (MockChunkDataSerialization->RxLoadFromFile.Num() == 1)
-					{
-						TEST_TRUE(MockChunkDataSerialization->RxLoadFromFile[0].Get<0>().StartsWith(StoreRootPath + TEXT("/")));
-					}
+					TEST_EQUAL(MockChunkDataSerialization->RxLoadFromArchive.Num(), 1);
 				});
 
-				Describe("and LoadFromFile will be successful", [this]()
+				Describe("and LoadFromArchive will be successful", [this]()
 				{
 					BeforeEach([this]()
 					{
-						MockChunkDataSerialization->TxLoadFromFile.Emplace(MockChunkDataAccess.Release(), EChunkLoadResult::Success);
-						MockChunkDataAccess.Reset(new FFakeChunkDataAccess());
+						MockChunkDataSerialization->TxLoadFromArchive.Emplace(FakeChunkDataAccessOne.Release(), EChunkLoadResult::Success);
+						FakeChunkDataAccessOne.Reset(new FFakeChunkDataAccess());
 					});
 
 					Describe("and when some chunk was last used with Get", [this]()
@@ -171,14 +283,14 @@ void FDiskChunkStoreSpec::Define()
 						BeforeEach([this]()
 						{
 							DiskChunkStore->Get(SomeChunk);
-							MockChunkDataSerialization->RxLoadFromFile.Empty();
+							MockChunkDataSerialization->RxLoadFromArchive.Empty();
 						});
 
 						It("should return some chunk without loading it.", [this]()
 						{
 							TUniquePtr<IChunkDataAccess> Removed = DiskChunkStore->Remove(SomeChunk);
 							TEST_TRUE(Removed.IsValid());
-							TEST_EQUAL(MockChunkDataSerialization->RxLoadFromFile.Num(), 0);
+							TEST_EQUAL(MockChunkDataSerialization->RxLoadFromArchive.Num(), 0);
 						});
 					});
 
@@ -187,25 +299,25 @@ void FDiskChunkStoreSpec::Define()
 						BeforeEach([this]()
 						{
 							DiskChunkStore->Remove(SomeChunk);
-							MockChunkDataSerialization->TxLoadFromFile.Emplace(MockChunkDataAccess.Release(), EChunkLoadResult::Success);
-							MockChunkDataAccess.Reset(new FFakeChunkDataAccess());
-							MockChunkDataSerialization->RxLoadFromFile.Empty();
+							MockChunkDataSerialization->TxLoadFromArchive.Emplace(FakeChunkDataAccessOne.Release(), EChunkLoadResult::Success);
+							FakeChunkDataAccessOne.Reset(new FFakeChunkDataAccess());
+							MockChunkDataSerialization->RxLoadFromArchive.Empty();
 						});
 
 						It("should need to reload some chunk.", [this]()
 						{
 							TUniquePtr<IChunkDataAccess> Removed = DiskChunkStore->Remove(SomeChunk);
 							TEST_TRUE(Removed.IsValid());
-							TEST_EQUAL(MockChunkDataSerialization->RxLoadFromFile.Num(), 1);
+							TEST_EQUAL(MockChunkDataSerialization->RxLoadFromArchive.Num(), 1);
 						});
 					});
 				});
 
-				Describe("and LoadFromFile will not be successful", [this]()
+				Describe("and LoadFromArchive will not be successful", [this]()
 				{
 					BeforeEach([this]()
 					{
-						MockChunkDataSerialization->TxLoadFromFile.Emplace(nullptr, EChunkLoadResult::SerializationError);
+						MockChunkDataSerialization->TxLoadFromArchive.Emplace(nullptr, EChunkLoadResult::SerializationError);
 					});
 
 					It("should return invalid ptr.", [this]()
@@ -217,7 +329,7 @@ void FDiskChunkStoreSpec::Define()
 					{
 						DiskChunkStore->Remove(SomeChunk);
 						DiskChunkStore->Remove(SomeChunk);
-						TEST_EQUAL(MockChunkDataSerialization->RxLoadFromFile.Num(), 1);
+						TEST_EQUAL(MockChunkDataSerialization->RxLoadFromArchive.Num(), 1);
 					});
 				});
 			});
@@ -235,25 +347,42 @@ void FDiskChunkStoreSpec::Define()
 				TEST_EQUAL(DiskChunkStore->GetSlack(), MAX_int32);
 			});
 		});
-
-		It("should use the same filename for Put and Get.", [this]()
-		{
-			FGuid ChunkId = FGuid::NewGuid();
-			DiskChunkStore->Put(ChunkId, TUniquePtr<IChunkDataAccess>(new FFakeChunkDataAccess()));
-			DiskChunkStore->Get(ChunkId);
-			TEST_EQUAL(MockChunkDataSerialization->RxSaveToFile.Num(), 1);
-			TEST_EQUAL(MockChunkDataSerialization->RxLoadFromFile.Num(), 1);
-			if (MockChunkDataSerialization->RxSaveToFile.Num() == 1 && MockChunkDataSerialization->RxLoadFromFile.Num() == 1)
-			{
-				TEST_EQUAL(MockChunkDataSerialization->RxSaveToFile[0].Get<0>(), MockChunkDataSerialization->RxLoadFromFile[0].Get<0>());
-			}
-		});
 	});
 
 	AfterEach([this]()
 	{
 		DiskChunkStore.Reset();
+		FakeChunkDataAccessOne.Reset();
+		FakeChunkDataAccessTwo.Reset();
+		MockChunkDataSerialization.Reset();
+		FakeFileSystem.Reset();
 	});
 }
 
+void FDiskChunkStoreSpec::MakeChunkData()
+{
+	FakeChunkDataAccessOne->ChunkData = SomeData.GetData();
+	FakeChunkDataAccessTwo->ChunkData = SomeData.GetData();
+	FakeChunkDataAccessOne->ChunkHeader.DataSize = SomeData.Num();
+	FakeChunkDataAccessTwo->ChunkHeader.DataSize = SomeData.Num();
+}
+
+void FDiskChunkStoreSpec::MakeUnit()
+{
+	using namespace BuildPatchServices;
+	FDiskChunkStoreConfig DiskChunkStoreConfig(StoreRootPath);
+	DiskChunkStoreConfig.MaxRetryTime = 0.01;
+	DiskChunkStore.Reset(FDiskChunkStoreFactory::Create(
+		FakeFileSystem.Get(),
+		MockChunkDataSerialization.Get(),
+		MockDiskChunkStoreStat.Get(),
+		MoveTemp(DiskChunkStoreConfig)));
+}
+
+int32 FDiskChunkStoreSpec::DiskDataNum()
+{
+	FScopeLock ScopeLock(&FakeFileSystem->ThreadLock);
+	return FakeFileSystem->DiskData.Num();
+}
+
 #endif //WITH_DEV_AUTOMATION_TESTS
diff --git a/Engine/Source/Runtime/Online/BuildPatchServices/Private/Tests/Unit/FileAttribution.spec.cpp b/Engine/Source/Runtime/Online/BuildPatchServices/Private/Tests/Unit/FileAttribution.spec.cpp
--- a/Engine/Source/Runtime/Online/BuildPatchServices/Private/Tests/Unit/FileAttribution.spec.cpp
+++ b/Engine/Source/Runtime/Online/BuildPatchServices/Private/Tests/Unit/FileAttribution.spec.cpp
@@ -88,7 +88,7 @@ void FFileAttributionSpec::Define()
 					for (const FString& File : NewFiles)
 					{
 						MockFileSystem->FileSizes.Add(StagedFileDirectory / File, 64);
-						MockFileSystem->FileAttributes.Add(StagedFileDirectory / File, EFileAttributes::Exists);
+						MockFileSystem->FileAttributes.Add(StagedFileDirectory / File, EAttributeFlags::Exists);
 					}
 				});
 
@@ -116,9 +116,9 @@ void FFileAttributionSpec::Define()
 					FFileManifest FileManifest;
 					FileManifest.Filename = MissingFile;
 					MockFileSystem->FileSizes.Add(StagedFileDirectory / MissingFile, INDEX_NONE);
-					MockFileSystem->FileAttributes.Add(StagedFileDirectory / MissingFile, EFileAttributes::None);
+					MockFileSystem->FileAttributes.Add(StagedFileDirectory / MissingFile, EAttributeFlags::None);
 					MockFileSystem->FileSizes.Add(InstallDirectory / MissingFile, INDEX_NONE);
-					MockFileSystem->FileAttributes.Add(InstallDirectory / MissingFile, EFileAttributes::None);
+					MockFileSystem->FileAttributes.Add(InstallDirectory / MissingFile, EAttributeFlags::None);
 					MockNewManifest->BuildFileList.Add(MissingFile);
 					MockNewManifest->FileManifests.Add(MissingFile, FileManifest);
 				});
@@ -346,7 +346,7 @@ void FFileAttributionSpec::MakeUnit()
 			MockOldManifest->FileManifests.Add(File, FileManifest);
 		}
 		MockFileSystem->FileSizes.Add(InstallDirectory / File, FileManifest.GetFileSize());
-		MockFileSystem->FileAttributes.Add(InstallDirectory / File, EFileAttributes::Exists);
+		MockFileSystem->FileAttributes.Add(InstallDirectory / File, EAttributeFlags::Exists);
 	}
 	FileAttribution.Reset(BuildPatchServices::FFileAttributionFactory::Create(
 		MockFileSystem.Get(),
diff --git a/Engine/Source/Runtime/Online/BuildPatchServices/Private/Tests/Unit/InstallChunkSource.spec.cpp b/Engine/Source/Runtime/Online/BuildPatchServices/Private/Tests/Unit/InstallChunkSource.spec.cpp
--- a/Engine/Source/Runtime/Online/BuildPatchServices/Private/Tests/Unit/InstallChunkSource.spec.cpp
+++ b/Engine/Source/Runtime/Online/BuildPatchServices/Private/Tests/Unit/InstallChunkSource.spec.cpp
@@ -270,7 +270,7 @@ void FInstallChunkSourceSpec::Define()
 				MockChunkReferenceTracker->NextReferences = SomeAvailableChunks.Array();
 				bHasPaused = false;
 				PauseTime = 0.1f;
-				MockInstallChunkSourceStat->OnLoadCompleteFunc = [this](const FGuid& ChunkId, IInstallChunkSourceStat::ELoadResult Result)
+				MockInstallChunkSourceStat->OnLoadCompleteFunc = [this](const FGuid& ChunkId, const IInstallChunkSourceStat::ELoadResult& Result, const ISpeedRecorder::FRecord& Record)
 				{
 					if (!bHasPaused)
 					{
@@ -370,13 +370,13 @@ void FInstallChunkSourceSpec::InventUsableChunkData()
 		switch (DataListIdx % 3)
 		{
 			case 0:
-			MockInstallationManifestA->ProducibleChunks.Add(TheChunk);
-			SomeAvailableChunks.Add(TheChunk);
-			break;
+				MockInstallationManifestA->ProducibleChunks.Add(TheChunk);
+				SomeAvailableChunks.Add(TheChunk);
+				break;
 			case 1:
-			MockInstallationManifestB->ProducibleChunks.Add(TheChunk);
-			SomeAvailableChunks.Add(TheChunk);
-			break;
+				MockInstallationManifestB->ProducibleChunks.Add(TheChunk);
+				SomeAvailableChunks.Add(TheChunk);
+				break;
 		}
 	}
 	uint32 FileCounter = 1;
diff --git a/Engine/Source/Runtime/Online/BuildPatchServices/Private/Tests/Unit/Verifier.spec.cpp b/Engine/Source/Runtime/Online/BuildPatchServices/Private/Tests/Unit/Verifier.spec.cpp
--- a/Engine/Source/Runtime/Online/BuildPatchServices/Private/Tests/Unit/Verifier.spec.cpp
+++ b/Engine/Source/Runtime/Online/BuildPatchServices/Private/Tests/Unit/Verifier.spec.cpp
@@ -7,6 +7,7 @@
 #include "Tests/Mock/VerifierStat.mock.h"
 #include "Tests/Mock/Manifest.mock.h"
 #include "Tests/Fake/FileSystem.fake.h"
+#include "Tests/Fake/InstallerError.fake.h"
 #include "Installer/Verifier.h"
 #include "BuildPatchVerify.h"
 #include "BuildPatchHash.h"
@@ -18,6 +19,7 @@ BEGIN_DEFINE_SPEC(FVerifierSpec, "BuildPatchServices.Unit", EAutomationTestFlags
 TUniquePtr<BuildPatchServices::IVerifier> Verifier;
 // Mock.
 TUniquePtr<BuildPatchServices::FFakeFileSystem> FakeFileSystem;
+TUniquePtr<BuildPatchServices::FFakeInstallerError> FakeInstallerError;
 TUniquePtr<BuildPatchServices::FMockVerifierStat> MockVerificationStat;
 BuildPatchServices::FMockManifestPtr MockManifest;
 // Data.
@@ -57,6 +59,7 @@ void FVerifierSpec::Define()
 	BeforeEach([this]()
 	{
 		FakeFileSystem.Reset(new FFakeFileSystem());
+		FakeInstallerError.Reset(new FFakeInstallerError());
 		MockVerificationStat.Reset(new FMockVerifierStat());
 		MockManifest = MakeShareable(new FMockManifest());
 		MakeFileData();
@@ -180,6 +183,117 @@ void FVerifierSpec::Define()
 					TEST_EQUAL(LoadedFiles(), AllFiles);
 				});
 			});
+
+			Describe("when there are multiple types of file errors", [this]()
+			{
+				Describe("when we are using SHA verify mode", [this]()
+				{
+					Describe("when the first error is a corruption", [this]()
+					{
+						BeforeEach([this]()
+						{
+							CorruptSomeFiles();
+							const FString LastFile = VerifyDirectory / SomeFiles.Array().Last();
+							FakeFileSystem->DiskData.Remove(LastFile);
+							MakeUnit(EVerifyMode::ShaVerifyAllFiles);
+						});
+
+						It("will return HashCheckFailed.", [this]()
+						{
+							EVerifyResult VerifyResult = Verifier->Verify(OutDatedFiles);
+							TEST_EQUAL(VerifyResult, EVerifyResult::HashCheckFailed);
+						});
+					});
+
+					Describe("the first error is missing file", [this]()
+					{
+						BeforeEach([this]()
+						{
+							CorruptSomeFiles();
+							const FString FirstFile = VerifyDirectory / *SomeFiles.CreateIterator();
+							FakeFileSystem->DiskData.Remove(FirstFile);
+							MakeUnit(EVerifyMode::ShaVerifyAllFiles);
+						});
+
+						It("will return FileMissing.", [this]()
+						{
+							EVerifyResult VerifyResult = Verifier->Verify(OutDatedFiles);
+							TEST_EQUAL(VerifyResult, EVerifyResult::FileMissing);
+						});
+					});
+
+					Describe("the first error is a file failed to open", [this]()
+					{
+						BeforeEach([this]()
+						{
+							CorruptSomeFiles();
+							const FString FirstFile = VerifyDirectory / *SomeFiles.CreateIterator();
+							FakeFileSystem->DiskDataOpenFailure.Add(FirstFile);
+							MakeUnit(EVerifyMode::ShaVerifyAllFiles);
+						});
+
+						It("will return OpenFileFailed.", [this]()
+						{
+							EVerifyResult VerifyResult = Verifier->Verify(OutDatedFiles);
+							TEST_EQUAL(VerifyResult, EVerifyResult::OpenFileFailed);
+						});
+					});
+
+					Describe("the first error is file was the wrong size", [this]()
+					{
+						BeforeEach([this]()
+						{
+							ResizeSomeFiles();
+							const FString LastFile = VerifyDirectory / SomeFiles.Array().Last();
+							FakeFileSystem->DiskData.Remove(LastFile);
+							MakeUnit(EVerifyMode::ShaVerifyAllFiles);
+						});
+
+						It("will return FileSizeFailed.", [this]()
+						{
+							EVerifyResult VerifyResult = Verifier->Verify(OutDatedFiles);
+							TEST_EQUAL(VerifyResult, EVerifyResult::FileSizeFailed);
+						});
+					});
+				});
+
+				Describe("when we are using file size verify mode", [this]()
+				{
+					Describe("the first error is missing file", [this]()
+					{
+						BeforeEach([this]()
+						{
+							ResizeSomeFiles();
+							const FString FirstFile = VerifyDirectory / *SomeFiles.CreateIterator();
+							FakeFileSystem->DiskData.Remove(FirstFile);
+							MakeUnit(EVerifyMode::FileSizeCheckAllFiles);
+						});
+
+						It("will return FileMissing.", [this]()
+						{
+							EVerifyResult VerifyResult = Verifier->Verify(OutDatedFiles);
+							TEST_EQUAL(VerifyResult, EVerifyResult::FileMissing);
+						});
+					});
+
+					Describe("the first error is file was the wrong size", [this]()
+					{
+						BeforeEach([this]()
+						{
+							ResizeSomeFiles();
+							const FString LastFile = VerifyDirectory / SomeFiles.Array().Last();
+							FakeFileSystem->DiskData.Remove(LastFile);
+							MakeUnit(EVerifyMode::FileSizeCheckAllFiles);
+						});
+
+						It("will return FileSizeFailed.", [this]()
+						{
+							EVerifyResult VerifyResult = Verifier->Verify(OutDatedFiles);
+							TEST_EQUAL(VerifyResult, EVerifyResult::FileSizeFailed);
+						});
+					});
+				});
+			});
 		});
 
 		Describe("SetPaused", [this]()
@@ -195,7 +309,7 @@ void FVerifierSpec::Define()
 				It("should delay the verification process.", [this]()
 				{
 					const float PauseTime = 0.15f;
-					MockVerificationStat->OnFileCompletedFunc = [this, PauseTime](const FString&, bool)
+					MockVerificationStat->OnFileCompletedFunc = [this, PauseTime](const FString&, EVerifyResult)
 					{
 						if (!bHasPaused)
 						{
@@ -228,7 +342,7 @@ void FVerifierSpec::Define()
 				It("should delay the verification process.", [this]()
 				{
 					const float PauseTime = 0.15f;
-					MockVerificationStat->OnFileCompletedFunc = [this, PauseTime](const FString&, bool)
+					MockVerificationStat->OnFileCompletedFunc = [this, PauseTime](const FString&, EVerifyResult)
 					{
 						if (!bHasPaused)
 						{
@@ -263,13 +377,23 @@ void FVerifierSpec::Define()
 
 				It("should halt process and stop.", [this]()
 				{
-					MockVerificationStat->OnFileCompletedFunc = [this](const FString&, bool)
+					MockVerificationStat->OnFileCompletedFunc = [this](const FString&, EVerifyResult)
 					{
 						Verifier->Abort();
 					};
 					Verifier->Verify(OutDatedFiles);
 					TEST_TRUE(LoadedFiles().Num() < MockManifest->BuildFileList.Num());
 				});
+
+				It("should return aborted.", [this]()
+				{
+					MockVerificationStat->OnFileCompletedFunc = [this](const FString&, EVerifyResult)
+					{
+						Verifier->Abort();
+					};
+					EVerifyResult VerifyResult = Verifier->Verify(OutDatedFiles);
+					TEST_EQUAL(VerifyResult, EVerifyResult::Aborted);
+				});
 			});
 
 			Describe("when verifying by size", [this]()
@@ -282,13 +406,23 @@ void FVerifierSpec::Define()
 
 				It("should halt process and stop.", [this]()
 				{
-					MockVerificationStat->OnFileCompletedFunc = [this](const FString&, bool)
+					MockVerificationStat->OnFileCompletedFunc = [this](const FString&, EVerifyResult)
 					{
 						Verifier->Abort();
 					};
 					Verifier->Verify(OutDatedFiles);
 					TEST_TRUE(LoadedFiles().Num() < MockManifest->BuildFileList.Num());
 				});
+
+				It("should return aborted.", [this]()
+				{
+					MockVerificationStat->OnFileCompletedFunc = [this](const FString&, EVerifyResult)
+					{
+						Verifier->Abort();
+					};
+					EVerifyResult VerifyResult = Verifier->Verify(OutDatedFiles);
+					TEST_EQUAL(VerifyResult, EVerifyResult::Aborted);
+				});
 			});
 		});
 	});
diff --git a/Engine/Source/Runtime/Online/BuildPatchServices/Public/BuildPatchInstall.h b/Engine/Source/Runtime/Online/BuildPatchServices/Public/BuildPatchInstall.h
new file mode 100644
--- /dev/null
+++ b/Engine/Source/Runtime/Online/BuildPatchServices/Public/BuildPatchInstall.h
@@ -0,0 +1,49 @@
+// Copyright 1998-2018 Epic Games, Inc. All Rights Reserved.
+#pragma once
+
+#include "CoreMinimal.h"
+
+namespace BuildPatchServices
+{
+	/**
+	 * An enum defining the installation mode that should be used.
+	 */
+	enum class EInstallMode : uint32
+	{
+		// Construct all required files, but only stage them ready to be completed later.
+		StageFiles,
+
+		// Full installation, allowing immediate changes to be made to existing files. The installation is unusable until complete.
+		DestructiveInstall,
+
+		// Full installation, staging all required files before moving them all into place in a final step. The installation is still usable if canceled before the moving staging begins.
+		NonDestructiveInstall,
+
+		// Execute the prerequisite installer only, downloading it first if necessary. If the specified manifest has no prerequisites, this will result in an error.
+		PrereqOnly
+	};
+
+	/**
+	 * Returns the string representation of the EInstallMode value. Used for analytics and logging only.
+	 * @param InstallMode     The value.
+	 * @return the enum's string representation.
+	 */
+	inline const FString& EnumToString(const EInstallMode& InstallMode)
+	{
+		// Const enum strings, special case no error.
+		static const FString StageFiles(TEXT("EInstallMode::StageFiles"));
+		static const FString DestructiveInstall(TEXT("EInstallMode::DestructiveInstall"));
+		static const FString NonDestructiveInstall(TEXT("EInstallMode::NonDestructiveInstall"));
+		static const FString PrereqOnly(TEXT("EInstallMode::PrereqOnly"));
+		static const FString InvalidOrMax(TEXT("InvalidOrMax"));
+
+		switch (InstallMode)
+		{
+			case EInstallMode::StageFiles: return StageFiles;
+			case EInstallMode::DestructiveInstall: return DestructiveInstall;
+			case EInstallMode::NonDestructiveInstall: return NonDestructiveInstall;
+			case EInstallMode::PrereqOnly: return PrereqOnly;
+			default: return InvalidOrMax;
+		}
+	}
+}
diff --git a/Engine/Source/Runtime/Online/BuildPatchServices/Public/BuildPatchMessage.h b/Engine/Source/Runtime/Online/BuildPatchServices/Public/BuildPatchMessage.h
--- a/Engine/Source/Runtime/Online/BuildPatchServices/Public/BuildPatchMessage.h
+++ b/Engine/Source/Runtime/Online/BuildPatchServices/Public/BuildPatchMessage.h
@@ -23,6 +23,28 @@ namespace BuildPatchServices
 		FString Location;
 	};
 
+	/**
+	 * A message describing an action taken to an installation file.
+	 */
+	struct FInstallationFileAction
+	{
+		// Describes the action type.
+		enum class EType : uint32
+		{
+			// The file was removed.
+			Removed = 0,
+			// The file was added.
+			Added,
+			// The file was updated.
+			Updated,
+		};
+
+		// The type of action that occurred.
+		EType Action;
+		// The filename affected, relative to the install location.
+		FString Filename;
+	};
+
 	/**
 	 * Base class of a message handler, this should be inherited from and passed to an installer to receive messages that you want to handle.
 	 */
@@ -37,5 +59,11 @@ namespace BuildPatchServices
 		 * @param Message   The message to be handled.
 		 */
 		virtual void HandleMessage(const FChunkSourceEvent& Message) {}
+
+		/**
+		 * Handles an installation file action message.
+		 * @param Message   The message to be handled.
+		 */
+		virtual void HandleMessage(const FInstallationFileAction& Message) {}
 	};
 }
diff --git a/Engine/Source/Runtime/Online/BuildPatchServices/Public/BuildPatchServicesSingleton.h b/Engine/Source/Runtime/Online/BuildPatchServices/Public/BuildPatchServicesSingleton.h
--- a/Engine/Source/Runtime/Online/BuildPatchServices/Public/BuildPatchServicesSingleton.h
+++ b/Engine/Source/Runtime/Online/BuildPatchServices/Public/BuildPatchServicesSingleton.h
@@ -4,11 +4,11 @@
 #include "CoreMinimal.h"
 #include "Modules/ModuleManager.h"
 #include "Interfaces/IBuildPatchServicesModule.h"
+#include "BuildPatchSettings.h"
 
 class BUILDPATCHSERVICES_API FBuildPatchServices
 {
 public:
-
 	static bool IsAvailable()
 	{
 		return FModuleManager::Get().IsModuleLoaded(ModuleName);
@@ -19,13 +19,21 @@ public:
 		return FModuleManager::Get().GetModuleChecked<IBuildPatchServicesModule>(ModuleName);
 	}
 
-	static void Set(const FName& Value)
+	static const BuildPatchServices::FBuildPatchServicesInitSettings& GetSettings()
+	{
+		return InitSettings;
+	}
+
+	static void Set(
+		const FName& Value, 
+		const BuildPatchServices::FBuildPatchServicesInitSettings& BuildPatchServicesInitSettings = BuildPatchServices::FBuildPatchServicesInitSettings())
 	{
 		if(IsAvailable())
 		{
 			Shutdown();
 		}
 		ModuleName = Value;
+		InitSettings = BuildPatchServicesInitSettings;
 		FModuleManager::Get().LoadModuleChecked<IBuildPatchServicesModule>(ModuleName);
 	}
 
@@ -40,4 +48,5 @@ public:
 private:
 
 	static FName ModuleName;
+	static BuildPatchServices::FBuildPatchServicesInitSettings InitSettings;
 };
