PR #4977: Fixed not being able to connect to a LwsWebSocket after shutting it down (Contributed by [REDACTED_PERSON])

diff --git a/Engine/Source/Runtime/Online/WebSockets/Private/Lws/LwsWebSocket.cpp b/Engine/Source/Runtime/Online/WebSockets/Private/Lws/LwsWebSocket.cpp
--- a/Engine/Source/Runtime/Online/WebSockets/Private/Lws/LwsWebSocket.cpp
+++ b/Engine/Source/Runtime/Online/WebSockets/Private/Lws/LwsWebSocket.cpp
@@ -91,7 +91,7 @@ void FLwsWebSocket::Connect()
 {
 	if (LastGameThreadState != EState::None)
 	{
-		UE_LOG(LogWebSockets, Warning, TEXT("FLwsWebSocket[%d]::Connect: State is not None (%s), unable to start connecting!"), Identifier, ToString(State));
+		UE_LOG(LogWebSockets, Warning, TEXT("FLwsWebSocket[%d]::Connect: State is not None (%s), unable to start connecting!"), Identifier, ToString(LastGameThreadState));
 		return;
 	}
 
@@ -478,15 +478,19 @@ void FLwsWebSocket::GameThreadFinalize()
 {
 	EState PreviousState;
 	FClosedReason LastClosedReason;
+
 	{
 		// TODO:  The contract requires the libwebsockets be done with this object prior to this being called.  Is there a better way to ensure we get the last set value by the libwebsockets thread?  FPlatformMisc::MemoryBarrier maybe?
 		FScopeLock ScopeLock(&StateLock);
 		PreviousState = State;
 		State = EState::None; // Will be re-usable on final delegate triggering
+		LastGameThreadState = State;
 		LastClosedReason = MoveTemp(ClosedReason);
 	}
+
 	UE_LOG(LogWebSockets, Verbose, TEXT("FLwsWebSocket[%d]::GameThreadFinalize: setting State=%s PreviousState=%s"),
 		Identifier, ToString(EState::None), ToString(PreviousState));
+
 	const bool bWasError = (PreviousState == EState::Error);
 	if (bWasError)
 	{
