iOS/macOS cert pubkey pinning: second edition (now with 190% more pinning!)

#rnx

diff --git a/Engine/Source/Runtime/Online/HTTP/Private/Apple/AppleHTTP.cpp b/Engine/Source/Runtime/Online/HTTP/Private/Apple/AppleHTTP.cpp
--- a/Engine/Source/Runtime/Online/HTTP/Private/Apple/AppleHTTP.cpp
+++ b/Engine/Source/Runtime/Online/HTTP/Private/Apple/AppleHTTP.cpp
@@ -4,11 +4,18 @@
 #include "AppleHTTP.h"
 #include "Misc/EngineVersion.h"
 #include "Security/Security.h"
+#include "CommonCrypto/CommonDigest.h"
+#include "Foundation/Foundation.h"
 #include "Misc/App.h"
+#include "Misc/Base64.h"
 #include "HAL/PlatformTime.h"
 #include "Http.h"
 #include "HttpModule.h"
 
+#if WITH_SSL
+#include "Ssl.h"
+#endif
+
 /****************************************************************************
  * FAppleHttpRequest implementation
  ***************************************************************************/
@@ -553,6 +560,187 @@ float FAppleHttpRequest::GetElapsedTime() const
 	}
 }
 
+#if WITH_SSL
+// CC gives the actual key, but strips the ASN.1 header... which means
+// we can't calulate a proper SPKI hash without reconstructing it. sigh.
+static const unsigned char rsa2048Asn1Header[] =
+{
+    0x30, 0x82, 0x01, 0x22, 0x30, 0x0d, 0x06, 0x09,
+    0x2a, 0x86, 0x48, 0x86, 0xf7, 0x0d, 0x01, 0x01,
+    0x01, 0x05, 0x00, 0x03, 0x82, 0x01, 0x0f, 0x00
+};
+static const unsigned char rsa4096Asn1Header[] =
+{
+    0x30, 0x82, 0x02, 0x22, 0x30, 0x0d, 0x06, 0x09,
+    0x2a, 0x86, 0x48, 0x86, 0xf7, 0x0d, 0x01, 0x01,
+    0x01, 0x05, 0x00, 0x03, 0x82, 0x02, 0x0f, 0x00
+};
+static const unsigned char ecdsaSecp256r1Asn1Header[] =
+{
+    0x30, 0x59, 0x30, 0x13, 0x06, 0x07, 0x2a, 0x86,
+    0x48, 0xce, 0x3d, 0x02, 0x01, 0x06, 0x08, 0x2a,
+    0x86, 0x48, 0xce, 0x3d, 0x03, 0x01, 0x07, 0x03,
+    0x42, 0x00
+};
+static const unsigned char ecdsaSecp384r1Asn1Header[] =
+{
+    0x30, 0x76, 0x30, 0x10, 0x06, 0x07, 0x2a, 0x86,
+    0x48, 0xce, 0x3d, 0x02, 0x01, 0x06, 0x05, 0x2b,
+    0x81, 0x04, 0x00, 0x22, 0x03, 0x62, 0x00
+};
+
+-(void) connection:(NSURLConnection *)connection willSendRequestForAuthenticationChallenge: (NSURLAuthenticationChallenge *)challenge
+{
+    if (ensure(ISslCertificateManager::PUBLIC_KEY_DIGEST_SIZE == CC_SHA256_DIGEST_LENGTH))
+    {
+        // we only care about challenges to the received certificate chain
+        if ([challenge.protectionSpace.authenticationMethod isEqualToString:NSURLAuthenticationMethodServerTrust])
+        {
+            SecTrustRef RemoteTrust = challenge.protectionSpace.serverTrust;
+            FString RemoteHost = FString(UTF8_TO_TCHAR([challenge.protectionSpace.host UTF8String]));
+            if ((RemoteTrust == NULL) || (RemoteHost.IsEmpty()))
+            {
+                UE_LOG(LogHttp, Error, TEXT("failed certificate pinning validation: could not parse parameters during certificate pinning evaluation"));
+                [challenge.sender cancelAuthenticationChallenge: challenge];
+                return;
+            }
+            // we check the default trust to verify against the system roots before we dig deeper
+            SecTrustResultType DefaultTrustResult;
+            if (SecTrustEvaluate(RemoteTrust, &DefaultTrustResult) != errSecSuccess)
+            {
+                UE_LOG(LogHttp, Error, TEXT("failed certificate pinning validation: could not evaluate default trust parameters for domain '%s'"), *RemoteHost);
+                [challenge.sender cancelAuthenticationChallenge: challenge];
+                return;
+            }
+            if ((DefaultTrustResult != kSecTrustResultProceed) && (DefaultTrustResult != kSecTrustResultUnspecified))
+            {
+                UE_LOG(LogHttp, Error, TEXT("failed certificate pinning validation: default certificate trust evaluation failed for domain '%s'"), *RemoteHost);
+                [challenge.sender cancelAuthenticationChallenge: challenge];
+                return;
+            }
+            
+            // look at all certs in the remote chain and calculate the SHA256 hash of their DER-encoded SPKI
+            // the chain starts with the server's cert itself, so walk backwards to optimize for roots first
+            TArray<TArray<uint8, TFixedAllocator<ISslCertificateManager::PUBLIC_KEY_DIGEST_SIZE>>> CertDigests;
+            
+            CFIndex NumCerts = SecTrustGetCertificateCount(RemoteTrust);
+            for (int i = static_cast<int>(NumCerts) - 1; i >= 0; i--)
+            {
+                SecCertificateRef Cert = SecTrustGetCertificateAtIndex(RemoteTrust, i);
+                
+                // this is not great, but the only way to extract a public key from a SecCertificateRef
+                // is to create an individual SecTrustRef for each cert that only contains itself and then
+                // evaluate that against an empty X509 policy.
+                SecTrustRef CertTrust;
+                SecPolicyRef TrustPolicy = SecPolicyCreateBasicX509();
+                SecTrustCreateWithCertificates(Cert, TrustPolicy, &CertTrust);
+                SecTrustResultType CertEvalResult;
+                SecTrustEvaluate(CertTrust, &CertEvalResult);
+                SecKeyRef CertPubKey = SecTrustCopyPublicKey(CertTrust);
+                NSData *CertPubKeyData = (NSData *)SecKeyCopyExternalRepresentation(CertPubKey, NULL);
+                if (CertPubKeyData == nil)
+                {
+                    UE_LOG(LogHttp, Warning, TEXT("could not extract public key from certificate %i for domain '%s'; skipping!"), i, *RemoteHost);
+                    continue;
+                }
+                
+				// we got the key. now we have to figure out what type of key it is; thanks, CommonCrypto.
+                CFDictionaryRef CertPubKeyAttr = SecKeyCopyAttributes(CertPubKey);
+                NSString *CertPubKeyType = static_cast<NSString *>(CFDictionaryGetValue(CertPubKeyAttr, kSecAttrKeyType));
+                NSNumber *CertPubKeySize = static_cast<NSNumber *>(CFDictionaryGetValue(CertPubKeyAttr, kSecAttrKeySizeInBits));
+                char *CertPubKeyASN1Header;
+                uint8_t CertPubKeyASN1HeaderSize = 0;
+                if ([CertPubKeyType isEqualToString: (NSString *)kSecAttrKeyTypeRSA])
+                {
+                    switch ([CertPubKeySize integerValue])
+                    {
+                        case 2048:
+                            UE_LOG(LogHttp, VeryVerbose, TEXT("found 2048 bit RSA pubkey"));
+                            CertPubKeyASN1Header = (char *)rsa2048Asn1Header;
+                            CertPubKeyASN1HeaderSize = sizeof(rsa2048Asn1Header);
+                            break;
+                        case 4096:
+                            UE_LOG(LogHttp, VeryVerbose, TEXT("found 4096 bit RSA pubkey"));
+                            CertPubKeyASN1Header = (char *)rsa4096Asn1Header;
+                            CertPubKeyASN1HeaderSize = sizeof(rsa4096Asn1Header);
+                            break;
+                        default:
+                            UE_LOG(LogHttp, Log, TEXT("unsupported RSA key length %i for certificate %i for domain '%s'; skipping!"), [CertPubKeySize integerValue], i, *RemoteHost);
+                            continue;
+                    }
+                }
+                else if ([CertPubKeyType isEqualToString: (NSString *)kSecAttrKeyTypeECSECPrimeRandom])
+                {
+                    switch ([CertPubKeySize integerValue])
+                    {
+                        case 256:
+                            UE_LOG(LogHttp, VeryVerbose, TEXT("found 256 bit ECDSA pubkey"));
+                            CertPubKeyASN1Header = (char *)ecdsaSecp256r1Asn1Header;
+                            CertPubKeyASN1HeaderSize = sizeof(ecdsaSecp256r1Asn1Header);
+                            break;
+                        case 384:
+                            UE_LOG(LogHttp, VeryVerbose, TEXT("found 384 bit ECDSA pubkey"));
+                            CertPubKeyASN1Header = (char *)ecdsaSecp384r1Asn1Header;
+                            CertPubKeyASN1HeaderSize = sizeof(ecdsaSecp384r1Asn1Header);
+                            break;
+                        default:
+                            UE_LOG(LogHttp, Log, TEXT("unsupported ECDSA key length %i for certificate %i for domain '%s'; skipping!"), [CertPubKeySize integerValue], i, *RemoteHost);
+                            continue;
+                    }
+                }
+                else {
+                    UE_LOG(LogHttp, Log, TEXT("unsupported key type (not RSA or ECDSA) for certificate %i for domain '%s'; skipping!"), i, *RemoteHost);
+                    continue;
+                }
+                
+                UE_LOG(LogHttp, VeryVerbose, TEXT("constructed key header: [%d] %s"), CertPubKeyASN1HeaderSize, UTF8_TO_TCHAR([[[NSData dataWithBytes:CertPubKeyASN1Header length:CertPubKeyASN1HeaderSize] description] UTF8String]));
+                UE_LOG(LogHttp, VeryVerbose, TEXT("current pubkey: [%d] %s"), [CertPubKeyData length], UTF8_TO_TCHAR([[[NSData dataWithBytes:[CertPubKeyData bytes] length:[CertPubKeyData length]] description] UTF8String]));
+                
+                // smash 'em together to get a proper key with an ASN.1 header
+                NSMutableData *ReconstructedPubKey = [NSMutableData data];
+                [ReconstructedPubKey appendData:[NSData dataWithBytes:CertPubKeyASN1Header length:CertPubKeyASN1HeaderSize]];
+                [ReconstructedPubKey appendData:CertPubKeyData];
+                UE_LOG(LogHttp, VeryVerbose, TEXT("reconstructed key: [%d] %s"), [ReconstructedPubKey length], UTF8_TO_TCHAR([[ReconstructedPubKey description] UTF8String]));
+                
+                TArray<uint8, TFixedAllocator<ISslCertificateManager::PUBLIC_KEY_DIGEST_SIZE>> CertCalcDigest;
+                CertCalcDigest.AddUninitialized(CC_SHA256_DIGEST_LENGTH);
+                if (!CC_SHA256([ReconstructedPubKey bytes], (CC_LONG)[ReconstructedPubKey length], CertCalcDigest.GetData()))
+                {
+                    UE_LOG(LogHttp, Warning, TEXT("could not calculate SHA256 digest of public key %d for domain '%s'; skipping!"), i, *RemoteHost);
+                    continue;
+                }
+                else
+                {
+                    CertDigests.Add(CertCalcDigest);
+                    UE_LOG(LogHttp, Verbose, TEXT("added SHA256 digest to list for evaluation: domain: '%s' digest: [%d] %s"), *RemoteHost, CertCalcDigest.Num(), UTF8_TO_TCHAR([[[NSData dataWithBytes:CertCalcDigest.GetData() length:CertCalcDigest.Num()] description] UTF8String]));
+                }
+            }
+            
+            //finally, see if any of the pubkeys in the chain match any of our pinned pubkey hashes
+            if (CertDigests.Num() <= 0 || !FSslModule::Get().GetCertificateManager().VerifySslCertificates(CertDigests, RemoteHost))
+            {
+                // we could not validate any of the provided certs in chain with the pinned hashes for this host
+                // so we tell the sender to cancel (which cancels the pending connection)
+                UE_LOG(LogHttp, Error, TEXT("failed certificate pinning validation: no SPKI hashes in request matched pinned hashes for domain '%s' (was provided %d certificates in request)"), *RemoteHost, CertDigests.Num());
+                [challenge.sender cancelAuthenticationChallenge:challenge];
+                return;
+            }
+        }
+    }
+    else
+    {
+        UE_LOG(LogHttp, Error, TEXT("failed certificate pinning validation: SslCertificateManager is using non-SHA256 SPKI hashes [expected %d bytes, got %d bytes]"), CC_SHA256_DIGEST_LENGTH, ISslCertificateManager::PUBLIC_KEY_DIGEST_SIZE);
+        [challenge.sender cancelAuthenticationChallenge:challenge];
+        return;
+    }
+    
+    // if we got this far, pinning validation either succeeded or was disabled (or this was checking for client auth, etc.)
+    // so tell the connection to keep going with whatever else it was trying to validate
+    UE_LOG(LogHttp, Verbose, TEXT("certificate public key pinning either succeeded, is disabled, or challenge was not a server trust; continuing with auth"));
+    [challenge.sender performDefaultHandlingForAuthenticationChallenge: challenge];
+}
+#endif
+
 
 -(void) connectionDidFinishLoading:(NSURLConnection *)connection
 {
diff --git a/Engine/Source/Runtime/Online/HTTP/Private/Apple/AppleHTTP.h b/Engine/Source/Runtime/Online/HTTP/Private/Apple/AppleHTTP.h
--- a/Engine/Source/Runtime/Online/HTTP/Private/Apple/AppleHTTP.h
+++ b/Engine/Source/Runtime/Online/HTTP/Private/Apple/AppleHTTP.h
@@ -146,6 +146,11 @@ private:
 /** Delegate called with we complete successfully. See Apple docs for when/how this should be used. */
 -(void) connectionDidFinishLoading:(NSURLConnection *)connection;
 
+#if WITH_SSL
+/** Delegate called when the connection is about to validate an auth challenge. We only care about server trust. See Apple docs for when/how this should be used. */
+-(void)connection:(NSURLConnection *)connection willSendRequestForAuthenticationChallenge: (NSURLAuthenticationChallenge *)challenge;
+#endif
+
 - (TArray<uint8>&)getPayload;
 - (int32)getBytesWritten;
 @end
diff --git a/Engine/Source/Runtime/Online/SSL/Private/SslCertificateManager.cpp b/Engine/Source/Runtime/Online/SSL/Private/SslCertificateManager.cpp
--- a/Engine/Source/Runtime/Online/SSL/Private/SslCertificateManager.cpp
+++ b/Engine/Source/Runtime/Online/SSL/Private/SslCertificateManager.cpp
@@ -187,25 +187,9 @@ bool FSslCertificateManager::VerifySslCertificates(X509_STORE_CTX* Context, cons
 		return false;
 	}
 
-	const TArray<TArray<uint8, TFixedAllocator<PUBLIC_KEY_DIGEST_SIZE>>>* PinnedKeys = nullptr;
-	for (const TPair<FString, TArray<TArray<uint8, TFixedAllocator<PUBLIC_KEY_DIGEST_SIZE>>>>& PinnedKeyPair: PinnedPublicKeys)
-	{
-		const FString& PinnedDomain = PinnedKeyPair.Key;
-		if ((PinnedDomain[0] == TEXT('.') && Domain.EndsWith(PinnedDomain))
-			|| Domain == PinnedDomain)
-		{
-			PinnedKeys = &PinnedKeyPair.Value;
-			break;
-		}
-	}
-
-	if (!PinnedKeys)
-	{
-		// No keys pinned for this domain
-		return true;
-	}
-
-	bool bFoundMatch = false;
+	TArray<TArray<uint8, TFixedAllocator<PUBLIC_KEY_DIGEST_SIZE>>> CertDigests;
+    
+    bool bFoundMatch = false;
 	for (int CertIndex = 0; CertIndex < NumCertsInChain; ++CertIndex)
 	{
 		X509* Certificate = sk_X509_value(Chain, CertIndex);
@@ -227,20 +211,55 @@ bool FSslCertificateManager::VerifySslCertificates(X509_STORE_CTX* Context, cons
 		SHA256_Init(&ShaContext);
 		SHA256_Update(&ShaContext, PubKey.GetData(), PubKey.Num());
 		SHA256_Final(Digest.GetData(), &ShaContext);
-
-		if (PinnedKeys->Contains(Digest))
-		{
-			bFoundMatch = true;
-			break;
-		}
-	}
-
-	if (!bFoundMatch)
-	{
-		X509_STORE_CTX_set_error(Context, X509_V_ERR_CERT_UNTRUSTED);
+		
+		CertDigests.Add(Digest);
 	}
+    
+    bFoundMatch = VerifySslCertificates(CertDigests, Domain);
+    if (!bFoundMatch)
+    {
+        X509_STORE_CTX_set_error(Context, X509_V_ERR_CERT_UNTRUSTED);
+    }
+    return bFoundMatch;
+}
 
-	return bFoundMatch;
+bool FSslCertificateManager::VerifySslCertificates(TArray<TArray<uint8, TFixedAllocator<PUBLIC_KEY_DIGEST_SIZE>>>& Digests, const FString& Domain) const
+{
+#if !UE_BUILD_SHIPPING
+    static const bool bPinningDisabled = FParse::Param(FCommandLine::Get(), TEXT("DisableSSLCertificatePinning"));
+    if (bPinningDisabled)
+    {
+        return true;
+    }
+#endif
+    const TArray<TArray<uint8, TFixedAllocator<PUBLIC_KEY_DIGEST_SIZE>>>* PinnedKeys = nullptr;
+    for (const TPair<FString, TArray<TArray<uint8, TFixedAllocator<PUBLIC_KEY_DIGEST_SIZE>>>>& PinnedKeyPair : PinnedPublicKeys)
+    {
+        const FString& PinnedDomain = PinnedKeyPair.Key;
+        if ((PinnedDomain[0] == TEXT('.') && Domain.EndsWith(PinnedDomain))
+            || Domain == PinnedDomain)
+        {
+            PinnedKeys = &PinnedKeyPair.Value;
+            break;
+        }
+    }
+    if (!PinnedKeys)
+    {
+        // No keys pinned for this domain
+        UE_LOG(LogSsl, Verbose, TEXT("no pinned key digests found for domain '%s'"), *Domain);
+        return true;
+    }
+    bool bFoundMatch = false;
+    for (int32 CertIndex = 0; CertIndex < Digests.Num(); ++CertIndex)
+    {
+        if (PinnedKeys->Contains(Digests[CertIndex]))
+        {
+            UE_LOG(LogSsl, Verbose, TEXT("found public key digest in request that matches a pinned key for '%s'"), *Domain);
+            bFoundMatch = true;
+            break;
+        }
+    }
+    return bFoundMatch;
 }
 
 void FSslCertificateManager::BuildRootCertificateArray()
diff --git a/Engine/Source/Runtime/Online/SSL/Private/SslCertificateManager.h b/Engine/Source/Runtime/Online/SSL/Private/SslCertificateManager.h
--- a/Engine/Source/Runtime/Online/SSL/Private/SslCertificateManager.h
+++ b/Engine/Source/Runtime/Online/SSL/Private/SslCertificateManager.h
@@ -22,6 +22,7 @@ public:
 	virtual bool IsDomainPinned(const FString& Domain) override;
 	virtual void SetPinnedPublicKeys(const FString& Domain, const FString& PinnedKeyDigests) override;
 	virtual bool VerifySslCertificates(X509_STORE_CTX* Context, const FString& Domain) const override;
+	virtual bool VerifySslCertificates(TArray<TArray<uint8, TFixedAllocator<PUBLIC_KEY_DIGEST_SIZE>>>& Digests, const FString& Domain) const override;
 
 	virtual void BuildRootCertificateArray();
 	virtual void EmptyRootCertificateArray();
@@ -30,7 +31,6 @@ protected:
 	void AddPEMFileToRootCertificateArray(const FString& Path);
 	void AddCertificateToRootCertificateArray(X509* Certificate);
 
-	static constexpr int PUBLIC_KEY_DIGEST_SIZE = 32;
 	TArray<X509*> RootCertificateArray;
 	TArray<TPair<FString, TArray<TArray<uint8, TFixedAllocator<PUBLIC_KEY_DIGEST_SIZE>>>>> PinnedPublicKeys;
 };
diff --git a/Engine/Source/Runtime/Online/SSL/Public/Interfaces/ISslCertificateManager.h b/Engine/Source/Runtime/Online/SSL/Public/Interfaces/ISslCertificateManager.h
--- a/Engine/Source/Runtime/Online/SSL/Public/Interfaces/ISslCertificateManager.h
+++ b/Engine/Source/Runtime/Online/SSL/Public/Interfaces/ISslCertificateManager.h
@@ -15,7 +15,9 @@ class ISslCertificateManager
 {
 public:
 	virtual ~ISslCertificateManager() {}
-
+    
+	static constexpr int PUBLIC_KEY_DIGEST_SIZE = 32;
+    
 	/**
 	 * Add trusted root certificates to the SSL context
 	 *
@@ -60,4 +62,14 @@ public:
 	 * @return false if validation fails
 	 */
 	virtual bool VerifySslCertificates(X509_STORE_CTX* Context, const FString& Domain) const = 0;
+    
+	/**
+	 * Performs additional ssl validation (certificate pinning)
+	 *
+	 * @param Digests Array of public key digests to check against pinned key digests
+	 * @param Domain Domain we are connected to
+	 *
+	 * @return false if validation fails
+	 */
+	virtual bool VerifySslCertificates(TArray<TArray<uint8, TFixedAllocator<PUBLIC_KEY_DIGEST_SIZE>>>& Digests, const FString& Domain) const = 0;
 };
