Add PATCH verb support for Curl

diff --git a/Engine/Source/Runtime/Online/HTTP/Private/Curl/CurlHttp.cpp b/Engine/Source/Runtime/Online/HTTP/Private/Curl/CurlHttp.cpp
--- a/Engine/Source/Runtime/Online/HTTP/Private/Curl/CurlHttp.cpp
+++ b/Engine/Source/Runtime/Online/HTTP/Private/Curl/CurlHttp.cpp
@@ -726,10 +726,15 @@ bool FCurlHttpRequest::SetupRequest()
 		curl_easy_setopt(EasyHandle, CURLOPT_POSTFIELDSIZE, RequestPayload->GetContentLength());
 		bUseReadFunction = true;
 	}
-	else if (Verb == TEXT("PUT"))
+	else if (Verb == TEXT("PUT") || Verb == TEXT("PATCH"))
 	{
 		curl_easy_setopt(EasyHandle, CURLOPT_UPLOAD, 1L);
 		curl_easy_setopt(EasyHandle, CURLOPT_INFILESIZE, RequestPayload->GetContentLength());
+		if (Verb != TEXT("PUT"))
+		{
+			curl_easy_setopt(EasyHandle, CURLOPT_CUSTOMREQUEST, TCHAR_TO_UTF8(*Verb));
+		}
+
 		bUseReadFunction = true;
 	}
 	else if (Verb == TEXT("GET"))
