Make sure AppleHTTP loads FSslModule before using it on the main thread.

diff --git a/Engine/Source/Runtime/Online/HTTP/Private/Apple/AppleHTTP.cpp b/Engine/Source/Runtime/Online/HTTP/Private/Apple/AppleHTTP.cpp
--- a/Engine/Source/Runtime/Online/HTTP/Private/Apple/AppleHTTP.cpp
+++ b/Engine/Source/Runtime/Online/HTTP/Private/Apple/AppleHTTP.cpp
@@ -43,6 +43,11 @@ FAppleHttpRequest::FAppleHttpRequest()
 	{
 		SetHeader(It.Key(), It.Value());
 	}
+
+#if WITH_SSL
+	// Make sure the module is loaded on the game thread before being used by FHttpResponseAppleWrapper, which will be called on the main thread
+	FSslModule::Get();
+#endif
 }
 
 
