#online Added typedefs for TSharedRefs to HTTP interfaces, to match the TSharedPtr ones

diff --git a/Engine/Source/Runtime/Online/HTTP/Public/Interfaces/IHttpRequest.h b/Engine/Source/Runtime/Online/HTTP/Public/Interfaces/IHttpRequest.h
--- a/Engine/Source/Runtime/Online/HTTP/Public/Interfaces/IHttpRequest.h
+++ b/Engine/Source/Runtime/Online/HTTP/Public/Interfaces/IHttpRequest.h
@@ -62,8 +62,14 @@ namespace EHttpRequestStatus
 	}
 }
 
-typedef TSharedPtr<class IHttpRequest> FHttpRequestPtr;
-typedef TSharedPtr<class IHttpResponse,ESPMode::ThreadSafe> FHttpResponsePtr;
+class IHttpRequest;
+class IHttpResponse;
+
+typedef TSharedPtr<IHttpRequest> FHttpRequestPtr;
+typedef TSharedPtr<IHttpResponse, ESPMode::ThreadSafe> FHttpResponsePtr;
+
+typedef TSharedRef<IHttpRequest> FHttpRequestRef;
+typedef TSharedRef<IHttpResponse, ESPMode::ThreadSafe> FHttpResponseRef;
 
 /**
  * Delegate called when an Http request completes
