Adding timeout for BackgroundHTTP Temp files after a configurable timeout time. Default is 3 days. -1 to turn off.
#tests Android Client

diff --git a/Engine/Config/BaseEngine.ini b/Engine/Config/BaseEngine.ini
--- a/Engine/Config/BaseEngine.ini
+++ b/Engine/Config/BaseEngine.ini
@@ -26,6 +26,10 @@ HttpConnectionTimeout=60
 HttpReceiveTimeout=30
 HttpSendTimeout=30
 
+[BackgroundHttp]
+; How many days we will wait to delete any found BackgroundHTTP Temp Files that haven't been claimed. Default is 259200 seconds (3 days). -1 to disable.
+BackgroundHttp.TempFileTimeOutSeconds=259200
+
 [WebSockets.LibWebSockets]
 ThreadStackSize=131072
 ThreadTargetFrameTimeInSeconds=0.0333
diff --git a/Engine/Source/Runtime/Online/BackgroundHTTP/Private/BackgroundHttpManagerImpl.cpp b/Engine/Source/Runtime/Online/BackgroundHTTP/Private/BackgroundHttpManagerImpl.cpp
--- a/Engine/Source/Runtime/Online/BackgroundHTTP/Private/BackgroundHttpManagerImpl.cpp
+++ b/Engine/Source/Runtime/Online/BackgroundHTTP/Private/BackgroundHttpManagerImpl.cpp
@@ -1,10 +1,12 @@
 // Copyright 1998-2019 Epic Games, Inc. All Rights Reserved.
 #include "BackgroundHttpManagerImpl.h"
 
+#include "HAL/FileManager.h"
 #include "HAL/PlatformFilemanager.h"
 #include "HAL/PlatformAtomics.h"
 #include "HAL/PlatformFile.h"
 
+#include "Misc/ConfigCacheIni.h"
 #include "Misc/ScopeRWLock.h"
 
 DEFINE_LOG_CATEGORY(LogBackgroundHttpManager);
@@ -24,6 +26,7 @@ FBackgroundHttpManagerImpl::~FBackgroundHttpManagerImpl()
 
 void FBackgroundHttpManagerImpl::Initialize()
 {
+	ClearAnyTempFilesFromTimeOut();
 }
 
 void FBackgroundHttpManagerImpl::Shutdown()
@@ -42,6 +45,38 @@ void FBackgroundHttpManagerImpl::Shutdown()
 	}
 }
 
+void FBackgroundHttpManagerImpl::ClearAnyTempFilesFromTimeOut()
+{
+	UE_LOG(LogBackgroundHttpManager, Log, TEXT("Checking for BackgroundHTTP temp files that should be deleted due to time out"));
+
+	TArray<FString> FilesToCheck;
+
+	//Find all files in our temp folder
+	IFileManager::Get().FindFiles(FilesToCheck, *FPlatformBackgroundHttp::GetTemporaryRootPath(), nullptr);
+
+	double FileAgeTimeOutSettings = -1;
+	GConfig->GetDouble(TEXT("BackgroundHttp"), TEXT("BackgroundHttp.TempFileTimeOutSeconds"), FileAgeTimeOutSettings, GEngineIni);
+
+	if (FileAgeTimeOutSettings >= 0)
+	{
+		for (const FString& File : FilesToCheck)
+		{
+			const double FileAge = IFileManager::Get().GetFileAgeSeconds(*File);
+			const bool bShouldDelete = (FileAge > FileAgeTimeOutSettings);
+
+			UE_LOG(LogBackgroundHttpManager, Log, TEXT("FoundTempFile: %s with age %lld -- bShouldDelete:%d"), *File, FileAge, (int)bShouldDelete);
+
+			if (bShouldDelete)
+			{
+				if (!IFileManager::Get().Delete(*File))
+				{
+					UE_LOG(LogBackgroundHttpManager, Error, TEXT("File %s failed to delete, but should have as as it is %lld seconds old!"), *File);
+				}
+			}
+		}
+	}
+}
+
 void FBackgroundHttpManagerImpl::CleanUpTemporaryFiles()
 {
 	UE_LOG(LogBackgroundHttpManager, Log, TEXT("Cleaning Up Temporary Files"));
diff --git a/Engine/Source/Runtime/Online/BackgroundHTTP/Private/IOS/ApplePlatformBackgroundHttpManager.cpp b/Engine/Source/Runtime/Online/BackgroundHTTP/Private/IOS/ApplePlatformBackgroundHttpManager.cpp
--- a/Engine/Source/Runtime/Online/BackgroundHTTP/Private/IOS/ApplePlatformBackgroundHttpManager.cpp
+++ b/Engine/Source/Runtime/Online/BackgroundHTTP/Private/IOS/ApplePlatformBackgroundHttpManager.cpp
@@ -43,6 +43,8 @@ void FApplePlatformBackgroundHttpManager::Initialize()
 	GConfig->GetInt(TEXT("BackgroundHttp.iOSSettings"), TEXT("BackgroundHttp.RetryResumeDataLimit"), RetryResumeDataLimitSetting, GEngineIni);
 
 	SetupNSURLSessionResponseDelegates();
+
+	Super::Initialize();
 }
 
 void FApplePlatformBackgroundHttpManager::PopulateUnAssociatedTasks()
diff --git a/Engine/Source/Runtime/Online/BackgroundHTTP/Public/BackgroundHttpManagerImpl.h b/Engine/Source/Runtime/Online/BackgroundHTTP/Public/BackgroundHttpManagerImpl.h
--- a/Engine/Source/Runtime/Online/BackgroundHTTP/Public/BackgroundHttpManagerImpl.h
+++ b/Engine/Source/Runtime/Online/BackgroundHTTP/Public/BackgroundHttpManagerImpl.h
@@ -37,6 +37,7 @@ protected:
 	virtual bool CheckForExistingCompletedDownload(const FBackgroundHttpRequestPtr Request, FString& ExistingFilePathOut, int64& ExistingFileSizeOut);
 	virtual void ActivatePendingRequests();
 
+	virtual void ClearAnyTempFilesFromTimeOut();
 protected:
 
 	/** List of Background Http requests that we have called AddRequest on, but have not yet started due to platform active download limits **/
