Fixing compile error on iOS with invalid Super:: call.
#tests none

diff --git a/Engine/Source/Runtime/Online/BackgroundHTTP/Private/IOS/ApplePlatformBackgroundHttpManager.cpp b/Engine/Source/Runtime/Online/BackgroundHTTP/Private/IOS/ApplePlatformBackgroundHttpManager.cpp
--- a/Engine/Source/Runtime/Online/BackgroundHTTP/Private/IOS/ApplePlatformBackgroundHttpManager.cpp
+++ b/Engine/Source/Runtime/Online/BackgroundHTTP/Private/IOS/ApplePlatformBackgroundHttpManager.cpp
@@ -44,7 +44,7 @@ void FApplePlatformBackgroundHttpManager::Initialize()
 
 	SetupNSURLSessionResponseDelegates();
 
-	Super::Initialize();
+	FBackgroundHttpManagerImpl::Initialize();
 }
 
 void FApplePlatformBackgroundHttpManager::PopulateUnAssociatedTasks()
