Fixing bug where BackgroundHTTP Requests could end up leaking memory or crashing when being removed due to a self-reference chain
#tests Android Client

diff --git a/Engine/Source/Runtime/Online/BackgroundHTTP/Private/GenericPlatform/GenericPlatformBackgroundHttpRequest.cpp b/Engine/Source/Runtime/Online/BackgroundHTTP/Private/GenericPlatform/GenericPlatformBackgroundHttpRequest.cpp
--- a/Engine/Source/Runtime/Online/BackgroundHTTP/Private/GenericPlatform/GenericPlatformBackgroundHttpRequest.cpp
+++ b/Engine/Source/Runtime/Online/BackgroundHTTP/Private/GenericPlatform/GenericPlatformBackgroundHttpRequest.cpp
@@ -12,14 +12,20 @@ FGenericPlatformBackgroundHttpRequest::FGenericPlatformBackgroundHttpRequest()
 {
 }
 
+FGenericPlatformBackgroundHttpRequest::~FGenericPlatformBackgroundHttpRequest()
+{
+	RequestWrapper.Release();
+}
+
 bool FGenericPlatformBackgroundHttpRequest::HandleDelayedProcess()
 {
 	//Our request may already have a response result from associating with another already completed request. If so, lets not spin up our RequestWrapper.
 	if (!GetResponse().IsValid())
 	{
+		//If we already have a Request Wrapper, lets keep it. Otherwise spin it up now.
 		if (!RequestWrapper.IsValid())
 		{
-			RequestWrapper = MakeShareable<FGenericPlatformBackgroundHttpWrapper>(new FGenericPlatformBackgroundHttpWrapper(SharedThis(this), NumberOfTotalRetries));
+			RequestWrapper.Reset(new FGenericPlatformBackgroundHttpWrapper(SharedThis(this), NumberOfTotalRetries));
 		}
 
 		if (ensureAlwaysMsgf(RequestWrapper.IsValid(), TEXT("Failure creating FGenericPlatformBackgroundHttpWrapper! Can not process BackgroundHttp request! RequestID:%s"), *GetRequestID()))
@@ -34,8 +40,8 @@ bool FGenericPlatformBackgroundHttpRequest::HandleDelayedProcess()
 
 void FGenericPlatformBackgroundHttpRequest::CancelRequest()
 {
-	FBackgroundHttpRequestImpl::CancelRequest();
 	RequestWrapper.Reset();
+	FBackgroundHttpRequestImpl::CancelRequest();
 }
 
 void FGenericPlatformBackgroundHttpRequest::FGenericPlatformBackgroundHttpRequest::CompleteWithExistingResponseData(FBackgroundHttpResponsePtr BackgroundResponse)
@@ -82,7 +88,7 @@ void FGenericPlatformBackgroundHttpRequest::FGenericPlatformBackgroundHttpWrappe
 		}
 		else
 		{
-			UE_LOG(LogBackgroundHttpRequest, Display, TEXT("No valid URL for Request. (No valid URL, or out of Retries) - RequestID:%s | RetryNumber:%d | MaxRetries:%d"), *OriginalRequest->GetRequestID(), CurrentRetryNumber, MaxRetries);
+			UE_LOG(LogBackgroundHttpRequest, Display, TEXT("No valid URL for Request. (No valid URL, or out of Retries) - RequestID:%s | RetryNumber:%d | MaxRetries:%d"), *OriginalRequest.Pin()->GetRequestID(), CurrentRetryNumber, MaxRetries);
 
 			//We don't have a valid URL, so just error immediately
 			HttpRequestComplete(HttpRequest, nullptr, false);
@@ -111,7 +117,7 @@ const FString FGenericPlatformBackgroundHttpRequest::FGenericPlatformBackgroundH
 	//Don't find a valid URL if we are already passed our retry count
 	if (URLIndexToUse <= MaxRetries)
 	{
-		const TArray<FString>& RequestURLList = OriginalRequest->GetURLList();
+		const TArray<FString>& RequestURLList = OriginalRequest.IsValid() ? OriginalRequest.Pin()->GetURLList() : TArray<FString>();
 		if (RequestURLList.Num() > 0)
 		{
 			//If our number is too high, lets loop back around to the start of the list
@@ -132,34 +138,40 @@ const FString FGenericPlatformBackgroundHttpRequest::FGenericPlatformBackgroundH
 
 void FGenericPlatformBackgroundHttpRequest::FGenericPlatformBackgroundHttpWrapper::HttpRequestComplete(FHttpRequestPtr HttpRequestIn, FHttpResponsePtr HttpResponse, bool bSuccess)
 {
-	//Create a new FGenericPlatformBackgroundHttpRequest from the HttpRequest/Response
-	FBackgroundHttpResponsePtr ConstructedResponse = MakeShareable(new FGenericPlatformBackgroundHttpResponse(HttpRequestIn, HttpResponse, bSuccess));
+	if (ensureAlwaysMsgf(OriginalRequest.IsValid(), TEXT("Recieved HttpRequestComplete callback with invalid OriginalRequest pointer! Can not complete request!")))
+	{
+		//Create a new FGenericPlatformBackgroundHttpRequest from the HttpRequest/Response
+		FBackgroundHttpResponsePtr ConstructedResponse = MakeShareable(new FGenericPlatformBackgroundHttpResponse(HttpRequestIn, HttpResponse, bSuccess));
 
-	//Let the FGenericPlatformBackgroundHttpResponse drive if this succeeded or not
-	const bool bRequestWasSuccessful = EHttpResponseCodes::IsOk(ConstructedResponse->GetResponseCode());
+		//Let the FGenericPlatformBackgroundHttpResponse drive if this succeeded or not
+		const bool bRequestWasSuccessful = EHttpResponseCodes::IsOk(ConstructedResponse->GetResponseCode());
 
-	UE_LOG(LogBackgroundHttpRequest, Display, TEXT("Underlying HttpRequest complete - RequestID:%s | bRequestWasSuccessfull:%d | CurrentRetryNumber:%d | MaxRetries:%d"), *OriginalRequest->GetRequestID(), (int)(bRequestWasSuccessful), CurrentRetryNumber, MaxRetries);
+		UE_LOG(LogBackgroundHttpRequest, Display, TEXT("Underlying HttpRequest complete - RequestID:%s | bRequestWasSuccessfull:%d | CurrentRetryNumber:%d | MaxRetries:%d"), *OriginalRequest.Pin()->GetRequestID(), (int)(bRequestWasSuccessful), CurrentRetryNumber, MaxRetries);
 
-	//Attempt a retry on failure
-	if (!bRequestWasSuccessful)
-	{
-		++CurrentRetryNumber;
-		MakeRequest();
-	}
-	else
-	{
-		//Pass constructed response to our BackgroundRequest so it can complete itself using this response
-		OriginalRequest->CompleteWithExistingResponseData(ConstructedResponse);
+		//Attempt a retry on failure
+		if (!bRequestWasSuccessful)
+		{
+			++CurrentRetryNumber;
+			MakeRequest();
+		}
+		else
+		{
+			//Pass constructed response to our BackgroundRequest so it can complete itself using this response
+			OriginalRequest.Pin()->CompleteWithExistingResponseData(ConstructedResponse);
+		}
 	}
 }
 
 void FGenericPlatformBackgroundHttpRequest::FGenericPlatformBackgroundHttpWrapper::UpdateHttpProgress(FHttpRequestPtr UnderlyingHttpRequest, int32 BytesSent, int32 BytesReceived)
 {
-	const int32 ByteDifference = (LastProgressUpdateBytes > 0) ? BytesReceived - LastProgressUpdateBytes : BytesReceived;
-	ensureAlwaysMsgf((ByteDifference > 0), TEXT("Invalid Byte Difference in UpdateHttpProgress -- ByteDifference:%d | LastProgressUpdateBytes:%d | BytesSent:%d | BytesReceived:%d"), ByteDifference, LastProgressUpdateBytes, BytesSent, BytesReceived);
-	LastProgressUpdateBytes = BytesReceived;
+	if (ensureAlwaysMsgf(OriginalRequest.IsValid(), TEXT("Recieved UpdateHttpProgress callback with invalid OriginalRequest pointer! Can not update request progress!")))
+	{
+		const int32 ByteDifference = (LastProgressUpdateBytes > 0) ? BytesReceived - LastProgressUpdateBytes : BytesReceived;
+		ensureAlwaysMsgf((ByteDifference > 0), TEXT("Invalid Byte Difference in UpdateHttpProgress -- ByteDifference:%d | LastProgressUpdateBytes:%d | BytesSent:%d | BytesReceived:%d"), ByteDifference, LastProgressUpdateBytes, BytesSent, BytesReceived);
+		LastProgressUpdateBytes = BytesReceived;
 
-	UE_LOG(LogBackgroundHttpRequest, VeryVerbose, TEXT("HttpRequest Progress Update- RequestID:%s | BytesSent: %d | BytesReceived:%d | BytesReceivedSinceLastUpdate:%d"), *OriginalRequest->GetRequestID(), BytesSent, BytesReceived, ByteDifference);
-	
-	OriginalRequest->OnProgressUpdated().ExecuteIfBound(OriginalRequest, BytesReceived, ByteDifference);
+		UE_LOG(LogBackgroundHttpRequest, VeryVerbose, TEXT("HttpRequest Progress Update- RequestID:%s | BytesSent: %d | BytesReceived:%d | BytesReceivedSinceLastUpdate:%d"), *OriginalRequest.Pin()->GetRequestID(), BytesSent, BytesReceived, ByteDifference);
+
+		OriginalRequest.Pin()->OnProgressUpdated().ExecuteIfBound(OriginalRequest.Pin(), BytesReceived, ByteDifference);
+	}
 }
diff --git a/Engine/Source/Runtime/Online/BackgroundHTTP/Public/GenericPlatform/GenericPlatformBackgroundHttpRequest.h b/Engine/Source/Runtime/Online/BackgroundHTTP/Public/GenericPlatform/GenericPlatformBackgroundHttpRequest.h
--- a/Engine/Source/Runtime/Online/BackgroundHTTP/Public/GenericPlatform/GenericPlatformBackgroundHttpRequest.h
+++ b/Engine/Source/Runtime/Online/BackgroundHTTP/Public/GenericPlatform/GenericPlatformBackgroundHttpRequest.h
@@ -17,7 +17,7 @@ class BACKGROUNDHTTP_API FGenericPlatformBackgroundHttpRequest
 public:
 
 	FGenericPlatformBackgroundHttpRequest();
-	virtual ~FGenericPlatformBackgroundHttpRequest() {}
+	virtual ~FGenericPlatformBackgroundHttpRequest();
 
 	//IHttpBackgroundHttpRequest
 	virtual bool HandleDelayedProcess() override;
@@ -50,7 +50,7 @@ protected:
 		FGenericPlatformBackgroundHttpWrapper() {}
 
 		//Reference to the creating request
-		FBackgroundHttpRequestPtr OriginalRequest;
+		TWeakPtr<class IBackgroundHttpRequest, ESPMode::ThreadSafe> OriginalRequest;
 
 		//Current Http Request being processed by this wrapper
 		TSharedPtr<IHttpRequest> HttpRequest;
@@ -65,5 +65,5 @@ protected:
 		int32 LastProgressUpdateBytes;
 	};
 
-	TSharedPtr<FGenericPlatformBackgroundHttpWrapper> RequestWrapper;
+	TUniquePtr<FGenericPlatformBackgroundHttpWrapper> RequestWrapper;
 };
