Fixing BackgroundHttp GenericPlatform Implementation Problems:
- Fixing infinite loop when request ran out of retries.
- Fixed possible missing initialization error for ResponseCode
#tests Android Client

diff --git a/Engine/Source/Runtime/Online/BackgroundHTTP/Private/GenericPlatform/GenericPlatformBackgroundHttpRequest.cpp b/Engine/Source/Runtime/Online/BackgroundHTTP/Private/GenericPlatform/GenericPlatformBackgroundHttpRequest.cpp
--- a/Engine/Source/Runtime/Online/BackgroundHTTP/Private/GenericPlatform/GenericPlatformBackgroundHttpRequest.cpp
+++ b/Engine/Source/Runtime/Online/BackgroundHTTP/Private/GenericPlatform/GenericPlatformBackgroundHttpRequest.cpp
@@ -115,7 +115,7 @@ const FString FGenericPlatformBackgroundHttpRequest::FGenericPlatformBackgroundH
 	int URLIndexToUse = CurrentRetryNumber;
 
 	//Don't find a valid URL if we are already passed our retry count
-	if (URLIndexToUse <= MaxRetries)
+	if (HasRetriesRemaining())
 	{
 		const TArray<FString>& RequestURLList = OriginalRequest.IsValid() ? OriginalRequest.Pin()->GetURLList() : TArray<FString>();
 		if (RequestURLList.Num() > 0)
@@ -149,7 +149,7 @@ void FGenericPlatformBackgroundHttpRequest::FGenericPlatformBackgroundHttpWrappe
 		UE_LOG(LogBackgroundHttpRequest, Display, TEXT("Underlying HttpRequest complete - RequestID:%s | bRequestWasSuccessfull:%d | CurrentRetryNumber:%d | MaxRetries:%d"), *OriginalRequest.Pin()->GetRequestID(), (int)(bRequestWasSuccessful), CurrentRetryNumber, MaxRetries);
 
 		//Attempt a retry on failure
-		if (!bRequestWasSuccessful)
+		if (!bRequestWasSuccessful && HasRetriesRemaining())
 		{
 			++CurrentRetryNumber;
 			MakeRequest();
@@ -162,6 +162,11 @@ void FGenericPlatformBackgroundHttpRequest::FGenericPlatformBackgroundHttpWrappe
 	}
 }
 
+bool FGenericPlatformBackgroundHttpRequest::FGenericPlatformBackgroundHttpWrapper::HasRetriesRemaining() const
+{
+	return (CurrentRetryNumber <= MaxRetries);
+}
+
 void FGenericPlatformBackgroundHttpRequest::FGenericPlatformBackgroundHttpWrapper::UpdateHttpProgress(FHttpRequestPtr UnderlyingHttpRequest, int32 BytesSent, int32 BytesReceived)
 {
 	if (ensureAlwaysMsgf(OriginalRequest.IsValid(), TEXT("Recieved UpdateHttpProgress callback with invalid OriginalRequest pointer! Can not update request progress!")))
diff --git a/Engine/Source/Runtime/Online/BackgroundHTTP/Private/GenericPlatform/GenericPlatformBackgroundHttpResponse.cpp b/Engine/Source/Runtime/Online/BackgroundHTTP/Private/GenericPlatform/GenericPlatformBackgroundHttpResponse.cpp
--- a/Engine/Source/Runtime/Online/BackgroundHTTP/Private/GenericPlatform/GenericPlatformBackgroundHttpResponse.cpp
+++ b/Engine/Source/Runtime/Online/BackgroundHTTP/Private/GenericPlatform/GenericPlatformBackgroundHttpResponse.cpp
@@ -16,6 +16,8 @@ FGenericPlatformBackgroundHttpResponse::FGenericPlatformBackgroundHttpResponse(i
 
 FGenericPlatformBackgroundHttpResponse::FGenericPlatformBackgroundHttpResponse(FHttpRequestPtr HttpRequestIn, FHttpResponsePtr HttpResponse, bool bSuccess)
 {
+	ResponseCode = EHttpResponseCodes::Unknown;
+
 	//Don't bother making a response out of a failed HttpRequest
 	if (bSuccess && HttpResponse.IsValid())
 	{
diff --git a/Engine/Source/Runtime/Online/BackgroundHTTP/Public/GenericPlatform/GenericPlatformBackgroundHttpRequest.h b/Engine/Source/Runtime/Online/BackgroundHTTP/Public/GenericPlatform/GenericPlatformBackgroundHttpRequest.h
--- a/Engine/Source/Runtime/Online/BackgroundHTTP/Public/GenericPlatform/GenericPlatformBackgroundHttpRequest.h
+++ b/Engine/Source/Runtime/Online/BackgroundHTTP/Public/GenericPlatform/GenericPlatformBackgroundHttpRequest.h
@@ -38,6 +38,8 @@ protected:
 		void HttpRequestComplete(FHttpRequestPtr HttpRequestIn, FHttpResponsePtr HttpResponse, bool bSuccess);
 
 	private:
+		
+		bool HasRetriesRemaining() const;
 		void UpdateHttpProgress(FHttpRequestPtr UnderlyingHttpRequest, int32 BytesSent, int32 BytesReceived);
 
 		void CleanUpHttpRequest();
