Fixing BackgroundHttp temp file time out. Was not ever deleting files previously.
#tests Android Client

diff --git a/Engine/Source/Runtime/Online/BackgroundHTTP/Private/BackgroundHttpManagerImpl.cpp b/Engine/Source/Runtime/Online/BackgroundHTTP/Private/BackgroundHttpManagerImpl.cpp
--- a/Engine/Source/Runtime/Online/BackgroundHTTP/Private/BackgroundHttpManagerImpl.cpp
+++ b/Engine/Source/Runtime/Online/BackgroundHTTP/Private/BackgroundHttpManagerImpl.cpp
@@ -47,31 +47,45 @@ void FBackgroundHttpManagerImpl::Shutdown()
 
 void FBackgroundHttpManagerImpl::ClearAnyTempFilesFromTimeOut()
 {
-	UE_LOG(LogBackgroundHttpManager, Log, TEXT("Checking for BackgroundHTTP temp files that should be deleted due to time out"));
-
-	TArray<FString> FilesToCheck;
+	const FString DirectoryToCheck = FPlatformBackgroundHttp::GetTemporaryRootPath();
 
 	//Find all files in our temp folder
-	IFileManager::Get().FindFiles(FilesToCheck, *FPlatformBackgroundHttp::GetTemporaryRootPath(), nullptr);
+	TArray<FString> FilesToCheck;
+	IFileManager::Get().FindFiles(FilesToCheck, *DirectoryToCheck, nullptr);
 
 	double FileAgeTimeOutSettings = -1;
 	GConfig->GetDouble(TEXT("BackgroundHttp"), TEXT("BackgroundHttp.TempFileTimeOutSeconds"), FileAgeTimeOutSettings, GEngineIni);
 
+	UE_LOG(LogBackgroundHttpManager, Log, TEXT("Checking for BackgroundHTTP temp files that should be deleted due to time out. NumTempFilesFound:%d | TempFileTimeOutSeconds:%lf"), FilesToCheck.Num(), FileAgeTimeOutSettings);
+
 	if (FileAgeTimeOutSettings >= 0)
 	{
 		for (const FString& File : FilesToCheck)
 		{
-			const double FileAge = IFileManager::Get().GetFileAgeSeconds(*File);
+			const FString FullFilePath = FPaths::Combine(DirectoryToCheck, File);
+			
+			FFileStatData FileData = IFileManager::Get().GetStatData(*FullFilePath);
+			FTimespan TimeSinceCreate = FDateTime::UtcNow() - FileData.CreationTime;
+
+			const double FileAge = TimeSinceCreate.GetTotalSeconds();
 			const bool bShouldDelete = (FileAge > FileAgeTimeOutSettings);
 
-			UE_LOG(LogBackgroundHttpManager, Log, TEXT("FoundTempFile: %s with age %lld -- bShouldDelete:%d"), *File, FileAge, (int)bShouldDelete);
+			UE_LOG(LogBackgroundHttpManager, Log, TEXT("FoundTempFile: %s with age %lf -- bShouldDelete:%d"), *FullFilePath, FileAge, (int)bShouldDelete);
 
 			if (bShouldDelete)
 			{
-				if (!IFileManager::Get().Delete(*File))
+				if (IFileManager::Get().Delete(*FullFilePath))
 				{
-					UE_LOG(LogBackgroundHttpManager, Error, TEXT("File %s failed to delete, but should have as as it is %lld seconds old!"), *File);
+					UE_LOG(LogBackgroundHttpManager, Log, TEXT("Successfully deleted %s due to time out settings"), *FullFilePath);
 				}
+				else
+				{
+					UE_LOG(LogBackgroundHttpManager, Error, TEXT("File %s failed to delete, but should have as as it is %lld seconds old!"), *FullFilePath, FileAge);
+				}
+			}
+			else
+			{
+				UE_LOG(LogBackgroundHttpManager, Log, TEXT("SKipping Delete of %s as it is more recent then the time out settings."), *FullFilePath);
 			}
 		}
 	}
