Fix typos in variable/type names and comments

diff --git a/Engine/Source/Runtime/Core/Public/Containers/Ticker.h b/Engine/Source/Runtime/Core/Public/Containers/Ticker.h
--- a/Engine/Source/Runtime/Core/Public/Containers/Ticker.h
+++ b/Engine/Source/Runtime/Core/Public/Containers/Ticker.h
@@ -11,7 +11,7 @@
 /**
  * Ticker delegates return true to automatically reschedule at the same delay or false for one-shot.
  * You will not get more than one fire per "frame", which is just a FTicker::Tick call.
- * DeltaTime argument iz the time since the last game frame, *not* since the last tick the
+ * DeltaTime argument is the time since the last game frame, *not* since the last tick the
  * delegate received.
  */
 DECLARE_DELEGATE_RetVal_OneParam(bool, FTickerDelegate, float);
diff --git a/Engine/Source/Runtime/Online/XMPP/Private/XmppJingle/XmppConnectionJingle.h b/Engine/Source/Runtime/Online/XMPP/Private/XmppJingle/XmppConnectionJingle.h
--- a/Engine/Source/Runtime/Online/XMPP/Private/XmppJingle/XmppConnectionJingle.h
+++ b/Engine/Source/Runtime/Online/XMPP/Private/XmppJingle/XmppConnectionJingle.h
@@ -35,7 +35,7 @@ public:
 	virtual const FXmppUserJid& GetUserJid() const override;
 
 	virtual FOnXmppLoginComplete& OnLoginComplete() override { return OnXmppLoginCompleteDelegate; }
-	virtual FOnXmppLogingChanged& OnLoginChanged() override { return OnXmppLogingChangedDelegate; }
+	virtual FOnXmppLoginChanged& OnLoginChanged() override { return OnXmppLoginChangedDelegate; }
 	virtual FOnXmppLogoutComplete& OnLogoutComplete() override { return OnXmppLogoutCompleteDelegate; }
 
 	virtual IXmppPresencePtr Presence() override;
@@ -96,7 +96,7 @@ private:
 
 	// completion delegates
 	FOnXmppLoginComplete OnXmppLoginCompleteDelegate;
-	FOnXmppLogingChanged OnXmppLogingChangedDelegate;
+	FOnXmppLoginChanged OnXmppLoginChangedDelegate;
 	FOnXmppLogoutComplete OnXmppLogoutCompleteDelegate;
 
 	/** access to presence implementation */
diff --git a/Engine/Source/Runtime/Online/XMPP/Private/XmppNull.cpp b/Engine/Source/Runtime/Online/XMPP/Private/XmppNull.cpp
--- a/Engine/Source/Runtime/Online/XMPP/Private/XmppNull.cpp
+++ b/Engine/Source/Runtime/Online/XMPP/Private/XmppNull.cpp
@@ -35,7 +35,7 @@ public:
 		return UserJid;
 	}
 	virtual FOnXmppLoginComplete& OnLoginComplete() override { return OnXmppLoginCompleteDelegate; }
-	virtual FOnXmppLogingChanged& OnLoginChanged() override { return OnXmppLogingChangedDelegate; }
+	virtual FOnXmppLoginChanged& OnLoginChanged() override { return OnXmppLoginChangedDelegate; }
 	virtual FOnXmppLogoutComplete& OnLogoutComplete() override { return OnXmppLogoutCompleteDelegate; }
 
 	virtual IXmppPresencePtr Presence() override { return NULL; }
@@ -56,7 +56,7 @@ private:
 
 	// completion delegates
 	FOnXmppLoginComplete OnXmppLoginCompleteDelegate;
-	FOnXmppLogingChanged OnXmppLogingChangedDelegate;
+	FOnXmppLoginChanged OnXmppLoginChangedDelegate;
 	FOnXmppLogoutComplete OnXmppLogoutCompleteDelegate;	
 };
 
diff --git a/Engine/Source/Runtime/Online/XMPP/Private/XmppStrophe/XmppConnectionStrophe.h b/Engine/Source/Runtime/Online/XMPP/Private/XmppStrophe/XmppConnectionStrophe.h
--- a/Engine/Source/Runtime/Online/XMPP/Private/XmppStrophe/XmppConnectionStrophe.h
+++ b/Engine/Source/Runtime/Online/XMPP/Private/XmppStrophe/XmppConnectionStrophe.h
@@ -51,7 +51,7 @@ public:
 	virtual const FXmppUserJid& GetUserJid() const override;
 
 	virtual FOnXmppLoginComplete& OnLoginComplete() override { return OnXmppLoginCompleteDelegate; }
-	virtual FOnXmppLogingChanged& OnLoginChanged() override { return OnXmppLogingChangedDelegate; }
+	virtual FOnXmppLoginChanged& OnLoginChanged() override { return OnXmppLoginChangedDelegate; }
 	virtual FOnXmppLogoutComplete& OnLogoutComplete() override { return OnXmppLogoutCompleteDelegate; }
 
 	virtual IXmppMessagesPtr Messages() override;
@@ -112,7 +112,7 @@ protected:
 
 	/** Login status delegates */
 	FOnXmppLoginComplete OnXmppLoginCompleteDelegate;
-	FOnXmppLogingChanged OnXmppLogingChangedDelegate;
+	FOnXmppLoginChanged OnXmppLoginChangedDelegate;
 	FOnXmppLogoutComplete OnXmppLogoutCompleteDelegate;
 
 	/** XMPP Implementation Shared Pointers */
diff --git a/Engine/Source/Runtime/Online/XMPP/Public/XmppConnection.h b/Engine/Source/Runtime/Online/XMPP/Public/XmppConnection.h
--- a/Engine/Source/Runtime/Online/XMPP/Public/XmppConnection.h
+++ b/Engine/Source/Runtime/Online/XMPP/Public/XmppConnection.h
@@ -284,12 +284,12 @@ public:
 	* @param UserJid jid of user that changed login state
 	* @param LoginState new login state (see EXmppLoginStatus)
 	*/
-	DECLARE_MULTICAST_DELEGATE_TwoParams(FOnXmppLogingChanged, const FXmppUserJid& /*UserJid*/, EXmppLoginStatus::Type /*LoginState*/);
+	DECLARE_MULTICAST_DELEGATE_TwoParams(FOnXmppLoginChanged, const FXmppUserJid& /*UserJid*/, EXmppLoginStatus::Type /*LoginState*/);
 
 	/** @return login complete delegate */
 	virtual FOnXmppLoginComplete& OnLoginComplete() = 0;
 	/** @return login changed delegate */
-	virtual FOnXmppLogingChanged& OnLoginChanged() = 0;
+	virtual FOnXmppLoginChanged& OnLoginChanged() = 0;
 	/** @return logout complete delegate */
 	virtual FOnXmppLogoutComplete& OnLogoutComplete() = 0;
 
