#http
Temp fix for code optimization issue with http thread inner loop sleep

diff --git a/Engine/Source/Runtime/Online/HTTP/Private/HttpThread.cpp b/Engine/Source/Runtime/Online/HTTP/Private/HttpThread.cpp
--- a/Engine/Source/Runtime/Online/HTTP/Private/HttpThread.cpp
+++ b/Engine/Source/Runtime/Online/HTTP/Private/HttpThread.cpp
@@ -67,6 +67,8 @@ bool FHttpThread::Init()
 	return true;
 }
 
+// @todo samz - needed due to InnerLoopEnd,InnerLoopBegin 
+PRAGMA_DISABLE_OPTIMIZATION
 uint32 FHttpThread::Run()
 {
 	// Arrays declared outside of loop to re-use memory
@@ -107,6 +109,8 @@ uint32 FHttpThread::Run()
 	}
 	return 0;
 }
+PRAGMA_ENABLE_OPTIMIZATION
+
 
 void FHttpThread::Tick()
 {
