Retry

Added cleanup of legacy build if we succeed patching in ESP.
Leigh.Swift, Justin.Marcus

diff --git a/Engine/Source/Runtime/Core/Public/GenericPlatform/GenericPlatformInstallBundleManager.h b/Engine/Source/Runtime/Core/Public/GenericPlatform/GenericPlatformInstallBundleManager.h
--- a/Engine/Source/Runtime/Core/Public/GenericPlatform/GenericPlatformInstallBundleManager.h
+++ b/Engine/Source/Runtime/Core/Public/GenericPlatform/GenericPlatformInstallBundleManager.h
@@ -218,6 +218,8 @@ public:
 	virtual bool IsNullInterface() const = 0;
 
 	virtual void SetErrorSimulationCommands(const FString& CommandLine) {}
+
+	virtual void CleanupMigrationDirectory() const {}
 };
 
 class IPlatformInstallBundleManagerModule : public IModuleInterface
diff --git a/Engine/Source/Runtime/Online/BuildPatchServices/Private/BuildPatchServicesModule.cpp b/Engine/Source/Runtime/Online/BuildPatchServices/Private/BuildPatchServicesModule.cpp
--- a/Engine/Source/Runtime/Online/BuildPatchServices/Private/BuildPatchServicesModule.cpp
+++ b/Engine/Source/Runtime/Online/BuildPatchServices/Private/BuildPatchServicesModule.cpp
@@ -381,6 +381,17 @@ void FBuildPatchServicesModule::RegisterAppInstallation(IBuildManifestRef AppMan
 	AvailableInstallations.Add(MoveTemp(AppInstallDirectory), MoveTemp(InternalRef));
 }
 
+bool FBuildPatchServicesModule::UnregisterAppInstallation(FString AppInstallDirectory)
+{
+	FPaths::NormalizeDirectoryName(AppInstallDirectory);
+	FPaths::CollapseRelativeDirectories(AppInstallDirectory);
+	if (AvailableInstallations.Remove(AppInstallDirectory) == 1)
+	{
+		return true;
+	}
+	return false;
+}
+
 void FBuildPatchServicesModule::CancelAllInstallers(bool WaitForThreads)
 {
 	// Using a local bool for this check will improve the assert message that gets displayed
diff --git a/Engine/Source/Runtime/Online/BuildPatchServices/Private/BuildPatchServicesModule.h b/Engine/Source/Runtime/Online/BuildPatchServices/Private/BuildPatchServicesModule.h
--- a/Engine/Source/Runtime/Online/BuildPatchServices/Private/BuildPatchServicesModule.h
+++ b/Engine/Source/Runtime/Online/BuildPatchServices/Private/BuildPatchServicesModule.h
@@ -51,6 +51,7 @@ public:
 	virtual void SetAnalyticsProvider(TSharedPtr< IAnalyticsProvider > AnalyticsProvider) override;
 	virtual void SetHttpTracker(TSharedPtr< FHttpServiceTracker > HttpTracker) override;
 	virtual void RegisterAppInstallation(IBuildManifestRef AppManifest, const FString AppInstallDirectory) override;
+	virtual bool UnregisterAppInstallation(const FString AppInstallDirectory) override;
 	virtual void CancelAllInstallers(bool WaitForThreads) override;
 	virtual bool ChunkBuildDirectory(const BuildPatchServices::FChunkBuildConfiguration& Configuration) override;
 	virtual bool OptimiseChunkDelta(const BuildPatchServices::FChunkDeltaOptimiserConfiguration& Configuration) override;
diff --git a/Engine/Source/Runtime/Online/BuildPatchServices/Public/Interfaces/IBuildPatchServicesModule.h b/Engine/Source/Runtime/Online/BuildPatchServices/Public/Interfaces/IBuildPatchServicesModule.h
--- a/Engine/Source/Runtime/Online/BuildPatchServices/Public/Interfaces/IBuildPatchServicesModule.h
+++ b/Engine/Source/Runtime/Online/BuildPatchServices/Public/Interfaces/IBuildPatchServicesModule.h
@@ -159,6 +159,13 @@ public:
 	 */
 	virtual void RegisterAppInstallation(IBuildManifestRef AppManifest, const FString AppInstallDirectory) = 0;
 
+	/**
+	 * Unregisters an installation on this machine
+	 * @param AppInstallDirectory	The install location which RegisterAppInstallation was executed with
+	 * @return returns true if successfully removed installation directory
+	 */
+	virtual bool UnregisterAppInstallation(const FString AppInstallDirectory) = 0;
+
 	/**
 	 * Call to force the exit out of all current installers, optionally blocks until threads have exited and complete delegates are called.
 	 * @param WaitForThreads		If true, will block on threads exit and completion delegates
