Fix for HTTP Request GetContent/GetContentLength crashing on Curl

diff --git a/Engine/Source/Runtime/Online/HTTP/Private/Curl/CurlHttp.cpp b/Engine/Source/Runtime/Online/HTTP/Private/Curl/CurlHttp.cpp
--- a/Engine/Source/Runtime/Online/HTTP/Private/Curl/CurlHttp.cpp
+++ b/Engine/Source/Runtime/Online/HTTP/Private/Curl/CurlHttp.cpp
@@ -266,12 +266,13 @@ FString FCurlHttpRequest::GetContentType() const
 
 int32 FCurlHttpRequest::GetContentLength() const
 {
-	return RequestPayload->GetContentLength();
+	return RequestPayload.IsValid() ? RequestPayload->GetContentLength() : 0;
 }
 
 const TArray<uint8>& FCurlHttpRequest::GetContent() const
 {
-	return RequestPayload->GetContent();
+	static const TArray<uint8> EmptyContent;
+	return RequestPayload.IsValid() ? RequestPayload->GetContent() : EmptyContent;
 }
 
 void FCurlHttpRequest::SetVerb(const FString& InVerb)
