Do not try to add a cert if the bundle path is empty

diff --git a/Engine/Source/Runtime/Online/SSL/Private/SslCertificateManager.cpp b/Engine/Source/Runtime/Online/SSL/Private/SslCertificateManager.cpp
--- a/Engine/Source/Runtime/Online/SSL/Private/SslCertificateManager.cpp
+++ b/Engine/Source/Runtime/Online/SSL/Private/SslCertificateManager.cpp
@@ -293,7 +293,10 @@ void FSslCertificateManager::BuildRootCertificateArray()
 		}
 	}
 
-	AddPEMFileToRootCertificateArray(CertificateBundlePath);
+	if (!CertificateBundlePath.IsEmpty())
+	{
+		AddPEMFileToRootCertificateArray(CertificateBundlePath);
+	}
 
 	FString DebuggingCertificatePath;
 	if (GConfig->GetString(TEXT("SSL"), TEXT("DebuggingCertificatePath"), DebuggingCertificatePath, GEngineIni) && DebuggingCertificatePath.Len() > 0)
