Fix AppleHTTP memory leak

diff --git a/Engine/Source/Runtime/Online/HTTP/Private/Apple/AppleHTTP.cpp b/Engine/Source/Runtime/Online/HTTP/Private/Apple/AppleHTTP.cpp
--- a/Engine/Source/Runtime/Online/HTTP/Private/Apple/AppleHTTP.cpp
+++ b/Engine/Source/Runtime/Online/HTTP/Private/Apple/AppleHTTP.cpp
@@ -16,6 +16,100 @@
 #include "Ssl.h"
 #endif
 
+namespace UE4FCRef_Private
+{
+	template <class CFRefType>
+	struct FTollFreeBridgeType;
+
+	// TODO: Specialize this for all the types that are interchangeable between Core Foundation and Foundation
+	template<> struct FTollFreeBridgeType<CFArrayRef> { using Type = NSArray*; };
+	template<> struct FTollFreeBridgeType<CFBooleanRef> { using Type = NSNumber*; };
+	template<> struct FTollFreeBridgeType<CFDataRef> { using Type = NSData*; };
+	template<> struct FTollFreeBridgeType<CFNumberRef> { using Type = NSNumber*; };
+	template<> struct FTollFreeBridgeType<CFStringRef> { using Type = NSString*; };
+}
+
+template <class CFRefType>
+class FCFRef
+{
+public:
+	FCFRef()
+	: Ref(nullptr)
+	{
+	}
+	
+	FCFRef(CFRefType InRef)
+	: Ref(InRef)
+	{
+	}
+	
+	FCFRef(const FCFRef& Other)
+	{
+		CFRetain(Other.Ref);
+		Ref = Other.Ref;
+	}
+	
+	FCFRef(FCFRef&& Other)
+	{
+		Ref = Other.Ref;
+		Other.Ref = nullptr;
+	}
+	
+	~FCFRef()
+	{
+		if (Ref)
+		{
+			CFRelease(Ref);
+		}
+	}
+	
+	FCFRef& operator=(const FCFRef& Other)
+	{
+		CFRetain(Other.Ref);
+		Ref = Other.Ref;
+		return *this;
+	}
+	
+	FCFRef& operator=(FCFRef&& Other)
+	{
+		if (this != &Other)
+		{
+			if (Ref)
+			{
+				CFRelease(Ref);
+			}
+			Ref = Other.Ref;
+			Other.Ref = nullptr;
+		}
+		return *this;
+	}
+	
+	CFRefType& GetRef()
+	{
+		return Ref;
+	}
+	
+	explicit operator bool() const
+	{
+		return Ref != nullptr;
+	}
+	
+	operator CFRefType() const
+	{
+		return Ref;
+	}
+	
+	template<typename T = CFRefType,
+			 typename = typename UE4FCRef_Private::FTollFreeBridgeType<T>::Type>
+	operator typename UE4FCRef_Private::FTollFreeBridgeType<T>::Type ()
+	{
+		return static_cast<typename UE4FCRef_Private::FTollFreeBridgeType<T>::Type>(Ref);
+	}
+
+private:
+	CFRefType Ref;
+};
+
 /****************************************************************************
  * FAppleHttpRequest implementation
  ***************************************************************************/
@@ -636,21 +730,21 @@ static const unsigned char ecdsaSecp384r1Asn1Header[] =
                 // this is not great, but the only way to extract a public key from a SecCertificateRef
                 // is to create an individual SecTrustRef for each cert that only contains itself and then
                 // evaluate that against an empty X509 policy.
-                SecTrustRef CertTrust;
-                SecPolicyRef TrustPolicy = SecPolicyCreateBasicX509();
-                SecTrustCreateWithCertificates(Cert, TrustPolicy, &CertTrust);
+                FCFRef<SecTrustRef> CertTrust;
+                FCFRef<SecPolicyRef> TrustPolicy = SecPolicyCreateBasicX509();
+                SecTrustCreateWithCertificates(Cert, TrustPolicy, &CertTrust.GetRef());
                 SecTrustResultType CertEvalResult;
                 SecTrustEvaluate(CertTrust, &CertEvalResult);
-                SecKeyRef CertPubKey = SecTrustCopyPublicKey(CertTrust);
-                NSData *CertPubKeyData = (NSData *)SecKeyCopyExternalRepresentation(CertPubKey, NULL);
-                if (CertPubKeyData == nil)
+                FCFRef<SecKeyRef> CertPubKey = SecTrustCopyPublicKey(CertTrust);
+				FCFRef<CFDataRef> CertPubKeyData = SecKeyCopyExternalRepresentation(CertPubKey, NULL);
+                if (!CertPubKeyData)
                 {
                     UE_LOG(LogHttp, Warning, TEXT("could not extract public key from certificate %i for domain '%s'; skipping!"), i, *RemoteHost);
                     continue;
                 }
                 
 				// we got the key. now we have to figure out what type of key it is; thanks, CommonCrypto.
-                CFDictionaryRef CertPubKeyAttr = SecKeyCopyAttributes(CertPubKey);
+                FCFRef<CFDictionaryRef> CertPubKeyAttr = SecKeyCopyAttributes(CertPubKey);
                 NSString *CertPubKeyType = static_cast<NSString *>(CFDictionaryGetValue(CertPubKeyAttr, kSecAttrKeyType));
                 NSNumber *CertPubKeySize = static_cast<NSNumber *>(CFDictionaryGetValue(CertPubKeyAttr, kSecAttrKeySizeInBits));
                 char *CertPubKeyASN1Header;
@@ -697,13 +791,13 @@ static const unsigned char ecdsaSecp384r1Asn1Header[] =
                     UE_LOG(LogHttp, Log, TEXT("unsupported key type (not RSA or ECDSA) for certificate %i for domain '%s'; skipping!"), i, *RemoteHost);
                     continue;
                 }
-                
+				
                 UE_LOG(LogHttp, VeryVerbose, TEXT("constructed key header: [%d] %s"), CertPubKeyASN1HeaderSize, UTF8_TO_TCHAR([[[NSData dataWithBytes:CertPubKeyASN1Header length:CertPubKeyASN1HeaderSize] description] UTF8String]));
-                UE_LOG(LogHttp, VeryVerbose, TEXT("current pubkey: [%d] %s"), [CertPubKeyData length], UTF8_TO_TCHAR([[[NSData dataWithBytes:[CertPubKeyData bytes] length:[CertPubKeyData length]] description] UTF8String]));
+                UE_LOG(LogHttp, VeryVerbose, TEXT("current pubkey: [%d] %s"), [(NSData*)CertPubKeyData length], UTF8_TO_TCHAR([[[NSData dataWithBytes:[(NSData*)CertPubKeyData bytes] length:[(NSData*)CertPubKeyData length]] description] UTF8String]));
                 
                 // smash 'em together to get a proper key with an ASN.1 header
                 NSMutableData *ReconstructedPubKey = [NSMutableData data];
-                [ReconstructedPubKey appendData:[NSData dataWithBytes:CertPubKeyASN1Header length:CertPubKeyASN1HeaderSize]];
+                [ReconstructedPubKey appendBytes:CertPubKeyASN1Header length:CertPubKeyASN1HeaderSize];
                 [ReconstructedPubKey appendData:CertPubKeyData];
                 UE_LOG(LogHttp, VeryVerbose, TEXT("reconstructed key: [%d] %s"), [ReconstructedPubKey length], UTF8_TO_TCHAR([[ReconstructedPubKey description] UTF8String]));
                 
@@ -712,7 +806,6 @@ static const unsigned char ecdsaSecp384r1Asn1Header[] =
                 if (!CC_SHA256([ReconstructedPubKey bytes], (CC_LONG)[ReconstructedPubKey length], CertCalcDigest.GetData()))
                 {
                     UE_LOG(LogHttp, Warning, TEXT("could not calculate SHA256 digest of public key %d for domain '%s'; skipping!"), i, *RemoteHost);
-                    continue;
                 }
                 else
                 {
