Increasing default notification schedule time for the BackgroundHttpNotificationObject to 15s instead of 1s to make it more reliable when the main thread dispatch of the notification takes a little longer then usual.
#tests iOS Client

diff --git a/Engine/Source/Runtime/Online/BackgroundHTTP/Private/BackgroundHttpNotificationObject.cpp b/Engine/Source/Runtime/Online/BackgroundHTTP/Private/BackgroundHttpNotificationObject.cpp
--- a/Engine/Source/Runtime/Online/BackgroundHTTP/Private/BackgroundHttpNotificationObject.cpp
+++ b/Engine/Source/Runtime/Online/BackgroundHTTP/Private/BackgroundHttpNotificationObject.cpp
@@ -42,7 +42,7 @@ FBackgroundHttpNotificationObject::~FBackgroundHttpNotificationObject()
 		{
 			//make a notification 1 second from now
 			FDateTime TargetTime = FDateTime::Now();
-			TargetTime += FTimespan::FromSeconds(1);
+			TargetTime += FTimespan::FromSeconds(15);
 
 			if (nullptr != PlatformNotificationService)
 			{
