Async tasks need to keep ticking the Engine - XMPP Strophe modules

diff --git a/Engine/Source/Runtime/Online/XMPP/Private/XmppStrophe/XmppMessagesStrophe.cpp b/Engine/Source/Runtime/Online/XMPP/Private/XmppStrophe/XmppMessagesStrophe.cpp
--- a/Engine/Source/Runtime/Online/XMPP/Private/XmppStrophe/XmppMessagesStrophe.cpp
+++ b/Engine/Source/Runtime/Online/XMPP/Private/XmppStrophe/XmppMessagesStrophe.cpp
@@ -14,17 +14,25 @@
 #include "Serialization/JsonSerializer.h"
 #include "Policies/CondensedJsonPrintPolicy.h"
 #include "Misc/Guid.h"
+#include "Misc/EmbeddedCommunication.h"
 
 #if WITH_XMPP_STROPHE
 
+const FName FXmppMessagesStrophe::TickRequesterId = FName("StropheMessages");
+
 FXmppMessagesStrophe::FXmppMessagesStrophe(FXmppConnectionStrophe& InConnectionManager)
 	: ConnectionManager(InConnectionManager)
 {
 }
 
+FXmppMessagesStrophe::~FXmppMessagesStrophe()
+{
+	CleanupMessages();
+}
+
 void FXmppMessagesStrophe::OnDisconnect()
 {
-	IncomingMessages.Empty();
+	CleanupMessages();
 }
 
 bool FXmppMessagesStrophe::ReceiveStanza(const FStropheStanza& IncomingStanza)
@@ -88,6 +96,7 @@ bool FXmppMessagesStrophe::HandleMessageStanza(const FStropheStanza& IncomingSta
 		}
 	}
 
+	FEmbeddedCommunication::KeepAwake(TickRequesterId, false);
 	IncomingMessages.Enqueue(MakeUnique<FXmppMessage>(MoveTemp(Message)));
 	return true;
 }
@@ -179,6 +188,7 @@ bool FXmppMessagesStrophe::Tick(float DeltaTime)
 		TUniquePtr<FXmppMessage> Message;
 		if (IncomingMessages.Dequeue(Message))
 		{
+			FEmbeddedCommunication::AllowSleep(TickRequesterId);
 			OnMessageReceived(MoveTemp(Message));
 		}
 	}
@@ -192,4 +202,14 @@ void FXmppMessagesStrophe::OnMessageReceived(TUniquePtr<FXmppMessage>&& Message)
 	OnMessageReceivedDelegate.Broadcast(ConnectionManager.AsShared(), MessageRef->FromJid, MessageRef);
 }
 
-#endif
+void FXmppMessagesStrophe::CleanupMessages()
+{
+	while (!IncomingMessages.IsEmpty())
+	{
+		TUniquePtr<FXmppMessage> Message;
+		IncomingMessages.Dequeue(Message);
+		FEmbeddedCommunication::AllowSleep(TickRequesterId);
+	}
+}
+
+#endif
\ No newline at end of file
diff --git a/Engine/Source/Runtime/Online/XMPP/Private/XmppStrophe/XmppMessagesStrophe.h b/Engine/Source/Runtime/Online/XMPP/Private/XmppStrophe/XmppMessagesStrophe.h
--- a/Engine/Source/Runtime/Online/XMPP/Private/XmppStrophe/XmppMessagesStrophe.h
+++ b/Engine/Source/Runtime/Online/XMPP/Private/XmppStrophe/XmppMessagesStrophe.h
@@ -20,7 +20,7 @@ class FXmppMessagesStrophe
 public:
 	// FXmppMessagesStrophe
 	FXmppMessagesStrophe(FXmppConnectionStrophe& InConnectionManager);
-	~FXmppMessagesStrophe() = default;
+	~FXmppMessagesStrophe();
 
 	// XMPP Thread
 	void OnDisconnect();
@@ -42,6 +42,9 @@ public:
 protected:
 	void OnMessageReceived(TUniquePtr<FXmppMessage>&& Message);
 
+	/** Remove pending messages and engine KeepAwake calls */
+	void CleanupMessages();
+
 protected:
 	/** Connection manager controls sending data to XMPP thread */
 	FXmppConnectionStrophe& ConnectionManager;
@@ -54,6 +57,9 @@ protected:
 
 	/** Delegate for game to listen to messages */
 	FOnXmppMessageReceived OnMessageReceivedDelegate;
+
+	/* Engine Tick Requester name */
+	const static FName TickRequesterId;
 };
 
-#endif
+#endif
\ No newline at end of file
diff --git a/Engine/Source/Runtime/Online/XMPP/Private/XmppStrophe/XmppMultiUserChatStrophe.cpp b/Engine/Source/Runtime/Online/XMPP/Private/XmppStrophe/XmppMultiUserChatStrophe.cpp
--- a/Engine/Source/Runtime/Online/XMPP/Private/XmppStrophe/XmppMultiUserChatStrophe.cpp
+++ b/Engine/Source/Runtime/Online/XMPP/Private/XmppStrophe/XmppMultiUserChatStrophe.cpp
@@ -9,27 +9,29 @@
 
 #include "Logging/LogScopedVerbosityOverride.h"
 #include "Misc/Guid.h"
+#include "Misc/EmbeddedCommunication.h"
 
 #if WITH_XMPP_STROPHE
 
+const FName FXmppMultiUserChatStrophe::TickRequesterId = FName("StropheMultichat");
+
 FXmppMultiUserChatStrophe::FXmppMultiUserChatStrophe(FXmppConnectionStrophe& InConnectionManager)
 	: ConnectionManager(InConnectionManager)
 {
 
 }
 
+FXmppMultiUserChatStrophe::~FXmppMultiUserChatStrophe()
+{
+	CleanupMessages();
+}
+
 void FXmppMultiUserChatStrophe::OnDisconnect()
 {
+	CleanupMessages();
+
 	Chatrooms.Empty();
 	PendingRoomCreateConfigs.Empty();
-
-	IncomingMucPresenceUpdates.Empty();
-	IncomingMucPresenceErrors.Empty();
-	IncomingGroupChatMessages.Empty();
-	IncomingRoomSubjects.Empty();
-	IncomingRoomConfigErrors.Empty();
-	IncomingRoomConfigWriteSuccesses.Empty();
-	IncomingRoomInfoUpdates.Empty();
 }
 
 bool FXmppMultiUserChatStrophe::ReceiveStanza(const FStropheStanza& IncomingStanza)
@@ -120,6 +122,7 @@ bool FXmppMultiUserChatStrophe::HandlePresenceStanza(const FStropheStanza& Incom
 		}
 	}
 
+	FEmbeddedCommunication::KeepAwake(TickRequesterId, false);
 	IncomingMucPresenceUpdates.Enqueue(MakeUnique<FXmppMucPresence>(MoveTemp(Presence)));
 	return true;
 }
@@ -186,6 +189,7 @@ bool FXmppMultiUserChatStrophe::HandlePresenceErrorStanza(const FStropheStanza&
 
 	if (!OutError.ErrorMessage.IsEmpty())
 	{
+		FEmbeddedCommunication::KeepAwake(TickRequesterId, false);
 		IncomingMucPresenceErrors.Enqueue(MoveTemp(OutError));
 	}
 
@@ -201,6 +205,7 @@ bool FXmppMultiUserChatStrophe::HandleGroupChatStanza(const FStropheStanza& Inco
 		SubjectUpdate.NewSubject = SubjectStanza->GetText();
 		SubjectUpdate.RoomId = IncomingStanza.GetFrom().Id;
 
+		FEmbeddedCommunication::KeepAwake(TickRequesterId, false);
 		IncomingRoomSubjects.Enqueue(MoveTemp(SubjectUpdate));
 		return true;
 	}
@@ -219,6 +224,7 @@ bool FXmppMultiUserChatStrophe::HandleGroupChatStanza(const FStropheStanza& Inco
 				FString Code = XChild.GetAttribute(Strophe::SA_CODE);
 				if (Code == Strophe::SC_104)
 				{
+					FEmbeddedCommunication::KeepAwake(TickRequesterId, false);
 					IncomingRoomInfoUpdates.Enqueue(FXmppRoomId(IncomingStanza.GetFrom().Id));
 				}
 			}
@@ -255,6 +261,7 @@ bool FXmppMultiUserChatStrophe::HandleGroupChatStanza(const FStropheStanza& Inco
 		ChatMessage.Timestamp = FDateTime::UtcNow();
 	}
 
+	FEmbeddedCommunication::KeepAwake(TickRequesterId, false);
 	IncomingGroupChatMessages.Enqueue(MakeUnique<FXmppChatMessage>(MoveTemp(ChatMessage)));
 	return true;
 }
@@ -317,6 +324,7 @@ bool FXmppMultiUserChatStrophe::HandleRoomConfigStanza(const FStropheStanza& Inc
 	TOptional<const FStropheStanza> QueryStanza = IncomingStanza.GetChild(Strophe::SN_QUERY);
 	if (!QueryStanza.IsSet())
 	{
+		FEmbeddedCommunication::KeepAwake(TickRequesterId, false);
 		IncomingRoomConfigWriteSuccesses.Enqueue(IncomingStanza.GetFrom().Id);
 		return true;
 	}
@@ -347,6 +355,7 @@ bool FXmppMultiUserChatStrophe::HandleRoomConfigErrorStanza(const FStropheStanza
 
 	if (!OutError.ErrorMessage.IsEmpty())
 	{
+		FEmbeddedCommunication::KeepAwake(TickRequesterId, false);
 		IncomingRoomConfigErrors.Enqueue(MoveTemp(OutError));
 	}
 
@@ -788,6 +797,7 @@ bool FXmppMultiUserChatStrophe::Tick(float DeltaTime)
 		TUniquePtr<FXmppMucPresence> MucPresence;
 		if (IncomingMucPresenceUpdates.Dequeue(MucPresence))
 		{
+			FEmbeddedCommunication::AllowSleep(TickRequesterId);
 			check(MucPresence.IsValid());
 			OnReceiveMucPresence(MoveTemp(*MucPresence));
 		}
@@ -797,6 +807,7 @@ bool FXmppMultiUserChatStrophe::Tick(float DeltaTime)
 		FXmppStropheErrorPair ErrorInfo;
 		if (IncomingMucPresenceErrors.Dequeue(ErrorInfo))
 		{
+			FEmbeddedCommunication::AllowSleep(TickRequesterId);
 			OnReceiveMucPresenceError(MoveTemp(ErrorInfo));
 		}
 	}
@@ -805,6 +816,7 @@ bool FXmppMultiUserChatStrophe::Tick(float DeltaTime)
 		TUniquePtr<FXmppChatMessage> GroupChatMessage;
 		if (IncomingGroupChatMessages.Dequeue(GroupChatMessage))
 		{
+			FEmbeddedCommunication::AllowSleep(TickRequesterId);
 			OnReceiveGroupChatMessage(MoveTemp(GroupChatMessage));
 		}
 	}
@@ -813,6 +825,7 @@ bool FXmppMultiUserChatStrophe::Tick(float DeltaTime)
 		FXmppStropheSubjectUpdate NewSubject;
 		if (IncomingRoomSubjects.Dequeue(NewSubject))
 		{
+			FEmbeddedCommunication::AllowSleep(TickRequesterId);
 			OnReceiveGroupChatSubject(MoveTemp(NewSubject));
 		}
 	}
@@ -821,6 +834,7 @@ bool FXmppMultiUserChatStrophe::Tick(float DeltaTime)
 		FXmppStropheErrorPair RoomConfigError;
 		if (IncomingRoomConfigErrors.Dequeue(RoomConfigError))
 		{
+			FEmbeddedCommunication::AllowSleep(TickRequesterId);
 			OnReceiveRoomConfigError(MoveTemp(RoomConfigError));
 		}
 	}
@@ -829,6 +843,7 @@ bool FXmppMultiUserChatStrophe::Tick(float DeltaTime)
 		FXmppRoomId RoomId;
 		if (IncomingRoomConfigWriteSuccesses.Dequeue(RoomId))
 		{
+			FEmbeddedCommunication::AllowSleep(TickRequesterId);
 			OnReceiveRoomConfigSuccess(MoveTemp(RoomId));
 		}
 	}
@@ -837,6 +852,7 @@ bool FXmppMultiUserChatStrophe::Tick(float DeltaTime)
 		FXmppRoomId RoomId;
 		if (IncomingRoomInfoUpdates.Dequeue(RoomId))
 		{
+			FEmbeddedCommunication::AllowSleep(TickRequesterId);
 			OnReceieveRoomInfoUpdate(MoveTemp(RoomId));
 		}
 	}
@@ -1299,4 +1315,56 @@ void FXmppMultiUserChatStrophe::HandleRoomMemberLeft(FXmppRoomStrophe& Room, FXm
 	OnXmppRoomMemberExitDelegate.Broadcast(ConnectionManager.AsShared(), Room.GetRoomId(), MemberPresence.UserJid);
 }
 
+void FXmppMultiUserChatStrophe::CleanupMessages()
+{
+	while (!IncomingMucPresenceUpdates.IsEmpty())
+	{
+		TUniquePtr<FXmppMucPresence> MucPresence;
+		IncomingMucPresenceUpdates.Dequeue(MucPresence);
+		FEmbeddedCommunication::AllowSleep(TickRequesterId);
+	}
+
+	while (!IncomingMucPresenceErrors.IsEmpty())
+	{
+		FXmppStropheErrorPair ErrorInfo;
+		IncomingMucPresenceErrors.Dequeue(ErrorInfo);
+		FEmbeddedCommunication::AllowSleep(TickRequesterId);
+	}
+
+	while (!IncomingGroupChatMessages.IsEmpty())
+	{
+		TUniquePtr<FXmppChatMessage> GroupChatMessage;
+		IncomingGroupChatMessages.Dequeue(GroupChatMessage);
+		FEmbeddedCommunication::AllowSleep(TickRequesterId);
+	}
+
+	while (!IncomingRoomSubjects.IsEmpty())
+	{
+		FXmppStropheSubjectUpdate NewSubject;
+		IncomingRoomSubjects.Dequeue(NewSubject);
+		FEmbeddedCommunication::AllowSleep(TickRequesterId);
+	}
+
+	while (!IncomingRoomConfigErrors.IsEmpty())
+	{
+		FXmppStropheErrorPair RoomConfigError;
+		IncomingRoomConfigErrors.Dequeue(RoomConfigError);
+		FEmbeddedCommunication::AllowSleep(TickRequesterId);
+	}
+
+	while (!IncomingRoomConfigWriteSuccesses.IsEmpty())
+	{
+		FXmppRoomId RoomId;
+		IncomingRoomConfigWriteSuccesses.Dequeue(RoomId);
+		FEmbeddedCommunication::AllowSleep(TickRequesterId);
+	}
+
+	while (!IncomingRoomInfoUpdates.IsEmpty())
+	{
+		FXmppRoomId RoomId;
+		IncomingRoomInfoUpdates.Dequeue(RoomId);
+		FEmbeddedCommunication::AllowSleep(TickRequesterId);
+	}
+}
+
 #endif
diff --git a/Engine/Source/Runtime/Online/XMPP/Private/XmppStrophe/XmppMultiUserChatStrophe.h b/Engine/Source/Runtime/Online/XMPP/Private/XmppStrophe/XmppMultiUserChatStrophe.h
--- a/Engine/Source/Runtime/Online/XMPP/Private/XmppStrophe/XmppMultiUserChatStrophe.h
+++ b/Engine/Source/Runtime/Online/XMPP/Private/XmppStrophe/XmppMultiUserChatStrophe.h
@@ -145,7 +145,7 @@ class FXmppMultiUserChatStrophe
 public:
 	// FXmppMultiUserChatStrophe
 	FXmppMultiUserChatStrophe(class FXmppConnectionStrophe& InConnectionManager);
-	virtual ~FXmppMultiUserChatStrophe() = default;
+	virtual ~FXmppMultiUserChatStrophe();
 
 	void OnDisconnect();
 	bool ReceiveStanza(const FStropheStanza& IncomingStanza);
@@ -216,6 +216,9 @@ protected:
 	void HandleRoomMemberChanged(FXmppRoomStrophe& Room, FXmppMucPresence&& MemberPresence);
 	void HandleRoomMemberLeft(FXmppRoomStrophe& Room, FXmppMucPresence&& MemberPresence);
 
+	/** Remove pending messages and engine KeepAwake calls */
+	void CleanupMessages();
+
 protected:
 	/** Connection manager controls sending data to XMPP thread */
 	FXmppConnectionStrophe& ConnectionManager;
@@ -254,5 +257,8 @@ protected:
 	FOnXmppRoomMemberExit OnXmppRoomMemberExitDelegate;
 	FOnXmppRoomMemberChanged OnXmppRoomMemberChangedDelegate;
 	FOnXmppRoomChatReceived OnXmppRoomChatReceivedDelegate;
+
+	/* Engine Tick Requester name */
+	const static FName TickRequesterId;
 };
-#endif
+#endif
\ No newline at end of file
diff --git a/Engine/Source/Runtime/Online/XMPP/Private/XmppStrophe/XmppPingStrophe.cpp b/Engine/Source/Runtime/Online/XMPP/Private/XmppStrophe/XmppPingStrophe.cpp
--- a/Engine/Source/Runtime/Online/XMPP/Private/XmppStrophe/XmppPingStrophe.cpp
+++ b/Engine/Source/Runtime/Online/XMPP/Private/XmppStrophe/XmppPingStrophe.cpp
@@ -6,9 +6,12 @@
 #include "XmppStrophe/XmppConnectionStrophe.h"
 #include "XmppLog.h"
 #include "Misc/Guid.h"
+#include "Misc/EmbeddedCommunication.h"
 
 #if WITH_XMPP_STROPHE
 
+const FName FXmppPingStrophe::TickRequesterId = FName("StrophePing");
+
 FXmppPingStrophe::FXmppPingStrophe(FXmppConnectionStrophe& InConnectionManager)
 	: ConnectionManager(InConnectionManager)
 	, TimeSinceLastClientPing(0.0f)
@@ -19,6 +22,11 @@ FXmppPingStrophe::FXmppPingStrophe(FXmppConnectionStrophe& InConnectionManager)
 {
 }
 
+FXmppPingStrophe::~FXmppPingStrophe()
+{
+	CleanupMessages();
+}
+
 bool FXmppPingStrophe::Tick(float DeltaTime)
 {
 	// Process our ping queue and send pongs
@@ -27,6 +35,7 @@ bool FXmppPingStrophe::Tick(float DeltaTime)
 		FXmppPingReceivedInfo ReceivedPing;
 		if (IncomingPings.Dequeue(ReceivedPing))
 		{
+			FEmbeddedCommunication::AllowSleep(TickRequesterId);
 			ReplyToPing(MoveTemp(ReceivedPing));
 		}
 	}
@@ -51,8 +60,7 @@ bool FXmppPingStrophe::Tick(float DeltaTime)
 
 void FXmppPingStrophe::OnDisconnect()
 {
-	// Clear out pending pongs when we disconnect
-	IncomingPings.Empty();
+	CleanupMessages();
 
 	// Reset our client ping timer
 	ResetPingTimer();
@@ -122,6 +130,7 @@ void FXmppPingStrophe::ResetPongTimer()
 
 bool FXmppPingStrophe::HandlePingStanza(const FStropheStanza& PingStanza)
 {
+	FEmbeddedCommunication::KeepAwake(TickRequesterId, false);
 	IncomingPings.Enqueue(FXmppPingReceivedInfo(PingStanza.GetFrom(), PingStanza.GetId()));
 	return true;
 }
@@ -204,4 +213,14 @@ void FXmppPingStrophe::CheckXmppPongTimeout(float DeltaTime)
 	}
 }
 
+void FXmppPingStrophe::CleanupMessages()
+{
+	while (!IncomingPings.IsEmpty())
+	{
+		FXmppPingReceivedInfo ReceivedPing;
+		IncomingPings.Dequeue(ReceivedPing);
+		FEmbeddedCommunication::AllowSleep(TickRequesterId);
+	}
+}
+
 #endif
diff --git a/Engine/Source/Runtime/Online/XMPP/Private/XmppStrophe/XmppPingStrophe.h b/Engine/Source/Runtime/Online/XMPP/Private/XmppStrophe/XmppPingStrophe.h
--- a/Engine/Source/Runtime/Online/XMPP/Private/XmppStrophe/XmppPingStrophe.h
+++ b/Engine/Source/Runtime/Online/XMPP/Private/XmppStrophe/XmppPingStrophe.h
@@ -38,7 +38,7 @@ class FXmppPingStrophe
 public:
 	// FXmppPingStrophe
 	FXmppPingStrophe(FXmppConnectionStrophe& InConnectionManager);
-	virtual ~FXmppPingStrophe() = default;
+	virtual ~FXmppPingStrophe();
 
 	// FTickerObjectBase
 	virtual bool Tick(float DeltaTime) override;
@@ -68,6 +68,12 @@ protected:
 	/** Check the current status of our Xmpp ping responses; may lead to disconnection requests */
 	void CheckXmppPongTimeout(float DeltaTime);
 
+	/** Remove pending messages and engine KeepAwake calls */
+	void CleanupMessages();
+
+	/* Engine Tick Requester name */
+	const static FName TickRequesterId;
+
 private:
 	/** Connection manager controls sending data to XMPP thread */
 	FXmppConnectionStrophe& ConnectionManager;
diff --git a/Engine/Source/Runtime/Online/XMPP/Private/XmppStrophe/XmppPresenceStrophe.cpp b/Engine/Source/Runtime/Online/XMPP/Private/XmppStrophe/XmppPresenceStrophe.cpp
--- a/Engine/Source/Runtime/Online/XMPP/Private/XmppStrophe/XmppPresenceStrophe.cpp
+++ b/Engine/Source/Runtime/Online/XMPP/Private/XmppStrophe/XmppPresenceStrophe.cpp
@@ -4,22 +4,31 @@
 #include "XmppStrophe/XmppConnectionStrophe.h"
 #include "XmppStrophe/StropheStanza.h"
 #include "XmppStrophe/StropheStanzaConstants.h"
+#include "Misc/EmbeddedCommunication.h"
 
 #if WITH_XMPP_STROPHE
 
+const FName FXmppPresenceStrophe::TickRequesterId = FName("StrophePresence");
+
 FXmppPresenceStrophe::FXmppPresenceStrophe(FXmppConnectionStrophe& InConnectionManager)
 	: ConnectionManager(InConnectionManager)
 {
 
 }
 
+FXmppPresenceStrophe::~FXmppPresenceStrophe()
+{
+	CleanupMessages();
+}
+
 void FXmppPresenceStrophe::OnDisconnect()
 {
 	if (RosterMembers.Num() > 0)
 	{
 		RosterMembers.Empty();
 	}
-	IncomingPresenceUpdates.Empty();
+
+	CleanupMessages();
 }
 
 bool FXmppPresenceStrophe::ReceiveStanza(const FStropheStanza& IncomingStanza)
@@ -96,6 +105,7 @@ bool FXmppPresenceStrophe::ReceiveStanza(const FStropheStanza& IncomingStanza)
 		Presence.UserJid.ParseResource(Presence.AppId, Presence.Platform, UnusedPlatformUserId);
 	}
 
+	FEmbeddedCommunication::KeepAwake(TickRequesterId, false);
 	return IncomingPresenceUpdates.Enqueue(MakeUnique<FXmppUserPresence>(MoveTemp(Presence)));
 }
 
@@ -200,6 +210,7 @@ bool FXmppPresenceStrophe::Tick(float DeltaTime)
 		TUniquePtr<FXmppUserPresence> PresencePtr;
 		if (IncomingPresenceUpdates.Dequeue(PresencePtr))
 		{
+			FEmbeddedCommunication::AllowSleep(TickRequesterId);
 			check(PresencePtr.IsValid());
 			OnPresenceUpdate(MoveTemp(PresencePtr));
 		}
@@ -216,4 +227,14 @@ void FXmppPresenceStrophe::OnPresenceUpdate(TUniquePtr<FXmppUserPresence>&& NewP
 	OnXmppPresenceReceivedDelegate.Broadcast(ConnectionManager.AsShared(), Presence->UserJid, Presence);
 }
 
-#endif
+void FXmppPresenceStrophe::CleanupMessages()
+{
+	while (!IncomingPresenceUpdates.IsEmpty())
+	{
+		TUniquePtr<FXmppUserPresence> PresencePtr;
+		IncomingPresenceUpdates.Dequeue(PresencePtr);
+		FEmbeddedCommunication::AllowSleep(TickRequesterId);
+	}
+}
+
+#endif
\ No newline at end of file
diff --git a/Engine/Source/Runtime/Online/XMPP/Private/XmppStrophe/XmppPresenceStrophe.h b/Engine/Source/Runtime/Online/XMPP/Private/XmppStrophe/XmppPresenceStrophe.h
--- a/Engine/Source/Runtime/Online/XMPP/Private/XmppStrophe/XmppPresenceStrophe.h
+++ b/Engine/Source/Runtime/Online/XMPP/Private/XmppStrophe/XmppPresenceStrophe.h
@@ -21,7 +21,7 @@ class FXmppPresenceStrophe
 public:
 	// FXmppPresenceStrophe
 	FXmppPresenceStrophe(FXmppConnectionStrophe& InConnectionManager);
-	virtual ~FXmppPresenceStrophe() = default;
+	virtual ~FXmppPresenceStrophe();
 
 	// XMPP Thread
 
@@ -44,6 +44,9 @@ public:
 protected:
 	void OnPresenceUpdate(TUniquePtr<FXmppUserPresence>&& UpdatedPresence);
 
+	/** Remove pending messages and engine KeepAwake calls */
+	void CleanupMessages();
+
 protected:
 	/** Connection manager controls sending data to XMPP thread */
 	FXmppConnectionStrophe& ConnectionManager;
@@ -59,6 +62,9 @@ protected:
 
 	/** Delegate to signal presence information has been received for a user */
 	FOnXmppPresenceReceived OnXmppPresenceReceivedDelegate;
+
+	/* Engine Tick Requester name */
+	const static FName TickRequesterId;
 };
 
 #endif
diff --git a/Engine/Source/Runtime/Online/XMPP/Private/XmppStrophe/XmppPrivateChatStrophe.cpp b/Engine/Source/Runtime/Online/XMPP/Private/XmppStrophe/XmppPrivateChatStrophe.cpp
--- a/Engine/Source/Runtime/Online/XMPP/Private/XmppStrophe/XmppPrivateChatStrophe.cpp
+++ b/Engine/Source/Runtime/Online/XMPP/Private/XmppStrophe/XmppPrivateChatStrophe.cpp
@@ -6,17 +6,25 @@
 #include "XmppStrophe/StropheStanza.h"
 #include "XmppStrophe/StropheStanzaConstants.h"
 #include "XmppLog.h"
+#include "Misc/EmbeddedCommunication.h"
 
 #if WITH_XMPP_STROPHE
 
+const FName FXmppPrivateChatStrophe::TickRequesterId = FName("StrophePrivateChat");
+
 FXmppPrivateChatStrophe::FXmppPrivateChatStrophe(FXmppConnectionStrophe& InConnectionManager)
 	: ConnectionManager(InConnectionManager)
 {
 }
 
+FXmppPrivateChatStrophe::~FXmppPrivateChatStrophe()
+{
+	CleanupMessages();
+}
+
 void FXmppPrivateChatStrophe::OnDisconnect()
 {
-	IncomingChatMessages.Empty();
+	CleanupMessages();
 }
 
 bool FXmppPrivateChatStrophe::ReceiveStanza(const FStropheStanza& IncomingStanza)
@@ -72,6 +80,7 @@ bool FXmppPrivateChatStrophe::ReceiveStanza(const FStropheStanza& IncomingStanza
 		ChatMessage.Timestamp = FDateTime::UtcNow();
 	}
 
+	FEmbeddedCommunication::KeepAwake(TickRequesterId, false);
 	IncomingChatMessages.Enqueue(MakeUnique<FXmppChatMessage>(MoveTemp(ChatMessage)));
 	return true;
 }
@@ -106,6 +115,7 @@ bool FXmppPrivateChatStrophe::Tick(float DeltaTime)
 		TUniquePtr<FXmppChatMessage> ChatMessage;
 		if (IncomingChatMessages.Dequeue(ChatMessage))
 		{
+			FEmbeddedCommunication::AllowSleep(TickRequesterId);
 			check(ChatMessage);
 			OnChatReceived(MoveTemp(ChatMessage));
 		}
@@ -120,4 +130,14 @@ void FXmppPrivateChatStrophe::OnChatReceived(TUniquePtr<FXmppChatMessage>&& Chat
 	OnChatReceivedDelegate.Broadcast(ConnectionManager.AsShared(), ChatRef->FromJid, ChatRef);
 }
 
+void FXmppPrivateChatStrophe::CleanupMessages()
+{
+	while (!IncomingChatMessages.IsEmpty())
+	{
+		TUniquePtr<FXmppChatMessage> ChatMessage;
+		IncomingChatMessages.Dequeue(ChatMessage);
+		FEmbeddedCommunication::AllowSleep(TickRequesterId);
+	}
+}
+
 #endif
diff --git a/Engine/Source/Runtime/Online/XMPP/Private/XmppStrophe/XmppPrivateChatStrophe.h b/Engine/Source/Runtime/Online/XMPP/Private/XmppStrophe/XmppPrivateChatStrophe.h
--- a/Engine/Source/Runtime/Online/XMPP/Private/XmppStrophe/XmppPrivateChatStrophe.h
+++ b/Engine/Source/Runtime/Online/XMPP/Private/XmppStrophe/XmppPrivateChatStrophe.h
@@ -23,7 +23,7 @@ class FXmppPrivateChatStrophe
 public:
 	// FXmppPrivateChatStrophe
 	FXmppPrivateChatStrophe(FXmppConnectionStrophe& InConnectionManager);
-	virtual ~FXmppPrivateChatStrophe() = default;
+	virtual ~FXmppPrivateChatStrophe();
 
 	// XMPP Thread
 	void OnDisconnect();
@@ -41,6 +41,9 @@ public:
 protected:
 	void OnChatReceived(TUniquePtr<FXmppChatMessage>&& Chat);
 
+	/** Remove pending messages and engine KeepAwake calls */
+	void CleanupMessages();
+
 protected:
 	/** Connection manager controls sending data to XMPP thread */
 	FXmppConnectionStrophe& ConnectionManager;
@@ -53,6 +56,9 @@ protected:
 
 	/** Delegate for game to listen to Chat Messages */
 	FOnXmppChatReceived OnChatReceivedDelegate;
+
+	/* Engine Tick Requester name */
+	const static FName TickRequesterId;
 };
 
 #endif
