Add OnBeforeFork/OnAfterFork to IHttpManager that can be called before and after forking to fix a hang in HTTPManager due to bad interaction between curl's connection cache and forking.

diff --git a/Engine/Source/Runtime/Online/HTTP/Private/Curl/CurlHttpManager.cpp b/Engine/Source/Runtime/Online/HTTP/Private/Curl/CurlHttpManager.cpp
--- a/Engine/Source/Runtime/Online/HTTP/Private/Curl/CurlHttpManager.cpp
+++ b/Engine/Source/Runtime/Online/HTTP/Private/Curl/CurlHttpManager.cpp
@@ -331,6 +331,22 @@ void FCurlHttpManager::ShutdownCurl()
 #endif // #if WITH_SSL
 }
 
+void FCurlHttpManager::OnBeforeFork()
+{
+	FHttpManager::OnBeforeFork();
+
+	Thread->StopThread();
+	ShutdownCurl();
+}
+
+void FCurlHttpManager::OnAfterFork()
+{
+	InitCurl();
+	Thread->StartThread();
+
+	FHttpManager::OnAfterFork();
+}
+
 FHttpThread* FCurlHttpManager::CreateHttpThread()
 {
 	return new FCurlHttpThread();
diff --git a/Engine/Source/Runtime/Online/HTTP/Private/Curl/CurlHttpManager.h b/Engine/Source/Runtime/Online/HTTP/Private/Curl/CurlHttpManager.h
--- a/Engine/Source/Runtime/Online/HTTP/Private/Curl/CurlHttpManager.h
+++ b/Engine/Source/Runtime/Online/HTTP/Private/Curl/CurlHttpManager.h
@@ -54,6 +54,9 @@ public:
 	CurlRequestOptions;
 
 	//~ Begin HttpManager Interface
+	virtual void OnBeforeFork() override;
+	virtual void OnAfterFork() override;
+
 public:
 	virtual bool SupportsDynamicProxy() const override;
 protected:
diff --git a/Engine/Source/Runtime/Online/HTTP/Private/HttpManager.cpp b/Engine/Source/Runtime/Online/HTTP/Private/HttpManager.cpp
--- a/Engine/Source/Runtime/Online/HTTP/Private/HttpManager.cpp
+++ b/Engine/Source/Runtime/Online/HTTP/Private/HttpManager.cpp
@@ -98,6 +98,16 @@ TFunction<FString()> FHttpManager::GetDefaultCorrelationIdMethod()
 	return []{ return FGuid::NewGuid().ToString(); };
 }
 
+void FHttpManager::OnBeforeFork()
+{
+	Flush(false);
+}
+
+void FHttpManager::OnAfterFork()
+{
+
+}
+
 FHttpThread* FHttpManager::CreateHttpThread()
 {
 	return new FHttpThread();
diff --git a/Engine/Source/Runtime/Online/HTTP/Public/HttpManager.h b/Engine/Source/Runtime/Online/HTTP/Public/HttpManager.h
--- a/Engine/Source/Runtime/Online/HTTP/Public/HttpManager.h
+++ b/Engine/Source/Runtime/Online/HTTP/Public/HttpManager.h
@@ -136,6 +136,16 @@ public:
 	 */
 	static TFunction<FString()> GetDefaultCorrelationIdMethod();
 
+	/**
+	 * Inform that HTTP Manager that we are about to fork(). Will block to flush all outstanding http requests
+	 */
+	virtual void OnBeforeFork();
+
+	/**
+	 * Inform that HTTP Manager that we have completed a fork(). Must be called in both the client and parent process
+	 */
+	virtual void OnAfterFork();
+
 protected:
 	/** 
 	 * Create HTTP thread object
