Changing BackgroundHTTP to process requests in the order they come in.
#tests Android Client

diff --git a/Engine/Source/Runtime/Online/BackgroundHTTP/Private/BackgroundHttpManagerImpl.cpp b/Engine/Source/Runtime/Online/BackgroundHTTP/Private/BackgroundHttpManagerImpl.cpp
--- a/Engine/Source/Runtime/Online/BackgroundHTTP/Private/BackgroundHttpManagerImpl.cpp
+++ b/Engine/Source/Runtime/Online/BackgroundHTTP/Private/BackgroundHttpManagerImpl.cpp
@@ -134,7 +134,7 @@ void FBackgroundHttpManagerImpl::RemoveRequest(const FBackgroundHttpRequestPtr R
 	//Check if this request was in active list first
 	{
 		FRWScopeLock ScopeLock(ActiveRequestLock, SLT_Write);
-		NumRequestsRemoved = ActiveRequests.RemoveSwap(Request);
+		NumRequestsRemoved = ActiveRequests.Remove(Request);
 
 		//If we removed an active request, lets decrement the NumCurrentlyActiveRequests accordingly
 		if (NumRequestsRemoved != 0)
@@ -147,7 +147,7 @@ void FBackgroundHttpManagerImpl::RemoveRequest(const FBackgroundHttpRequestPtr R
 	if (NumRequestsRemoved == 0)
 	{
 		FRWScopeLock ScopeLock(PendingRequestLock, SLT_Write);
-		NumRequestsRemoved = PendingStartRequests.RemoveSwap(Request);
+		NumRequestsRemoved = PendingStartRequests.Remove(Request);
 	}
 
 	UE_LOG(LogBackgroundHttpManager, Log, TEXT("FGenericPlatformBackgroundHttpManager::RemoveRequest Called - RequestID:%s | NumRequestsActuallyRemoved:%d | NumCurrentlyActiveRequests:%d"), *Request->GetRequestID(), NumRequestsRemoved, NumCurrentlyActiveRequests);
@@ -225,23 +225,23 @@ void FBackgroundHttpManagerImpl::ActivatePendingRequests()
 			const int NumRequestsWeCanProcess = (FPlatformBackgroundHttp::GetPlatformMaxActiveDownloads() - NumCurrentlyActiveRequests);
 			if (NumRequestsWeCanProcess > 0)
 			{
-				for (int RequestAddedCount = 0; RequestAddedCount < NumRequestsWeCanProcess; ++RequestAddedCount)
+				TArray<FBackgroundHttpRequestPtr> RemainingRequests;
+				RemainingRequests.Reserve(PendingStartRequests.Num() - NumRequestsWeCanProcess);
+
+				for (int RequestIndex = 0; RequestIndex < PendingStartRequests.Num(); ++RequestIndex)
 				{
-					//Don't do anything once we are out of PendingStartRequests
-					if (PendingStartRequests.Num() > 0)
+					FBackgroundHttpRequestPtr& PendingRequest = PendingStartRequests[RequestIndex];
+					if (RequestIndex <= NumRequestsWeCanProcess)
 					{
-						//Lets just always use the first element and then remove it
-						FBackgroundHttpRequestPtr& PendingRequest = PendingStartRequests[0];
 						RequestsStartingThisTick.Add(PendingRequest);
-						PendingStartRequests.RemoveAtSwap(0);
-
-						++NumCurrentlyActiveRequests;
 					}
 					else
 					{
-						break;
+						RemainingRequests.Add(PendingRequest);
 					}
 				}
+
+				PendingStartRequests = RemainingRequests;
 			}
 		}
 	}
@@ -255,6 +255,8 @@ void FBackgroundHttpManagerImpl::ActivatePendingRequests()
 		FRWScopeLock ScopeLock(ActiveRequestLock, SLT_Write);
 		ActiveRequests.Add(RequestToStart);
 
+		++NumCurrentlyActiveRequests;
+
 		//Call Handle for that task to now kick itself off
 		RequestToStart->HandleDelayedProcess();
 	}
