iOS BackgroundHTTP Fixes:
- Fixing bug where queueing same URL with 2 different BackgroundHttpRequests on iOS could result in receiving a callback 2x for the same request.
- Crash fix for iOS when an invalid URL is present in a task when we process UnAssociatedTasks.
#tests iOS Client

diff --git a/Engine/Source/Runtime/Online/BackgroundHTTP/Private/IOS/ApplePlatformBackgroundHttpManager.cpp b/Engine/Source/Runtime/Online/BackgroundHTTP/Private/IOS/ApplePlatformBackgroundHttpManager.cpp
--- a/Engine/Source/Runtime/Online/BackgroundHTTP/Private/IOS/ApplePlatformBackgroundHttpManager.cpp
+++ b/Engine/Source/Runtime/Online/BackgroundHTTP/Private/IOS/ApplePlatformBackgroundHttpManager.cpp
@@ -59,7 +59,15 @@ void FApplePlatformBackgroundHttpManager::PopulateUnAssociatedTasks()
 				//Store all existing tasks by their URL
 				for (id task in tasks)
 				{
-					[UnAssociatedTasks setObject:task forKey:[[[task currentRequest] URL] absoluteString]];
+					//Make sure we have a valid absolute string version of the URL to use for our task's key. Otherwise, we just disregard this task.
+					if (	([task] != nullptr)
+						&&	([task currentRequest] != nullptr)
+						&&  ([[task currentRequest] URL] != nullptr)
+						&&	([[[task currentRequest] URL] absoluteString] != nullptr)
+						&&	([[[[task currentRequest] URL] absoluteString] length] != 0))
+					{
+						[UnAssociatedTasks setObject : task forKey : [[[task currentRequest] URL] absoluteString]];
+					}
 				}
 
 				bHasFinishedPopulatingUnassociatedTasks = true;
@@ -113,52 +121,62 @@ void FApplePlatformBackgroundHttpManager::AddRequest(const FBackgroundHttpReques
     
     //See if our request is an AppleBackgroundHttpRequest so we can do more detailed checks on it.
     FAppleBackgroundHttpRequestPtr AppleRequest = StaticCastSharedPtr<FApplePlatformBackgroundHttpRequest>(Request);
-    if (ensureAlwaysMsgf(AppleRequest.IsValid(), TEXT("Adding a non-Apple background request to our Apple Background Http Manager! This is not supported or expected!")))
-    {
-        GenerateURLMapEntriesForRequest(AppleRequest);
-    }
-    else
-    {
-        //just error out if we don't have the expected AppleRequest
-        return;
-    }
-    
-	if (!AssociateWithAnyExistingRequest(Request))
+	if (ensureAlwaysMsgf(AppleRequest.IsValid(), TEXT("Adding a non-Apple background request to our Apple Background Http Manager! This is not supported or expected!")))
 	{
-        if (!AssociateWithAnyExistingUnAssociatedTasks(Request))
-        {
-            StartRequest(AppleRequest);
-        }
-        
-		FRWScopeLock ScopeLock(ActiveRequestLock, SLT_Write);
-		ActiveRequests.Add(Request);
-
-		//Increment our underlying FBackgroundHttpManagerImpl tracker for active requests as we
-		//don't implement the method it uses to increase this number.
-		// NOTE: We don't make use of this number in Apple Platform functions as all requests are "Active" but their
-		// underlying Task might not be, see NumCurrentlyActiveTasks instead to track how many current Tasks are downloading data.
-		++NumCurrentlyActiveRequests;
+		//If we fail to generate URLMapEntries or AssociateWithAnyExistingRequest, then we will have already sent a completion handler immediately, so only start work and monitor
+		//these requests if those didn't already complete this Request
+		if (GenerateURLMapEntriesForRequest(AppleRequest) && !AssociateWithAnyExistingRequest(Request))
+		{
+			if (!AssociateWithAnyExistingUnAssociatedTasks(Request))
+			{
+				StartRequest(AppleRequest);
+			}
+
+			FRWScopeLock ScopeLock(ActiveRequestLock, SLT_Write);
+			ActiveRequests.Add(Request);
+
+			//Increment our underlying FBackgroundHttpManagerImpl tracker for active requests as we
+			//don't implement the method it uses to increase this number.
+			// NOTE: We don't make use of this number in Apple Platform functions as all requests are "Active" but their
+			// underlying Task might not be, see NumCurrentlyActiveTasks instead to track how many current Tasks are downloading data.
+			++NumCurrentlyActiveRequests;
+		}
 	}
 }
 
-void FApplePlatformBackgroundHttpManager::GenerateURLMapEntriesForRequest( FAppleBackgroundHttpRequestPtr Request)
+bool FApplePlatformBackgroundHttpManager::GenerateURLMapEntriesForRequest( FAppleBackgroundHttpRequestPtr Request)
 {
-    FRWScopeLock ScopeLock(URLToRequestMapLock, SLT_Write);
-    for (const FString& URL : Request->GetURLList())
-    {
-        FBackgroundHttpURLMappedRequestPtr& FoundRequest = URLToRequestMap.FindOrAdd(URL);
-        
-        const bool bRequestAlreadyExistsForURL = ((FoundRequest.IsValid()) && (Request != FoundRequest));
-        if (ensureAlwaysMsgf(!bRequestAlreadyExistsForURL, TEXT("URL is represented by 2 different Requests! Immediately completing new request with error.")))
-        {
-            FoundRequest = Request;
-        }
-        else
-        {
-            FBackgroundHttpResponsePtr NewResponse = FPlatformBackgroundHttp::ConstructBackgroundResponse(EHttpResponseCodes::Unknown, FString());
-            Request->CompleteWithExistingResponseData(NewResponse);
-        }
-    }
+	bool bWasGenerateSuccess = true;
+
+	//Attempt to add entries for all URLs
+	{
+		FRWScopeLock ScopeLock(URLToRequestMapLock, SLT_Write);
+		for (const FString& URL : Request->GetURLList())
+		{
+			FBackgroundHttpURLMappedRequestPtr& FoundRequest = URLToRequestMap.FindOrAdd(URL);
+
+			const bool bRequestAlreadyExistsForURL = ((FoundRequest.IsValid()) && (Request != FoundRequest));
+			if (ensureAlwaysMsgf(!bRequestAlreadyExistsForURL, TEXT("URL is represented by 2 different Requests! Immediately completing new request with error.")))
+			{
+				FoundRequest = Request;
+			}
+			else
+			{
+				bWasGenerateSuccess = false;
+
+				FBackgroundHttpResponsePtr NewResponse = FPlatformBackgroundHttp::ConstructBackgroundResponse(EHttpResponseCodes::Unknown, FString());
+				Request->CompleteWithExistingResponseData(NewResponse);
+			}
+		}
+	}
+
+	//if we didn't succeed, make sure we don't have any stale partial URL Map entries for this request
+	if (!bWasGenerateSuccess)
+	{
+		RemoveURLMapEntriesForRequest(Request);
+	}
+
+	return bWasGenerateSuccess;
 }
 
 void FApplePlatformBackgroundHttpManager::RemoveURLMapEntriesForRequest(FAppleBackgroundHttpRequestPtr Request)
diff --git a/Engine/Source/Runtime/Online/BackgroundHTTP/Public/IOS/ApplePlatformBackgroundHttpManager.h b/Engine/Source/Runtime/Online/BackgroundHTTP/Public/IOS/ApplePlatformBackgroundHttpManager.h
--- a/Engine/Source/Runtime/Online/BackgroundHTTP/Public/IOS/ApplePlatformBackgroundHttpManager.h
+++ b/Engine/Source/Runtime/Online/BackgroundHTTP/Public/IOS/ApplePlatformBackgroundHttpManager.h
@@ -86,7 +86,7 @@ private:
 	void FinishRequest(FAppleBackgroundHttpRequestPtr Request);
 	void RetryRequest(FAppleBackgroundHttpRequestPtr Request, bool bShouldIncreaseRetryCount, bool bShouldStartImmediately, NSData* RetryData);
     void RemoveSessionTasksForRequest(FAppleBackgroundHttpRequestPtr Request);
-    void GenerateURLMapEntriesForRequest(FAppleBackgroundHttpRequestPtr Request);
+    bool GenerateURLMapEntriesForRequest(FAppleBackgroundHttpRequestPtr Request);
     void RemoveURLMapEntriesForRequest(FAppleBackgroundHttpRequestPtr Request);
     
 	bool IsRetryDataValid(NSData* RetryData) const;
