Fixing iOS compile error by previous submit.
#tests none

diff --git a/Engine/Source/Runtime/Online/BackgroundHTTP/Private/IOS/ApplePlatformBackgroundHttpManager.cpp b/Engine/Source/Runtime/Online/BackgroundHTTP/Private/IOS/ApplePlatformBackgroundHttpManager.cpp
--- a/Engine/Source/Runtime/Online/BackgroundHTTP/Private/IOS/ApplePlatformBackgroundHttpManager.cpp
+++ b/Engine/Source/Runtime/Online/BackgroundHTTP/Private/IOS/ApplePlatformBackgroundHttpManager.cpp
@@ -60,7 +60,7 @@ void FApplePlatformBackgroundHttpManager::PopulateUnAssociatedTasks()
 				for (id task in tasks)
 				{
 					//Make sure we have a valid absolute string version of the URL to use for our task's key. Otherwise, we just disregard this task.
-					if (	([task] != nullptr)
+					if (	(task != nullptr)
 						&&	([task currentRequest] != nullptr)
 						&&  ([[task currentRequest] URL] != nullptr)
 						&&	([[[task currentRequest] URL] absoluteString] != nullptr)
