Adding option for BackgroundNotificationObject to only send notification if its in the background (default behavior)
#tests Android Client, iOS Client

diff --git a/Engine/Source/Runtime/Online/BackgroundHTTP/Private/BackgroundHttpNotificationObject.cpp b/Engine/Source/Runtime/Online/BackgroundHTTP/Private/BackgroundHttpNotificationObject.cpp
--- a/Engine/Source/Runtime/Online/BackgroundHTTP/Private/BackgroundHttpNotificationObject.cpp
+++ b/Engine/Source/Runtime/Online/BackgroundHTTP/Private/BackgroundHttpNotificationObject.cpp
@@ -3,6 +3,7 @@
 #include "BackgroundHttpNotificationObject.h"
 
 #include "Misc/ConfigCacheIni.h"
+#include "Misc/CoreDelegates.h"
 
 #include "Modules/ModuleInterface.h"
 #include "Modules/ModuleManager.h"
@@ -10,13 +11,22 @@
 #include "LocalNotification.h"
 
 FBackgroundHttpNotificationObject::FBackgroundHttpNotificationObject(FText InNotificationTitle, FText InNotificationBody, FText InNotificationAction, const FString& InNotificationActivationString, bool InNotifyOnlyOnFullSuccess)
+	: FBackgroundHttpNotificationObject(InNotificationTitle, InNotificationBody, InNotificationAction, InNotificationActivationString, InNotifyOnlyOnFullSuccess, true)
+{
+}
+
+FBackgroundHttpNotificationObject::FBackgroundHttpNotificationObject(FText InNotificationTitle, FText InNotificationBody, FText InNotificationAction, const FString& InNotificationActivationString, bool InNotifyOnlyOnFullSuccess, bool InbOnlySendNotificationInBackground)
 	: NotificationTitle(InNotificationTitle)
     , NotificationAction(InNotificationAction)
     , NotificationBody(InNotificationBody)
 	, NotificationActivationString(InNotificationActivationString)
-	, bNotifyOnlyOnFullSuccess(InNotifyOnlyOnFullSuccess)
+    , bOnlySendNotificationInBackground(InbOnlySendNotificationInBackground)
+    , bNotifyOnlyOnFullSuccess(InNotifyOnlyOnFullSuccess)
+	, bIsInBackground(false)
 	, NumFailedDownloads(0)
     , PlatformNotificationService(nullptr)
+    , OnApp_EnteringForegroundHandle()
+	, OnApp_EnteringBackgroundHandle()
 {
 	if (GConfig)
 	{
@@ -32,10 +42,41 @@ FBackgroundHttpNotificationObject::FBackgroundHttpNotificationObject(FText InNot
 			}
 		}
 	}
+
+	//Only register with this delegate if we are actually going to monitor background notifications
+	if (bOnlySendNotificationInBackground)
+	{
+		OnApp_EnteringBackgroundHandle = FCoreDelegates::ApplicationWillEnterBackgroundDelegate.AddRaw(this, &FBackgroundHttpNotificationObject::OnApp_EnteringBackground);
+		OnApp_EnteringForegroundHandle = FCoreDelegates::ApplicationHasEnteredForegroundDelegate.AddRaw(this, &FBackgroundHttpNotificationObject::OnApp_EnteringForeground);
+	}
+}
+
+void FBackgroundHttpNotificationObject::OnApp_EnteringForeground()
+{
+	bIsInBackground = false;
+}
+
+void FBackgroundHttpNotificationObject::OnApp_EnteringBackground()
+{
+	bIsInBackground = true;
 }
 
 FBackgroundHttpNotificationObject::~FBackgroundHttpNotificationObject()
 {
+	if (bOnlySendNotificationInBackground)
+	{
+		//These should only be registered if bOnlySendNotificationInBackground is set
+		FCoreDelegates::ApplicationWillEnterBackgroundDelegate.Remove(OnApp_EnteringBackgroundHandle);
+		FCoreDelegates::ApplicationHasEnteredForegroundDelegate.Remove(OnApp_EnteringForegroundHandle);
+
+		//If we have flagged as only sending notifications when we are in the BG, and we are not in the BG, just early out
+		//so we don't send a notification
+		if (!bIsInBackground)
+		{
+			return;
+		}
+	}
+
 	if (nullptr != PlatformNotificationService)
 	{
 		if (!bNotifyOnlyOnFullSuccess || (NumFailedDownloads == 0))
diff --git a/Engine/Source/Runtime/Online/BackgroundHTTP/Public/BackgroundHttpNotificationObject.h b/Engine/Source/Runtime/Online/BackgroundHTTP/Public/BackgroundHttpNotificationObject.h
--- a/Engine/Source/Runtime/Online/BackgroundHTTP/Public/BackgroundHttpNotificationObject.h
+++ b/Engine/Source/Runtime/Online/BackgroundHTTP/Public/BackgroundHttpNotificationObject.h
@@ -3,28 +3,39 @@
 #pragma once
 
 #include "CoreMinimal.h"
-
+#include "Delegates/Delegate.h"
 
 class BACKGROUNDHTTP_API FBackgroundHttpNotificationObject
 	: public TSharedFromThis<FBackgroundHttpNotificationObject, ESPMode::ThreadSafe>
 {
 public:
 	FBackgroundHttpNotificationObject(FText InNotificationTitle, FText InNotificationAction, FText InNotificationBody, const FString& InNotificationActivationString, bool InNotifyOnlyOnFullSuccess);
+	FBackgroundHttpNotificationObject(FText InNotificationTitle, FText InNotificationAction, FText InNotificationBody, const FString& InNotificationActivationString, bool InNotifyOnlyOnFullSuccess, bool bOnlySendNotificationInBackground);
 	~FBackgroundHttpNotificationObject();
 
 	void NotifyOfDownloadResult(bool bWasSuccess);
 
+private:
+	void OnApp_EnteringBackground();
+	void OnApp_EnteringForeground();
+
 private:
 	FText NotificationTitle;
 	FText NotificationAction;
 	FText NotificationBody;
 	FString NotificationActivationString;
 
+	bool bOnlySendNotificationInBackground;
 	bool bNotifyOnlyOnFullSuccess;
+	
+	volatile bool bIsInBackground;
 	volatile int32 NumFailedDownloads;
 
 	class ILocalNotificationService* PlatformNotificationService;
 
+	FDelegateHandle OnApp_EnteringForegroundHandle;
+	FDelegateHandle OnApp_EnteringBackgroundHandle;
+
 	//No default constructor
 	FBackgroundHttpNotificationObject() {}
 };
