- Potential fix for wesocket header buffer overflow.

diff --git a/Engine/Source/Runtime/Online/WebSockets/Private/Lws/LwsWebSocketsManager.cpp b/Engine/Source/Runtime/Online/WebSockets/Private/Lws/LwsWebSocketsManager.cpp
--- a/Engine/Source/Runtime/Online/WebSockets/Private/Lws/LwsWebSocketsManager.cpp
+++ b/Engine/Source/Runtime/Online/WebSockets/Private/Lws/LwsWebSocketsManager.cpp
@@ -97,6 +97,9 @@ void FLwsWebSocketsManager::InitWebSockets(TArrayView<const FString> Protocols)
 	ContextInfo.uid = -1;
 	ContextInfo.gid = -1;
 	ContextInfo.options |= LWS_SERVER_OPTION_PEER_CERT_NOT_REQUIRED | LWS_SERVER_OPTION_DISABLE_OS_CA_CERTS | LWS_SERVER_OPTION_DO_SSL_GLOBAL_INIT;
+	ContextInfo.max_http_header_data = 0;
+	ContextInfo.max_http_header_data2 = 1024 * 32;
+	ContextInfo.pt_serv_buf_size = 1024 * 32;
 	
 	// HTTP proxy
 	const FString& ProxyAddress = FHttpModule::Get().GetProxyAddress();
