fix a thread safety issue in the curl module where an HTTP response would be considered complete before the result code is set and make the bool atomic to prevent further shenanigans. also, ultimately fail requests that return OK but have an invalid/zero HTTP response code.

diff --git a/Engine/Source/Runtime/Online/HTTP/Private/Curl/CurlHttp.cpp b/Engine/Source/Runtime/Online/HTTP/Private/Curl/CurlHttp.cpp
--- a/Engine/Source/Runtime/Online/HTTP/Private/Curl/CurlHttp.cpp
+++ b/Engine/Source/Runtime/Online/HTTP/Private/Curl/CurlHttp.cpp
@@ -1042,6 +1042,13 @@ void FCurlHttpRequest::FinishedRequest()
 			{
 				Response->ContentLength = static_cast< int32 >(ContentLengthDownload);
 			}
+
+			if (Response->HttpCode <= 0)
+			{
+				UE_LOG(LogHttp, Warning, TEXT("%p: invalid HTTP response code received. URL: %s, HTTP code: %d, content length: %d, actual payload size: %d"),
+					this, *GetURL(), Response->HttpCode, Response->ContentLength, Response->Payload.Num());
+				Response->bSucceeded = false;
+			}
 		}
 	}
 	
diff --git a/Engine/Source/Runtime/Online/HTTP/Private/Curl/CurlHttp.h b/Engine/Source/Runtime/Online/HTTP/Private/Curl/CurlHttp.h
--- a/Engine/Source/Runtime/Online/HTTP/Private/Curl/CurlHttp.h
+++ b/Engine/Source/Runtime/Online/HTTP/Private/Curl/CurlHttp.h
@@ -8,6 +8,7 @@
 #include "IHttpThreadedRequest.h"
 #include "Containers/Queue.h"
 #include "GenericPlatform/HttpRequestPayload.h"
+#include "HAL/ThreadSafeBool.h"
 
 class FCurlHttpResponse;
 
@@ -179,8 +180,8 @@ public:
 	 */
 	inline void MarkAsCompleted(CURLcode InCurlCompletionResult)
 	{
-		bCurlRequestCompleted = true;
 		CurlCompletionResult = InCurlCompletionResult;
+		bCurlRequestCompleted = true;
 	}
 	
 	/** 
@@ -328,7 +329,7 @@ private:
 	/** Set to true if request has been canceled */
 	bool			bCanceled;
 	/** Set to true when request has been completed */
-	bool			bCurlRequestCompleted;
+	FThreadSafeBool	bCurlRequestCompleted;
 	/** Set to true when request has "30* Multiple Choices" (e.g. 301 Moved Permanently, 302 temporary redirect, 308 Permanent Redirect, etc.) */
 	bool			bRedirected;
 	/** Set to true if request failed to be added to curl multi */
