- Make the libwebSockets max header buffer size configurable.

diff --git a/Engine/Config/BaseEngine.ini b/Engine/Config/BaseEngine.ini
--- a/Engine/Config/BaseEngine.ini
+++ b/Engine/Config/BaseEngine.ini
@@ -34,6 +34,7 @@ BackgroundHttp.TempFileTimeOutSeconds=259200
 ThreadStackSize=131072
 ThreadTargetFrameTimeInSeconds=0.0333
 ThreadMinimumSleepTimeInSeconds=0.0
+MaxHttpHeaderData=32768
 
 [Ping]
 StackSize=1048576
diff --git a/Engine/Source/Runtime/Online/WebSockets/Private/Lws/LwsWebSocketsManager.cpp b/Engine/Source/Runtime/Online/WebSockets/Private/Lws/LwsWebSocketsManager.cpp
--- a/Engine/Source/Runtime/Online/WebSockets/Private/Lws/LwsWebSocketsManager.cpp
+++ b/Engine/Source/Runtime/Online/WebSockets/Private/Lws/LwsWebSocketsManager.cpp
@@ -98,8 +98,11 @@ void FLwsWebSocketsManager::InitWebSockets(TArrayView<const FString> Protocols)
 	ContextInfo.gid = -1;
 	ContextInfo.options |= LWS_SERVER_OPTION_PEER_CERT_NOT_REQUIRED | LWS_SERVER_OPTION_DISABLE_OS_CA_CERTS | LWS_SERVER_OPTION_DO_SSL_GLOBAL_INIT;
 	ContextInfo.max_http_header_data = 0;
-	ContextInfo.max_http_header_data2 = 1024 * 32;
-	ContextInfo.pt_serv_buf_size = 1024 * 32;
+
+	int32 MaxHttpHeaderData = 1024 * 32;
+	GConfig->GetInt(TEXT("WebSockets.LibWebSockets"), TEXT("MaxHttpHeaderData"), MaxHttpHeaderData, GEngineIni);
+	ContextInfo.max_http_header_data2 = MaxHttpHeaderData;
+	ContextInfo.pt_serv_buf_size = MaxHttpHeaderData;
 	
 	// HTTP proxy
 	const FString& ProxyAddress = FHttpModule::Get().GetProxyAddress();
