Fixing bug on GenericPlatform BackgroundHttp implementation requests would fail as we were not setting the Get HTTP Verb.

diff --git a/Engine/Source/Runtime/Online/BackgroundHTTP/Private/GenericPlatform/GenericPlatformBackgroundHttpRequest.cpp b/Engine/Source/Runtime/Online/BackgroundHTTP/Private/GenericPlatform/GenericPlatformBackgroundHttpRequest.cpp
--- a/Engine/Source/Runtime/Online/BackgroundHTTP/Private/GenericPlatform/GenericPlatformBackgroundHttpRequest.cpp
+++ b/Engine/Source/Runtime/Online/BackgroundHTTP/Private/GenericPlatform/GenericPlatformBackgroundHttpRequest.cpp
@@ -83,6 +83,7 @@ void FGenericPlatformBackgroundHttpRequest::FGenericPlatformBackgroundHttpWrappe
 			//Reset this value as we are restarting our progress
 			LastProgressUpdateBytes = 0;
 
+			HttpRequest->SetVerb(TEXT("GET"));
 			HttpRequest->SetURL(RequestURL);
 			HttpRequest->ProcessRequest();
 		}
