[HTML5] Add missing NtoHS, HtoNS, NtoHL, and HtoNL implementations in ICMP module

diff --git a/Engine/Plugins/Online/OnlineFramework/OnlineFramework.uplugin b/Engine/Plugins/Online/OnlineFramework/OnlineFramework.uplugin
--- a/Engine/Plugins/Online/OnlineFramework/OnlineFramework.uplugin
+++ b/Engine/Plugins/Online/OnlineFramework/OnlineFramework.uplugin
@@ -13,10 +13,7 @@
 	"Modules": [
 		{
 			"Name": "Qos",
-			"Type": "Runtime",
-			"BlacklistPlatforms": [
-				"HTML5"
-			]
+			"Type": "Runtime"
 		},
 		{
 			"Name": "Party",
diff --git a/Engine/Source/Runtime/Online/ICMP/Private/HTML5/IcmpHTML5.cpp b/Engine/Source/Runtime/Online/ICMP/Private/HTML5/IcmpHTML5.cpp
new file mode 100644
--- /dev/null
+++ b/Engine/Source/Runtime/Online/ICMP/Private/HTML5/IcmpHTML5.cpp
@@ -0,0 +1,26 @@
+// Copyright 1998-2019 Epic Games, Inc. All Rights Reserved.
+
+#include "Icmp.h"
+#include "IcmpModule.h"
+#include "IcmpPrivate.h"
+#include <netinet/in.h>
+
+uint16 NtoHS(uint16 val)
+{
+	return ntohs(val);
+}
+
+uint16 HtoNS(uint16 val)
+{
+	return htons(val);
+}
+
+uint32 NtoHL(uint32 val)
+{
+	return ntohl(val);
+}
+
+uint32 HtoNL(uint32 val)
+{
+	return htonl(val);
+}
