Fix FBackgroundHttpManagerImpl trying to reserve a negative number sometimes.

diff --git a/Engine/Source/Runtime/Online/BackgroundHTTP/Private/BackgroundHttpManagerImpl.cpp b/Engine/Source/Runtime/Online/BackgroundHTTP/Private/BackgroundHttpManagerImpl.cpp
--- a/Engine/Source/Runtime/Online/BackgroundHTTP/Private/BackgroundHttpManagerImpl.cpp
+++ b/Engine/Source/Runtime/Online/BackgroundHTTP/Private/BackgroundHttpManagerImpl.cpp
@@ -226,7 +226,10 @@ void FBackgroundHttpManagerImpl::ActivatePendingRequests()
 			if (NumRequestsWeCanProcess > 0)
 			{
 				TArray<FBackgroundHttpRequestPtr> RemainingRequests;
-				RemainingRequests.Reserve(PendingStartRequests.Num() - NumRequestsWeCanProcess);
+				if (PendingStartRequests.Num() > NumRequestsWeCanProcess)
+				{
+					RemainingRequests.Reserve(PendingStartRequests.Num() - NumRequestsWeCanProcess);
+				}
 
 				for (int RequestIndex = 0; RequestIndex < PendingStartRequests.Num(); ++RequestIndex)
 				{
@@ -241,7 +244,7 @@ void FBackgroundHttpManagerImpl::ActivatePendingRequests()
 					}
 				}
 
-				PendingStartRequests = RemainingRequests;
+				PendingStartRequests = MoveTemp(RemainingRequests);
 			}
 		}
 	}
