Install Bundle Manager
 - Get actual largest file from BuildPatchManifest size instead of just estimating.
 - Add GetFileSize() to IBuildManifest

diff --git a/Engine/Source/Runtime/Online/BuildPatchServices/Private/BuildPatchManifest.h b/Engine/Source/Runtime/Online/BuildPatchServices/Private/BuildPatchManifest.h
--- a/Engine/Source/Runtime/Online/BuildPatchServices/Private/BuildPatchManifest.h
+++ b/Engine/Source/Runtime/Online/BuildPatchServices/Private/BuildPatchManifest.h
@@ -115,6 +115,9 @@ public:
 	virtual int64 GetBuildSize(const TSet<FString>& Tags) const override;
 	virtual TArray<FString> GetBuildFileList() const override;
 	virtual TArray<FString> GetBuildFileList(const TSet<FString>& Tags) const override;
+	virtual int64 GetFileSize(const FString& Filename) const override;
+	virtual int64 GetFileSize(const TArray<FString>& Filenames) const override;
+	virtual int64 GetFileSize(const TSet  <FString>& Filenames) const override;
 	virtual void GetFileTagList(TSet<FString>& Tags) const override;
 	virtual void GetRemovableFiles(const IBuildManifestRef& OldManifest, TArray< FString >& RemovableFiles) const override;
 	virtual void GetRemovableFiles(const TCHAR* InstallPath, TArray< FString >& RemovableFiles) const override;
@@ -204,23 +207,6 @@ public:
 	virtual int64 GetDataSize(const TArray<FGuid>& DataGuids) const;
 	virtual int64 GetDataSize(const TSet  <FGuid>& DataGuids) const;
 
-	/**
-	 * Returns the size of a particular file in the build
-	 * VALID FOR ANY MANIFEST
-	 * @param Filename		The file.
-	 * @return		File size.
-	 */
-	virtual int64 GetFileSize(const FString& Filename) const;
-
-	/**
-	 * Returns the total size of all files in the array
-	 * VALID FOR ANY MANIFEST
-	 * @param Filenames		The array of files.
-	 * @return		Total size of files in array.
-	 */
-	virtual int64 GetFileSize(const TArray<FString>& Filenames) const;
-	virtual int64 GetFileSize(const TSet  <FString>& Filenames) const;
-
 	/**
 	 * Returns the number of files in this build.
 	 * @return		The number of files.
diff --git a/Engine/Source/Runtime/Online/BuildPatchServices/Public/Interfaces/IBuildManifest.h b/Engine/Source/Runtime/Online/BuildPatchServices/Public/Interfaces/IBuildManifest.h
--- a/Engine/Source/Runtime/Online/BuildPatchServices/Public/Interfaces/IBuildManifest.h
+++ b/Engine/Source/Runtime/Online/BuildPatchServices/Public/Interfaces/IBuildManifest.h
@@ -169,6 +169,22 @@ public:
 	 */
 	virtual TArray<FString> GetBuildFileList(const TSet<FString>& Tags) const = 0;
 
+	/**
+	 * Returns the size of a particular file in the build
+	 * @param Filename		The file.
+	 * @return		File size.
+	 */
+	virtual int64 GetFileSize(const FString& Filename) const = 0;
+
+	/**
+	 * Returns the total size of all files in the array
+	 * @param Filenames		The array of files.
+	 * @return		Total size of files in array.
+	 */
+	virtual int64 GetFileSize(const TArray<FString>& Filenames) const = 0;
+	virtual int64 GetFileSize(const TSet  <FString>& Filenames) const = 0;
+
+
 	/**
 	 * Get the list of install tags in this manifest
 	 * @param Tags			OUT		Receives the tags referenced.
