HttpServerModule - Initial
+ Simple HttpServer implemenation. Single-threaded/yielding request/response handling.
+ HttpServerModule creates HttpListeners and exposes IHttpRouter instances via the HttpServeMmodule public api
+ Clients leverage IHttpRouter abstraction to register for HTTP routes akin to modern listener frameworks, eg:  BindRoute(/my/path,  verb, MyHandler)
+ Tiered/REST-style http routing is supported
+ HttpListeners manage unique pools of HttpConnections
+ HttpConnections are implemented as a state-driven tickable abstraction built atop FSocket and related subsystems.
+ Multi-frame read/write is supported via helper read/write context helpers.

diff --git a/Engine/Source/Runtime/Online/HTTPServer/HttpServer.Build.cs b/Engine/Source/Runtime/Online/HTTPServer/HttpServer.Build.cs
new file mode 100644
--- /dev/null
+++ b/Engine/Source/Runtime/Online/HTTPServer/HttpServer.Build.cs
@@ -0,0 +1,32 @@
+// Copyright 1998-2019 Epic Games, Inc. All Rights Reserved.
+
+using UnrealBuildTool;
+
+public class HTTPServer : ModuleRules
+{
+    public HTTPServer(ReadOnlyTargetRules Target) : base(Target)
+    {
+        //OptimizeCode = CodeOptimization.Never;
+
+        PublicDefinitions.Add("HTTPSERVER_PACKAGE=1");
+
+        PublicIncludePaths.AddRange(
+			new string[] {
+				"Runtime/Online/HTTPServer/Public",
+            }
+        );
+
+        PrivateIncludePaths.AddRange(
+            new string[] {
+                "Runtime/Online/HTTPServer/Private",
+            }
+        );
+
+        PrivateDependencyModuleNames.AddRange(
+            new string[] {
+                "Core",
+                "Sockets",
+            }
+        );
+    }
+}
diff --git a/Engine/Source/Runtime/Online/HTTPServer/Private/HttpConnection.cpp b/Engine/Source/Runtime/Online/HTTPServer/Private/HttpConnection.cpp
new file mode 100644
--- /dev/null
+++ b/Engine/Source/Runtime/Online/HTTPServer/Private/HttpConnection.cpp
@@ -0,0 +1,301 @@
+// Copyright 1998-2019 Epic Games, Inc. All Rights Reserved.
+
+#include "HttpConnection.h"
+#include "HttpConnectionContext.h"
+#include "HttpConnectionRequestReadContext.h"
+#include "HttpConnectionResponseWriteContext.h"
+#include "HttpRouter.h"
+#include "HttpServerRequest.h"
+#include "HttpServerResponse.h"
+#include "HttpRequestHandlerIterator.h"
+#include "HttpServerConstants.h"
+#include "HttpServerConstantsPrivate.h"
+
+#include "Sockets.h"
+#include "SocketSubsystem.h"
+#include "Containers/Array.h"
+
+DEFINE_LOG_CATEGORY(LogHttpConnection)
+
+FHttpConnection::FHttpConnection(FSocket* InSocket, TSharedPtr<FHttpRouter> InRouter)
+	: Socket(InSocket)
+	,Router(InRouter)
+	,ReadContext(InSocket)
+	,WriteContext(InSocket)
+{
+	check(nullptr != Socket);
+}
+
+FHttpConnection::~FHttpConnection()
+{
+	check(nullptr == Socket);
+}
+
+void FHttpConnection::Tick(float DeltaTime)
+{
+	const float AwaitReadTimeout = bKeepAlive ?
+		ConnectionKeepAliveTimeout : ConnectionTimeout;
+
+	switch (State)
+	{
+	case EHttpConnectionState::AwaitingRead:
+		if (ReadContext.GetElapsedIdleTime() > AwaitReadTimeout)
+		{
+			Destroy();
+			return;
+		}
+		BeginRead(DeltaTime);
+		break;
+
+	case EHttpConnectionState::Reading:
+		if (ReadContext.GetElapsedIdleTime() > ConnectionTimeout)
+		{
+			Destroy();
+			return;
+		}
+		ContinueRead(DeltaTime);
+		break;
+
+	case EHttpConnectionState::AwaitingProcessing:
+		break;
+
+	case EHttpConnectionState::Writing:
+		if (WriteContext.GetElapsedIdleTime() > ConnectionTimeout)
+		{
+			Destroy();
+			return;
+		}
+		ContinueWrite(DeltaTime);
+		break;
+
+	case EHttpConnectionState::Destroyed:
+		ensure(false);
+		break;
+	}
+}
+
+void FHttpConnection::ChangeState(EHttpConnectionState NewState)
+{
+	check(EHttpConnectionState::Destroyed != State);
+	check(NewState != State);
+
+	UE_LOG(LogHttpConnection, Verbose,
+		TEXT("ChangingState: %d => %d"), State, NewState);
+	State = NewState;
+}
+
+void FHttpConnection::TransferState(EHttpConnectionState CurrentState, EHttpConnectionState NextState)
+{
+	check(CurrentState == State);
+	check(NextState != CurrentState);
+	ChangeState(NextState);
+}
+
+void FHttpConnection::BeginRead(float DeltaTime)
+{
+	// Wait should always return true if the connection is valid
+	if (!Socket->Wait(ESocketWaitConditions::WaitForRead, FTimespan::Zero()))
+	{
+		Destroy();
+		return;
+	}
+
+	// The socket is reachable, however there may not be data in the pipe
+	uint32 PendingDataSize = 0;
+	if (Socket->HasPendingData(PendingDataSize))
+	{
+		TransferState(EHttpConnectionState::AwaitingRead, EHttpConnectionState::Reading);
+		ReadContext.ResetContext();
+		ContinueRead(DeltaTime);
+	}
+	else
+	{
+		ReadContext.AddElapsedIdleTime(DeltaTime);
+	}
+}
+
+void FHttpConnection::ContinueRead(float DeltaTime)
+{
+	check(State == EHttpConnectionState::Reading);
+
+	auto ReaderState = ReadContext.ReadStream(DeltaTime);
+
+	switch (ReaderState)
+	{
+	case EHttpConnectionContextState::Continue:
+		break;
+
+	case EHttpConnectionContextState::Done:
+		CompleteRead(ReadContext.GetRequest());
+		break;
+
+	case EHttpConnectionContextState::Error:
+		HandleReadError(*ReadContext.GetErrorStr());
+		break;
+	}
+}
+
+void FHttpConnection::CompleteRead(const TSharedPtr<FHttpServerRequest>& Request)
+{
+	TArray<FString>* ConnectionHeaders = Request->Headers.Find(FHttpServerHeaderKeys::CONNECTION);
+	if (ConnectionHeaders &&
+		ConnectionHeaders->Contains(TEXT("Keep-Alive")))
+	{
+		bKeepAlive = true;
+	}
+
+	TWeakPtr<FHttpConnection> WeakThisPtr(AsShared());
+	FHttpResultCallback OnProcessingComplete = [WeakThisPtr, Request, LastRequestNumberStateCapture = ++LastRequestNumber]
+	(TUniquePtr<FHttpServerResponse>&& Response)
+	{
+		TSharedPtr<FHttpConnection> SharedThisPtr = WeakThisPtr.Pin();
+		if (SharedThisPtr.IsValid())
+		{
+			UE_LOG(LogHttpConnection, Log,
+				TEXT("Completed Processing Request [%u]"), LastRequestNumberStateCapture)
+
+			// Ensure this result callback was called once
+			check(EHttpConnectionState::AwaitingProcessing == SharedThisPtr->GetState());
+
+			// Begin response flow
+			SharedThisPtr->BeginWrite(MoveTemp(Response), LastRequestNumberStateCapture);
+		}
+	};
+
+	ProcessRequest(Request, OnProcessingComplete);
+}
+
+void FHttpConnection::ProcessRequest(const TSharedPtr<FHttpServerRequest>& Request, const FHttpResultCallback& OnProcessingComplete)
+{
+	TransferState(EHttpConnectionState::Reading, EHttpConnectionState::AwaitingProcessing);
+
+	UE_LOG(LogHttpConnection, Log,
+		TEXT("Begin Processing Request [%u]: %s"), LastRequestNumber, *Request->RelativePath.GetPath());
+
+	bool bRequestHandled = false;
+	auto RequestHandlerIterator = Router->CreateRequestHandlerIterator(Request);
+	while (const auto& RequestHandlerPtr = RequestHandlerIterator.Next())
+	{
+		(*RequestHandlerPtr).CheckCallable();
+		bRequestHandled = (*RequestHandlerPtr)(*Request, OnProcessingComplete);
+		if (bRequestHandled)
+		{
+			break;
+		}
+		// If this handler didn't accept, ensure they did not invoke the result callback
+		check(State == EHttpConnectionState::AwaitingProcessing);
+	}
+
+	if (!bRequestHandled)
+	{
+		const FString& ResponseError(FHttpServerErrorStrings::NotFound);
+		auto Response = FHttpServerResponse::Create(EHttpServerResponseCodes::NotFound, ResponseError);
+		OnProcessingComplete(MoveTemp(Response));
+	}
+}
+
+void FHttpConnection::BeginWrite(TUniquePtr<FHttpServerResponse>&& Response, uint32 RequestNumber)
+{
+	// Ensure the passed-in request number is the one we expect
+	check(RequestNumber == LastRequestNumber);
+
+	ChangeState(EHttpConnectionState::Writing);
+
+	if (bKeepAlive)
+	{
+		const FString& KeepAliveStr = FString::Printf(TEXT("timeout=%f"), ConnectionKeepAliveTimeout);
+		const TArray<FString>& KeepAliveValue = { KeepAliveStr };
+		Response->Headers.Add(FHttpServerHeaderKeys::KEEP_ALIVE, KeepAliveValue);
+	}
+
+	WriteContext.ResetContext(MoveTemp(Response));
+	ContinueWrite(0.0f);
+}
+
+void FHttpConnection::ContinueWrite(float DeltaTime)
+{
+	check(State == EHttpConnectionState::Writing);
+
+	auto WriterState = WriteContext.WriteStream(DeltaTime);
+	switch (WriterState)
+	{
+	case EHttpConnectionContextState::Continue:
+		break;
+
+	case EHttpConnectionContextState::Done:
+		CompleteWrite();
+		break;
+
+	case EHttpConnectionContextState::Error:
+		HandleWriteError(*WriteContext.GetErrorStr());
+		break;
+	}
+}
+
+void FHttpConnection::CompleteWrite()
+{
+	check(EHttpConnectionState::Writing == State);
+
+	if (bKeepAlive && !bGracefulDestroyRequested)
+	{
+		ChangeState(EHttpConnectionState::AwaitingRead);
+	}
+	else
+	{
+		Destroy();
+	}
+}
+
+void FHttpConnection::RequestDestroy(bool bGraceful)
+{
+	if (EHttpConnectionState::Destroyed == State)
+	{
+		return;
+	}
+
+	bGracefulDestroyRequested = bGraceful;
+
+	// If we aren't gracefully destroying, or we are otherwise already 
+	// awaiting a read operation (not started yet), destroy() immediately.
+	if (!bGracefulDestroyRequested || State == EHttpConnectionState::AwaitingRead)
+	{
+		Destroy();
+	}
+}
+
+void FHttpConnection::Destroy()
+{
+	check(State != EHttpConnectionState::Destroyed);
+	ChangeState(EHttpConnectionState::Destroyed);
+
+	if (Socket)
+	{
+		ISocketSubsystem *SocketSubsystem = ISocketSubsystem::Get(PLATFORM_SOCKETSUBSYSTEM);
+		if (SocketSubsystem)
+		{
+			SocketSubsystem->DestroySocket(Socket);
+		}
+		Socket = nullptr;
+	}
+}
+
+void FHttpConnection::HandleReadError(const TCHAR* ErrorCode)
+{
+	UE_LOG(LogHttpConnection, Error, TEXT("%s"), ErrorCode);
+
+	// Forcibly Reply
+	bKeepAlive = false;
+	auto ResponseCode = EHttpServerResponseCodes::BadRequest;
+	auto Response = FHttpServerResponse::Create(ResponseCode, ErrorCode);
+
+	BeginWrite(MoveTemp(Response), ++LastRequestNumber);
+}
+
+void FHttpConnection::HandleWriteError(const TCHAR* ErrorCode)
+{
+	UE_LOG(LogHttpConnection, Error, TEXT("%s"), ErrorCode);
+
+	// Forcibly Close
+	bKeepAlive = false;
+	Destroy();
+}
\ No newline at end of file
diff --git a/Engine/Source/Runtime/Online/HTTPServer/Private/HttpConnection.h b/Engine/Source/Runtime/Online/HTTPServer/Private/HttpConnection.h
new file mode 100644
--- /dev/null
+++ b/Engine/Source/Runtime/Online/HTTPServer/Private/HttpConnection.h
@@ -0,0 +1,193 @@
+// Copyright 1998-2019 Epic Games, Inc. All Rights Reserved.
+#pragma once
+
+#include "CoreMinimal.h"
+#include "HttpConnectionTypes.h"
+#include "HttpConnectionRequestReadContext.h"
+#include "HttpConnectionResponseWriteContext.h"
+#include "HttpResultCallback.h"
+
+class FSocket;
+class ISocketSubsystem;
+class FHttpRouter;
+struct FHttpServerRequest;
+struct FHttpServerResponse;
+
+DECLARE_LOG_CATEGORY_EXTERN(LogHttpConnection, Log, All);
+
+struct FHttpConnection final : public TSharedFromThis<FHttpConnection>
+{
+
+public:
+
+	/**
+	 * Constructor
+	 * 
+	 * @param InSocket The underlying file descriptor
+	 */
+	FHttpConnection(FSocket* Socket, TSharedPtr<FHttpRouter> Router);
+
+	/**
+	 * Destructor
+	 */
+	~FHttpConnection();
+
+	/** 
+	* Ticks the connection to drive internal state
+	*
+	* @param DeltaTime The elapsed time since the last Tick()
+	*/
+	void Tick(float DeltaTime);
+
+	/**
+	 * Returns the current state of the connection
+	 */
+	FORCEINLINE EHttpConnectionState GetState() const
+	{
+		return State;
+	}
+
+	/**
+	 * Signals the connection to destroy itself
+	 *
+	 * @param bGraceful Whether to gracefully destroy pending current operations
+	 */
+	void RequestDestroy(bool bGraceful);
+
+	/**
+	*  Determines whether this connection should remain open after writing
+	*/
+	FORCEINLINE bool IsHttpKeepAliveEnabled() const 
+	{ 
+		return bKeepAlive;
+	}
+
+	FORCEINLINE bool operator==(const FHttpConnection& Other) const
+	{
+		return this->Socket == Other.Socket;
+	}
+
+	static friend uint32 GetTypeHash(const FHttpConnection& Conn) 
+	{
+		return GetTypeHash(Conn.Socket);
+	}
+
+private:
+
+	/**
+     * Changes the internal state of the connection to NewState
+	 * 
+     * @param NewState The state to transition to
+     */
+	void ChangeState(EHttpConnectionState NewState);
+
+	/**
+	 * Changes the internal state from the caller-supplied current state to the next state
+	 *
+	 * @param CurrentState The expected current state
+	 * @param NextState    The next state
+	 */
+	void TransferState(EHttpConnectionState CurrentState, EHttpConnectionState NextState);
+
+	/**
+     * Begins a read operation
+	 *
+	 * @param DeltaTime The elapsed time since the last invocation
+     */
+	void BeginRead(float DeltaTime);
+
+	/**
+	 * Continues an in-progress read operation
+	 *
+	 * @param DeltaTime The elapsed time since the last invocation
+	 */
+	void ContinueRead(float DeltaTime);
+
+	/**
+	 * Completes a previously begun read operation
+	 *
+	 * @param Request The instantiated request object 
+	 */
+	void CompleteRead(const TSharedPtr<FHttpServerRequest>& Request);
+
+
+	/**
+	 * Proxies the respective request to a bound handler
+	 * @param Request              The request to process
+	 * @param OnProcessingComplete The callback to be invoked upon completion
+	 * @return                     true if the request was accepted by a handler, false otherwise
+	 */
+	void ProcessRequest(const TSharedPtr<FHttpServerRequest>& Request, const FHttpResultCallback& OnProcessingComplete);
+
+	/**
+     * Begins a write operation
+	 * 
+     * @param Response      The response to write
+	 * @param RequestNumber The expected and unique request number
+     */
+	void BeginWrite(TUniquePtr<FHttpServerResponse>&& Response, uint32 RequestNumber);
+
+	/**
+	 * Continues an in-progress write operation
+     *
+	 * @param DeltaTime The elapsed time since the last invocation
+	 */
+	void ContinueWrite(float DeltaTime);
+
+	/**
+	 * Completes a previously begun write operation
+	 */
+	void CompleteWrite();
+
+	/**
+	 * Closes and nullifies the underlying connection
+	 */
+	void Destroy();
+
+	/**
+	 * Logs and responds with the caller-supplied error code
+	 *
+	 * @param ErrorCode The machine-readable error description
+	 */
+	void HandleReadError(const TCHAR* ErrorCode);
+
+	/**
+	 * Logs the caller-supplied error code and closes the connection
+	 *
+	 * @param ErrorCode The machine-readable error description
+	 */
+	void HandleWriteError(const TCHAR* ErrorCode);
+
+private:
+
+	/** Accepted external socket */
+	FSocket* Socket;
+
+	/** State of the connection  */
+	EHttpConnectionState State = EHttpConnectionState::AwaitingRead;
+
+	/** Routing mechanism  */
+	const TSharedPtr<FHttpRouter> Router;
+
+	/** Helper reader context to track the state of streaming request reads */
+	FHttpConnectionRequestReadContext ReadContext;
+
+	/** Helper writer context to track the state of streaming response writes */
+	FHttpConnectionResponseWriteContext WriteContext;
+
+	/** Whether to keep this connection alive after writing */
+	bool bKeepAlive = false;
+
+	/** Whether to gracefully close pending current operations */
+	bool bGracefulDestroyRequested = false;
+
+	/** Internal state tracker (incremented per-request-read) used to validate request/response throughput */
+	uint32 LastRequestNumber = 0;
+
+	/** The duration (seconds) at which connections are forcefully timed out */
+	static constexpr float ConnectionTimeout = 5.0f;
+
+	/** The duration (seconds) at which idle keep-alive connections are forcefully timed out */
+	static constexpr float ConnectionKeepAliveTimeout = 15.0f;
+};
+
diff --git a/Engine/Source/Runtime/Online/HTTPServer/Private/HttpConnectionContext.cpp b/Engine/Source/Runtime/Online/HTTPServer/Private/HttpConnectionContext.cpp
new file mode 100644
--- /dev/null
+++ b/Engine/Source/Runtime/Online/HTTPServer/Private/HttpConnectionContext.cpp
@@ -0,0 +1,34 @@
+// Copyright 1998-2019 Epic Games, Inc. All Rights Reserved.
+
+#include "HttpConnectionContext.h"
+#include "CoreMinimal.h"
+
+FHttpConnectionContext::FHttpConnectionContext()
+{
+	ErrorBuilder.SetAutoEmitLineTerminator(true);
+}
+
+FHttpConnectionContext::~FHttpConnectionContext()
+{
+}
+
+void FHttpConnectionContext::AddElapsedIdleTime(float DeltaTime)
+{
+	ElapsedIdleTime += DeltaTime;
+}
+
+float FHttpConnectionContext::GetElapsedIdleTime() const
+{
+	return ElapsedIdleTime;
+}
+
+const FString& FHttpConnectionContext::GetErrorStr() const
+{
+	return ErrorBuilder;
+}
+
+void FHttpConnectionContext::AddError(const TCHAR* ErrorCode)
+{
+	ErrorBuilder.Append(ErrorCode);
+	ErrorBuilder.AppendChar(TCHAR(' '));
+}
diff --git a/Engine/Source/Runtime/Online/HTTPServer/Private/HttpConnectionContext.h b/Engine/Source/Runtime/Online/HTTPServer/Private/HttpConnectionContext.h
new file mode 100644
--- /dev/null
+++ b/Engine/Source/Runtime/Online/HTTPServer/Private/HttpConnectionContext.h
@@ -0,0 +1,60 @@
+// Copyright 1998-2019 Epic Games, Inc. All Rights Reserved.
+#pragma once
+#include "CoreMinimal.h"
+
+enum class EHttpConnectionContextState
+{
+	Continue,
+	Done,
+	Error
+};
+
+
+struct FHttpConnectionContext
+{
+
+public:
+
+	/**
+	 * Adds to the context elapsed time
+	 * 
+	 * @param DeltaTime The elapsed time (seconds) to add
+	 */
+	void AddElapsedIdleTime(float DeltaTime);
+
+	/**
+	 * Gets the elapsed time (seconds) since last activity
+	 */
+	FORCEINLINE float GetElapsedIdleTime() const;
+
+	/**
+	 * Gets the cumulative error string
+	 */
+	FORCEINLINE const FString& GetErrorStr() const;
+
+protected:
+
+	/**
+	 * Constructor
+	 */
+	FHttpConnectionContext();
+
+	/**
+	 * Destructor
+	 */
+	virtual ~FHttpConnectionContext();
+
+	/**
+	 * Adds the caller-supplied error to the context
+
+	 * @param ErrorCode The machine-readable error code
+	 */
+	void AddError(const TCHAR* ErrorCode);
+
+	/** Tracks time since last read/write activity */
+	float ElapsedIdleTime = 0.0f;
+
+	/** Tracks cumulative context errors */
+	FStringOutputDevice ErrorBuilder;
+
+};
\ No newline at end of file
diff --git a/Engine/Source/Runtime/Online/HTTPServer/Private/HttpConnectionRequestReadContext.cpp b/Engine/Source/Runtime/Online/HTTPServer/Private/HttpConnectionRequestReadContext.cpp
new file mode 100644
--- /dev/null
+++ b/Engine/Source/Runtime/Online/HTTPServer/Private/HttpConnectionRequestReadContext.cpp
@@ -0,0 +1,370 @@
+
+#include "HttpConnectionRequestReadContext.h"
+#include "HttpServerConstantsPrivate.h"
+#include "HttpServerRequest.h"
+#include "HttpConnectionContext.h"
+#include "Sockets.h"
+
+FHttpConnectionRequestReadContext::FHttpConnectionRequestReadContext(FSocket* InSocket)
+	: Socket(InSocket)
+{
+}
+
+void FHttpConnectionRequestReadContext::ResetContext()
+{
+	Request = nullptr;
+	ElapsedIdleTime = 0.0f;
+	ErrorBuilder.Empty();
+	HeaderBytes.Empty();
+	IncomingRequestBodyBytesToRead = 0;
+
+	bParseHeaderComplete = false;
+	bParseBodyComplete = false;
+}
+
+EHttpConnectionContextState FHttpConnectionRequestReadContext::ReadStream(float DeltaTime)
+{
+	ElapsedIdleTime += DeltaTime;
+
+	uint8 ByteBuffer[2 * 1024] = { 0 };
+	int32 BytesRead = 0;
+	if (!Socket->Recv(ByteBuffer, sizeof(ByteBuffer) - 1, BytesRead, ESocketReceiveFlags::None))
+	{
+		AddError(FHttpServerErrorStrings::SocketRecvFailure);
+		return EHttpConnectionContextState::Error;
+	}
+
+	if (0 == BytesRead)
+	{
+		return EHttpConnectionContextState::Continue;
+	}
+	else
+	{
+		ElapsedIdleTime = 0.0f;
+	}
+
+	if (!bParseHeaderComplete)
+	{
+		if (!ParseHeader(ByteBuffer, BytesRead))
+		{
+			return EHttpConnectionContextState::Error;
+		}
+	}
+	else if (!bParseBodyComplete)
+	{
+		if (!ParseBody(ByteBuffer, BytesRead))
+		{
+			return EHttpConnectionContextState::Error;
+		}
+	}
+		
+	// If parsing is complete, we are done
+	if (bParseHeaderComplete && bParseBodyComplete)
+	{
+		check(Request);
+		return EHttpConnectionContextState::Done;
+	}
+	
+	return EHttpConnectionContextState::Continue;
+}
+
+bool FHttpConnectionRequestReadContext::ParseHeader(uint8* ByteBuffer, int32 BufferLen)
+{
+	// Append this chunk to our header bytes
+	int PreviousHeaderBytesLen = HeaderBytes.Num();
+	HeaderBytes.Append(ByteBuffer, BufferLen);
+
+	// We are hunting for the header terminator sequence '\r\n\r\n'
+	// Start 3 bytes prior to this most recent appendage
+	int32 SearchIndex = FMath::Max(0, PreviousHeaderBytesLen - 3);
+	for (SearchIndex; SearchIndex < HeaderBytes.Num() - 3; ++SearchIndex)
+	{
+		// To find this sequence we will parse byte-by-byte in ANSICHAR
+		if (ANSICHAR(HeaderBytes[SearchIndex])     == '\r' &&
+			ANSICHAR(HeaderBytes[SearchIndex + 1]) == '\n' &&
+			ANSICHAR(HeaderBytes[SearchIndex + 2]) == '\r' &&
+			ANSICHAR(HeaderBytes[SearchIndex + 3]) == '\n')
+		{
+			// We found the terminator
+			bParseHeaderComplete = true;
+
+			// Remove any extra header bytes, but retain the first trailing \r\n 
+			HeaderBytes.SetNum(SearchIndex + 2);
+			HeaderBytes.Add('\0');
+
+			// Build header string
+			FUTF8ToTCHAR WByteBuffer(reinterpret_cast<const ANSICHAR*>(HeaderBytes.GetData()), HeaderBytes.Num());
+			FString IncomingRequestHeaderStr(WByteBuffer.Get(), WByteBuffer.Length());
+			Request = BuildRequest(IncomingRequestHeaderStr);
+			if (!Request)
+			{
+				// Error if we cannot build a request object
+				return false;
+			}
+
+			if (!IsRequestValid(*Request.Get()))
+			{
+				// Error if request object is invalid
+				return false;
+			}
+
+			// Try to parse the rest of the request body here in the same frame
+			ParseContentLength(*Request.Get(), IncomingRequestBodyBytesToRead);
+
+			//                    base         relative offset    absolute offset
+			auto BodyBytesPtr = ByteBuffer +  (SearchIndex + 4 - PreviousHeaderBytesLen);
+			auto BufferBytesRemaining = BufferLen - (BodyBytesPtr - ByteBuffer);
+			return ParseBody(BodyBytesPtr, BufferBytesRemaining);
+		}
+	}
+
+	// Keep parsing.
+	return true;
+}
+
+bool FHttpConnectionRequestReadContext::ParseBody(uint8* ByteBuffer, int32 BufferLen)
+{
+	if (BufferLen > IncomingRequestBodyBytesToRead)
+	{
+		// Error - Sent data size exceeds expected 
+		AddError(FHttpServerErrorStrings::MismatchedContentLengthBodyTooLarge);
+		return false;
+	}
+
+	// Append bytes
+	if (BufferLen > 0)
+	{
+		Request->Body.Append(ByteBuffer, BufferLen);
+		IncomingRequestBodyBytesToRead -= BufferLen;
+	}
+
+	bParseBodyComplete = (0 == IncomingRequestBodyBytesToRead);
+	return true;
+}
+
+bool FHttpConnectionRequestReadContext::IsRequestValid(const FHttpServerRequest& InRequest) 
+{
+	int32 RequestContentLength; 
+	bool bContentLengthSpecified = ParseContentLength(InRequest, RequestContentLength);
+
+	switch (InRequest.Verb)
+	{
+	case EHttpServerRequestVerbs::GET:
+		// Enforce content length missing or 0
+		if (bContentLengthSpecified && 0 != RequestContentLength)
+		{
+			AddError(FHttpServerErrorStrings::InvalidContentLengthHeader);
+			return false;
+		}
+		break;
+
+	case EHttpServerRequestVerbs::PUT:
+	case EHttpServerRequestVerbs::POST:
+		// Content length must be set
+		if (!bContentLengthSpecified)
+		{
+			AddError(FHttpServerErrorStrings::MissingContentLengthHeader);
+			return false;
+		}
+		// Content length must be valid
+		if(RequestContentLength < 0)
+		{
+			AddError(FHttpServerErrorStrings::InvalidContentLengthHeader);
+			return false;
+		}
+		break;
+	}
+
+	return true;
+}
+
+bool FHttpConnectionRequestReadContext::ParseContentLength(const FHttpServerRequest& InRequest, int32& OutContentLength)
+{
+	const TArray<FString>* ContentLengthValues = InRequest.Headers.Find(FHttpServerHeaderKeys::CONTENT_LENGTH);
+
+	if (ContentLengthValues && ContentLengthValues->Num() > 0)
+	{
+		const FString& ContentLengthStr = (*ContentLengthValues)[0];
+		OutContentLength = FMath::Max(0, FCString::Atoi(*ContentLengthStr));
+		return true;
+	}
+	return false;
+}
+
+TSharedPtr<FHttpServerRequest> FHttpConnectionRequestReadContext::BuildRequest(const FString& RequestHeader)
+{
+	TArray<FString> ParsedHeader;
+	RequestHeader.ParseIntoArrayLines(ParsedHeader);
+
+	if (0 == ParsedHeader.Num())
+	{
+		AddError(FHttpServerErrorStrings::MissingRequestHeaders);
+		return nullptr;
+	}
+
+	// Split HTTP Method Line (Path/Verb)
+	TArray<FString> HttpMethodTokens;
+	const FString& HttpMethod = ParsedHeader[0];
+	HttpMethod.ParseIntoArrayWS(HttpMethodTokens);
+	if (HttpMethodTokens.Num() < 2)
+	{
+		AddError(FHttpServerErrorStrings::MalformedRequestHeaders);
+		return nullptr;
+	}
+
+	Request = MakeShared<FHttpServerRequest>();
+
+	auto RequestVerb = HttpMethodTokens[0];
+	if (0 == RequestVerb.Compare(TEXT("GET"),
+		ESearchCase::IgnoreCase))
+	{
+		Request->Verb = EHttpServerRequestVerbs::GET;
+	}
+	else if (0 == RequestVerb.Compare(TEXT("POST"),
+		ESearchCase::IgnoreCase))
+	{
+		Request->Verb = EHttpServerRequestVerbs::POST;
+	}
+	else if (0 == RequestVerb.Compare(TEXT("PUT"),
+		ESearchCase::IgnoreCase))
+	{
+		Request->Verb = EHttpServerRequestVerbs::PUT;
+	}
+	else if (0 == RequestVerb.Compare(TEXT("DELETE"),
+		ESearchCase::IgnoreCase))
+	{
+		Request->Verb = EHttpServerRequestVerbs::DELETE;
+	}
+	else if (0 == RequestVerb.Compare(TEXT("PATCH"),
+		ESearchCase::IgnoreCase))
+	{
+		Request->Verb = EHttpServerRequestVerbs::PATCH;
+	}
+	else
+	{
+		// Unknown Verb
+		AddError(FHttpServerErrorStrings::UnknownRequestVerb);
+		return nullptr;
+	}
+
+	// Parse/store path+query params
+	auto RequestHttpPath = HttpMethodTokens[1].TrimStartAndEnd();
+	int32 QueryParamsIndex = 0;
+	if (RequestHttpPath.FindChar(TCHAR('?'), QueryParamsIndex))
+	{
+		FString QueryParamsStr = RequestHttpPath.Mid(QueryParamsIndex+1);
+		RequestHttpPath = RequestHttpPath.Mid(0, QueryParamsIndex);
+
+		// Split query params
+		TArray<FString> QueryParamPairs;
+		const TCHAR QueryParamDelimiter[] = { TCHAR('&') };
+		QueryParamsStr.ParseIntoArray(QueryParamPairs, QueryParamDelimiter, true);
+		for (const FString& QueryParamPair : QueryParamPairs)
+		{
+			int32 Equalsindex = 0;
+			if (QueryParamPair.FindChar(TCHAR('='), Equalsindex))
+			{
+				FString QueryParamKey = UrlDecode(QueryParamPair.Mid(0, Equalsindex));
+				FString QueryParamValue = UrlDecode(QueryParamPair.Mid(Equalsindex + 1));
+				Request->QueryParams.Emplace(MoveTemp(QueryParamKey), MoveTemp(QueryParamValue));
+			}
+		}
+	}
+	Request->RelativePath.SetPath(RequestHttpPath);
+
+	// Parse/store headers
+	for (int i = 1; i < ParsedHeader.Num(); ++i)
+	{
+		const FString& HeaderLine = ParsedHeader[i];
+		int32 SplitIndex = 0;
+		if (HeaderLine.FindChar(TCHAR(':'), SplitIndex))
+		{
+			const auto& HeaderKey = HeaderLine.Mid(0, SplitIndex).TrimStartAndEnd().ToLower();
+			const auto& HeaderValuesStr = HeaderLine.Mid(SplitIndex + 1).TrimStartAndEnd();
+
+			TArray<FString> HeaderValues;
+			const TCHAR HeaderValueDelimiters[] = { TCHAR(',') };
+			HeaderValuesStr.ParseIntoArray(HeaderValues, HeaderValueDelimiters, true);
+
+			if (HeaderValues.Num() > 0)
+			{
+				TArray<FString>* ExistingHeaderValues = Request->Headers.Find(HeaderKey);
+				if (ExistingHeaderValues)
+				{
+					ExistingHeaderValues->Append(HeaderValues);
+				}
+				else
+				{
+					Request->Headers.Emplace(HeaderKey, MoveTemp(HeaderValues));
+				}
+			}
+		}
+	}
+
+	return Request;
+}
+
+FString FHttpConnectionRequestReadContext::UrlDecode(const FString &EncodedString)
+{
+	FTCHARToUTF8 Converter(*EncodedString);
+	const UTF8CHAR* UTF8Data = (UTF8CHAR*)Converter.Get();
+
+	TArray<ANSICHAR> Data;
+	Data.Reserve(EncodedString.Len());
+
+	for (int32 CharIdx = 0; CharIdx < Converter.Length();)
+	{
+		if (UTF8Data[CharIdx] == '%')
+		{
+			int32 Value = 0;
+			if (UTF8Data[CharIdx + 1] == 'u')
+			{
+				if (CharIdx + 6 <= Converter.Length())
+				{
+					// Treat all %uXXXX as code point
+					Value = FParse::HexDigit(UTF8Data[CharIdx + 2]) << 12;
+					Value += FParse::HexDigit(UTF8Data[CharIdx + 3]) << 8;
+					Value += FParse::HexDigit(UTF8Data[CharIdx + 4]) << 4;
+					Value += FParse::HexDigit(UTF8Data[CharIdx + 5]);
+					CharIdx += 6;
+
+					ANSICHAR Buffer[8] = { 0 };
+					ANSICHAR* BufferPtr = Buffer;
+					const int32 Len = ARRAY_COUNT(Buffer);
+					const int32 WrittenChars = FTCHARToUTF8_Convert::Utf8FromCodepoint(Value, BufferPtr, Len);
+
+					Data.Append(Buffer, WrittenChars);
+				}
+				else
+				{
+					// Not enough in the buffer for valid decoding, skip it
+					CharIdx++;
+					continue;
+				}
+			}
+			else if (CharIdx + 3 <= Converter.Length())
+			{
+				// Treat all %XX as straight byte
+				Value = FParse::HexDigit(UTF8Data[CharIdx + 1]) << 4;
+				Value += FParse::HexDigit(UTF8Data[CharIdx + 2]);
+				CharIdx += 3;
+				Data.Add((ANSICHAR)(Value));
+			}
+			else
+			{
+				// Not enough in the buffer for valid decoding, skip it
+				CharIdx++;
+				continue;
+			}
+		}
+		else
+		{
+			// Non escaped characters
+			Data.Add(UTF8Data[CharIdx]);
+			CharIdx++;
+		}
+	}
+
+	Data.Add('\0');
+	return FString(UTF8_TO_TCHAR(Data.GetData()));
+}
\ No newline at end of file
diff --git a/Engine/Source/Runtime/Online/HTTPServer/Private/HttpConnectionRequestReadContext.h b/Engine/Source/Runtime/Online/HTTPServer/Private/HttpConnectionRequestReadContext.h
new file mode 100644
--- /dev/null
+++ b/Engine/Source/Runtime/Online/HTTPServer/Private/HttpConnectionRequestReadContext.h
@@ -0,0 +1,115 @@
+// Copyright 1998-2019 Epic Games, Inc. All Rights Reserved.
+#pragma once
+#include "CoreMinimal.h"
+#include "HttpConnectionContext.h"
+
+struct FHttpServerRequest;
+struct FHttpConnection;
+
+class FSocket;
+
+
+struct FHttpConnectionRequestReadContext final : public FHttpConnectionContext
+{
+public:
+
+	/**
+	 * Constructor
+	 *
+	 * @param Socket The underlying connection
+	 */
+	FHttpConnectionRequestReadContext(FSocket* InSocket);
+
+	/**
+	 * Reads a request from the connection
+
+	 * @param DeltaTime The elapsed time since the last invocation
+	 * @return The state of the read context
+	 */
+	EHttpConnectionContextState ReadStream(float DeltaTime);
+
+	/**
+	 * Gets the internally stored request object
+	 */
+	FORCEINLINE TSharedPtr<FHttpServerRequest> GetRequest() const
+	{
+		return Request;
+	}
+
+	/**
+     * Resets the internal state. 
+	 * Should be called for every read.
+     */
+	void ResetContext();
+
+private:
+
+	/**
+	 * Parses the request header
+	 *
+	 * @param  ByteBuffer The bytes to parse
+	 * @param  BufferLen The length of the supplied byte buffer
+	 * @return true if no error was encountered, false otherwise 
+	 */
+	bool ParseHeader(uint8* ByteBuffer, int32 BufferLen);
+
+	/**
+	 * Parses the request body
+	 *
+	 * @param  ByteBuffer The bytes to parse
+	 * @param  BufferLen The length of the supplied byte buffer
+	 * @return true if no error was encountered, false otherwise
+	 */
+	bool ParseBody(uint8* ByteBuffer, int32 BufferLen);
+
+	/**
+	 * Determines whether the member Request object is valid
+	 * @param InRequest The request to validate
+	 * eg: POST requests with ContentLength > 0 are invalid
+	 * eg: Body-capable requests with mismatched content-length headers and bodies are invalid
+	 *
+	 * @return true if the current Request object is valid, false otherwise
+	 */
+	bool IsRequestValid(const FHttpServerRequest& InRequest);
+
+	/**
+	* Builds an FHttpServerRequest from the caller-supplied http header string
+	*
+	* @param RequestHeader The complete http request header string
+	* @return An instantiated request object with valid input, nullptr otherwise
+	*/
+	TSharedPtr<FHttpServerRequest> BuildRequest(const FString& RequestHeader);
+
+	/**
+	 * Gets the expected Content-Length from the member Request object
+
+	 * @param InRequest The request to validate
+	 * @param OutContentLength Set to the respective content length when returns true
+	 * @return true if content length is specified, false otherwise
+	 */
+	static bool ParseContentLength(const FHttpServerRequest& InRequest, int32& OutContentLength);
+
+	/**
+	 * Returns a decoded version of the percent-encoded passed in string
+	 *
+	 * @param EncodedString The percent encoded string to convert to string
+	 * @return The decoded string
+	 */
+	static FString UrlDecode(const FString& EncodedString);
+
+
+private:
+
+	/** The underlying connection */
+	FSocket* Socket;
+	/** Internal request state. */
+	TSharedPtr<FHttpServerRequest> Request;
+	/** Whether the header has been fully parsed */
+	bool bParseHeaderComplete = false;
+	/** Whether the body has been fully parsed */
+	bool bParseBodyComplete = false;
+	/** State variable which stores incoming header bytes */
+	TArray<uint8> HeaderBytes;
+	/** State variable which tracks bytes read for an incoming request */
+	int32 IncomingRequestBodyBytesToRead = 0;
+};
\ No newline at end of file
diff --git a/Engine/Source/Runtime/Online/HTTPServer/Private/HttpConnectionResponseWriteContext.cpp b/Engine/Source/Runtime/Online/HTTPServer/Private/HttpConnectionResponseWriteContext.cpp
new file mode 100644
--- /dev/null
+++ b/Engine/Source/Runtime/Online/HTTPServer/Private/HttpConnectionResponseWriteContext.cpp
@@ -0,0 +1,123 @@
+
+#include "HttpConnectionResponseWriteContext.h"
+#include "HttpServerResponse.h"
+#include "HttpServerConstantsPrivate.h"
+#include "Sockets.h"
+
+FHttpConnectionResponseWriteContext::FHttpConnectionResponseWriteContext(FSocket* InSocket)
+	: Socket(InSocket)
+{
+}
+
+void FHttpConnectionResponseWriteContext::ResetContext(TUniquePtr<FHttpServerResponse>&& Resp)
+{
+	Response = MoveTemp(Resp);
+	ElapsedIdleTime = 0.0f;
+	ErrorBuilder.Empty();
+
+	HeaderBytesWritten = 0;
+	BodyBytesWritten = 0;
+	HeaderBytes.Empty();
+
+	if (Response)
+	{
+		// Affix Content-Length Header
+		TArray<FString> ContentLengthValue = { FString::FromInt(Response->Body.Num()) };
+		Response->Headers.Add(FHttpServerHeaderKeys::CONTENT_LENGTH, MoveTemp(ContentLengthValue));
+
+		// Serialize Headers
+		HeaderBytes.Append(SerializeHeadersUtf8(Response->Code, Response->Headers));
+	}
+}
+
+EHttpConnectionContextState FHttpConnectionResponseWriteContext::WriteStream(float DeltaTime)
+{
+	ElapsedIdleTime += DeltaTime;
+
+	int32 BytesWritten = 0;
+
+	if (!IsWriteHeaderComplete())
+	{
+		const uint8* DataOffset = HeaderBytes.GetData() + HeaderBytesWritten;
+		int32 DataLen = HeaderBytes.Num() - HeaderBytesWritten;
+		if (!WriteBytes(DataOffset, DataLen, BytesWritten))
+		{
+			AddError(FHttpServerErrorStrings::SocketSendFailure);
+			return EHttpConnectionContextState::Error;
+		}
+		HeaderBytesWritten += BytesWritten;
+	}
+
+	if (IsWriteHeaderComplete() && !IsWriteBodyComplete())
+	{
+		const uint8* DataOffset = Response->Body.GetData() + BodyBytesWritten;
+		int32 DataLen = Response->Body.Num() - BodyBytesWritten;
+		if (!WriteBytes(DataOffset, DataLen, BytesWritten))
+		{
+			AddError(FHttpServerErrorStrings::SocketSendFailure);
+			return EHttpConnectionContextState::Error;
+		}
+		BodyBytesWritten += BytesWritten;
+	}
+
+	if (IsWriteHeaderComplete() && IsWriteBodyComplete())
+	{
+		ensure(BodyBytesWritten+HeaderBytesWritten == HeaderBytes.Num()+Response->Body.Num());
+		return EHttpConnectionContextState::Done;
+	}
+
+	return EHttpConnectionContextState::Continue;
+}
+
+bool FHttpConnectionResponseWriteContext::WriteBytes(const uint8* Bytes, int32 BytesLen, int32 &OutBytesWritten)
+{
+	OutBytesWritten = 0;
+	bool bWriteSuccess = Socket->Send(Bytes, BytesLen, OutBytesWritten);
+	if (!bWriteSuccess)
+	{
+		UE_LOG(LogHttpConnection, Warning,
+			TEXT("WriteBytes sent %d/%d bytes"), OutBytesWritten, BytesLen);
+	}
+
+	if (OutBytesWritten > 0)
+	{
+		ElapsedIdleTime = 0.0f;
+	}
+
+	return bWriteSuccess;
+}
+
+bool FHttpConnectionResponseWriteContext::IsWriteHeaderComplete() const
+{
+	return HeaderBytesWritten >= HeaderBytes.Num();
+}
+
+bool FHttpConnectionResponseWriteContext::IsWriteBodyComplete() const
+{
+	check(Response);
+	return BodyBytesWritten >= Response->Body.Num();
+}
+
+TArray<uint8> FHttpConnectionResponseWriteContext::SerializeHeadersUtf8(int32 ResponseCode, const TMap<FString, TArray<FString>>& HeadersMap)
+{
+	FString ResponseHeaderStr = FString::Printf(TEXT("HTTP/1.1 %d\r\n"), ResponseCode);
+
+	for (const auto& KeyValuePair : HeadersMap)
+	{
+		const FString& HeaderKey = KeyValuePair.Key;
+		const TArray<FString>& HeaderValues = KeyValuePair.Value;
+		for (const auto& HeaderValue : HeaderValues)
+		{
+			const FString HeaderItem = FString::Printf(TEXT("%s: %s\r\n"), *HeaderKey, *HeaderValue);
+			ResponseHeaderStr.Append(HeaderItem);
+		}
+	}
+	ResponseHeaderStr.Append(TEXT("\r\n"));
+
+	TArray<uint8> HeaderRawBytes;
+	FTCHARToUTF8 HeaderUtf8(*ResponseHeaderStr);
+	const uint8* HeaderUtf8Bytes = reinterpret_cast<const uint8*>(HeaderUtf8.Get());
+	HeaderRawBytes.Append(HeaderUtf8Bytes, HeaderUtf8.Length());
+	return HeaderRawBytes;
+}
+
diff --git a/Engine/Source/Runtime/Online/HTTPServer/Private/HttpConnectionResponseWriteContext.h b/Engine/Source/Runtime/Online/HTTPServer/Private/HttpConnectionResponseWriteContext.h
new file mode 100644
--- /dev/null
+++ b/Engine/Source/Runtime/Online/HTTPServer/Private/HttpConnectionResponseWriteContext.h
@@ -0,0 +1,82 @@
+// Copyright 1998-2019 Epic Games, Inc. All Rights Reserved.
+#pragma once
+#include "CoreMinimal.h"
+#include "HttpConnectionContext.h"
+
+struct FHttpServerResponse;
+class FSocket;
+
+struct FHttpConnectionResponseWriteContext final : public FHttpConnectionContext
+{
+public:
+
+	/**
+	 * Constructor
+	 *
+	 * @param InSocket The underlying connection
+	 */
+	FHttpConnectionResponseWriteContext(FSocket* InSocket);
+
+	/**
+     * Writes to the connection
+	 *
+	 * @param DeltaTime The elapsed time since the last invocation
+	 * @return The state of the write context
+     */
+	EHttpConnectionContextState WriteStream(float DeltaTime);
+
+	/**
+	 * Resets the internal state.  
+	 * Should be called before every write.
+	 *
+	 * @param Response The response to be written
+	 */
+	void ResetContext(TUniquePtr<FHttpServerResponse>&&  Response);
+
+private:
+
+	/**
+	 * Writes the caller-supplied Bytes up to BytesLen as permitted by the socket
+	 *
+     * @param  Bytes The bytes to write
+	 * @param  The desired number of bytes to write
+	 * @param  OutBytesWritten The number of bytes actually written
+	 * @return true if the write was successful, false otherwise
+	 */
+	bool WriteBytes(const uint8* Bytes, int32 BytesLen, int32 &OutBytesWritten);
+
+
+	/**
+	 * Determines whether the header has been completely written
+	 * @return true if the header has been completely written, false otherwise
+	 */
+	bool IsWriteHeaderComplete() const;
+
+	/**
+	 * Determines whether the body has been completely written
+	 * @return true if the body has been completely written, false otherwise
+	 */
+	bool IsWriteBodyComplete() const;
+
+	/**
+	 * Creates a UTF8-serialized HTTP header in byte representation
+	 *
+	 * @param  The HTTP response code
+	 * @param  HeadersMap The associative header key to values map
+	 * @return The UTF8-encoded byte representation
+	 */
+	static TArray<uint8> SerializeHeadersUtf8(int32 ResponseCode, const TMap<FString, TArray<FString>>& HeadersMap);
+
+private:
+
+	/** The underlying connection */
+	FSocket* Socket;
+	/** The response to write */
+	TUniquePtr<FHttpServerResponse> Response;
+	/** State  */
+	TArray<uint8> HeaderBytes;
+	/** The total number of header bytes written */
+	int32 HeaderBytesWritten = 0;
+	/** The total number of body bytes written */
+	int32 BodyBytesWritten = 0;
+};
\ No newline at end of file
diff --git a/Engine/Source/Runtime/Online/HTTPServer/Private/HttpConnectionTypes.h b/Engine/Source/Runtime/Online/HTTPServer/Private/HttpConnectionTypes.h
new file mode 100644
--- /dev/null
+++ b/Engine/Source/Runtime/Online/HTTPServer/Private/HttpConnectionTypes.h
@@ -0,0 +1,34 @@
+// Copyright 1998-2019 Epic Games, Inc. All Rights Reserved.
+#pragma once
+#include "CoreMinimal.h"
+#include "Containers/Set.h"
+
+struct FHttpConnection;
+
+/**
+ * Utility unique identifier for FHttpConnections
+ */
+typedef FGuid HttpConnectionIdType;
+
+/**
+ * FHttpConnectionPool
+ */
+typedef TSet<TSharedPtr<FHttpConnection>> FHttpConnectionPool;
+
+/**
+ * FHttpConnection state model
+ */
+enum class EHttpConnectionState
+{
+	// Awaiting available data
+	AwaitingRead,
+	// Reading available data
+	Reading,
+	// Awaiting request handler completion
+	AwaitingProcessing,
+	// Writing response to client
+	Writing,
+	// Awaiting destruction via listener
+	Destroyed,
+};
+
diff --git a/Engine/Source/Runtime/Online/HTTPServer/Private/HttpListener.cpp b/Engine/Source/Runtime/Online/HTTPServer/Private/HttpListener.cpp
new file mode 100644
--- /dev/null
+++ b/Engine/Source/Runtime/Online/HTTPServer/Private/HttpListener.cpp
@@ -0,0 +1,225 @@
+// Copyright 1998-2019 Epic Games, Inc. All Rights Reserved.
+
+#include "HttpListener.h"
+#include "HttpConnection.h"
+#include "HttpRequestHandler.h"
+#include "HttpRouter.h"
+
+#include "Sockets.h"
+#include "SocketSubsystem.h"
+#include "IPAddress.h"
+
+DEFINE_LOG_CATEGORY(LogHttpListener)
+
+FHttpListener::FHttpListener(uint32 Port)
+{ 
+	check(Port > 0);
+	ListenPort = Port;
+	Router = MakeShared<FHttpRouter>();
+}
+
+FHttpListener::~FHttpListener() 
+{ 
+	check(nullptr == ListenSocket);
+	check(!bIsListening);
+
+	const bool bRequestGracefulExit = false;
+	for (const auto& Connection : Connections)
+	{
+		Connection->RequestDestroy(bRequestGracefulExit);
+	}
+	Connections.Empty();
+}
+
+// --------------------------------------------------------------------------------------------
+// Public Interface
+// --------------------------------------------------------------------------------------------
+bool FHttpListener::StartListening()
+{
+	check(nullptr == ListenSocket);
+	check(!bIsListening);
+	bIsListening = true;
+
+	ISocketSubsystem* SocketSubsystem = ISocketSubsystem::Get(PLATFORM_SOCKETSUBSYSTEM);
+	if (nullptr == SocketSubsystem)
+	{
+		UE_LOG(LogHttpListener, Error, 
+			TEXT("IHttpListener - SocketSubsystem Initialization Failed"));
+		return false;
+	}
+
+	ListenSocket = SocketSubsystem->CreateSocket(NAME_Stream, TEXT("HttpListenerSocket"));
+	if (nullptr == ListenSocket)
+	{
+		UE_LOG(LogHttpListener, Error, 
+			TEXT("IHttpListener - Unable to allocate stream socket"));
+		return false;
+	}
+	ListenSocket->SetNonBlocking(true);
+
+	// Bind to Localhost/Caller-defined port
+	TSharedRef<FInternetAddr> LocalhostAddr = SocketSubsystem->CreateInternetAddr(ListenerInterfaceAddress, ListenPort);
+	if (!ListenSocket->Bind(*LocalhostAddr))
+	{
+		UE_LOG(LogHttpListener, Error, 
+			TEXT("IHttpListener unable to bind to %s"), 
+			*LocalhostAddr->ToString(true));
+		return false;
+	}
+
+	int32 ActualBufferSize;
+	ListenSocket->SetSendBufferSize(ListenerBufferSize, ActualBufferSize);
+	if (ActualBufferSize != ListenerBufferSize)
+	{
+		UE_LOG(LogHttpListener, Warning, 
+			TEXT("IHttpListener unable to set desired buffer size (%d): Limited to %d"),
+			ListenerBufferSize, ActualBufferSize);
+	}
+
+	if (!ListenSocket->Listen(ListenerConnectionBacklogSize))
+	{
+		UE_LOG(LogHttpListener, Error, 
+			TEXT("IHttpListener unable to listen on socket"));
+		return false;
+	}
+
+	UE_LOG(LogHttpListener, Log, 
+		TEXT("Created new HttpListener on port %d"), ListenPort);
+	return true;
+}
+
+void FHttpListener::StopListening()
+{
+	check(bIsListening);
+
+	// Tear down our top-level listener first
+	if (ListenSocket)
+	{
+		ISocketSubsystem* SocketSubsystem = ISocketSubsystem::Get(PLATFORM_SOCKETSUBSYSTEM);
+		if (SocketSubsystem)
+		{
+			SocketSubsystem->DestroySocket(ListenSocket);
+		}
+		ListenSocket = nullptr;
+	}
+	bIsListening = false;
+
+	const bool bRequestGracefulExit = true;
+	for (const auto& Connection : Connections)
+	{
+		Connection->RequestDestroy(bRequestGracefulExit);
+	}
+}
+
+void FHttpListener::Tick(float DeltaTime)
+{
+	// Accept new connections
+	AcceptConnections(MaxConnectionsToAcceptPerFrame);
+
+	// Tick Connections
+	TickConnections(DeltaTime);
+
+	// Remove any destroyed connections
+	RemoveDestroyedConnections();
+}
+
+bool FHttpListener::HasPendingConnections() const 
+{
+	for (const auto& Connection : Connections)
+	{
+		switch (Connection->GetState())
+		{
+		case EHttpConnectionState::Reading:
+		case EHttpConnectionState::AwaitingProcessing:
+		case EHttpConnectionState::Writing:
+			return true;
+		}
+	}
+	return false;
+}
+
+// --------------------------------------------------------------------------------------------
+// Private Implementation
+// --------------------------------------------------------------------------------------------
+void FHttpListener::AcceptConnections(uint32 MaxConnectionsToAccept)
+{
+	check(ListenSocket);
+
+	for (uint32 i = 0; i < MaxConnectionsToAccept; ++i)
+	{
+		// Check pending prior to Accept()ing
+		bool bHasPendingConnection = false;
+		if (!ListenSocket->HasPendingConnection(bHasPendingConnection))
+		{
+			UE_LOG(LogHttpListener, 
+				Error, TEXT("ListenSocket failed to query pending connection"));
+			return;
+		}
+
+		if (bHasPendingConnection)
+		{
+			FSocket* IncomingConnection = ListenSocket->Accept(TEXT("HttpRequest"));
+
+			if (nullptr == IncomingConnection)
+			{
+				ESocketErrors ErrorCode = ESocketErrors::SE_NO_ERROR;
+				FString ErrorStr = TEXT("SocketSubsystem Unavialble");
+
+				ISocketSubsystem* SocketSubsystem = ISocketSubsystem::Get(PLATFORM_SOCKETSUBSYSTEM);
+				if (SocketSubsystem)
+				{
+					ErrorCode = SocketSubsystem->GetLastErrorCode();
+					ErrorStr = SocketSubsystem->GetSocketError();
+				}
+				UE_LOG(LogHttpListener, Error,
+					TEXT("Error accepting expected connection [%d] %s"), (int32)ErrorCode, *ErrorStr);
+				return;
+			}
+
+			IncomingConnection->SetNonBlocking(true);
+			TSharedPtr<FHttpConnection> Connection = MakeShared<FHttpConnection>(IncomingConnection, Router);
+			Connections.Add(Connection);
+		}
+	}
+}
+
+void FHttpListener::TickConnections(float DeltaTime)
+{
+	for (const auto& Connection : Connections)
+	{
+		check(Connection.IsValid());
+
+		switch (Connection->GetState())
+		{
+		case EHttpConnectionState::AwaitingRead:
+		case EHttpConnectionState::Reading:
+			Connection->Tick(DeltaTime);
+			break;
+		}
+	}
+
+	for (const auto& Connection : Connections)
+	{
+		check(Connection.IsValid());
+
+		switch (Connection->GetState())
+		{
+		case EHttpConnectionState::Writing:
+			Connection->Tick(DeltaTime);
+			break;
+		}
+	}
+}
+
+void FHttpListener::RemoveDestroyedConnections()
+{
+	for (auto ConnectionsIter = Connections.CreateIterator(); ConnectionsIter; ++ConnectionsIter)
+	{
+		// Remove any destroyed connections
+		if (EHttpConnectionState::Destroyed == ConnectionsIter->Get()->GetState())
+		{
+			ConnectionsIter.RemoveCurrent();
+			continue;
+		}
+	}
+}
\ No newline at end of file
diff --git a/Engine/Source/Runtime/Online/HTTPServer/Private/HttpListener.h b/Engine/Source/Runtime/Online/HTTPServer/Private/HttpListener.h
new file mode 100644
--- /dev/null
+++ b/Engine/Source/Runtime/Online/HTTPServer/Private/HttpListener.h
@@ -0,0 +1,122 @@
+// Copyright 1998-2019 Epic Games, Inc. All Rights Reserved.
+#pragma once
+
+#include "HttpConnectionTypes.h"
+
+struct FHttpConnection;
+struct FHttpPath;
+class FSocket;
+
+class IHttpRouter;
+class FHttpRouter;
+
+DECLARE_LOG_CATEGORY_EXTERN(LogHttpListener, Log, All);
+
+class FHttpListener final 
+{
+
+public:
+
+	/**
+	 * Constructor
+	 *
+	 *  @param ListenPort The port on which to listen for incoming connections
+	 */
+	FHttpListener(uint32 ListenPort);
+
+	/**
+	 * Destructor
+	 */
+	~FHttpListener();
+
+	/**
+	 * Starts listening for and accepting incoming connections
+	 *
+	 * @return true if the listener was able to start listening, false otherwise
+	 */
+	bool StartListening();
+
+	/**
+	 * Stops listening for and accepting incoming connections
+	 */
+	void StopListening();
+
+	/**
+	 * Tick the listener to otherwise commence connection flows
+	 *
+	 * @param DeltaTime The elapsed time since the last tick
+	 */
+	void Tick(float DeltaTime);
+
+	/**
+	 * Determines whether this listener has pending connections in-flight
+	 *
+	 * @return true if there are pending connections, false otherwise
+	 */
+	bool HasPendingConnections() const;
+
+	/**
+	 * Determines whether the listener has been initialized
+	 * @return true if this listener is bound and listening, false otherwise
+	 */
+	FORCEINLINE bool IsListening() const
+	{
+		return bIsListening;
+	}
+
+	/**
+	 * Gets the respective router
+	 * @return The respective router
+	 */
+	FORCEINLINE TSharedPtr<IHttpRouter> GetRouter() const
+	{
+		return Router;
+	}
+
+private:
+
+	/**
+	 * Accepts a single available connection
+	 * 
+	 * @param MaxConnectionsToAccept The maximum number of connections to accept
+	 */
+	void AcceptConnections(uint32 MaxConnectionsToAccept);
+
+	/**
+	 * Ticks connections in reading/writing phases
+	 * 
+	 * @param DeltaTime The elapsed time since the last tick
+	 */
+	void TickConnections(float DeltaTime);
+
+	/**
+	 * Removes connections that have been destroyed
+	 */
+	void RemoveDestroyedConnections();
+
+private:
+	
+	/** Whether this listeners has begun listening */
+	bool bIsListening = false;
+
+	/** The port on which the binding socket listens */
+	uint32 ListenPort = 0;
+
+	/** The binding socket which accepts incoming connections */
+	FSocket* ListenSocket = nullptr;
+
+	/** The mechanism that routes requests to respective handlers  */
+	TSharedPtr<FHttpRouter> Router = nullptr;
+
+	/** The collection of unique connections */
+	FHttpConnectionPool Connections;
+
+	/** Maximum number of connections to accept per frame */
+	static constexpr uint32 MaxConnectionsToAcceptPerFrame = 1;
+	/** Maximum number of pending connections to queue */
+	static constexpr uint32 ListenerConnectionBacklogSize = 16;
+	/** Maximum send buffer size */
+	static constexpr uint32 ListenerBufferSize = 512 * 1024;
+	/** The interface to bind to */
+	static constexpr uint32 ListenerInterfaceAddress = 0x0;
+};
\ No newline at end of file
diff --git a/Engine/Source/Runtime/Online/HTTPServer/Private/HttpPath.cpp b/Engine/Source/Runtime/Online/HTTPServer/Private/HttpPath.cpp
new file mode 100644
--- /dev/null
+++ b/Engine/Source/Runtime/Online/HTTPServer/Private/HttpPath.cpp
@@ -0,0 +1,89 @@
+// Copyright 1998-2019 Epic Games, Inc. All Rights Reserved.
+#include "HttpPath.h"
+
+FHttpPath::FHttpPath()
+{
+}
+
+FHttpPath::FHttpPath(FString InPath)
+	:Path(MoveTemp(InPath))
+{
+	NormalizePath();
+}
+
+const FString& FHttpPath::GetPath() const
+{
+	return Path;
+}
+
+uint32 FHttpPath::ParsePathTokens(TArray<FString>& OutPathTokens) const
+{
+	const TCHAR PathDelimiter[] = { TCHAR('/') };
+	return Path.ParseIntoArray(OutPathTokens, PathDelimiter, true);
+}
+
+void FHttpPath::SetPath(FString NewPath)
+{
+	Path = NewPath;
+	NormalizePath();
+}
+
+bool FHttpPath::IsValidPath() const
+{
+	if (IsRoot())
+	{
+		return false;
+	}
+	if (!Path.StartsWith(TEXT("/")))
+	{
+		return false;
+	}
+
+	auto IsInvalidUriChar = [](TCHAR C) 
+	{ 
+		return 
+			(C <= 32) || (C >= 127) ||
+			(C == TEXT(' ')) ||	(C == TEXT('.')) ||	(C == TEXT(',')) ||
+			(C == TEXT(';')) ||	(C == TEXT(':')) || (C == TEXT('<')) ||
+			(C == TEXT('>')) ||	(C == TEXT(']')) ||	(C == TEXT('[')) ||
+			(C == TEXT('}')) ||	(C == TEXT('{')) ||	(C == TEXT('#')) ||
+			(C == TEXT('|')) ||	(C == TEXT('^')) ||	(C == TEXT('~')) ||
+			(C == TEXT('\\'));
+	}; 
+
+	return (INDEX_NONE == Path.FindLastCharByPredicate(IsInvalidUriChar, 0));
+}
+
+bool FHttpPath::IsRoot() const
+{
+	return Path == TEXT("/");
+}
+
+void FHttpPath::MakeRelative(const FString& OtherPath)
+{
+	const bool bAllowShrinking = false;
+
+	if (OtherPath == TEXT("/"))
+	{
+		return;
+	}
+
+	if (Path == OtherPath)
+	{
+		Path.RemoveAt(1, Path.Len() - 1, bAllowShrinking);
+	}
+
+	if (Path.StartsWith(OtherPath))
+	{
+		Path.RemoveAt(0, OtherPath.Len(), bAllowShrinking);
+	}
+}
+
+void FHttpPath::NormalizePath()
+{
+	if (!IsRoot() &&  Path.EndsWith(TEXT("/")))
+	{
+		Path.RemoveFromEnd(TEXT("/"));
+	}
+}
+
diff --git a/Engine/Source/Runtime/Online/HTTPServer/Private/HttpRequestHandlerIterator.cpp b/Engine/Source/Runtime/Online/HTTPServer/Private/HttpRequestHandlerIterator.cpp
new file mode 100644
--- /dev/null
+++ b/Engine/Source/Runtime/Online/HTTPServer/Private/HttpRequestHandlerIterator.cpp
@@ -0,0 +1,81 @@
+// Copyright 1998-2019 Epic Games, Inc. All Rights Reserved.
+#include "HttpRequestHandlerIterator.h"
+#include "HttpServerRequest.h"
+#include "HttpRouteHandle.h"
+
+
+FHttpRequestHandlerIterator::FHttpRequestHandlerIterator(
+	const TSharedPtr<FHttpServerRequest>& InRequest, 
+	const FHttpRequestHandlerRegistrar& InRequestHandlerRegistrar)
+	: Request(InRequest)
+	, RequestHandlerRegistrar(InRequestHandlerRegistrar)
+	, HttpPathIterator(InRequest->RelativePath)
+{
+}
+
+const FHttpRequestHandler * const FHttpRequestHandlerIterator::Next()
+{
+	while (HttpPathIterator.HasNext())
+	{
+		// Determine if we have a matching handler for the next route
+		const auto& NextRoute = HttpPathIterator.Next();
+
+		// Filter by http route
+		const auto RouteHandlePtr = RequestHandlerRegistrar->Find(NextRoute);
+		if (!RouteHandlePtr)
+		{
+			// Not a matching route
+			continue;
+		}
+		const auto& RouteHandle = *RouteHandlePtr;
+
+		// Filter by http verb
+		const EHttpServerRequestVerbs VerbFilterResult = RouteHandle->Verbs & Request->Verb;
+		if (EHttpServerRequestVerbs::NONE == VerbFilterResult)
+		{
+			// Not a matching verb
+			continue;
+		}
+
+		// Make request path relative to the respective handler
+		Request->RelativePath.MakeRelative(NextRoute);
+
+		return &(RouteHandle->Handler);
+	}
+	return nullptr;
+}
+
+FHttpRequestHandlerIterator::FHttpPathIterator::FHttpPathIterator(const FHttpPath& HttpPath)
+{
+	NextPath = HttpPath.GetPath();
+}
+
+bool FHttpRequestHandlerIterator::FHttpPathIterator::HasNext() const
+{
+	return !bLastIteration;
+}
+
+const FString& FHttpRequestHandlerIterator::FHttpPathIterator::Next()
+{
+	// Callers  should always test HasNext() first!
+	check(!bLastIteration); 
+
+	if (!bFirstIteration)
+	{
+		int32 SlashIndex = 0;
+		if (NextPath.FindLastChar(TCHAR('/'), SlashIndex))
+		{
+			const bool bAllowShrinking = false;
+			NextPath.RemoveAt(SlashIndex, NextPath.Len() - SlashIndex, bAllowShrinking);
+
+			if (0 == NextPath.Len())
+			{
+				NextPath.AppendChar(TCHAR('/'));
+				bLastIteration = true;
+			}
+		}
+	}
+	bFirstIteration = false;
+	return NextPath;
+}
+
diff --git a/Engine/Source/Runtime/Online/HTTPServer/Private/HttpRequestHandlerIterator.h b/Engine/Source/Runtime/Online/HTTPServer/Private/HttpRequestHandlerIterator.h
new file mode 100644
--- /dev/null
+++ b/Engine/Source/Runtime/Online/HTTPServer/Private/HttpRequestHandlerIterator.h
@@ -0,0 +1,72 @@
+// Copyright 1998-2019 Epic Games, Inc. All Rights Reserved.
+#pragma once
+
+#include "CoreMinimal.h"
+#include "HttpRequestHandler.h"
+#include "HttpRequestHandlerRegistrar.h"
+
+struct FHttpPath;
+struct FHttpServerRequest;
+
+/**
+ * Utility to iterate qualified FHttpRequestHandlers
+ */
+struct FHttpRequestHandlerIterator final
+{
+
+	/**
+	 *  Utility (internal) to iterate FHttpPaths in-place
+	 */
+	struct FHttpPathIterator final
+	{
+	public:
+		/**
+		 * Constructor
+		 * @param HttpPath The iterator starting path
+		 */
+		FHttpPathIterator(const FHttpPath& HttpPath);
+
+		/**
+		 * Gets the next route 
+		 * @return The base path of the current path
+		 */
+		const FString& Next();
+
+		/**
+		 * Determines whether there is a Next() path to get
+		 * @return true if there is a next path, false otherwise
+		 */
+		FORCEINLINE bool HasNext() const;
+
+	private:
+		FString NextPath;
+		bool bFirstIteration = true;
+		bool bLastIteration = false;
+	};
+
+
+public:
+
+	/**
+	 * Constructor
+	 */
+	FHttpRequestHandlerIterator(const TSharedPtr<FHttpServerRequest>& InRequest, const FHttpRequestHandlerRegistrar& InRequestHandlerRegistrar);
+
+	/** 
+	* Determines the next registered request handler
+	* @return The next registered request handler if found, nullptr otherwise
+	*/
+	const FHttpRequestHandler* const Next();
+
+private:
+
+	/** The basis request */
+	const TSharedPtr<FHttpServerRequest> Request;
+
+	/** The associative route/handler registration  */
+	const FHttpRequestHandlerRegistrar RequestHandlerRegistrar;
+
+	/** Utility to iterate FHttpRoutes in-place */
+	FHttpPathIterator HttpPathIterator;
+
+};
diff --git a/Engine/Source/Runtime/Online/HTTPServer/Private/HttpRequestHandlerRegistrar.h b/Engine/Source/Runtime/Online/HTTPServer/Private/HttpRequestHandlerRegistrar.h
new file mode 100644
--- /dev/null
+++ b/Engine/Source/Runtime/Online/HTTPServer/Private/HttpRequestHandlerRegistrar.h
@@ -0,0 +1,11 @@
+// Copyright 1998-2019 Epic Games, Inc. All Rights Reserved.
+#pragma once
+
+#include "CoreMinimal.h"
+#include "HttpRouteHandle.h"
+
+/**
+ * FHttpRequestHandlerRegistrar
+ *  Represents the associative relationship between Http request paths and respective route handles
+ */
+typedef TSharedPtr<TMap<const FString, const FHttpRouteHandle>> FHttpRequestHandlerRegistrar;
diff --git a/Engine/Source/Runtime/Online/HTTPServer/Private/HttpRouter.cpp b/Engine/Source/Runtime/Online/HTTPServer/Private/HttpRouter.cpp
new file mode 100644
--- /dev/null
+++ b/Engine/Source/Runtime/Online/HTTPServer/Private/HttpRouter.cpp
@@ -0,0 +1,49 @@
+// Copyright 1998-2019 Epic Games, Inc. All Rights Reserved.
+
+#include "HttpRouter.h"
+#include "HttpRouteHandle.h"
+#include "HttpPath.h"
+#include "HttpConnection.h"
+#include "HttpServerRequest.h"
+#include "HttpServerResponse.h"
+#include "HttpRequestHandler.h"
+#include "HttpRequestHandlerIterator.h"
+
+FHttpRouter::FHttpRouter()
+{
+	RequestHandlerRegistrar = MakeShared<TMap<const FString, const FHttpRouteHandle>>();
+}
+
+FHttpRouteHandle FHttpRouter::BindRoute(const FHttpPath& HttpPath,  const EHttpServerRequestVerbs& HttpVerbs,  const FHttpRequestHandler& Handler)
+{
+	check(HttpPath.IsValidPath());
+	check(EHttpServerRequestVerbs::NONE != HttpVerbs);
+
+	if (RequestHandlerRegistrar->Contains(HttpPath.GetPath()))
+	{
+		return nullptr;
+	}
+
+	auto RouteHandle = MakeShared<FHttpRouteHandleInternal>(HttpPath.GetPath(), HttpVerbs, Handler);
+	RequestHandlerRegistrar->Add(HttpPath.GetPath(), RouteHandle);
+
+	return RouteHandle;
+}
+
+void FHttpRouter::UnbindRoute(const FHttpRouteHandle& RouteHandle)
+{
+	auto ExistingRouteHandle = RequestHandlerRegistrar->Find(RouteHandle->Path);
+
+	if (ExistingRouteHandle)
+	{
+		// Ensure caller is unbinding a route they actually own
+		check(*ExistingRouteHandle == RouteHandle);
+		RequestHandlerRegistrar->Remove(RouteHandle->Path);
+	}
+}
+
+FHttpRequestHandlerIterator FHttpRouter::CreateRequestHandlerIterator(const TSharedPtr<FHttpServerRequest>& Request) const
+{
+	FHttpRequestHandlerIterator Iterator(Request, RequestHandlerRegistrar);
+	return Iterator;
+}
diff --git a/Engine/Source/Runtime/Online/HTTPServer/Private/HttpRouter.h b/Engine/Source/Runtime/Online/HTTPServer/Private/HttpRouter.h
new file mode 100644
--- /dev/null
+++ b/Engine/Source/Runtime/Online/HTTPServer/Private/HttpRouter.h
@@ -0,0 +1,38 @@
+// Copyright 1998-2019 Epic Games, Inc. All Rights Reserved.
+#pragma once
+
+#include "IHttpRouter.h"
+#include "HttpRequestHandler.h"
+#include "HttpRouteHandle.h"
+#include "HttpRequestHandlerRegistrar.h"
+
+struct FHttpPath;
+struct FHttpRequestHandlerIterator;
+
+class FHttpRouter final : public IHttpRouter
+{
+public:
+
+	/**
+	 * Constructor
+	 */
+	FHttpRouter();
+
+	// IHttpRouter Overrides
+	FHttpRouteHandle BindRoute(const FHttpPath& HttpPath, const EHttpServerRequestVerbs& RequestVerbs, const FHttpRequestHandler& Handler) override;
+	void UnbindRoute(const FHttpRouteHandle& RouteHandle) override;
+
+	/**
+	 * Creates a request handler iterator to facilitate http routing
+	 * @param Request The basis request for which matching request handlers are found
+	 * @return The instantiated FHttpRequestHandlerIterator
+	 */
+	FHttpRequestHandlerIterator CreateRequestHandlerIterator(const TSharedPtr<FHttpServerRequest>& Request) const;
+
+private:
+
+	/** The associative pairing of Http routes to respective request handlers */
+	FHttpRequestHandlerRegistrar RequestHandlerRegistrar;
+};
+
+
diff --git a/Engine/Source/Runtime/Online/HTTPServer/Private/HttpServerConstantsPrivate.h b/Engine/Source/Runtime/Online/HTTPServer/Private/HttpServerConstantsPrivate.h
new file mode 100644
--- /dev/null
+++ b/Engine/Source/Runtime/Online/HTTPServer/Private/HttpServerConstantsPrivate.h
@@ -0,0 +1,34 @@
+// Copyright 1998-2019 Epic Games, Inc. All Rights Reserved.
+#pragma once
+
+#include "CoreMinimal.h"
+
+struct FHttpServerHeaderKeys
+{
+	static constexpr TCHAR* CONTENT_TYPE = TEXT("content-type");
+	static constexpr TCHAR* CONTENT_LENGTH = TEXT("content-length");
+	static constexpr TCHAR* CONNECTION = TEXT("connection");
+	static constexpr TCHAR* KEEP_ALIVE = TEXT("keep-alive");
+};
+
+struct FHttpServerErrorStrings
+{
+	// Connection
+	static constexpr TCHAR* SocketClosedFailure = TEXT("errors.com.epicgames.httpserver.socket_closed_failure");
+	static constexpr TCHAR* SocketRecvFailure = TEXT("errors.com.epicgames.httpserver.socket_recv_failure");
+	static constexpr TCHAR* SocketSendFailure = TEXT("errors.com.epicgames.httpserver.socket_send_failure");
+
+	// Routing
+	static constexpr TCHAR* NotFound = TEXT("errors.com.epicgames.httpserver.route_handler_not_found");
+
+	// Serialization
+	static constexpr TCHAR* MalformedRequestSize = TEXT("errors.com.epicgames.httpserver.malformed_request_size");
+	static constexpr TCHAR* MalformedRequestHeaders = TEXT("errors.com.epicgames.httpserver.malformed_request_header");
+	static constexpr TCHAR* MissingRequestHeaders = TEXT("errors.com.epicgames.httpserver.missing_request_headers");
+	static constexpr TCHAR* MalformedRequestBody = TEXT("errors.com.epicgames.httpserver.malformed_request_body");
+	static constexpr TCHAR* UnknownRequestVerb = TEXT("errors.com.epicgames.httpserver.unknown_request_verb");
+	static constexpr TCHAR* InvalidContentLengthHeader = TEXT("errors.com.epicgames.httpserver.invalid_content_length_header");
+	static constexpr TCHAR* MissingContentLengthHeader = TEXT("errors.com.epicgames.httpserver.missing_content_length_header");
+	static constexpr TCHAR* MismatchedContentLengthBodyTooLarge =  TEXT("errors.com.epicgames.httpserver.mismatched_content_length_body_too_large");
+};
+
diff --git a/Engine/Source/Runtime/Online/HTTPServer/Private/HttpServerModule.cpp b/Engine/Source/Runtime/Online/HTTPServer/Private/HttpServerModule.cpp
new file mode 100644
--- /dev/null
+++ b/Engine/Source/Runtime/Online/HTTPServer/Private/HttpServerModule.cpp
@@ -0,0 +1,104 @@
+// Copyright 1998-2019 Epic Games, Inc. All Rights Reserved.
+
+#include "HttpServerModule.h"
+#include "HttpListener.h"
+#include "Modules/ModuleManager.h"
+
+DEFINE_LOG_CATEGORY(LogHttpServerModule); 
+
+// FHttpServerModule 
+IMPLEMENT_MODULE(FHttpServerModule, HTTPServer);
+
+FHttpServerModule* FHttpServerModule::Singleton = nullptr;
+
+void FHttpServerModule::StartupModule()
+{
+	Singleton = this;
+	bInitialized = true;
+}
+
+void FHttpServerModule::ShutdownModule()
+{
+	bInitialized = false;
+
+	// stop all listeners
+	StopAllListeners();
+
+	// destroy all listeners
+	Listeners.Empty();
+}
+
+void FHttpServerModule::StopAllListeners()
+{
+	for (const auto& Listener : Listeners)
+	{
+		if (Listener.Value->IsListening())
+		{
+			UE_LOG(LogHttpServerModule, Log,
+				TEXT("Stopping listener on Port: %u"), Listener.Key);
+
+			Listener.Value->StopListening();
+		}
+	}
+}
+
+bool FHttpServerModule::HasPendingListeners() const 
+{
+	for (const auto& Listener : Listeners)
+	{
+		if (Listener.Value->HasPendingConnections())
+		{
+			return true;
+		}
+	}
+	return false;
+}
+
+FHttpServerModule& FHttpServerModule::Get()
+{
+	if (nullptr == Singleton)
+	{
+		check(IsInGameThread());
+		FModuleManager::LoadModuleChecked<FHttpServerModule>("HTTPServer");
+	}
+	check(Singleton);
+	return *Singleton;
+}
+
+TSharedPtr<IHttpRouter> FHttpServerModule::GetHttpRouter(uint32 Port)
+{
+	if (!bInitialized)
+	{
+		return nullptr;
+	}
+
+	// We may already be listening on this port
+	TUniquePtr<FHttpListener>* ExistingListener = Listeners.Find(Port);
+	if (ExistingListener)
+	{
+		return ExistingListener->Get()->GetRouter();
+	}
+	else
+	{
+		// Otherwise create a new one
+		TUniquePtr<FHttpListener> NewListener = MakeUnique<FHttpListener>(Port);
+		if (NewListener->StartListening())
+		{
+			const auto& NewListenerRef = Listeners.Add(Port, MoveTemp(NewListener));
+			return NewListenerRef->GetRouter();
+		}
+	}
+	return nullptr;
+}
+
+bool FHttpServerModule::Tick(float DeltaTime)
+{
+	check(Singleton == this);
+
+	for (const auto& Listener : Listeners)
+	{
+		check(Listener.Value);
+		Listener.Value->Tick(DeltaTime);
+	}
+	return true;
+}
diff --git a/Engine/Source/Runtime/Online/HTTPServer/Private/HttpServerResponse.cpp b/Engine/Source/Runtime/Online/HTTPServer/Private/HttpServerResponse.cpp
new file mode 100644
--- /dev/null
+++ b/Engine/Source/Runtime/Online/HTTPServer/Private/HttpServerResponse.cpp
@@ -0,0 +1,55 @@
+
+#include "HttpServerResponse.h"
+#include "HttpServerConstants.h"
+
+TUniquePtr<FHttpServerResponse> FHttpServerResponse::Create(const FString& Text, FString ContentType)
+{
+	TUniquePtr<FHttpServerResponse> Response = MakeUnique<FHttpServerResponse>();
+	Response->Code = EHttpServerResponseCodes::Ok;
+
+	FTCHARToUTF8 ConvertToUtf8(*Text);
+	const uint8* ConvertToUtf8Bytes = (reinterpret_cast<const uint8*>(ConvertToUtf8.Get()));
+	Response->Body.Append(ConvertToUtf8Bytes, ConvertToUtf8.Length());
+
+	FString Utf8CharsetContentType = FString::Printf(TEXT("%s ;charset=utf-8"), *ContentType);
+	TArray<FString> ContentTypeValue = { Utf8CharsetContentType };
+	Response->Headers.Add(FHttpServerHeaderKeys::CONTENT_TYPE, ContentTypeValue);
+
+	return Response;
+}
+
+TUniquePtr<FHttpServerResponse> FHttpServerResponse::Create(TArray<uint8>&& RawBytes, FString ContentType)
+{
+	TUniquePtr<FHttpServerResponse> Response = MakeUnique<FHttpServerResponse>(MoveTemp(RawBytes));
+	Response->Code = EHttpServerResponseCodes::Ok;
+
+	TArray<FString> ContentTypeValue = { MoveTemp(ContentType) };
+	Response->Headers.Add(FHttpServerHeaderKeys::CONTENT_TYPE, MoveTemp(ContentTypeValue));
+	return Response;
+}
+
+TUniquePtr<FHttpServerResponse> FHttpServerResponse::Create(const TArrayView<uint8>& RawBytes, FString ContentType)
+{
+	TUniquePtr<FHttpServerResponse> Response = MakeUnique<FHttpServerResponse>();
+	Response->Code = EHttpServerResponseCodes::Ok;
+	Response->Body.Append(RawBytes.GetData(), RawBytes.Num());
+
+	TArray<FString> ContentTypeValue = { MoveTemp(ContentType) };
+	Response->Headers.Add(FHttpServerHeaderKeys::CONTENT_TYPE, MoveTemp(ContentTypeValue));
+	return Response;
+}
+
+TUniquePtr<FHttpServerResponse> FHttpServerResponse::Create(int32 HttpResponseCode, const FString& ErrorCode)
+{
+	FString ResponseBody = FString::Printf(TEXT("{\"errorCode\": \"%s\"}"), *ErrorCode);
+	auto Response = Create(ResponseBody, TEXT("application/json"));
+	Response->Code = HttpResponseCode;
+	return Response;
+}
+
+TUniquePtr<FHttpServerResponse> FHttpServerResponse::Ok()
+{
+	auto Response = MakeUnique<FHttpServerResponse>();
+	Response->Code = EHttpServerResponseCodes::Ok;
+	return Response;
+}
\ No newline at end of file
diff --git a/Engine/Source/Runtime/Online/HTTPServer/Private/IHttpRouter.cpp b/Engine/Source/Runtime/Online/HTTPServer/Private/IHttpRouter.cpp
new file mode 100644
--- /dev/null
+++ b/Engine/Source/Runtime/Online/HTTPServer/Private/IHttpRouter.cpp
@@ -0,0 +1,8 @@
+// Copyright 1998-2019 Epic Games, Inc. All Rights Reserved.
+
+#include "IHttpRouter.h"
+
+IHttpRouter::~IHttpRouter()
+{
+}
+
diff --git a/Engine/Source/Runtime/Online/HTTPServer/Private/Tests/HttpServerTests.cpp b/Engine/Source/Runtime/Online/HTTPServer/Private/Tests/HttpServerTests.cpp
new file mode 100644
--- /dev/null
+++ b/Engine/Source/Runtime/Online/HTTPServer/Private/Tests/HttpServerTests.cpp
@@ -0,0 +1,48 @@
+// Copyright 1998-2019 Epic Games, Inc. All Rights Reserved.
+
+#include "CoreMinimal.h"
+#include "Misc/AutomationTest.h"
+
+#include "HttpServerModule.h"
+#include "IHttpRouter.h"
+#include "HttpRouteHandle.h"
+#include "HttpRequestHandler.h"
+#include "HttpServerRequest.h"
+#include "HttpServerResponse.h"
+
+
+IMPLEMENT_SIMPLE_AUTOMATION_TEST(FHttpServerIntegrationTest, "System.Online.HttpServer.Integration", EAutomationTestFlags::EditorContext | EAutomationTestFlags::EngineFilter)
+bool FHttpServerIntegrationTest::RunTest(const FString& Parameters)
+{
+	const uint32 HttpRouterPort = 8888;
+	const FHttpPath HttpPath(TEXT("/TestHttpServer"));
+
+	// Ensure router creation
+	TSharedPtr<IHttpRouter> HttpRouter = FHttpServerModule::Get().GetHttpRouter(HttpRouterPort);
+	TestTrue(TEXT("HttpRouter.IsValid()"), HttpRouter.IsValid());
+
+	// Ensure unique routers per-port
+	TSharedPtr<IHttpRouter> DuplicateHttpRouter = FHttpServerModule::Get().GetHttpRouter(HttpRouterPort);
+	TestEqual(TEXT("HttpRouter Duplicates"), HttpRouter, DuplicateHttpRouter);
+
+	// Ensure we can create route bindings
+	const FHttpRequestHandler RequestHandler = [this]
+	(const FHttpServerRequest& Request, const FHttpResultCallback& OnComplete)
+	{
+		return true;
+	};
+	FHttpRouteHandle HttpRouteHandle = HttpRouter->BindRoute(HttpPath, EHttpServerRequestVerbs::GET, RequestHandler);
+	TestTrue(TEXT("HttpRouteHandle.IsValid()"), HttpRouteHandle.IsValid());
+
+	// Disallow duplicate route bindings
+	FHttpRouteHandle DuplicateHandle = HttpRouter->BindRoute(HttpPath, EHttpServerRequestVerbs::GET, RequestHandler);
+	TestFalse(TEXT("HttpRouteHandle Duplicated"), DuplicateHandle.IsValid());
+
+	// Make a request
+	/*
+	ADD_LATENT_AUTOMATION_COMMAND(FExecStringLatentCommand(TEXT("HTTP TEST 1 http://localhost:8888/TestHttpServer")));
+	*/
+
+	HttpRouter->UnbindRoute(HttpRouteHandle);
+	return true;
+}
\ No newline at end of file
diff --git a/Engine/Source/Runtime/Online/HTTPServer/Public/HttpPath.h b/Engine/Source/Runtime/Online/HTTPServer/Public/HttpPath.h
new file mode 100644
--- /dev/null
+++ b/Engine/Source/Runtime/Online/HTTPServer/Public/HttpPath.h
@@ -0,0 +1,85 @@
+// Copyright 1998-2019 Epic Games, Inc. All Rights Reserved.
+#pragma once
+
+#include "CoreMinimal.h"
+#include "HttpRequestHandler.h"
+
+struct FHttpPath
+{
+
+public:
+
+	/**
+	 * Constructor
+	 */
+	HTTPSERVER_API FHttpPath();
+
+	/**
+	 * Constructor
+	 *
+	 * @param Path - The http path
+	 */
+	HTTPSERVER_API FHttpPath(FString InPath);
+
+	/**
+     * Gets the respective http path as a string
+     */
+	HTTPSERVER_API const FString& GetPath() const;
+
+	/**
+	 * Parses the respective http path into a caller-supplied array
+	 * 
+	 * @param OutPathTokens The caller-allocated tokens array
+	 * @return The number of tokens parsed
+	 */
+	HTTPSERVER_API uint32 ParsePathTokens(TArray<FString>& OutPathTokens) const;
+
+	/**
+	 * Sets the server-relative http path
+	 *
+	 * @param Path the new path to assign
+	 */
+	HTTPSERVER_API void SetPath(FString Path);
+
+	/**
+	 * Determines if the path is valid
+	 */
+	HTTPSERVER_API bool IsValidPath() const;
+
+	/**
+	 * Determines if the path is /
+	 */
+	HTTPSERVER_API bool IsRoot() const;
+
+	/**
+	 * Re-path this path sans the path of another
+	 * MakeRelative(/a/b/c/d, /a/b) => /c/d  
+	 *
+	 * @param OtherPath The path to make relative against
+	 */
+	HTTPSERVER_API void MakeRelative(const FString& OtherPath);
+
+	// TMap<> Comparer
+	HTTPSERVER_API bool operator==(const FHttpPath& Other) const
+	{
+		bool bEqual = (0 == Path.Compare(Other.Path, ESearchCase::CaseSensitive));
+		return bEqual;
+	}
+	static friend uint32 GetTypeHash(const FHttpPath& Path)
+	{
+		return GetTypeHash(Path.GetPath());
+	}
+
+private:
+
+	/**
+	 * Normalizes the internally stored Path
+	 */
+	void NormalizePath();
+
+	/** The respective Http path */
+	FString Path = TEXT("/");
+
+};
+
+
diff --git a/Engine/Source/Runtime/Online/HTTPServer/Public/HttpRequestHandler.h b/Engine/Source/Runtime/Online/HTTPServer/Public/HttpRequestHandler.h
new file mode 100644
--- /dev/null
+++ b/Engine/Source/Runtime/Online/HTTPServer/Public/HttpRequestHandler.h
@@ -0,0 +1,20 @@
+// Copyright 1998-2019 Epic Games, Inc. All Rights Reserved.
+#pragma once
+
+#include "CoreMinimal.h"
+#include "HttpResultCallback.h"
+
+struct FHttpServerRequest;
+
+/**
+ * FHttpRequestHandler 
+ *
+ *  NOTE - Returning true implies that the delegate will eventually invoke OnComplete
+ *  NOTE - Returning false implies that the delegate will never invoke OnComplete
+ * 
+ * @param Request The incoming http request to be handled
+ * @param OnComplete The callback to invoke to write an http response
+ * @return True if the request has been handled, false otherwise
+ */
+typedef TFunction<bool(const FHttpServerRequest& Request, const FHttpResultCallback& OnComplete)> FHttpRequestHandler;
+
diff --git a/Engine/Source/Runtime/Online/HTTPServer/Public/HttpResultCallback.h b/Engine/Source/Runtime/Online/HTTPServer/Public/HttpResultCallback.h
new file mode 100644
--- /dev/null
+++ b/Engine/Source/Runtime/Online/HTTPServer/Public/HttpResultCallback.h
@@ -0,0 +1,15 @@
+// Copyright 1998-2019 Epic Games, Inc. All Rights Reserved.
+#pragma once
+
+#include "CoreMinimal.h"
+struct FHttpServerResponse;
+
+/**
+* FHttpResultCallback
+* This callback is intended to be invoked exclusively by FHttpRequestHandler delegates
+* 
+* @param Response The response to write
+*/
+typedef TFunction<void(TUniquePtr<FHttpServerResponse>&& Response)> FHttpResultCallback;
+
+
diff --git a/Engine/Source/Runtime/Online/HTTPServer/Public/HttpRouteHandle.h b/Engine/Source/Runtime/Online/HTTPServer/Public/HttpRouteHandle.h
new file mode 100644
--- /dev/null
+++ b/Engine/Source/Runtime/Online/HTTPServer/Public/HttpRouteHandle.h
@@ -0,0 +1,27 @@
+// Copyright 1998-2019 Epic Games, Inc. All Rights Reserved.
+#pragma once
+#include "CoreMinimal.h"
+#include "HttpServerRequest.h"
+
+struct FHttpRouteHandleInternal
+{
+	FHttpRouteHandleInternal(const FString& InPath, const EHttpServerRequestVerbs& InVerbs, const FHttpRequestHandler& InHandler)
+		: Path(InPath)
+		, Verbs(InVerbs)
+		, Handler(InHandler)
+	{ };
+
+	/** The respective bound http path */
+	const FString Path;
+	/** The qualifying http verbs for which this route is valid */
+	const EHttpServerRequestVerbs Verbs;
+	/** The respective handler to be invoked */
+	const FHttpRequestHandler Handler;
+};
+
+/**
+ * FHttpRouteHandle 
+ *  Returned by IHttpRouter BindRoute() invocations. 
+ *  Clients may test this value to validate and relinquish http router bindings
+*/
+typedef TSharedPtr<const FHttpRouteHandleInternal> FHttpRouteHandle;
diff --git a/Engine/Source/Runtime/Online/HTTPServer/Public/HttpServerConstants.h b/Engine/Source/Runtime/Online/HTTPServer/Public/HttpServerConstants.h
new file mode 100644
--- /dev/null
+++ b/Engine/Source/Runtime/Online/HTTPServer/Public/HttpServerConstants.h
@@ -0,0 +1,96 @@
+// Copyright 1998-2019 Epic Games, Inc. All Rights Reserved.
+#pragma once
+#include "CoreMinimal.h"
+
+namespace EHttpServerResponseCodes
+{
+	/**
+	 * Response codes that can come back from an Http request
+	 */
+	enum Type
+	{
+		// status code not set yet
+		Unknown = 0,
+		// the request can be continued.
+		Continue = 100,
+		// the server has switched protocols in an upgrade header.
+		SwitchProtocol = 101,
+		// the request completed successfully.
+		Ok = 200,
+		// the request has been fulfilled and resulted in the creation of a new resource.
+		Created = 201,
+		// the request has been accepted for processing, but the processing has not been completed.
+		Accepted = 202,
+		// the returned meta information in the entity-header is not the definitive set available from the origin server.
+		Partial = 203,
+		// the server has fulfilled the request, but there is no new information to send back.
+		NoContent = 204,
+		// the request has been completed, and the client program should reset the document view that caused the request to be sent to allow the user to easily initiate another input action.
+		ResetContent = 205,
+		// the server has fulfilled the partial get request for the resource.
+		PartialContent = 206,
+		// the server couldn't decide what to return.
+		Ambiguous = 300,
+		// the requested resource has been assigned to a new permanent uri (uniform resource identifier), and any future references to this resource should be done using one of the returned uris.
+		Moved = 301,
+		// the requested resource resides temporarily under a different uri (uniform resource identifier).
+		Redirect = 302,
+		// the response to the request can be found under a different uri (uniform resource identifier) and should be retrieved using a get http verb on that resource.
+		RedirectMethod = 303,
+		// the requested resource has not been modified.
+		NotModified = 304,
+		// the requested resource must be accessed through the proxy given by the location field.
+		UseProxy = 305,
+		// the redirected request keeps the same http verb. http/1.1 behavior.
+		RedirectKeepVerb = 307,
+		// the request could not be processed by the server due to invalid syntax.
+		BadRequest = 400,
+		// the requested resource requires user authentication.
+		Denied = 401,
+		// not currently implemented in the http protocol.
+		PaymentReq = 402,
+		// the server understood the request, but is refusing to fulfill it.
+		Forbidden = 403,
+		// the server has not found anything matching the requested uri (uniform resource identifier).
+		NotFound = 404,
+		// the http verb used is not allowed.
+		BadMethod = 405,
+		// no responses acceptable to the client were found.
+		NoneAcceptable = 406,
+		// proxy authentication required.
+		ProxyAuthReq = 407,
+		// the server timed out waiting for the request.
+		RequestTimeout = 408,
+		// the request could not be completed due to a conflict with the current state of the resource. the user should resubmit with more information.
+		Conflict = 409,
+		// the requested resource is no longer available at the server, and no forwarding address is known.
+		Gone = 410,
+		// the server refuses to accept the request without a defined content length.
+		LengthRequired = 411,
+		// the precondition given in one or more of the request header fields evaluated to false when it was tested on the server.
+		PrecondFailed = 412,
+		// the server is refusing to process a request because the request entity is larger than the server is willing or able to process.
+		RequestTooLarge = 413,
+		// the server is refusing to service the request because the request uri (uniform resource identifier) is longer than the server is willing to interpret.
+		UriTooLong = 414,
+		// the server is refusing to service the request because the entity of the request is in a format not supported by the requested resource for the requested method.
+		UnsupportedMedia = 415,
+		// too many requests, the server is throttling
+		TooManyRequests = 429,
+		// the request should be retried after doing the appropriate action.
+		RetryWith = 449,
+		// the server encountered an unexpected condition that prevented it from fulfilling the request.
+		ServerError = 500,
+		// the server does not support the functionality required to fulfill the request.
+		NotSupported = 501,
+		// the server, while acting as a gateway or proxy, received an invalid response from the upstream server it accessed in attempting to fulfill the request.
+		BadGateway = 502,
+		// the service is temporarily overloaded.
+		ServiceUnavail = 503,
+		// the request was timed out waiting for a gateway.
+		GatewayTimeout = 504,
+		// the server does not support, or refuses to support, the http protocol version that was used in the request message.
+		VersionNotSup = 505
+	};
+}
+
diff --git a/Engine/Source/Runtime/Online/HTTPServer/Public/HttpServerModule.h b/Engine/Source/Runtime/Online/HTTPServer/Public/HttpServerModule.h
new file mode 100644
--- /dev/null
+++ b/Engine/Source/Runtime/Online/HTTPServer/Public/HttpServerModule.h
@@ -0,0 +1,89 @@
+// Copyright 1998-2019 Epic Games, Inc. All Rights Reserved.
+#pragma once
+
+#include "CoreMinimal.h"
+#include "Misc/CoreMisc.h"
+#include "Modules/ModuleInterface.h"
+#include "Modules/ModuleManager.h"
+#include "Containers/Ticker.h" 
+
+
+class FHttpListener;
+class IHttpRouter;
+
+DECLARE_LOG_CATEGORY_EXTERN(LogHttpServerModule, Log, All);
+
+/**
+ * Module for HtttpServer Implementation
+ */
+class FHttpServerModule : 
+	public IModuleInterface
+	,public FTickerObjectBase
+{
+
+public:
+
+	/**
+	 * Singleton-like access to this module's interface.  This is just for convenience!
+	 * Beware of calling this during the shutdown phase, though.  Your module might have been unloaded already.
+	 *
+	 * @return Returns singleton instance, loading the module on demand if needed
+	 */
+	HTTPSERVER_API static FHttpServerModule& Get();
+
+	/**
+	 * Per-port-binding access to an http router
+	 *
+	 * @param  Port The listener's bound port 
+	 * @return An IHttpRouter instance that can be leveraged to respond to HTTP requests
+	 */
+	HTTPSERVER_API TSharedPtr<IHttpRouter> GetHttpRouter(uint32 Port);
+
+	/**
+	 * FTicker callback
+	 * 
+	 * @param DeltaTime  The time in seconds since the last tick
+	 * @return           false if no longer needs ticking, true otherwise
+	 */
+	 bool Tick(float DeltaTime) override;
+
+	 /**
+	  * Stops all listeners
+	  */
+	 HTTPSERVER_API void StopAllListeners();
+
+	 /**
+	  * Determines if any listeners are pending operations
+	  *
+	  * @return true if there are pending listeners, false otherwise
+	  */
+	 HTTPSERVER_API bool HasPendingListeners() const;
+
+private:
+
+	// IModuleInterface
+
+	/**
+	 * Called when voice module is loaded
+	 * Initialize platform specific parts of template handling
+	 */
+	virtual void StartupModule() override;
+	
+	/**
+	 * Called when voice module is unloaded
+	 * Shutdown platform specific parts of template handling
+	 */
+	virtual void ShutdownModule() override;
+
+private:
+
+	/** Singleton Instance */
+	static FHttpServerModule* Singleton;
+
+	/** The association of port bindings and respective HTTP listeners */
+	TMap<uint32, TUniquePtr<FHttpListener>> Listeners;
+
+	/** Whether this module has been initialized */
+	bool bInitialized = false;
+};
+
diff --git a/Engine/Source/Runtime/Online/HTTPServer/Public/HttpServerRequest.h b/Engine/Source/Runtime/Online/HTTPServer/Public/HttpServerRequest.h
new file mode 100644
--- /dev/null
+++ b/Engine/Source/Runtime/Online/HTTPServer/Public/HttpServerRequest.h
@@ -0,0 +1,43 @@
+// Copyright 1998-2019 Epic Games, Inc. All Rights Reserved.
+#pragma once
+
+#include "CoreMinimal.h"
+#include "HttpPath.h"
+
+enum class EHttpServerRequestVerbs : uint16
+{
+	NONE   = 0,
+	GET    = 1 << 0,
+	POST   = 1 << 1,
+	PUT    = 1 << 2,
+	PATCH  = 1 << 3,
+	DELETE = 1 << 4,
+	ALL = GET | POST | PUT | PATCH | DELETE,
+};
+
+ENUM_CLASS_FLAGS(EHttpServerRequestVerbs);
+
+struct FHttpServerRequest
+{
+public:
+
+	/** Constructor */
+	FHttpServerRequest() { };
+
+	/** The handler-route-relative HTTP path */
+	FHttpPath RelativePath;
+
+	/** The HTTP-compliant verb  */
+	EHttpServerRequestVerbs Verb;
+
+	/** The HTTP headers */
+	TMap<FString, TArray<FString>> Headers;
+
+	/** The query parameters */
+	TMap<FString, FString> QueryParams;
+
+	/** The raw body contents */
+	TArray<uint8> Body;
+
+};
+
diff --git a/Engine/Source/Runtime/Online/HTTPServer/Public/HttpServerResponse.h b/Engine/Source/Runtime/Online/HTTPServer/Public/HttpServerResponse.h
new file mode 100644
--- /dev/null
+++ b/Engine/Source/Runtime/Online/HTTPServer/Public/HttpServerResponse.h
@@ -0,0 +1,78 @@
+// Copyright 1998-2019 Epic Games, Inc. All Rights Reserved.
+#pragma once
+#include "CoreMinimal.h"
+
+struct FHttpServerResponse final
+{
+public:
+	/**
+	 * Constructor
+	 */
+	FHttpServerResponse() 
+	{ }
+
+	/**
+	 * Constructor
+	 * Facilitates in-place body construction
+     *
+	 * @param InBody The r-value body data
+	 */
+	FHttpServerResponse(TArray<uint8>&& InBody)
+		: Body(MoveTemp(InBody))
+	{ }
+
+	/** Http Response Code */
+	int32 Code;
+
+	/** Http Headers */
+	TMap<FString, TArray<FString>> Headers;
+
+	/** Http Body Content */
+	TArray<uint8> Body;
+
+public:
+
+	/**
+	 * Creates an FHttpServerResponse from a string
+	 * 
+	 * @param  Text         The text to serialize
+	 * @param  ContentType  The HTTP response content type
+	 * @return              A unique pointer to an initialized response object
+	 */
+	HTTPSERVER_API static TUniquePtr<FHttpServerResponse> Create(const FString& Text, FString ContentType);
+
+	/**
+	 * Creates an FHttpServerResponse from a raw byte buffer
+	 *
+	 * @param  RawBytes     The byte buffer to serialize
+	 * @param  ContentType  The HTTP response content type
+	 * @return              A unique pointer to an initialized response object
+	 */
+	HTTPSERVER_API static TUniquePtr<FHttpServerResponse> Create(TArray<uint8>&& RawBytes, FString ContentType);
+
+	/**
+	 * Creates an FHttpServerResponse from a raw byte buffer
+	 *
+	 * @param  RawBytes     The byte buffer view to serialize
+	 * @param  ContentType  The HTTP response content type
+	 * @return              A unique pointer to an initialized response object
+	 */
+	HTTPSERVER_API static TUniquePtr<FHttpServerResponse> Create(const TArrayView<uint8>& RawBytes, FString ContentType);
+
+	/**
+	 * Creates an FHttpServerResponse from a response code
+	 *
+	 * @param  HttpResponseCode  The HTTP response/error code
+	 * @param  ErrorCode         The string error code/message
+	 * @return                   A unique pointer to an initialized response object
+	 */
+	HTTPSERVER_API static TUniquePtr<FHttpServerResponse> Create(int32 HttpResponseCode, const FString& ErrorCode);
+
+	/**
+	 * Creates an FHttpServerResponse 200 
+	 * @return A unique pointer to an initialized response object
+	 */
+	HTTPSERVER_API static TUniquePtr<FHttpServerResponse> Ok();
+};
+
+
diff --git a/Engine/Source/Runtime/Online/HTTPServer/Public/IHttpRouter.h b/Engine/Source/Runtime/Online/HTTPServer/Public/IHttpRouter.h
new file mode 100644
--- /dev/null
+++ b/Engine/Source/Runtime/Online/HTTPServer/Public/IHttpRouter.h
@@ -0,0 +1,37 @@
+// Copyright 1998-2019 Epic Games, Inc. All Rights Reserved.
+#pragma once
+
+#include "HttpRequestHandler.h"
+#include "HttpRouteHandle.h"
+
+struct FHttpPath;
+
+class IHttpRouter : public TSharedFromThis<IHttpRouter>
+{
+public:
+
+	/**
+	 * Binds the caller-supplied Uri to the caller-supplied handler
+	 *  @param  HttpPath   The respective http path to bind
+	 *  @param  HttpVerbs  The respective HTTP verbs to bind
+	 *  @param  Handler    The caller-defined closure to execute when the binding is invoked
+	 *  @return            An FHttpRouteHandle on success, nullptr otherwise. 
+	 */
+	virtual FHttpRouteHandle BindRoute(const FHttpPath& HttpPath, const EHttpServerRequestVerbs& HttpVerbs, const FHttpRequestHandler& Handler) = 0;
+
+	/**
+	 * Unbinds the caller-supplied Route 
+	 *
+	 *  @param  RouteHandle The handle to the route to unbind (callers must retain from BindRoute)
+	 */
+	virtual void UnbindRoute(const FHttpRouteHandle& RouteHandle) = 0;
+
+protected:
+
+	/**
+	 * Destructor
+	 */
+	virtual ~IHttpRouter();
+};
+
+
