[HttpServer] FHttpPath::ParsePathTokens() correctly parses http paths with "/" delimiter.

diff --git a/Engine/Source/Runtime/Online/HTTPServer/Private/HttpPath.cpp b/Engine/Source/Runtime/Online/HTTPServer/Private/HttpPath.cpp
--- a/Engine/Source/Runtime/Online/HTTPServer/Private/HttpPath.cpp
+++ b/Engine/Source/Runtime/Online/HTTPServer/Private/HttpPath.cpp
@@ -18,8 +18,7 @@ const FString& FHttpPath::GetPath() const
 
 uint32 FHttpPath::ParsePathTokens(TArray<FString>& OutPathTokens) const
 {
-	const TCHAR PathDelimiter[] = { TCHAR('/') };
-	return Path.ParseIntoArray(OutPathTokens, PathDelimiter, true);
+	return Path.ParseIntoArray(OutPathTokens, TEXT("/"), true);
 }
 
 void FHttpPath::SetPath(FString NewPath)
