[HttpServer] Non-unity build compiles and links succesfully
+ HttpConnectionContext: FORCEINLINE function definitions moved to header
+ HttpConnectionResponseWriteContext:  Now has its own log category.  It was previously piggy-backing off of LogHttpConnection, which would require inclusion of "HttpConnection.h" in non-unity builds, which would otherwise introduce a circular dependency.
+ HttpListener.h:  Now includes "HttpRouter.h" such that the non-unity compilation process can facilitate casting from an FHttpRouter to an IHttpRouter

diff --git a/Engine/Source/Runtime/Online/HTTPServer/Private/HttpConnectionContext.cpp b/Engine/Source/Runtime/Online/HTTPServer/Private/HttpConnectionContext.cpp
--- a/Engine/Source/Runtime/Online/HTTPServer/Private/HttpConnectionContext.cpp
+++ b/Engine/Source/Runtime/Online/HTTPServer/Private/HttpConnectionContext.cpp
@@ -17,16 +17,6 @@ void FHttpConnectionContext::AddElapsedIdleTime(float DeltaTime)
 	ElapsedIdleTime += DeltaTime;
 }
 
-float FHttpConnectionContext::GetElapsedIdleTime() const
-{
-	return ElapsedIdleTime;
-}
-
-const FString& FHttpConnectionContext::GetErrorStr() const
-{
-	return ErrorBuilder;
-}
-
 void FHttpConnectionContext::AddError(const TCHAR* ErrorCode)
 {
 	ErrorBuilder.Append(ErrorCode);
diff --git a/Engine/Source/Runtime/Online/HTTPServer/Private/HttpConnectionContext.h b/Engine/Source/Runtime/Online/HTTPServer/Private/HttpConnectionContext.h
--- a/Engine/Source/Runtime/Online/HTTPServer/Private/HttpConnectionContext.h
+++ b/Engine/Source/Runtime/Online/HTTPServer/Private/HttpConnectionContext.h
@@ -25,12 +25,18 @@ public:
 	/**
 	 * Gets the elapsed time (seconds) since last activity
 	 */
-	FORCEINLINE float GetElapsedIdleTime() const;
+	FORCEINLINE float GetElapsedIdleTime() const
+	{
+		return ElapsedIdleTime;
+	}
 
 	/**
 	 * Gets the cumulative error string
 	 */
-	FORCEINLINE const FString& GetErrorStr() const;
+	FORCEINLINE const FString& GetErrorStr() const
+	{
+		return ErrorBuilder;
+	}
 
 protected:
 
diff --git a/Engine/Source/Runtime/Online/HTTPServer/Private/HttpConnectionResponseWriteContext.cpp b/Engine/Source/Runtime/Online/HTTPServer/Private/HttpConnectionResponseWriteContext.cpp
--- a/Engine/Source/Runtime/Online/HTTPServer/Private/HttpConnectionResponseWriteContext.cpp
+++ b/Engine/Source/Runtime/Online/HTTPServer/Private/HttpConnectionResponseWriteContext.cpp
@@ -4,6 +4,8 @@
 #include "HttpServerConstantsPrivate.h"
 #include "Sockets.h"
 
+DEFINE_LOG_CATEGORY(LogHttpConnectionResponseWriteContext);
+
 FHttpConnectionResponseWriteContext::FHttpConnectionResponseWriteContext(FSocket* InSocket)
 	: Socket(InSocket)
 {
@@ -75,7 +77,7 @@ bool FHttpConnectionResponseWriteContext::WriteBytes(const uint8* Bytes, int32 B
 	bool bWriteSuccess = Socket->Send(Bytes, BytesLen, OutBytesWritten);
 	if (!bWriteSuccess)
 	{
-		UE_LOG(LogHttpConnection, Warning,
+		UE_LOG(LogHttpConnectionResponseWriteContext, Warning,
 			TEXT("WriteBytes sent %d/%d bytes"), OutBytesWritten, BytesLen);
 	}
 
diff --git a/Engine/Source/Runtime/Online/HTTPServer/Private/HttpConnectionResponseWriteContext.h b/Engine/Source/Runtime/Online/HTTPServer/Private/HttpConnectionResponseWriteContext.h
--- a/Engine/Source/Runtime/Online/HTTPServer/Private/HttpConnectionResponseWriteContext.h
+++ b/Engine/Source/Runtime/Online/HTTPServer/Private/HttpConnectionResponseWriteContext.h
@@ -6,6 +6,8 @@
 struct FHttpServerResponse;
 class FSocket;
 
+DECLARE_LOG_CATEGORY_EXTERN(LogHttpConnectionResponseWriteContext, Log, All);
+
 struct FHttpConnectionResponseWriteContext final : public FHttpConnectionContext
 {
 public:
diff --git a/Engine/Source/Runtime/Online/HTTPServer/Private/HttpListener.h b/Engine/Source/Runtime/Online/HTTPServer/Private/HttpListener.h
--- a/Engine/Source/Runtime/Online/HTTPServer/Private/HttpListener.h
+++ b/Engine/Source/Runtime/Online/HTTPServer/Private/HttpListener.h
@@ -2,14 +2,12 @@
 #pragma once
 
 #include "HttpConnectionTypes.h"
+#include "HttpRouter.h"
 
 struct FHttpConnection;
 struct FHttpPath;
 class FSocket;
 
-class IHttpRouter;
-class FHttpRouter;
-
 DECLARE_LOG_CATEGORY_EXTERN(LogHttpListener, Log, All);
 
 class FHttpListener final 
