[HttpServer] FHttpResponse exposes simplified/convenient  Ok() / Error() response builders
+ Additionally addressed inconsisten usage of Move semantics within the creation call hierarchies
+ Additionally addressed potential non-unity-build failure re: inclusion of HttpServerConstantsPrivate.h

diff --git a/Engine/Source/Runtime/Online/HTTPServer/Private/HttpServerResponse.cpp b/Engine/Source/Runtime/Online/HTTPServer/Private/HttpServerResponse.cpp
--- a/Engine/Source/Runtime/Online/HTTPServer/Private/HttpServerResponse.cpp
+++ b/Engine/Source/Runtime/Online/HTTPServer/Private/HttpServerResponse.cpp
@@ -1,6 +1,7 @@
 
 #include "HttpServerResponse.h"
 #include "HttpServerConstants.h"
+#include "HttpServerConstantsPrivate.h"
 
 TUniquePtr<FHttpServerResponse> FHttpServerResponse::Create(const FString& Text, FString ContentType)
 {
@@ -12,8 +13,8 @@ TUniquePtr<FHttpServerResponse> FHttpServerResponse::Create(const FString& Text,
 	Response->Body.Append(ConvertToUtf8Bytes, ConvertToUtf8.Length());
 
 	FString Utf8CharsetContentType = FString::Printf(TEXT("%s ;charset=utf-8"), *ContentType);
-	TArray<FString> ContentTypeValue = { Utf8CharsetContentType };
-	Response->Headers.Add(FHttpServerHeaderKeys::CONTENT_TYPE, ContentTypeValue);
+	TArray<FString> ContentTypeValue = { MoveTemp(Utf8CharsetContentType) };
+	Response->Headers.Add(FHttpServerHeaderKeys::CONTENT_TYPE, MoveTemp(ContentTypeValue));
 
 	return Response;
 }
@@ -41,15 +42,22 @@ TUniquePtr<FHttpServerResponse> FHttpServerResponse::Create(const TArrayView<uin
 
 TUniquePtr<FHttpServerResponse> FHttpServerResponse::Create(int32 HttpResponseCode, const FString& ErrorCode)
 {
-	FString ResponseBody = FString::Printf(TEXT("{\"errorCode\": \"%s\"}"), *ErrorCode);
+	const FString& ResponseBody = FString::Printf(TEXT("{\"errorCode\": \"%s\"}"), *ErrorCode);
 	auto Response = Create(ResponseBody, TEXT("application/json"));
 	Response->Code = HttpResponseCode;
 	return Response;
 }
 
-TUniquePtr<FHttpServerResponse> FHttpServerResponse::Ok()
+TUniquePtr<FHttpServerResponse> FHttpServerResponse::Ok(const FString& Message, FString ContentType)
 {
-	auto Response = MakeUnique<FHttpServerResponse>();
+	auto Response = Create(Message, MoveTemp(ContentType));
 	Response->Code = EHttpServerResponseCodes::Ok;
 	return Response;
+}
+
+TUniquePtr<FHttpServerResponse> FHttpServerResponse::Error(const FString& Message, FString ContentType)
+{
+	auto Response = Create(Message, MoveTemp(ContentType));
+	Response->Code = EHttpServerResponseCodes::ServerError;
+	return Response;
 }
\ No newline at end of file
diff --git a/Engine/Source/Runtime/Online/HTTPServer/Public/HttpServerResponse.h b/Engine/Source/Runtime/Online/HTTPServer/Public/HttpServerResponse.h
--- a/Engine/Source/Runtime/Online/HTTPServer/Public/HttpServerResponse.h
+++ b/Engine/Source/Runtime/Online/HTTPServer/Public/HttpServerResponse.h
@@ -69,10 +69,22 @@ public:
 	HTTPSERVER_API static TUniquePtr<FHttpServerResponse> Create(int32 HttpResponseCode, const FString& ErrorCode);
 
 	/**
-	 * Creates an FHttpServerResponse 200 
+	 * Creates an FHttpServerResponse 200
+	 * 
+	 * @param Message The optional and respective success message
+	 * @param ContentType The optional and respective content type
 	 * @return A unique pointer to an initialized response object
 	 */
-	HTTPSERVER_API static TUniquePtr<FHttpServerResponse> Ok();
+	HTTPSERVER_API static TUniquePtr<FHttpServerResponse> Ok(const FString& Message = TEXT(""), FString ContentType = TEXT("text/text"));
+
+	/**
+    * Creates an FHttpServerResponse 500
+	*
+	* @param Message The optional and respective error message
+	* @param ContentType The optional and respective content type
+    * @return A unique pointer to an initialized response object
+    */
+	HTTPSERVER_API static TUniquePtr<FHttpServerResponse> Error(const FString& Message = TEXT(""), FString ContentType = TEXT("text/text"));
 };
 
 
