[HttpServer] Fix clang compilation: HttpServer constants error codes are const TCHAR[] instead of const TCHAR*

diff --git a/Engine/Source/Runtime/Online/HTTPServer/Private/HttpServerConstantsPrivate.h b/Engine/Source/Runtime/Online/HTTPServer/Private/HttpServerConstantsPrivate.h
--- a/Engine/Source/Runtime/Online/HTTPServer/Private/HttpServerConstantsPrivate.h
+++ b/Engine/Source/Runtime/Online/HTTPServer/Private/HttpServerConstantsPrivate.h
@@ -5,30 +5,31 @@
 
 struct FHttpServerHeaderKeys
 {
-	static constexpr TCHAR* CONTENT_TYPE = TEXT("content-type");
-	static constexpr TCHAR* CONTENT_LENGTH = TEXT("content-length");
-	static constexpr TCHAR* CONNECTION = TEXT("connection");
-	static constexpr TCHAR* KEEP_ALIVE = TEXT("keep-alive");
+	static constexpr TCHAR CONTENT_TYPE[] = TEXT("content-type");
+	static constexpr TCHAR CONTENT_LENGTH[] = TEXT("content-length");
+	static constexpr TCHAR CONNECTION[] = TEXT("connection");
+	static constexpr TCHAR KEEP_ALIVE[] = TEXT("keep-alive");
 };
 
 struct FHttpServerErrorStrings
 {
 	// Connection
-	static constexpr TCHAR* SocketClosedFailure = TEXT("errors.com.epicgames.httpserver.socket_closed_failure");
-	static constexpr TCHAR* SocketRecvFailure = TEXT("errors.com.epicgames.httpserver.socket_recv_failure");
-	static constexpr TCHAR* SocketSendFailure = TEXT("errors.com.epicgames.httpserver.socket_send_failure");
+	static constexpr TCHAR SocketClosedFailure[] = TEXT("errors.com.epicgames.httpserver.socket_closed_failure");
+	static constexpr TCHAR SocketRecvFailure[] = TEXT("errors.com.epicgames.httpserver.socket_recv_failure");
+	static constexpr TCHAR SocketSendFailure[] = TEXT("errors.com.epicgames.httpserver.socket_send_failure");
 
 	// Routing
-	static constexpr TCHAR* NotFound = TEXT("errors.com.epicgames.httpserver.route_handler_not_found");
+	static constexpr TCHAR NotFound[] = TEXT("errors.com.epicgames.httpserver.route_handler_not_found");
 
 	// Serialization
-	static constexpr TCHAR* MalformedRequestSize = TEXT("errors.com.epicgames.httpserver.malformed_request_size");
-	static constexpr TCHAR* MalformedRequestHeaders = TEXT("errors.com.epicgames.httpserver.malformed_request_header");
-	static constexpr TCHAR* MissingRequestHeaders = TEXT("errors.com.epicgames.httpserver.missing_request_headers");
-	static constexpr TCHAR* MalformedRequestBody = TEXT("errors.com.epicgames.httpserver.malformed_request_body");
-	static constexpr TCHAR* UnknownRequestVerb = TEXT("errors.com.epicgames.httpserver.unknown_request_verb");
-	static constexpr TCHAR* InvalidContentLengthHeader = TEXT("errors.com.epicgames.httpserver.invalid_content_length_header");
-	static constexpr TCHAR* MissingContentLengthHeader = TEXT("errors.com.epicgames.httpserver.missing_content_length_header");
-	static constexpr TCHAR* MismatchedContentLengthBodyTooLarge =  TEXT("errors.com.epicgames.httpserver.mismatched_content_length_body_too_large");
+	static constexpr TCHAR MalformedRequestSize[] = TEXT("errors.com.epicgames.httpserver.malformed_request_size");
+	static constexpr TCHAR MalformedRequestHeaders[] = TEXT("errors.com.epicgames.httpserver.malformed_request_header");
+	static constexpr TCHAR MissingRequestHeaders[] = TEXT("errors.com.epicgames.httpserver.missing_request_headers");
+	static constexpr TCHAR MalformedRequestBody[] = TEXT("errors.com.epicgames.httpserver.malformed_request_body");
+	static constexpr TCHAR UnknownRequestVerb[] = TEXT("errors.com.epicgames.httpserver.unknown_request_verb");
+	static constexpr TCHAR UnsupportedHttpVersion[] = TEXT("errors.com.epicgames.httpserver.unsupported_http_version");
+	static constexpr TCHAR InvalidContentLengthHeader[] = TEXT("errors.com.epicgames.httpserver.invalid_content_length_header");
+	static constexpr TCHAR MissingContentLengthHeader[] = TEXT("errors.com.epicgames.httpserver.missing_content_length_header");
+	static constexpr TCHAR MismatchedContentLengthBodyTooLarge[] =  TEXT("errors.com.epicgames.httpserver.mismatched_content_length_body_too_large");
 };
 
