[HttpServer] Fix non-win64-compiler discrepencies
+ Removed 'static' specifier from friend uint32 GetTypeHash() signatures

diff --git a/Engine/Source/Runtime/Online/HTTPServer/Private/HttpConnection.h b/Engine/Source/Runtime/Online/HTTPServer/Private/HttpConnection.h
--- a/Engine/Source/Runtime/Online/HTTPServer/Private/HttpConnection.h
+++ b/Engine/Source/Runtime/Online/HTTPServer/Private/HttpConnection.h
@@ -67,7 +67,7 @@ public:
 		return this->Socket == Other.Socket;
 	}
 
-	static friend uint32 GetTypeHash(const FHttpConnection& Conn) 
+	friend uint32 GetTypeHash(const FHttpConnection& Conn) 
 	{
 		return GetTypeHash(Conn.Socket);
 	}
diff --git a/Engine/Source/Runtime/Online/HTTPServer/Public/HttpPath.h b/Engine/Source/Runtime/Online/HTTPServer/Public/HttpPath.h
--- a/Engine/Source/Runtime/Online/HTTPServer/Public/HttpPath.h
+++ b/Engine/Source/Runtime/Online/HTTPServer/Public/HttpPath.h
@@ -65,7 +65,8 @@ public:
 		bool bEqual = (0 == Path.Compare(Other.Path, ESearchCase::CaseSensitive));
 		return bEqual;
 	}
-	static friend uint32 GetTypeHash(const FHttpPath& Path)
+
+	friend uint32 GetTypeHash(const FHttpPath& Path)
 	{
 		return GetTypeHash(Path.GetPath());
 	}
