[HttpServer] Fix X1 compiler error
+ Removed trailing "," on enum definitions (suspect), removed EHttpServerRequestVerbs::ALL (potential suspect)

diff --git a/Engine/Source/Runtime/Online/HTTPServer/Public/HttpServerRequest.h b/Engine/Source/Runtime/Online/HTTPServer/Public/HttpServerRequest.h
--- a/Engine/Source/Runtime/Online/HTTPServer/Public/HttpServerRequest.h
+++ b/Engine/Source/Runtime/Online/HTTPServer/Public/HttpServerRequest.h
@@ -11,8 +11,7 @@ enum class EHttpServerRequestVerbs : uint16
 	POST   = 1 << 1,
 	PUT    = 1 << 2,
 	PATCH  = 1 << 3,
-	DELETE = 1 << 4,
-	ALL = GET | POST | PUT | PATCH | DELETE,
+	DELETE = 1 << 4
 };
 
 ENUM_CLASS_FLAGS(EHttpServerRequestVerbs);
