[HttpServer] Fix clang compiler discrepency

diff --git a/Engine/Source/Runtime/Online/HTTPServer/Private/HttpConnectionRequestReadContext.cpp b/Engine/Source/Runtime/Online/HTTPServer/Private/HttpConnectionRequestReadContext.cpp
--- a/Engine/Source/Runtime/Online/HTTPServer/Private/HttpConnectionRequestReadContext.cpp
+++ b/Engine/Source/Runtime/Online/HTTPServer/Private/HttpConnectionRequestReadContext.cpp
@@ -77,7 +77,7 @@ bool FHttpConnectionRequestReadContext::ParseHeader(uint8* ByteBuffer, int32 Buf
 	// We are hunting for the header terminator sequence '\r\n\r\n'
 	// Start 3 bytes prior to this most recent appendage
 	int32 SearchIndex = FMath::Max(0, PreviousHeaderBytesLen - 3);
-	for (SearchIndex; SearchIndex < HeaderBytes.Num() - 3; ++SearchIndex)
+	for ( ; SearchIndex < HeaderBytes.Num() - 3; ++SearchIndex)
 	{
 		// To find this sequence we will parse byte-by-byte in ANSICHAR
 		if (ANSICHAR(HeaderBytes[SearchIndex])     == '\r' &&
