[HttpServer] Fix clang compiler discrepency

diff --git a/Engine/Source/Runtime/Online/HTTPServer/Public/HttpPath.h b/Engine/Source/Runtime/Online/HTTPServer/Public/HttpPath.h
--- a/Engine/Source/Runtime/Online/HTTPServer/Public/HttpPath.h
+++ b/Engine/Source/Runtime/Online/HTTPServer/Public/HttpPath.h
@@ -66,9 +66,9 @@ public:
 		return bEqual;
 	}
 
-	friend uint32 GetTypeHash(const FHttpPath& Path)
+	friend uint32 GetTypeHash(const FHttpPath& InPath)
 	{
-		return GetTypeHash(Path.GetPath());
+		return GetTypeHash(InPath.GetPath());
 	}
 
 private:
