[HttpServer] Fix clang linkage discrepencies

diff --git a/Engine/Source/Runtime/Online/HTTPServer/Private/HttpConnectionContext.cpp b/Engine/Source/Runtime/Online/HTTPServer/Private/HttpConnectionContext.cpp
--- a/Engine/Source/Runtime/Online/HTTPServer/Private/HttpConnectionContext.cpp
+++ b/Engine/Source/Runtime/Online/HTTPServer/Private/HttpConnectionContext.cpp
@@ -17,7 +17,7 @@ void FHttpConnectionContext::AddElapsedIdleTime(float DeltaTime)
 	ElapsedIdleTime += DeltaTime;
 }
 
-void FHttpConnectionContext::AddError(const TCHAR* ErrorCode)
+void FHttpConnectionContext::AddError(const FString& ErrorCode)
 {
 	ErrorBuilder.Append(ErrorCode);
 	ErrorBuilder.AppendChar(TCHAR(' '));
diff --git a/Engine/Source/Runtime/Online/HTTPServer/Private/HttpConnectionContext.h b/Engine/Source/Runtime/Online/HTTPServer/Private/HttpConnectionContext.h
--- a/Engine/Source/Runtime/Online/HTTPServer/Private/HttpConnectionContext.h
+++ b/Engine/Source/Runtime/Online/HTTPServer/Private/HttpConnectionContext.h
@@ -55,7 +55,7 @@ protected:
 
 	 * @param ErrorCode The machine-readable error code
 	 */
-	void AddError(const TCHAR* ErrorCode);
+	void AddError(const FString& ErrorCode);
 
 	/** Tracks time since last read/write activity */
 	float ElapsedIdleTime = 0.0f;
diff --git a/Engine/Source/Runtime/Online/HTTPServer/Private/HttpConnectionRequestReadContext.cpp b/Engine/Source/Runtime/Online/HTTPServer/Private/HttpConnectionRequestReadContext.cpp
--- a/Engine/Source/Runtime/Online/HTTPServer/Private/HttpConnectionRequestReadContext.cpp
+++ b/Engine/Source/Runtime/Online/HTTPServer/Private/HttpConnectionRequestReadContext.cpp
@@ -1,3 +1,4 @@
+// Copyright 1998-2019 Epic Games, Inc. All Rights Reserved.
 
 #include "HttpConnectionRequestReadContext.h"
 #include "HttpServerConstantsPrivate.h"
diff --git a/Engine/Source/Runtime/Online/HTTPServer/Private/HttpConnectionResponseWriteContext.cpp b/Engine/Source/Runtime/Online/HTTPServer/Private/HttpConnectionResponseWriteContext.cpp
--- a/Engine/Source/Runtime/Online/HTTPServer/Private/HttpConnectionResponseWriteContext.cpp
+++ b/Engine/Source/Runtime/Online/HTTPServer/Private/HttpConnectionResponseWriteContext.cpp
@@ -1,3 +1,4 @@
+// Copyright 1998-2019 Epic Games, Inc. All Rights Reserved.
 
 #include "HttpConnectionResponseWriteContext.h"
 #include "HttpServerResponse.h"
diff --git a/Engine/Source/Runtime/Online/HTTPServer/Private/HttpServerConstantsPrivate.h b/Engine/Source/Runtime/Online/HTTPServer/Private/HttpServerConstantsPrivate.h
--- a/Engine/Source/Runtime/Online/HTTPServer/Private/HttpServerConstantsPrivate.h
+++ b/Engine/Source/Runtime/Online/HTTPServer/Private/HttpServerConstantsPrivate.h
@@ -3,33 +3,33 @@
 
 #include "CoreMinimal.h"
 
-struct FHttpServerHeaderKeys
+namespace FHttpServerHeaderKeys
 {
-	static constexpr TCHAR CONTENT_TYPE[] = TEXT("content-type");
-	static constexpr TCHAR CONTENT_LENGTH[] = TEXT("content-length");
-	static constexpr TCHAR CONNECTION[] = TEXT("connection");
-	static constexpr TCHAR KEEP_ALIVE[] = TEXT("keep-alive");
+	const FString CONTENT_TYPE = TEXT("content-type");
+	const FString CONTENT_LENGTH = TEXT("content-length");
+	const FString CONNECTION = TEXT("connection");
+	const FString KEEP_ALIVE = TEXT("keep-alive");
 };
 
-struct FHttpServerErrorStrings
+namespace FHttpServerErrorStrings
 {
 	// Connection
-	static constexpr TCHAR SocketClosedFailure[] = TEXT("errors.com.epicgames.httpserver.socket_closed_failure");
-	static constexpr TCHAR SocketRecvFailure[] = TEXT("errors.com.epicgames.httpserver.socket_recv_failure");
-	static constexpr TCHAR SocketSendFailure[] = TEXT("errors.com.epicgames.httpserver.socket_send_failure");
+	const FString SocketClosedFailure = TEXT("errors.com.epicgames.httpserver.socket_closed_failure");
+	const FString SocketRecvFailure = TEXT("errors.com.epicgames.httpserver.socket_recv_failure");
+	const FString SocketSendFailure = TEXT("errors.com.epicgames.httpserver.socket_send_failure");
 
 	// Routing
-	static constexpr TCHAR NotFound[] = TEXT("errors.com.epicgames.httpserver.route_handler_not_found");
+	const FString NotFound = TEXT("errors.com.epicgames.httpserver.route_handler_not_found");
 
 	// Serialization
-	static constexpr TCHAR MalformedRequestSize[] = TEXT("errors.com.epicgames.httpserver.malformed_request_size");
-	static constexpr TCHAR MalformedRequestHeaders[] = TEXT("errors.com.epicgames.httpserver.malformed_request_header");
-	static constexpr TCHAR MissingRequestHeaders[] = TEXT("errors.com.epicgames.httpserver.missing_request_headers");
-	static constexpr TCHAR MalformedRequestBody[] = TEXT("errors.com.epicgames.httpserver.malformed_request_body");
-	static constexpr TCHAR UnknownRequestVerb[] = TEXT("errors.com.epicgames.httpserver.unknown_request_verb");
-	static constexpr TCHAR UnsupportedHttpVersion[] = TEXT("errors.com.epicgames.httpserver.unsupported_http_version");
-	static constexpr TCHAR InvalidContentLengthHeader[] = TEXT("errors.com.epicgames.httpserver.invalid_content_length_header");
-	static constexpr TCHAR MissingContentLengthHeader[] = TEXT("errors.com.epicgames.httpserver.missing_content_length_header");
-	static constexpr TCHAR MismatchedContentLengthBodyTooLarge[] =  TEXT("errors.com.epicgames.httpserver.mismatched_content_length_body_too_large");
+	const FString MalformedRequestSize = TEXT("errors.com.epicgames.httpserver.malformed_request_size");
+	const FString MalformedRequestHeaders = TEXT("errors.com.epicgames.httpserver.malformed_request_header");
+	const FString MissingRequestHeaders = TEXT("errors.com.epicgames.httpserver.missing_request_headers");
+	const FString MalformedRequestBody = TEXT("errors.com.epicgames.httpserver.malformed_request_body");
+	const FString UnknownRequestVerb = TEXT("errors.com.epicgames.httpserver.unknown_request_verb");
+	const FString UnsupportedHttpVersion = TEXT("errors.com.epicgames.httpserver.unsupported_http_version");
+	const FString InvalidContentLengthHeader = TEXT("errors.com.epicgames.httpserver.invalid_content_length_header");
+	const FString MissingContentLengthHeader = TEXT("errors.com.epicgames.httpserver.missing_content_length_header");
+	const FString MismatchedContentLengthBodyTooLarge =  TEXT("errors.com.epicgames.httpserver.mismatched_content_length_body_too_large");
 };
 
