[HttpServer] Fix X1 compile: EHttpServerRequestVerbs are prefixed with VERB_ to avoid namespace collisions

diff --git a/Engine/Source/Runtime/Online/HTTPServer/Private/HttpConnectionRequestReadContext.cpp b/Engine/Source/Runtime/Online/HTTPServer/Private/HttpConnectionRequestReadContext.cpp
--- a/Engine/Source/Runtime/Online/HTTPServer/Private/HttpConnectionRequestReadContext.cpp
+++ b/Engine/Source/Runtime/Online/HTTPServer/Private/HttpConnectionRequestReadContext.cpp
@@ -150,7 +150,7 @@ bool FHttpConnectionRequestReadContext::IsRequestValid(const FHttpServerRequest&
 
 	switch (InRequest.Verb)
 	{
-	case EHttpServerRequestVerbs::GET:
+	case EHttpServerRequestVerbs::VERB_GET:
 		// Enforce content length missing or 0
 		if (bContentLengthSpecified && 0 != RequestContentLength)
 		{
@@ -159,8 +159,8 @@ bool FHttpConnectionRequestReadContext::IsRequestValid(const FHttpServerRequest&
 		}
 		break;
 
-	case EHttpServerRequestVerbs::PUT:
-	case EHttpServerRequestVerbs::POST:
+	case EHttpServerRequestVerbs::VERB_PUT:
+	case EHttpServerRequestVerbs::VERB_POST:
 		// Content length must be set
 		if (!bContentLengthSpecified)
 		{
@@ -219,27 +219,27 @@ TSharedPtr<FHttpServerRequest> FHttpConnectionRequestReadContext::BuildRequest(c
 	if (0 == RequestVerb.Compare(TEXT("GET"),
 		ESearchCase::IgnoreCase))
 	{
-		Request->Verb = EHttpServerRequestVerbs::GET;
+		Request->Verb = EHttpServerRequestVerbs::VERB_GET;
 	}
 	else if (0 == RequestVerb.Compare(TEXT("POST"),
 		ESearchCase::IgnoreCase))
 	{
-		Request->Verb = EHttpServerRequestVerbs::POST;
+		Request->Verb = EHttpServerRequestVerbs::VERB_POST;
 	}
 	else if (0 == RequestVerb.Compare(TEXT("PUT"),
 		ESearchCase::IgnoreCase))
 	{
-		Request->Verb = EHttpServerRequestVerbs::PUT;
+		Request->Verb = EHttpServerRequestVerbs::VERB_PUT;
 	}
 	else if (0 == RequestVerb.Compare(TEXT("DELETE"),
 		ESearchCase::IgnoreCase))
 	{
-		Request->Verb = EHttpServerRequestVerbs::DELETE;
+		Request->Verb = EHttpServerRequestVerbs::VERB_DELETE;
 	}
 	else if (0 == RequestVerb.Compare(TEXT("PATCH"),
 		ESearchCase::IgnoreCase))
 	{
-		Request->Verb = EHttpServerRequestVerbs::PATCH;
+		Request->Verb = EHttpServerRequestVerbs::VERB_PATCH;
 	}
 	else
 	{
diff --git a/Engine/Source/Runtime/Online/HTTPServer/Private/HttpRequestHandlerIterator.cpp b/Engine/Source/Runtime/Online/HTTPServer/Private/HttpRequestHandlerIterator.cpp
--- a/Engine/Source/Runtime/Online/HTTPServer/Private/HttpRequestHandlerIterator.cpp
+++ b/Engine/Source/Runtime/Online/HTTPServer/Private/HttpRequestHandlerIterator.cpp
@@ -31,7 +31,7 @@ const FHttpRequestHandler * const FHttpRequestHandlerIterator::Next()
 
 		// Filter by http verb
 		const EHttpServerRequestVerbs VerbFilterResult = RouteHandle->Verbs & Request->Verb;
-		if (EHttpServerRequestVerbs::NONE == VerbFilterResult)
+		if (EHttpServerRequestVerbs::VERB_NONE == VerbFilterResult)
 		{
 			// Not a matching verb
 			continue;
diff --git a/Engine/Source/Runtime/Online/HTTPServer/Private/HttpRouter.cpp b/Engine/Source/Runtime/Online/HTTPServer/Private/HttpRouter.cpp
--- a/Engine/Source/Runtime/Online/HTTPServer/Private/HttpRouter.cpp
+++ b/Engine/Source/Runtime/Online/HTTPServer/Private/HttpRouter.cpp
@@ -17,7 +17,7 @@ FHttpRouter::FHttpRouter()
 FHttpRouteHandle FHttpRouter::BindRoute(const FHttpPath& HttpPath,  const EHttpServerRequestVerbs& HttpVerbs,  const FHttpRequestHandler& Handler)
 {
 	check(HttpPath.IsValidPath());
-	check(EHttpServerRequestVerbs::NONE != HttpVerbs);
+	check(EHttpServerRequestVerbs::VERB_NONE != HttpVerbs);
 
 	if (RequestHandlerRegistrar->Contains(HttpPath.GetPath()))
 	{
diff --git a/Engine/Source/Runtime/Online/HTTPServer/Private/Tests/HttpServerTests.cpp b/Engine/Source/Runtime/Online/HTTPServer/Private/Tests/HttpServerTests.cpp
--- a/Engine/Source/Runtime/Online/HTTPServer/Private/Tests/HttpServerTests.cpp
+++ b/Engine/Source/Runtime/Online/HTTPServer/Private/Tests/HttpServerTests.cpp
@@ -31,11 +31,11 @@ bool FHttpServerIntegrationTest::RunTest(const FString& Parameters)
 	{
 		return true;
 	};
-	FHttpRouteHandle HttpRouteHandle = HttpRouter->BindRoute(HttpPath, EHttpServerRequestVerbs::GET, RequestHandler);
+	FHttpRouteHandle HttpRouteHandle = HttpRouter->BindRoute(HttpPath, EHttpServerRequestVerbs::VERB_GET, RequestHandler);
 	TestTrue(TEXT("HttpRouteHandle.IsValid()"), HttpRouteHandle.IsValid());
 
 	// Disallow duplicate route bindings
-	FHttpRouteHandle DuplicateHandle = HttpRouter->BindRoute(HttpPath, EHttpServerRequestVerbs::GET, RequestHandler);
+	FHttpRouteHandle DuplicateHandle = HttpRouter->BindRoute(HttpPath, EHttpServerRequestVerbs::VERB_GET, RequestHandler);
 	TestFalse(TEXT("HttpRouteHandle Duplicated"), DuplicateHandle.IsValid());
 
 	// Make a request
diff --git a/Engine/Source/Runtime/Online/HTTPServer/Public/HttpServerRequest.h b/Engine/Source/Runtime/Online/HTTPServer/Public/HttpServerRequest.h
--- a/Engine/Source/Runtime/Online/HTTPServer/Public/HttpServerRequest.h
+++ b/Engine/Source/Runtime/Online/HTTPServer/Public/HttpServerRequest.h
@@ -6,12 +6,12 @@
 
 enum class EHttpServerRequestVerbs : uint16
 {
-	NONE   = 0,
-	GET    = 1 << 0,
-	POST   = 1 << 1,
-	PUT    = 1 << 2,
-	PATCH  = 1 << 3,
-	DELETE = 1 << 4
+	VERB_NONE   = 0,
+	VERB_GET    = 1 << 0,
+	VERB_POST   = 1 << 1,
+	VERB_PUT    = 1 << 2,
+	VERB_PATCH  = 1 << 3,
+	VERB_DELETE = 1 << 4
 };
 
 ENUM_CLASS_FLAGS(EHttpServerRequestVerbs);
