[HttpServer] HttpConnection request logging clearly and uniqely identifies requests by [Port][ConnectionId]-RequestId
+ HttpListener now supplies incoming connection instances with the OriginPort on which the connection was accepted
+ HttpConnection now supplies HttpConnectionWriteContext with the HTTP protocol version that matches the respective request http protocol version
+ HttpConnection now increments 'LastRequestNumber' whenever a new read operation is initiated (prior: was incremented after a completed read)
+ HttpConnectionContext now exposes a respective HTTP error code that can optionally be set when appending error messages.  This overload facilitates an early-out flow if/when an unsupported HTTP version is specified in the request.
+ HttpServerResponseCodes are now an enum class : uint16 - This change yields type coherence/safety across the variant classes that utilize these values
+ Additionally, KeepAlive connections are now correctly pooled as per HTTP1.1 protocol, such that KeepAlive is true by default and otherwise disabled via an explicit "Connection: close" header.

diff --git a/Engine/Source/Runtime/Online/HTTPServer/Private/HttpConnection.cpp b/Engine/Source/Runtime/Online/HTTPServer/Private/HttpConnection.cpp
--- a/Engine/Source/Runtime/Online/HTTPServer/Private/HttpConnection.cpp
+++ b/Engine/Source/Runtime/Online/HTTPServer/Private/HttpConnection.cpp
@@ -17,9 +17,11 @@
 
 DEFINE_LOG_CATEGORY(LogHttpConnection)
 
-FHttpConnection::FHttpConnection(FSocket* InSocket, TSharedPtr<FHttpRouter> InRouter)
+FHttpConnection::FHttpConnection(FSocket* InSocket, TSharedPtr<FHttpRouter> InRouter, uint32 InOriginPort, uint32 InConnectionId)
 	: Socket(InSocket)
 	,Router(InRouter)
+	,OriginPort(InOriginPort)
+	,ConnectionId(InConnectionId)
 	,ReadContext(InSocket)
 	,WriteContext(InSocket)
 {
@@ -105,6 +107,7 @@ void FHttpConnection::BeginRead(float DeltaTime)
 	if (Socket->HasPendingData(PendingDataSize))
 	{
 		TransferState(EHttpConnectionState::AwaitingRead, EHttpConnectionState::Reading);
+		LastRequestNumber++;
 		ReadContext.ResetContext();
 		ContinueRead(DeltaTime);
 	}
@@ -130,7 +133,7 @@ void FHttpConnection::ContinueRead(float DeltaTime)
 		break;
 
 	case EHttpConnectionContextState::Error:
-		HandleReadError(*ReadContext.GetErrorStr());
+		HandleReadError(ReadContext.GetErrorCode(), *ReadContext.GetErrorStr());
 		break;
 	}
 }
@@ -138,27 +141,33 @@ void FHttpConnection::ContinueRead(float DeltaTime)
 void FHttpConnection::CompleteRead(const TSharedPtr<FHttpServerRequest>& Request)
 {
 	TArray<FString>* ConnectionHeaders = Request->Headers.Find(FHttpServerHeaderKeys::CONNECTION);
-	if (ConnectionHeaders &&
-		ConnectionHeaders->Contains(TEXT("Keep-Alive")))
+	if (ConnectionHeaders)
 	{
-		bKeepAlive = true;
+		bKeepAlive = ResolveKeepAlive(Request->HttpVersion, *ConnectionHeaders);
 	}
 
 	TWeakPtr<FHttpConnection> WeakThisPtr(AsShared());
-	FHttpResultCallback OnProcessingComplete = [WeakThisPtr, Request, LastRequestNumberStateCapture = ++LastRequestNumber]
+	FHttpResultCallback OnProcessingComplete = 
+		[WeakThisPtr, 
+		OriginPortCapture = OriginPort,
+		ConnectionIdCapture = ConnectionId, 
+		LastRequestNumberCapture = LastRequestNumber,
+		ResponseVersionCapture = Request->HttpVersion]
 	(TUniquePtr<FHttpServerResponse>&& Response)
 	{
 		TSharedPtr<FHttpConnection> SharedThisPtr = WeakThisPtr.Pin();
 		if (SharedThisPtr.IsValid())
 		{
 			UE_LOG(LogHttpConnection, Log,
-				TEXT("Completed Processing Request [%u]"), LastRequestNumberStateCapture)
+				TEXT("EndProcessRequest\t [%d][%u]-%u : %u"), 
+				OriginPortCapture, ConnectionIdCapture, LastRequestNumberCapture, Response->Code);
 
 			// Ensure this result callback was called once
 			check(EHttpConnectionState::AwaitingProcessing == SharedThisPtr->GetState());
 
 			// Begin response flow
-			SharedThisPtr->BeginWrite(MoveTemp(Response), LastRequestNumberStateCapture);
+			Response->HttpVersion = ResponseVersionCapture;
+			SharedThisPtr->BeginWrite(MoveTemp(Response), LastRequestNumberCapture);
 		}
 	};
 
@@ -170,7 +179,8 @@ void FHttpConnection::ProcessRequest(const TSharedPtr<FHttpServerRequest>& Reque
 	TransferState(EHttpConnectionState::Reading, EHttpConnectionState::AwaitingProcessing);
 
 	UE_LOG(LogHttpConnection, Log,
-		TEXT("Begin Processing Request [%u]: %s"), LastRequestNumber, *Request->RelativePath.GetPath());
+		TEXT("BeginProcessRequest\t [%d][%u]-%u : %s"), 
+		OriginPort, ConnectionId, LastRequestNumber, *Request->RelativePath.GetPath());
 
 	bool bRequestHandled = false;
 	auto RequestHandlerIterator = Router->CreateRequestHandlerIterator(Request);
@@ -189,7 +199,7 @@ void FHttpConnection::ProcessRequest(const TSharedPtr<FHttpServerRequest>& Reque
 	if (!bRequestHandled)
 	{
 		const FString& ResponseError(FHttpServerErrorStrings::NotFound);
-		auto Response = FHttpServerResponse::Create(EHttpServerResponseCodes::NotFound, ResponseError);
+		auto Response = FHttpServerResponse::Error(EHttpServerResponseCodes::NotFound, ResponseError);
 		OnProcessingComplete(MoveTemp(Response));
 	}
 }
@@ -203,9 +213,9 @@ void FHttpConnection::BeginWrite(TUniquePtr<FHttpServerResponse>&& Response, uin
 
 	if (bKeepAlive)
 	{
-		const FString& KeepAliveStr = FString::Printf(TEXT("timeout=%f"), ConnectionKeepAliveTimeout);
-		const TArray<FString>& KeepAliveValue = { KeepAliveStr };
-		Response->Headers.Add(FHttpServerHeaderKeys::KEEP_ALIVE, KeepAliveValue);
+		FString KeepAliveTimeoutStr = FString::Printf(TEXT("timeout=%f"), ConnectionKeepAliveTimeout);
+		TArray<FString> KeepAliveTimeoutValue = { MoveTemp(KeepAliveTimeoutStr) };
+		Response->Headers.Add(FHttpServerHeaderKeys::KEEP_ALIVE, MoveTemp(KeepAliveTimeoutValue));
 	}
 
 	WriteContext.ResetContext(MoveTemp(Response));
@@ -279,23 +289,36 @@ void FHttpConnection::Destroy()
 	}
 }
 
-void FHttpConnection::HandleReadError(const TCHAR* ErrorCode)
+void FHttpConnection::HandleReadError(EHttpServerResponseCodes ErrorCode, const TCHAR* ErrorCodeStr)
 {
-	UE_LOG(LogHttpConnection, Error, TEXT("%s"), ErrorCode);
+	UE_LOG(LogHttpConnection, Error, TEXT("%s"), ErrorCodeStr);
 
 	// Forcibly Reply
 	bKeepAlive = false;
-	auto ResponseCode = EHttpServerResponseCodes::BadRequest;
-	auto Response = FHttpServerResponse::Create(ResponseCode, ErrorCode);
-
-	BeginWrite(MoveTemp(Response), ++LastRequestNumber);
+	auto Response = FHttpServerResponse::Error(ErrorCode, ErrorCodeStr);
+	BeginWrite(MoveTemp(Response), LastRequestNumber);
 }
 
-void FHttpConnection::HandleWriteError(const TCHAR* ErrorCode)
+void FHttpConnection::HandleWriteError(const TCHAR* ErrorCodeStr)
 {
-	UE_LOG(LogHttpConnection, Error, TEXT("%s"), ErrorCode);
+	UE_LOG(LogHttpConnection, Error, TEXT("%s"), ErrorCodeStr);
 
 	// Forcibly Close
 	bKeepAlive = false;
 	Destroy();
-}
\ No newline at end of file
+}
+
+bool FHttpConnection::ResolveKeepAlive(HttpVersion::EHttpServerHttpVersion HttpVersion, const TArray<FString>& ConnectionHeaders)
+{
+	const bool bDefaultKeepAlive = HttpVersion == HttpVersion::EHttpServerHttpVersion::HTTP_VERSION_1_1;
+
+	if (bDefaultKeepAlive)
+	{
+		return !ConnectionHeaders.Contains(TEXT("close"));
+	}
+	else
+	{
+		return ConnectionHeaders.Contains(TEXT("Keep-Alive"));
+	}
+	return true;
+}
diff --git a/Engine/Source/Runtime/Online/HTTPServer/Private/HttpConnection.h b/Engine/Source/Runtime/Online/HTTPServer/Private/HttpConnection.h
--- a/Engine/Source/Runtime/Online/HTTPServer/Private/HttpConnection.h
+++ b/Engine/Source/Runtime/Online/HTTPServer/Private/HttpConnection.h
@@ -6,6 +6,8 @@
 #include "HttpConnectionRequestReadContext.h"
 #include "HttpConnectionResponseWriteContext.h"
 #include "HttpResultCallback.h"
+#include "HttpserverConstants.h"
+#include "HttpServerHttpVersion.h"
 
 class FSocket;
 class ISocketSubsystem;
@@ -25,7 +27,7 @@ public:
 	 * 
 	 * @param InSocket The underlying file descriptor
 	 */
-	FHttpConnection(FSocket* Socket, TSharedPtr<FHttpRouter> Router);
+	FHttpConnection(FSocket* Socket, TSharedPtr<FHttpRouter> Router, uint32 OriginPort, uint32 ConnectionId);
 
 	/**
 	 * Destructor
@@ -67,7 +69,7 @@ public:
 		return this->Socket == Other.Socket;
 	}
 
-	friend uint32 GetTypeHash(const FHttpConnection& Conn) 
+	friend uint32 GetTypeHash(const FHttpConnection& Conn)
 	{
 		return GetTypeHash(Conn.Socket);
 	}
@@ -147,16 +149,27 @@ private:
 	/**
 	 * Logs and responds with the caller-supplied error code
 	 *
-	 * @param ErrorCode The machine-readable error description
+	 * @param ErrorCode The HTTP error code
+	 * @param ErrorCodeStr The machine-readable error description
 	 */
-	void HandleReadError(const TCHAR* ErrorCode);
+	void HandleReadError(EHttpServerResponseCodes ErrorCode,  const TCHAR* ErrorCodeStr);
 
 	/**
 	 * Logs the caller-supplied error code and closes the connection
 	 *
-	 * @param ErrorCode The machine-readable error description
+	 * @param ErrorCodeStr The machine-readable error description
 	 */
-	void HandleWriteError(const TCHAR* ErrorCode);
+	void HandleWriteError(const TCHAR* ErrorCodeStr);
+
+	/**
+	* Determines whether KeepAlive is set based on the caller-supplied http version and connection headers
+	*
+	* @param HttpVersion         The http version of the request
+	* @param ConnectionHeaders   The request "Connection:" headers
+	* @return                    true if KeepAlive should be set, false otherwise
+    */
+	static bool ResolveKeepAlive(HttpVersion::EHttpServerHttpVersion HttpVersion, const TArray<FString>& ConnectionHeaders);
+
 
 private:
 
@@ -169,6 +182,12 @@ private:
 	/** Routing mechanism  */
 	const TSharedPtr<FHttpRouter> Router;
 
+	/** The origin port on which this connection was accepted */
+	uint32 OriginPort;
+
+	/** The connection identifier (used for logging purposes) */
+	uint32 ConnectionId;
+
 	/** Helper reader context to track the state of streaming request reads */
 	FHttpConnectionRequestReadContext ReadContext;
 
@@ -176,7 +195,7 @@ private:
 	FHttpConnectionResponseWriteContext WriteContext;
 
 	/** Whether to keep this connection alive after writing */
-	bool bKeepAlive = false;
+	bool bKeepAlive = true;
 
 	/** Whether to gracefully close pending current operations */
 	bool bGracefulDestroyRequested = false;
diff --git a/Engine/Source/Runtime/Online/HTTPServer/Private/HttpConnectionContext.cpp b/Engine/Source/Runtime/Online/HTTPServer/Private/HttpConnectionContext.cpp
--- a/Engine/Source/Runtime/Online/HTTPServer/Private/HttpConnectionContext.cpp
+++ b/Engine/Source/Runtime/Online/HTTPServer/Private/HttpConnectionContext.cpp
@@ -17,8 +17,16 @@ void FHttpConnectionContext::AddElapsedIdleTime(float DeltaTime)
 	ElapsedIdleTime += DeltaTime;
 }
 
-void FHttpConnectionContext::AddError(const FString& ErrorCode)
+void FHttpConnectionContext::AddError(const FString& ErrorCodeStr, EHttpServerResponseCodes HttpErrorCode)
 {
-	ErrorBuilder.Append(ErrorCode);
-	ErrorBuilder.AppendChar(TCHAR(' '));
+	if (!ErrorBuilder.IsEmpty())
+	{
+		ErrorBuilder.AppendChar(TCHAR(' '));
+	}
+	ErrorBuilder.Append(ErrorCodeStr);
+
+	if (EHttpServerResponseCodes::Unknown != HttpErrorCode)
+	{
+		ErrorCode = HttpErrorCode;
+	}
 }
diff --git a/Engine/Source/Runtime/Online/HTTPServer/Private/HttpConnectionContext.h b/Engine/Source/Runtime/Online/HTTPServer/Private/HttpConnectionContext.h
--- a/Engine/Source/Runtime/Online/HTTPServer/Private/HttpConnectionContext.h
+++ b/Engine/Source/Runtime/Online/HTTPServer/Private/HttpConnectionContext.h
@@ -1,6 +1,7 @@
 // Copyright 1998-2019 Epic Games, Inc. All Rights Reserved.
 #pragma once
 #include "CoreMinimal.h"
+#include "HttpServerConstants.h"
 
 enum class EHttpConnectionContextState
 {
@@ -30,6 +31,14 @@ public:
 		return ElapsedIdleTime;
 	}
 
+	/**
+    * Gets the respective error code
+	*/
+	FORCEINLINE EHttpServerResponseCodes GetErrorCode() const
+	{
+		return ErrorCode;
+	}
+
 	/**
 	 * Gets the cumulative error string
 	 */
@@ -55,12 +64,14 @@ protected:
 
 	 * @param ErrorCode The machine-readable error code
 	 */
-	void AddError(const FString& ErrorCode);
+	void AddError(const FString& ErrorCodeStr, EHttpServerResponseCodes ErrorCode = EHttpServerResponseCodes::Unknown);
 
 	/** Tracks time since last read/write activity */
 	float ElapsedIdleTime = 0.0f;
 
+	/** Tracks the respective ErrorCode  */
+	EHttpServerResponseCodes ErrorCode = EHttpServerResponseCodes::Unknown;
+
 	/** Tracks cumulative context errors */
 	FStringOutputDevice ErrorBuilder;
-
 };
\ No newline at end of file
diff --git a/Engine/Source/Runtime/Online/HTTPServer/Private/HttpConnectionRequestReadContext.cpp b/Engine/Source/Runtime/Online/HTTPServer/Private/HttpConnectionRequestReadContext.cpp
--- a/Engine/Source/Runtime/Online/HTTPServer/Private/HttpConnectionRequestReadContext.cpp
+++ b/Engine/Source/Runtime/Online/HTTPServer/Private/HttpConnectionRequestReadContext.cpp
@@ -1,6 +1,7 @@
 // Copyright 1998-2019 Epic Games, Inc. All Rights Reserved.
 
 #include "HttpConnectionRequestReadContext.h"
+#include "HttpServerConstants.h"
 #include "HttpServerConstantsPrivate.h"
 #include "HttpServerRequest.h"
 #include "HttpConnectionContext.h"
@@ -128,7 +129,7 @@ bool FHttpConnectionRequestReadContext::ParseBody(uint8* ByteBuffer, int32 Buffe
 	if (BufferLen > IncomingRequestBodyBytesToRead)
 	{
 		// Error - Sent data size exceeds expected 
-		AddError(FHttpServerErrorStrings::MismatchedContentLengthBodyTooLarge);
+		AddError(FHttpServerErrorStrings::MismatchedContentLengthBodyTooLarge, EHttpServerResponseCodes::BadRequest);
 		return false;
 	}
 
@@ -154,7 +155,7 @@ bool FHttpConnectionRequestReadContext::IsRequestValid(const FHttpServerRequest&
 		// Enforce content length missing or 0
 		if (bContentLengthSpecified && 0 != RequestContentLength)
 		{
-			AddError(FHttpServerErrorStrings::InvalidContentLengthHeader);
+			AddError(FHttpServerErrorStrings::InvalidContentLengthHeader, EHttpServerResponseCodes::BadRequest);
 			return false;
 		}
 		break;
@@ -164,13 +165,13 @@ bool FHttpConnectionRequestReadContext::IsRequestValid(const FHttpServerRequest&
 		// Content length must be set
 		if (!bContentLengthSpecified)
 		{
-			AddError(FHttpServerErrorStrings::MissingContentLengthHeader);
+			AddError(FHttpServerErrorStrings::MissingContentLengthHeader, EHttpServerResponseCodes::LengthRequired);
 			return false;
 		}
 		// Content length must be valid
 		if(RequestContentLength < 0)
 		{
-			AddError(FHttpServerErrorStrings::InvalidContentLengthHeader);
+			AddError(FHttpServerErrorStrings::InvalidContentLengthHeader, EHttpServerResponseCodes::LengthRequired);
 			return false;
 		}
 		break;
@@ -199,17 +200,17 @@ TSharedPtr<FHttpServerRequest> FHttpConnectionRequestReadContext::BuildRequest(c
 
 	if (0 == ParsedHeader.Num())
 	{
-		AddError(FHttpServerErrorStrings::MissingRequestHeaders);
+		AddError(FHttpServerErrorStrings::MissingRequestHeaders, EHttpServerResponseCodes::BadRequest);
 		return nullptr;
 	}
 
-	// Split HTTP Method Line (Path/Verb)
+	// Split HTTP Method Line (Path/Verb/Version)
 	TArray<FString> HttpMethodTokens;
 	const FString& HttpMethod = ParsedHeader[0];
 	HttpMethod.ParseIntoArrayWS(HttpMethodTokens);
-	if (HttpMethodTokens.Num() < 2)
+	if (HttpMethodTokens.Num() < 3)
 	{
-		AddError(FHttpServerErrorStrings::MalformedRequestHeaders);
+		AddError(FHttpServerErrorStrings::MalformedRequestHeaders, EHttpServerResponseCodes::BadRequest);
 		return nullptr;
 	}
 
@@ -244,7 +245,7 @@ TSharedPtr<FHttpServerRequest> FHttpConnectionRequestReadContext::BuildRequest(c
 	else
 	{
 		// Unknown Verb
-		AddError(FHttpServerErrorStrings::UnknownRequestVerb);
+		AddError(FHttpServerErrorStrings::UnknownRequestVerb, EHttpServerResponseCodes::BadMethod);
 		return nullptr;
 	}
 
@@ -273,6 +274,14 @@ TSharedPtr<FHttpServerRequest> FHttpConnectionRequestReadContext::BuildRequest(c
 	}
 	Request->RelativePath.SetPath(RequestHttpPath);
 
+	// Parse/store http version
+	const FString& RequestHttpVersion = HttpMethodTokens[2];
+	if(!HttpVersion::FromString(RequestHttpVersion, Request->HttpVersion))
+	{
+		AddError(FHttpServerErrorStrings::UnsupportedHttpVersion, EHttpServerResponseCodes::VersionNotSup);
+		return nullptr;
+	}
+
 	// Parse/store headers
 	for (int i = 1; i < ParsedHeader.Num(); ++i)
 	{
diff --git a/Engine/Source/Runtime/Online/HTTPServer/Private/HttpConnectionResponseWriteContext.cpp b/Engine/Source/Runtime/Online/HTTPServer/Private/HttpConnectionResponseWriteContext.cpp
--- a/Engine/Source/Runtime/Online/HTTPServer/Private/HttpConnectionResponseWriteContext.cpp
+++ b/Engine/Source/Runtime/Online/HTTPServer/Private/HttpConnectionResponseWriteContext.cpp
@@ -2,6 +2,7 @@
 
 #include "HttpConnectionResponseWriteContext.h"
 #include "HttpServerResponse.h"
+#include "HttpServerHttpVersion.h"
 #include "HttpServerConstantsPrivate.h"
 #include "Sockets.h"
 
@@ -29,7 +30,7 @@ void FHttpConnectionResponseWriteContext::ResetContext(TUniquePtr<FHttpServerRes
 		Response->Headers.Add(FHttpServerHeaderKeys::CONTENT_LENGTH, MoveTemp(ContentLengthValue));
 
 		// Serialize Headers
-		HeaderBytes.Append(SerializeHeadersUtf8(Response->Code, Response->Headers));
+		HeaderBytes.Append(SerializeHeadersUtf8(Response->HttpVersion, Response->Code, Response->Headers));
 	}
 }
 
@@ -101,9 +102,10 @@ bool FHttpConnectionResponseWriteContext::IsWriteBodyComplete() const
 	return BodyBytesWritten >= Response->Body.Num();
 }
 
-TArray<uint8> FHttpConnectionResponseWriteContext::SerializeHeadersUtf8(int32 ResponseCode, const TMap<FString, TArray<FString>>& HeadersMap)
+TArray<uint8> FHttpConnectionResponseWriteContext::SerializeHeadersUtf8(HttpVersion::EHttpServerHttpVersion HttpVersion, EHttpServerResponseCodes ResponseCode, const TMap<FString, TArray<FString>>& HeadersMap)
 {
-	FString ResponseHeaderStr = FString::Printf(TEXT("HTTP/1.1 %d\r\n"), ResponseCode);
+	FString ResponseHeaderStr = FString::Printf(TEXT("HTTP/%.1f %d\r\n"),  
+		*HttpVersion::ToString(HttpVersion), ResponseCode);
 
 	for (const auto& KeyValuePair : HeadersMap)
 	{
diff --git a/Engine/Source/Runtime/Online/HTTPServer/Private/HttpConnectionResponseWriteContext.h b/Engine/Source/Runtime/Online/HTTPServer/Private/HttpConnectionResponseWriteContext.h
--- a/Engine/Source/Runtime/Online/HTTPServer/Private/HttpConnectionResponseWriteContext.h
+++ b/Engine/Source/Runtime/Online/HTTPServer/Private/HttpConnectionResponseWriteContext.h
@@ -2,6 +2,8 @@
 #pragma once
 #include "CoreMinimal.h"
 #include "HttpConnectionContext.h"
+#include "HttpServerHttpVersion.h"
+
 
 struct FHttpServerResponse;
 class FSocket;
@@ -63,11 +65,12 @@ private:
 	/**
 	 * Creates a UTF8-serialized HTTP header in byte representation
 	 *
-	 * @param  The HTTP response code
-	 * @param  HeadersMap The associative header key to values map
-	 * @return The UTF8-encoded byte representation
+	 * @param  HttpVersion  The HTTP version
+	 * @param  ResponseCode The HTTP response code
+	 * @param  HeadersMap   The associative header key to values map
+	 * @return              The UTF8-encoded byte representation
 	 */
-	static TArray<uint8> SerializeHeadersUtf8(int32 ResponseCode, const TMap<FString, TArray<FString>>& HeadersMap);
+	static TArray<uint8> SerializeHeadersUtf8(HttpVersion::EHttpServerHttpVersion HttpVersion, EHttpServerResponseCodes ResponseCode, const TMap<FString, TArray<FString>>& HeadersMap);
 
 private:
 
diff --git a/Engine/Source/Runtime/Online/HTTPServer/Private/HttpListener.cpp b/Engine/Source/Runtime/Online/HTTPServer/Private/HttpListener.cpp
--- a/Engine/Source/Runtime/Online/HTTPServer/Private/HttpListener.cpp
+++ b/Engine/Source/Runtime/Online/HTTPServer/Private/HttpListener.cpp
@@ -44,7 +44,7 @@ bool FHttpListener::StartListening()
 	if (nullptr == SocketSubsystem)
 	{
 		UE_LOG(LogHttpListener, Error, 
-			TEXT("IHttpListener - SocketSubsystem Initialization Failed"));
+			TEXT("HttpListener - SocketSubsystem Initialization Failed"));
 		return false;
 	}
 
@@ -52,7 +52,7 @@ bool FHttpListener::StartListening()
 	if (nullptr == ListenSocket)
 	{
 		UE_LOG(LogHttpListener, Error, 
-			TEXT("IHttpListener - Unable to allocate stream socket"));
+			TEXT("HttpListener - Unable to allocate stream socket"));
 		return false;
 	}
 	ListenSocket->SetNonBlocking(true);
@@ -62,7 +62,7 @@ bool FHttpListener::StartListening()
 	if (!ListenSocket->Bind(*LocalhostAddr))
 	{
 		UE_LOG(LogHttpListener, Error, 
-			TEXT("IHttpListener unable to bind to %s"), 
+			TEXT("HttpListener unable to bind to %s"), 
 			*LocalhostAddr->ToString(true));
 		return false;
 	}
@@ -72,19 +72,19 @@ bool FHttpListener::StartListening()
 	if (ActualBufferSize != ListenerBufferSize)
 	{
 		UE_LOG(LogHttpListener, Warning, 
-			TEXT("IHttpListener unable to set desired buffer size (%d): Limited to %d"),
+			TEXT("HttpListener unable to set desired buffer size (%d): Limited to %d"),
 			ListenerBufferSize, ActualBufferSize);
 	}
 
 	if (!ListenSocket->Listen(ListenerConnectionBacklogSize))
 	{
 		UE_LOG(LogHttpListener, Error, 
-			TEXT("IHttpListener unable to listen on socket"));
+			TEXT("HttpListener unable to listen on socket"));
 		return false;
 	}
 
 	UE_LOG(LogHttpListener, Log, 
-		TEXT("Created new HttpListener on port %d"), ListenPort);
+		TEXT("Created new HttpListener on port %u"), ListenPort);
 	return true;
 }
 
@@ -95,6 +95,9 @@ void FHttpListener::StopListening()
 	// Tear down our top-level listener first
 	if (ListenSocket)
 	{
+		UE_LOG(LogHttpListener, Log,
+			TEXT("HttListener stopping listening on Port %u"), ListenPort);
+
 		ISocketSubsystem* SocketSubsystem = ISocketSubsystem::Get(PLATFORM_SOCKETSUBSYSTEM);
 		if (SocketSubsystem)
 		{
@@ -177,7 +180,8 @@ void FHttpListener::AcceptConnections(uint32 MaxConnectionsToAccept)
 			}
 
 			IncomingConnection->SetNonBlocking(true);
-			TSharedPtr<FHttpConnection> Connection = MakeShared<FHttpConnection>(IncomingConnection, Router);
+			TSharedPtr<FHttpConnection> Connection = 
+				MakeShared<FHttpConnection>(IncomingConnection, Router, ListenPort, NumConnectionsAccepted++);
 			Connections.Add(Connection);
 		}
 	}
diff --git a/Engine/Source/Runtime/Online/HTTPServer/Private/HttpListener.h b/Engine/Source/Runtime/Online/HTTPServer/Private/HttpListener.h
--- a/Engine/Source/Runtime/Online/HTTPServer/Private/HttpListener.h
+++ b/Engine/Source/Runtime/Online/HTTPServer/Private/HttpListener.h
@@ -109,6 +109,9 @@ private:
 	/** The collection of unique connections */
 	FHttpConnectionPool Connections;
 
+	/** The total number of connections accepted by this listener */
+	uint32 NumConnectionsAccepted = 0;
+
 	/** Maximum number of connections to accept per frame */
 	static constexpr uint32 MaxConnectionsToAcceptPerFrame = 1;
 	/** Maximum number of pending connections to queue */
diff --git a/Engine/Source/Runtime/Online/HTTPServer/Private/HttpServerResponse.cpp b/Engine/Source/Runtime/Online/HTTPServer/Private/HttpServerResponse.cpp
--- a/Engine/Source/Runtime/Online/HTTPServer/Private/HttpServerResponse.cpp
+++ b/Engine/Source/Runtime/Online/HTTPServer/Private/HttpServerResponse.cpp
@@ -1,3 +1,4 @@
+// Copyright 1998-2019 Epic Games, Inc. All Rights Reserved.
 
 #include "HttpServerResponse.h"
 #include "HttpServerConstants.h"
@@ -40,24 +41,17 @@ TUniquePtr<FHttpServerResponse> FHttpServerResponse::Create(const TArrayView<uin
 	return Response;
 }
 
-TUniquePtr<FHttpServerResponse> FHttpServerResponse::Create(int32 HttpResponseCode, const FString& ErrorCode)
+TUniquePtr<FHttpServerResponse> FHttpServerResponse::Ok()
 {
-	const FString& ResponseBody = FString::Printf(TEXT("{\"errorCode\": \"%s\"}"), *ErrorCode);
-	auto Response = Create(ResponseBody, TEXT("application/json"));
-	Response->Code = HttpResponseCode;
-	return Response;
-}
-
-TUniquePtr<FHttpServerResponse> FHttpServerResponse::Ok(const FString& Message, FString ContentType)
-{
-	auto Response = Create(Message, MoveTemp(ContentType));
-	Response->Code = EHttpServerResponseCodes::Ok;
+	TUniquePtr<FHttpServerResponse> Response = MakeUnique<FHttpServerResponse>();
+	Response->Code = EHttpServerResponseCodes::NoContent;
 	return Response;
 }
 
-TUniquePtr<FHttpServerResponse> FHttpServerResponse::Error(const FString& Message, FString ContentType)
+TUniquePtr<FHttpServerResponse> FHttpServerResponse::Error(EHttpServerResponseCodes ResponseCode, const FString& ErrorCode)
 {
-	auto Response = Create(Message, MoveTemp(ContentType));
-	Response->Code = EHttpServerResponseCodes::ServerError;
+	const FString& ResponseBody = FString::Printf(TEXT("{\"errorCode\": \"%s\"}"), *ErrorCode);
+	auto Response = Create(ResponseBody, TEXT("application/json"));
+	Response->Code = ResponseCode;
 	return Response;
 }
\ No newline at end of file
diff --git a/Engine/Source/Runtime/Online/HTTPServer/Public/HttpServerConstants.h b/Engine/Source/Runtime/Online/HTTPServer/Public/HttpServerConstants.h
--- a/Engine/Source/Runtime/Online/HTTPServer/Public/HttpServerConstants.h
+++ b/Engine/Source/Runtime/Online/HTTPServer/Public/HttpServerConstants.h
@@ -2,95 +2,92 @@
 #pragma once
 #include "CoreMinimal.h"
 
-namespace EHttpServerResponseCodes
+enum class EHttpServerResponseCodes : uint16
 {
 	/**
 	 * Response codes that can come back from an Http request
 	 */
-	enum Type
-	{
-		// status code not set yet
-		Unknown = 0,
-		// the request can be continued.
-		Continue = 100,
-		// the server has switched protocols in an upgrade header.
-		SwitchProtocol = 101,
-		// the request completed successfully.
-		Ok = 200,
-		// the request has been fulfilled and resulted in the creation of a new resource.
-		Created = 201,
-		// the request has been accepted for processing, but the processing has not been completed.
-		Accepted = 202,
-		// the returned meta information in the entity-header is not the definitive set available from the origin server.
-		Partial = 203,
-		// the server has fulfilled the request, but there is no new information to send back.
-		NoContent = 204,
-		// the request has been completed, and the client program should reset the document view that caused the request to be sent to allow the user to easily initiate another input action.
-		ResetContent = 205,
-		// the server has fulfilled the partial get request for the resource.
-		PartialContent = 206,
-		// the server couldn't decide what to return.
-		Ambiguous = 300,
-		// the requested resource has been assigned to a new permanent uri (uniform resource identifier), and any future references to this resource should be done using one of the returned uris.
-		Moved = 301,
-		// the requested resource resides temporarily under a different uri (uniform resource identifier).
-		Redirect = 302,
-		// the response to the request can be found under a different uri (uniform resource identifier) and should be retrieved using a get http verb on that resource.
-		RedirectMethod = 303,
-		// the requested resource has not been modified.
-		NotModified = 304,
-		// the requested resource must be accessed through the proxy given by the location field.
-		UseProxy = 305,
-		// the redirected request keeps the same http verb. http/1.1 behavior.
-		RedirectKeepVerb = 307,
-		// the request could not be processed by the server due to invalid syntax.
-		BadRequest = 400,
-		// the requested resource requires user authentication.
-		Denied = 401,
-		// not currently implemented in the http protocol.
-		PaymentReq = 402,
-		// the server understood the request, but is refusing to fulfill it.
-		Forbidden = 403,
-		// the server has not found anything matching the requested uri (uniform resource identifier).
-		NotFound = 404,
-		// the http verb used is not allowed.
-		BadMethod = 405,
-		// no responses acceptable to the client were found.
-		NoneAcceptable = 406,
-		// proxy authentication required.
-		ProxyAuthReq = 407,
-		// the server timed out waiting for the request.
-		RequestTimeout = 408,
-		// the request could not be completed due to a conflict with the current state of the resource. the user should resubmit with more information.
-		Conflict = 409,
-		// the requested resource is no longer available at the server, and no forwarding address is known.
-		Gone = 410,
-		// the server refuses to accept the request without a defined content length.
-		LengthRequired = 411,
-		// the precondition given in one or more of the request header fields evaluated to false when it was tested on the server.
-		PrecondFailed = 412,
-		// the server is refusing to process a request because the request entity is larger than the server is willing or able to process.
-		RequestTooLarge = 413,
-		// the server is refusing to service the request because the request uri (uniform resource identifier) is longer than the server is willing to interpret.
-		UriTooLong = 414,
-		// the server is refusing to service the request because the entity of the request is in a format not supported by the requested resource for the requested method.
-		UnsupportedMedia = 415,
-		// too many requests, the server is throttling
-		TooManyRequests = 429,
-		// the request should be retried after doing the appropriate action.
-		RetryWith = 449,
-		// the server encountered an unexpected condition that prevented it from fulfilling the request.
-		ServerError = 500,
-		// the server does not support the functionality required to fulfill the request.
-		NotSupported = 501,
-		// the server, while acting as a gateway or proxy, received an invalid response from the upstream server it accessed in attempting to fulfill the request.
-		BadGateway = 502,
-		// the service is temporarily overloaded.
-		ServiceUnavail = 503,
-		// the request was timed out waiting for a gateway.
-		GatewayTimeout = 504,
-		// the server does not support, or refuses to support, the http protocol version that was used in the request message.
-		VersionNotSup = 505
-	};
-}
+	 // status code not set yet
+	Unknown = 0,
+	// the request can be continued.
+	Continue = 100,
+	// the server has switched protocols in an upgrade header.
+	SwitchProtocol = 101,
+	// the request completed successfully.
+	Ok = 200,
+	// the request has been fulfilled and resulted in the creation of a new resource.
+	Created = 201,
+	// the request has been accepted for processing, but the processing has not been completed.
+	Accepted = 202,
+	// the returned meta information in the entity-header is not the definitive set available from the origin server.
+	Partial = 203,
+	// the server has fulfilled the request, but there is no new information to send back.
+	NoContent = 204,
+	// the request has been completed, and the client program should reset the document view that caused the request to be sent to allow the user to easily initiate another input action.
+	ResetContent = 205,
+	// the server has fulfilled the partial get request for the resource.
+	PartialContent = 206,
+	// the server couldn't decide what to return.
+	Ambiguous = 300,
+	// the requested resource has been assigned to a new permanent uri (uniform resource identifier), and any future references to this resource should be done using one of the returned uris.
+	Moved = 301,
+	// the requested resource resides temporarily under a different uri (uniform resource identifier).
+	Redirect = 302,
+	// the response to the request can be found under a different uri (uniform resource identifier) and should be retrieved using a get http verb on that resource.
+	RedirectMethod = 303,
+	// the requested resource has not been modified.
+	NotModified = 304,
+	// the requested resource must be accessed through the proxy given by the location field.
+	UseProxy = 305,
+	// the redirected request keeps the same http verb. http/1.1 behavior.
+	RedirectKeepVerb = 307,
+	// the request could not be processed by the server due to invalid syntax.
+	BadRequest = 400,
+	// the requested resource requires user authentication.
+	Denied = 401,
+	// not currently implemented in the http protocol.
+	PaymentReq = 402,
+	// the server understood the request, but is refusing to fulfill it.
+	Forbidden = 403,
+	// the server has not found anything matching the requested uri (uniform resource identifier).
+	NotFound = 404,
+	// the http verb used is not allowed.
+	BadMethod = 405,
+	// no responses acceptable to the client were found.
+	NoneAcceptable = 406,
+	// proxy authentication required.
+	ProxyAuthReq = 407,
+	// the server timed out waiting for the request.
+	RequestTimeout = 408,
+	// the request could not be completed due to a conflict with the current state of the resource. the user should resubmit with more information.
+	Conflict = 409,
+	// the requested resource is no longer available at the server, and no forwarding address is known.
+	Gone = 410,
+	// the server refuses to accept the request without a defined content length.
+	LengthRequired = 411,
+	// the precondition given in one or more of the request header fields evaluated to false when it was tested on the server.
+	PrecondFailed = 412,
+	// the server is refusing to process a request because the request entity is larger than the server is willing or able to process.
+	RequestTooLarge = 413,
+	// the server is refusing to service the request because the request uri (uniform resource identifier) is longer than the server is willing to interpret.
+	UriTooLong = 414,
+	// the server is refusing to service the request because the entity of the request is in a format not supported by the requested resource for the requested method.
+	UnsupportedMedia = 415,
+	// too many requests, the server is throttling
+	TooManyRequests = 429,
+	// the request should be retried after doing the appropriate action.
+	RetryWith = 449,
+	// the server encountered an unexpected condition that prevented it from fulfilling the request.
+	ServerError = 500,
+	// the server does not support the functionality required to fulfill the request.
+	NotSupported = 501,
+	// the server, while acting as a gateway or proxy, received an invalid response from the upstream server it accessed in attempting to fulfill the request.
+	BadGateway = 502,
+	// the service is temporarily overloaded.
+	ServiceUnavail = 503,
+	// the request was timed out waiting for a gateway.
+	GatewayTimeout = 504,
+	// the server does not support, or refuses to support, the http protocol version that was used in the request message.
+	VersionNotSup = 505
+};
 
diff --git a/Engine/Source/Runtime/Online/HTTPServer/Public/HttpServerRequest.h b/Engine/Source/Runtime/Online/HTTPServer/Public/HttpServerRequest.h
--- a/Engine/Source/Runtime/Online/HTTPServer/Public/HttpServerRequest.h
+++ b/Engine/Source/Runtime/Online/HTTPServer/Public/HttpServerRequest.h
@@ -3,6 +3,7 @@
 
 #include "CoreMinimal.h"
 #include "HttpPath.h"
+#include "HttpServerHttpVersion.h"
 
 enum class EHttpServerRequestVerbs : uint16
 {
@@ -29,6 +30,9 @@ public:
 	/** The HTTP-compliant verb  */
 	EHttpServerRequestVerbs Verb;
 
+	/** The request HTTP protocol version */
+	HttpVersion::EHttpServerHttpVersion HttpVersion;
+
 	/** The HTTP headers */
 	TMap<FString, TArray<FString>> Headers;
 
diff --git a/Engine/Source/Runtime/Online/HTTPServer/Public/HttpServerResponse.h b/Engine/Source/Runtime/Online/HTTPServer/Public/HttpServerResponse.h
--- a/Engine/Source/Runtime/Online/HTTPServer/Public/HttpServerResponse.h
+++ b/Engine/Source/Runtime/Online/HTTPServer/Public/HttpServerResponse.h
@@ -1,6 +1,8 @@
 // Copyright 1998-2019 Epic Games, Inc. All Rights Reserved.
 #pragma once
 #include "CoreMinimal.h"
+#include "HttpServerConstants.h"
+#include "HttpServerHttpVersion.h"
 
 struct FHttpServerResponse final
 {
@@ -21,8 +23,11 @@ public:
 		: Body(MoveTemp(InBody))
 	{ }
 
+	/** Http protocol version */
+	HttpVersion::EHttpServerHttpVersion HttpVersion;
+
 	/** Http Response Code */
-	int32 Code;
+	EHttpServerResponseCodes Code;
 
 	/** Http Headers */
 	TMap<FString, TArray<FString>> Headers;
@@ -60,31 +65,20 @@ public:
 	HTTPSERVER_API static TUniquePtr<FHttpServerResponse> Create(const TArrayView<uint8>& RawBytes, FString ContentType);
 
 	/**
-	 * Creates an FHttpServerResponse from a response code
-	 *
-	 * @param  HttpResponseCode  The HTTP response/error code
-	 * @param  ErrorCode         The string error code/message
-	 * @return                   A unique pointer to an initialized response object
-	 */
-	HTTPSERVER_API static TUniquePtr<FHttpServerResponse> Create(int32 HttpResponseCode, const FString& ErrorCode);
-
-	/**
-	 * Creates an FHttpServerResponse 200
+	 * Creates an FHttpServerResponse 204
 	 * 
-	 * @param Message The optional and respective success message
-	 * @param ContentType The optional and respective content type
 	 * @return A unique pointer to an initialized response object
 	 */
-	HTTPSERVER_API static TUniquePtr<FHttpServerResponse> Ok(const FString& Message = TEXT(""), FString ContentType = TEXT("text/text"));
+	HTTPSERVER_API static TUniquePtr<FHttpServerResponse> Ok();
 
 	/**
-    * Creates an FHttpServerResponse 500
+    * Creates an FHttpServerResponse with the caller-supplied response and error codes
 	*
-	* @param Message The optional and respective error message
-	* @param ContentType The optional and respective content type
+	* @param ResponseCode The HTTP response code
+	* @param ErrorCode    The machine-readable error code
     * @return A unique pointer to an initialized response object
     */
-	HTTPSERVER_API static TUniquePtr<FHttpServerResponse> Error(const FString& Message = TEXT(""), FString ContentType = TEXT("text/text"));
+	HTTPSERVER_API static TUniquePtr<FHttpServerResponse> Error(EHttpServerResponseCodes ResponseCode, const FString& ErrorCode = TEXT(""));
 };
 
 
