[HttpServer] Missing file from shelf

diff --git a/Engine/Source/Runtime/Online/HTTPServer/Public/HttpServerHttpVersion.h b/Engine/Source/Runtime/Online/HTTPServer/Public/HttpServerHttpVersion.h
--- a/Engine/Source/Runtime/Online/HTTPServer/Public/HttpServerHttpVersion.h
+++ b/Engine/Source/Runtime/Online/HTTPServer/Public/HttpServerHttpVersion.h
@@ -34,7 +34,7 @@ namespace HttpVersion
 		case EHttpServerHttpVersion::HTTP_VERSION_1_0:
 			Result = TEXT("HTTP/1.0");
 			break;
-			EHttpServerHttpVersion::HTTP_VERSION_1_1;
+		case EHttpServerHttpVersion::HTTP_VERSION_1_1:
 			Result = TEXT("HTTP/1.1");
 			break;
 		default:
