[HttpServer] Compile fix - HttpServerConstants.h case matches the on-disk case.

diff --git a/Engine/Source/Runtime/Online/HTTPServer/Private/HttpConnection.h b/Engine/Source/Runtime/Online/HTTPServer/Private/HttpConnection.h
--- a/Engine/Source/Runtime/Online/HTTPServer/Private/HttpConnection.h
+++ b/Engine/Source/Runtime/Online/HTTPServer/Private/HttpConnection.h
@@ -6,7 +6,7 @@
 #include "HttpConnectionRequestReadContext.h"
 #include "HttpConnectionResponseWriteContext.h"
 #include "HttpResultCallback.h"
-#include "HttpserverConstants.h"
+#include "HttpServerConstants.h"
 #include "HttpServerHttpVersion.h"
 
 class FSocket;
