[EmbeddedCommunication] Fix very long delays due to websocket not waking game thread

diff --git a/Engine/Source/Runtime/Core/Private/Misc/EmbeddedCommunication.cpp b/Engine/Source/Runtime/Core/Private/Misc/EmbeddedCommunication.cpp
--- a/Engine/Source/Runtime/Core/Private/Misc/EmbeddedCommunication.cpp
+++ b/Engine/Source/Runtime/Core/Private/Misc/EmbeddedCommunication.cpp
@@ -43,6 +43,7 @@ void* FEmbeddedDelegates::GetNamedObject(const FString& Name)
 #if BUILD_EMBEDDED_APP
 
 static FEvent* GSleepEvent = nullptr;
+static TAtomic<bool> GTickAnotherFrame;
 static FCriticalSection GEmbeddedLock;
 //static TArray<TFunction<void()>> GEmbeddedQueues[5];
 static TQueue<TFunction<void()>> GEmbeddedQueues[5];
@@ -70,6 +71,7 @@ void FEmbeddedCommunication::Init()
 #if BUILD_EMBEDDED_APP
  	FScopeLock Lock(&GEmbeddedLock);
  	GSleepEvent = FPlatformProcess::GetSynchEventFromPool(false);
+	GTickAnotherFrame = false;
 
 	FTicker::GetCoreTicker().AddTicker(FTickerDelegate::CreateStatic(&FEmbeddedCommunication::TickGameThread));
 #endif
@@ -242,6 +244,8 @@ void FEmbeddedCommunication::RunOnGameThread(int Priority, TFunction<void()> Lam
 void FEmbeddedCommunication::WakeGameThread()
 {
 #if BUILD_EMBEDDED_APP
+	// If the game thread is already running, this will make it loop again.
+	GTickAnotherFrame = true;
 	// wake up the game thread!
 	if (GSleepEvent)
 	{
@@ -273,12 +277,15 @@ bool FEmbeddedCommunication::TickGameThread(float DeltaTime)
 	{
 		LambdaToCall();
 	}
-	// sleep if nothing is going on 
-	else if (GRenderingWakeMap.Num() == 0 && GTickWakeMap.Num() == 0)
+	// sleep if nothing is going on
+	else if (GRenderingWakeMap.Num() == 0
+		&& GTickWakeMap.Num() == 0
+		&& !GTickAnotherFrame) // Don't sleep if we have been asked to tick another frame
  	{
  		// wake up every 5 seconds even if nothing to do
  		GSleepEvent->Wait(5000);
  	}
+	GTickAnotherFrame = false;
 #endif
 
 	return true;
diff --git a/Engine/Source/Runtime/Online/WebSockets/Private/Lws/LwsWebSocket.cpp b/Engine/Source/Runtime/Online/WebSockets/Private/Lws/LwsWebSocket.cpp
--- a/Engine/Source/Runtime/Online/WebSockets/Private/Lws/LwsWebSocket.cpp
+++ b/Engine/Source/Runtime/Online/WebSockets/Private/Lws/LwsWebSocket.cpp
@@ -10,6 +10,7 @@
 #if WITH_SSL
 #include "Ssl.h"
 #endif
+#include "Misc/EmbeddedCommunication.h"
 #include "Misc/ScopeLock.h"
 #include "HttpModule.h"
 #include "HttpManager.h"
@@ -272,20 +273,27 @@ int FLwsWebSocket::LwsCallback(lws* Instance, lws_callback_reasons Reason, void*
 	{
 		const SIZE_T BytesLeft = lws_remaining_packet_payload(Instance);
 		UE_LOG(LogWebSockets, VeryVerbose, TEXT("FLwsWebSocket[%d]::LwsCallback: Received LWS_CALLBACK_CLIENT_RECEIVE Length=%d BytesLeft=%d"), Identifier, Length, BytesLeft);
+		bool bWakeGameThread = false;
 		if (bWantsMessageEvents)
 		{
 			FUTF8ToTCHAR Convert((const ANSICHAR*)Data, Length);
 			ReceiveBuffer.Append(Convert.Get(), Convert.Length());
 			if (BytesLeft == 0)
 			{
+				bWakeGameThread = true;
 				ReceiveTextQueue.Enqueue(MakeUnique<FLwsReceiveBufferText>(MoveTemp(ReceiveBuffer)));
 				ReceiveBuffer.Empty();
 			}
 		}
 		if (bWantsRawMessageEvents)
 		{
+			bWakeGameThread = true;
 			ReceiveBinaryQueue.Enqueue(MakeUnique<FLwsReceiveBufferBinary>(static_cast<const uint8*>(Data), Length, BytesLeft));
 		}
+		if (bWakeGameThread)
+		{
+			FEmbeddedCommunication::WakeGameThread();
+		}
 		break;
 	}
 	case LWS_CALLBACK_WS_PEER_INITIATED_CLOSE:
@@ -321,6 +329,7 @@ int FLwsWebSocket::LwsCallback(lws* Instance, lws_callback_reasons Reason, void*
 				ClosedReason.CloseStatus = CloseStatus;
 				ClosedReason.bWasClean = true;
 			}
+			FEmbeddedCommunication::WakeGameThread();
 		}
 		else
 		{
@@ -342,23 +351,30 @@ int FLwsWebSocket::LwsCallback(lws* Instance, lws_callback_reasons Reason, void*
 				State = EState::Error;
 				ClosedReason.Reason = MoveTemp(CloseReasonString);
 			}
+			FEmbeddedCommunication::WakeGameThread();
 		}
 		else
 		{
 			if (State == EState::Connected)
 			{
-				FScopeLock ScopeLock(&StateLock);
-				State = EState::Closed;
-				ClosedReason.bWasClean = false;
-				ClosedReason.CloseStatus = LWS_CLOSE_STATUS_ABNORMAL_CLOSE;
+				{
+					FScopeLock ScopeLock(&StateLock);
+					State = EState::Closed;
+					ClosedReason.bWasClean = false;
+					ClosedReason.CloseStatus = LWS_CLOSE_STATUS_ABNORMAL_CLOSE;
+				}
+				FEmbeddedCommunication::WakeGameThread();
 			}
 			else if (State == EState::ClosingByRequest)
 			{
-				FScopeLock ScopeLock(&StateLock);
-				State = EState::Closed;
-				ClosedReason.Reason = TEXT("Successfully closed connection to our peer");
-				ClosedReason.CloseStatus = LWS_CLOSE_STATUS_NORMAL;
-				ClosedReason.bWasClean = true;
+				{
+					FScopeLock ScopeLock(&StateLock);
+					State = EState::Closed;
+					ClosedReason.Reason = TEXT("Successfully closed connection to our peer");
+					ClosedReason.CloseStatus = LWS_CLOSE_STATUS_NORMAL;
+					ClosedReason.bWasClean = true;
+				}
+				FEmbeddedCommunication::WakeGameThread();
 			}
 			UE_LOG(LogWebSockets, Verbose, TEXT("FLwsWebSocket[%d]::LwsCallback: Received LWS_CALLBACK_WSI_DESTROY, State=%s"), Identifier, ToString(State));
 		}
@@ -385,6 +401,7 @@ int FLwsWebSocket::LwsCallback(lws* Instance, lws_callback_reasons Reason, void*
 				ClosedReason.CloseStatus = LWS_CLOSE_STATUS_NORMAL;
 				ClosedReason.bWasClean = bClosingByRequest;
 			}
+			FEmbeddedCommunication::WakeGameThread();
 		}
 		else
 		{
@@ -399,6 +416,7 @@ int FLwsWebSocket::LwsCallback(lws* Instance, lws_callback_reasons Reason, void*
 
 		FUTF8ToTCHAR Convert((const ANSICHAR*)Data, Length);
 		FString CloseReasonString(Convert.Length(), Convert.Get());
+		bool bWakeGameThread = false;
 		UE_LOG(LogWebSockets, Verbose, TEXT("FLwsWebSocket[%d]::LwsCallback: Received LWS_CALLBACK_CLIENT_CONNECTION_ERROR, setting State=%s CloseReason=%s PreviousState=%s"),
 			Identifier, ToString(EState::Error), *CloseReasonString, ToString(State));
 		{
@@ -409,13 +427,20 @@ int FLwsWebSocket::LwsCallback(lws* Instance, lws_callback_reasons Reason, void*
 				State = EState::Closed;
 				ClosedReason.bWasClean = false;
 				ClosedReason.CloseStatus = LWS_CLOSE_STATUS_ABNORMAL_CLOSE;
+				bWakeGameThread = true;
 			}
 			else if (State != EState::Closed)
 			{
 				State = EState::Error;
+				bWakeGameThread = true;
 			}
 			ClosedReason.Reason = MoveTemp(CloseReasonString);
 		}
+		if (bWakeGameThread)
+		{
+			FEmbeddedCommunication::WakeGameThread();
+		}
+
 		return -1;
 	}
 	case LWS_CALLBACK_RECEIVE_PONG:
@@ -476,6 +501,7 @@ int FLwsWebSocket::LwsCallback(lws* Instance, lws_callback_reasons Reason, void*
 	default:
 		break;
 	}
+
 	return 0;
 }
 
