[HttpServer] HttpConnectionResponseWriteContext - submitting shelved file that was otherwise out of sync with on-disk commit.

diff --git a/Engine/Source/Runtime/Online/HTTPServer/Private/HttpConnectionResponseWriteContext.cpp b/Engine/Source/Runtime/Online/HTTPServer/Private/HttpConnectionResponseWriteContext.cpp
--- a/Engine/Source/Runtime/Online/HTTPServer/Private/HttpConnectionResponseWriteContext.cpp
+++ b/Engine/Source/Runtime/Online/HTTPServer/Private/HttpConnectionResponseWriteContext.cpp
@@ -104,7 +104,7 @@ bool FHttpConnectionResponseWriteContext::IsWriteBodyComplete() const
 
 TArray<uint8> FHttpConnectionResponseWriteContext::SerializeHeadersUtf8(HttpVersion::EHttpServerHttpVersion HttpVersion, EHttpServerResponseCodes ResponseCode, const TMap<FString, TArray<FString>>& HeadersMap)
 {
-	FString ResponseHeaderStr = FString::Printf(TEXT("HTTP/%.1f %d\r\n"),  
+	FString ResponseHeaderStr = FString::Printf(TEXT("%s %d\r\n"),  
 		*HttpVersion::ToString(HttpVersion), ResponseCode);
 
 	for (const auto& KeyValuePair : HeadersMap)
