[HttpServer] HttpServer dependency is naturally induced, however http listeners are effectively disabled for Shipping-Client configuration
+ Additionally, HttpServerModule requires that StartAllListeners() is called prior to ticking internal http listeners.

diff --git a/Engine/Source/Runtime/Online/HTTPServer/HttpServer.Build.cs b/Engine/Source/Runtime/Online/HTTPServer/HttpServer.Build.cs
--- a/Engine/Source/Runtime/Online/HTTPServer/HttpServer.Build.cs
+++ b/Engine/Source/Runtime/Online/HTTPServer/HttpServer.Build.cs
@@ -6,10 +6,6 @@ public class HTTPServer : ModuleRules
 {
     public HTTPServer(ReadOnlyTargetRules Target) : base(Target)
     {
-        //OptimizeCode = CodeOptimization.Never;
-
-        PublicDefinitions.Add("HTTPSERVER_PACKAGE=1");
-
         PublicIncludePaths.AddRange(
 			new string[] {
 				"Runtime/Online/HTTPServer/Public",
diff --git a/Engine/Source/Runtime/Online/HTTPServer/Private/HttpServerModule.cpp b/Engine/Source/Runtime/Online/HTTPServer/Private/HttpServerModule.cpp
--- a/Engine/Source/Runtime/Online/HTTPServer/Private/HttpServerModule.cpp
+++ b/Engine/Source/Runtime/Online/HTTPServer/Private/HttpServerModule.cpp
@@ -28,18 +28,40 @@ void FHttpServerModule::ShutdownModule()
 	Listeners.Empty();
 }
 
+void FHttpServerModule::StartAllListeners()
+{
+	bHttpListenersEnabled = true;
+
+	UE_LOG(LogHttpServerModule, Log,
+		TEXT("Starting all listeners..."));
+
+	for (const auto& Listener : Listeners)
+	{
+		if (!Listener.Value->IsListening())
+		{
+			Listener.Value->StartListening();
+		}
+	}
+	UE_LOG(LogHttpServerModule, Log,
+		TEXT("All listeners started"));
+}
+
 void FHttpServerModule::StopAllListeners()
 {
+	UE_LOG(LogHttpServerModule, Log, 
+		TEXT("Stopping all listeners..."));
+
 	for (const auto& Listener : Listeners)
 	{
 		if (Listener.Value->IsListening())
 		{
-			UE_LOG(LogHttpServerModule, Log,
-				TEXT("Stopping listener on Port: %u"), Listener.Key);
-
 			Listener.Value->StopListening();
 		}
 	}
+
+	UE_LOG(LogHttpServerModule, Log,
+		TEXT("All listeners stopped"));
+
 }
 
 bool FHttpServerModule::HasPendingListeners() const 
@@ -67,10 +89,7 @@ FHttpServerModule& FHttpServerModule::Get()
 
 TSharedPtr<IHttpRouter> FHttpServerModule::GetHttpRouter(uint32 Port)
 {
-	if (!bInitialized)
-	{
-		return nullptr;
-	}
+	check(bInitialized);
 
 	// We may already be listening on this port
 	TUniquePtr<FHttpListener>* ExistingListener = Listeners.Find(Port);
@@ -78,27 +97,30 @@ TSharedPtr<IHttpRouter> FHttpServerModule::GetHttpRouter(uint32 Port)
 	{
 		return ExistingListener->Get()->GetRouter();
 	}
-	else
+
+	// Otherwise create a new one
+	TUniquePtr<FHttpListener> NewListener = MakeUnique<FHttpListener>(Port);
+
+    // Try to start this listener now
+	if (bHttpListenersEnabled)
 	{
-		// Otherwise create a new one
-		TUniquePtr<FHttpListener> NewListener = MakeUnique<FHttpListener>(Port);
-		if (NewListener->StartListening())
-		{
-			const auto& NewListenerRef = Listeners.Add(Port, MoveTemp(NewListener));
-			return NewListenerRef->GetRouter();
-		}
+		NewListener->StartListening();
 	}
-	return nullptr;
+	const auto& NewListenerRef = Listeners.Add(Port, MoveTemp(NewListener));
+	return NewListenerRef->GetRouter();
 }
 
 bool FHttpServerModule::Tick(float DeltaTime)
 {
 	check(Singleton == this);
+	check(bInitialized);
 
-	for (const auto& Listener : Listeners)
+	if (bHttpListenersEnabled)
 	{
-		check(Listener.Value);
-		Listener.Value->Tick(DeltaTime);
+		for (const auto& Listener : Listeners)
+		{
+			Listener.Value->Tick(DeltaTime);
+		}
 	}
 	return true;
 }
diff --git a/Engine/Source/Runtime/Online/HTTPServer/Public/HttpServerModule.h b/Engine/Source/Runtime/Online/HTTPServer/Public/HttpServerModule.h
--- a/Engine/Source/Runtime/Online/HTTPServer/Public/HttpServerModule.h
+++ b/Engine/Source/Runtime/Online/HTTPServer/Public/HttpServerModule.h
@@ -7,7 +7,6 @@
 #include "Modules/ModuleManager.h"
 #include "Containers/Ticker.h" 
 
-
 class FHttpListener;
 class IHttpRouter;
 
@@ -47,6 +46,11 @@ public:
 	 */
 	 bool Tick(float DeltaTime) override;
 
+	 /**
+	  * Starts all listeners
+	  */
+	 HTTPSERVER_API void StartAllListeners();
+
 	 /**
 	  * Stops all listeners
 	  */
@@ -85,5 +89,8 @@ private:
 
 	/** Whether this module has been initialized */
 	bool bInitialized = false;
+
+	/** Whether listeners can be started */
+	bool bHttpListenersEnabled = false;
 };
 
