Fix so that iOS BackgroundHttp no longer auto-resumes paused tasks.
#tests iOS Client

diff --git a/Engine/Source/Runtime/Online/BackgroundHTTP/Private/IOS/ApplePlatformBackgroundHttpManager.cpp b/Engine/Source/Runtime/Online/BackgroundHTTP/Private/IOS/ApplePlatformBackgroundHttpManager.cpp
--- a/Engine/Source/Runtime/Online/BackgroundHTTP/Private/IOS/ApplePlatformBackgroundHttpManager.cpp
+++ b/Engine/Source/Runtime/Online/BackgroundHTTP/Private/IOS/ApplePlatformBackgroundHttpManager.cpp
@@ -781,15 +781,21 @@ void FApplePlatformBackgroundHttpManager::TickTasks(float DeltaTime)
                                      FAppleBackgroundHttpRequestPtr FoundRequest = (nullptr != WeakRequestInMap) ? WeakRequestInMap->Pin() : nullptr;
                                      
                                      const bool bIsPaused = FoundRequest.IsValid() ? FoundRequest->IsUnderlyingTaskPaused() : false;
-                                     
-                                     if (FoundRequest.IsValid())
+                                     if (!bIsPaused)
                                      {
                                          UE_LOG(LogBackgroundHttpManager, Display, TEXT("Manager Calling to Active Task For Request -- RequestDebugID:%s | TaskURL:%s | TaskIdentifier:%d | CurrentlyActiveRequests:%d"), *(FoundRequest->GetRequestDebugID()), *TaskURL, TaskIdentifier, NewRequestCount);
                                          FoundRequest->ActivateUnderlyingTask();
                                      }
                                      else
                                      {
-                                         UE_LOG(LogBackgroundHttpManager, Display, TEXT("Skipping Activating Task as there is no associated Request or Request is paused. Once a Request associates with this task, it can then be activated. -- TaskURL:%s | TaskIdentifier:%d| bIsPaused:%d"), *TaskURL, TaskIdentifier, (int)bIsPaused);
+                                         if (bIsPaused)
+                                         {
+                                             UE_LOG(LogBackgroundHttpManager, Verbose, TEXT("Skipping Activating Task as the associated request is paused. -- TaskURL:%s | TaskIdentifier:%d"), *TaskURL, TaskIdentifier);
+                                         }
+                                         else
+                                         {
+                                             UE_LOG(LogBackgroundHttpManager, Display, TEXT("Skipping Activating Task as there is no associated Request. Once a Request associates with this task, it can then be activated. -- TaskURL:%s | TaskIdentifier:%d"), *TaskURL, TaskIdentifier);
+                                         }
                                          
                                          //Don't activate and remove our increment from above because something put us over the limit before we resumed
                                          FPlatformAtomics::InterlockedDecrement(&NumCurrentlyActiveTasks);
