Previous fix could crash when handling unassociated tasks. Also fixed seperate potential crash with uninitialized variable for debug name.
#tests iOS Client

diff --git a/Engine/Source/Runtime/Online/BackgroundHTTP/Private/IOS/ApplePlatformBackgroundHttpManager.cpp b/Engine/Source/Runtime/Online/BackgroundHTTP/Private/IOS/ApplePlatformBackgroundHttpManager.cpp
--- a/Engine/Source/Runtime/Online/BackgroundHTTP/Private/IOS/ApplePlatformBackgroundHttpManager.cpp
+++ b/Engine/Source/Runtime/Online/BackgroundHTTP/Private/IOS/ApplePlatformBackgroundHttpManager.cpp
@@ -781,14 +781,14 @@ void FApplePlatformBackgroundHttpManager::TickTasks(float DeltaTime)
                                      FAppleBackgroundHttpRequestPtr FoundRequest = (nullptr != WeakRequestInMap) ? WeakRequestInMap->Pin() : nullptr;
                                      
                                      const bool bIsPaused = FoundRequest.IsValid() ? FoundRequest->IsUnderlyingTaskPaused() : false;
-                                     if (!bIsPaused)
+                                     if (FoundRequest.IsValid() && !bIsPaused)
                                      {
                                          UE_LOG(LogBackgroundHttpManager, Display, TEXT("Manager Calling to Active Task For Request -- RequestDebugID:%s | TaskURL:%s | TaskIdentifier:%d | CurrentlyActiveRequests:%d"), *(FoundRequest->GetRequestDebugID()), *TaskURL, TaskIdentifier, NewRequestCount);
                                          FoundRequest->ActivateUnderlyingTask();
                                      }
                                      else
                                      {
-                                         if (bIsPaused)
+                                         if (FoundRequest.IsValid() && bIsPaused)
                                          {
                                              UE_LOG(LogBackgroundHttpManager, Verbose, TEXT("Skipping Activating Task as the associated request is paused. -- TaskURL:%s | TaskIdentifier:%d"), *TaskURL, TaskIdentifier);
                                          }
diff --git a/Engine/Source/Runtime/Online/BackgroundHTTP/Private/IOS/ApplePlatformBackgroundHttpRequest.cpp b/Engine/Source/Runtime/Online/BackgroundHTTP/Private/IOS/ApplePlatformBackgroundHttpRequest.cpp
--- a/Engine/Source/Runtime/Online/BackgroundHTTP/Private/IOS/ApplePlatformBackgroundHttpRequest.cpp
+++ b/Engine/Source/Runtime/Online/BackgroundHTTP/Private/IOS/ApplePlatformBackgroundHttpRequest.cpp
@@ -12,6 +12,7 @@ FApplePlatformBackgroundHttpRequest::FApplePlatformBackgroundHttpRequest()
 	, ResumeDataRetryCount(0)
     , FirstTask(nullptr)
 	, FirstTaskIdentifier(0)
+    , CombinedRequestID("")
     , bIsTaskActive(false)
     , bIsTaskPaused(false)
     , bIsCompleted(false)
