Submitted on behalf of thomas.ross

Updating to correct download complete notification.
#tests none

diff --git a/Engine/Source/Runtime/Online/BackgroundHTTP/Public/BackgroundHttpNotificationObject.h b/Engine/Source/Runtime/Online/BackgroundHTTP/Public/BackgroundHttpNotificationObject.h
--- a/Engine/Source/Runtime/Online/BackgroundHTTP/Public/BackgroundHttpNotificationObject.h
+++ b/Engine/Source/Runtime/Online/BackgroundHTTP/Public/BackgroundHttpNotificationObject.h
@@ -9,8 +9,8 @@ class BACKGROUNDHTTP_API FBackgroundHttpNotificationObject
 	: public TSharedFromThis<FBackgroundHttpNotificationObject, ESPMode::ThreadSafe>
 {
 public:
-	FBackgroundHttpNotificationObject(FText InNotificationTitle, FText InNotificationAction, FText InNotificationBody, const FString& InNotificationActivationString, bool InNotifyOnlyOnFullSuccess);
-	FBackgroundHttpNotificationObject(FText InNotificationTitle, FText InNotificationAction, FText InNotificationBody, const FString& InNotificationActivationString, bool InNotifyOnlyOnFullSuccess, bool bOnlySendNotificationInBackground);
+	FBackgroundHttpNotificationObject(FText InNotificationTitle, FText InNotificationBody, FText InNotificationAction, const FString& InNotificationActivationString, bool InNotifyOnlyOnFullSuccess);
+	FBackgroundHttpNotificationObject(FText InNotificationTitle, FText InNotificationBody, FText InNotificationAction, const FString& InNotificationActivationString, bool InNotifyOnlyOnFullSuccess, bool bOnlySendNotificationInBackground);
 	~FBackgroundHttpNotificationObject();
 
 	void NotifyOfDownloadResult(bool bWasSuccess);
