Fix for http url handling so that the web link before "?" is removed before parsing "&" separated value parameter pairs.

diff --git a/Engine/Source/Runtime/Online/HTTP/Private/Apple/AppleHTTP.cpp b/Engine/Source/Runtime/Online/HTTP/Private/Apple/AppleHTTP.cpp
--- a/Engine/Source/Runtime/Online/HTTP/Private/Apple/AppleHTTP.cpp
+++ b/Engine/Source/Runtime/Online/HTTP/Private/Apple/AppleHTTP.cpp
@@ -91,17 +91,23 @@ FString FAppleHttpRequest::GetURLParameter(const FString& ParameterName) const
 	SCOPED_AUTORELEASE_POOL;
 	UE_LOG(LogHttp, Verbose, TEXT("FAppleHttpRequest::GetURLParameter() - %s"), *ParameterName);
 
-	NSString* ParameterNameStr = ParameterName.GetNSString();
-	NSArray* Parameters = [Request.URL.query componentsSeparatedByString:@"&"];
-	for (NSString* Parameter in Parameters)
+	NSRange* ParametersStart = [Request.URL.query rangeOfString:@"?"];
+	if (ParametersStart.location != NSNotFound && ParametersStart.length > 0)
 	{
-		NSArray* KeyValue = [Parameter componentsSeparatedByString:@"="];
-		NSString* Key = KeyValue[0];
-		if ([Key compare:ParameterNameStr] == NSOrderedSame)
+		NSString* ParametersStr = [Request.URL.query substringFromIndex:ParametersStart.location + 1];
+		NSString* ParameterNameStr = ParameterName.GetNSString();
+		NSArray* Parameters = [ParametersStr componentsSeparatedByString:@"&"];
+		for (NSString* Parameter in Parameters)
 		{
-			return FString(KeyValue[1]);
+			NSArray* KeyValue = [Parameter componentsSeparatedByString:@"="];
+			NSString* Key = KeyValue[0];
+			if ([Key compare:ParameterNameStr] == NSOrderedSame)
+			{
+				return FString(KeyValue[1]);
+			}
 		}
 	}
+
 	return FString();
 }
 
diff --git a/Engine/Source/Runtime/Online/HTTP/Private/Curl/CurlHttp.cpp b/Engine/Source/Runtime/Online/HTTP/Private/Curl/CurlHttp.cpp
--- a/Engine/Source/Runtime/Online/HTTP/Private/Curl/CurlHttp.cpp
+++ b/Engine/Source/Runtime/Online/HTTP/Private/Curl/CurlHttp.cpp
@@ -196,27 +196,32 @@ FString FCurlHttpRequest::GetURLParameter(const FString& ParameterName) const
 {
 	TArray<FString> StringElements;
 
-	int32 NumElems = URL.ParseIntoArray(StringElements, TEXT("&"), true);
-	check(NumElems == StringElements.Num());
-	
-	FString ParamValDelimiter(TEXT("="));
-	for (int Idx = 0; Idx < NumElems; ++Idx )
-	{
-		FString Param, Value;
-		if (StringElements[Idx].Split(ParamValDelimiter, &Param, &Value) && Param == ParameterName)
+	//Parameters start after "?" in url
+	FString Path, Parameters;
+	if (URL.Split(TEXT("?"), &Path, &Parameters))
+	{
+		int32 NumElems = Parameters.ParseIntoArray(StringElements, TEXT("&"), true);
+		check(NumElems == StringElements.Num());
+		
+		FString ParamValDelimiter(TEXT("="));
+		for (int Idx = 0; Idx < NumElems; ++Idx )
 		{
-			// unescape
-			auto Converter = StringCast<ANSICHAR>(*Value);
-			char * EscapedAnsi = (char *)Converter.Get();
-			int32 EscapedLength = Converter.Length();
-
-			int32 UnescapedLength = 0;	
-			char * UnescapedAnsi = curl_easy_unescape(EasyHandle, EscapedAnsi, EscapedLength, &UnescapedLength);
-			
-			FString UnescapedValue(ANSI_TO_TCHAR(UnescapedAnsi));
-			curl_free(UnescapedAnsi);
-			
-			return UnescapedValue;
+			FString Param, Value;
+			if (StringElements[Idx].Split(ParamValDelimiter, &Param, &Value) && Param == ParameterName)
+			{
+				// unescape
+				auto Converter = StringCast<ANSICHAR>(*Value);
+				char * EscapedAnsi = (char *)Converter.Get();
+				int32 EscapedLength = Converter.Length();
+
+				int32 UnescapedLength = 0;	
+				char * UnescapedAnsi = curl_easy_unescape(EasyHandle, EscapedAnsi, EscapedLength, &UnescapedLength);
+				
+				FString UnescapedValue(ANSI_TO_TCHAR(UnescapedAnsi));
+				curl_free(UnescapedAnsi);
+				
+				return UnescapedValue;
+			}
 		}
 	}
 
diff --git a/Engine/Source/Runtime/Online/HTTP/Private/HTML5/HTML5HTTP.cpp b/Engine/Source/Runtime/Online/HTTP/Private/HTML5/HTML5HTTP.cpp
--- a/Engine/Source/Runtime/Online/HTTP/Private/HTML5/HTML5HTTP.cpp
+++ b/Engine/Source/Runtime/Online/HTTP/Private/HTML5/HTML5HTTP.cpp
@@ -56,16 +56,21 @@ FString FHTML5HttpRequest::GetURLParameter(const FString& ParameterName) const
 
 	TArray<FString> StringElements;
 
-	int32 NumElems = URL.ParseIntoArray(StringElements, TEXT("&"), true);
-	check(NumElems == StringElements.Num());
-
-	FString ParamValDelimiter(TEXT("="));
-	for (int Idx = 0; Idx < NumElems; ++Idx )
+	//Parameters start after "?" in url
+	FString Path, Parameters;
+	if (URL.Split(TEXT("?"), &Path, &Parameters))
 	{
-		FString Param, Value;
-		if (StringElements[Idx].Split(ParamValDelimiter, &Param, &Value) && Param == ParameterName)
+		int32 NumElems = Parameters.ParseIntoArray(StringElements, TEXT("&"), true);
+		check(NumElems == StringElements.Num());
+
+		FString ParamValDelimiter(TEXT("="));
+		for (int Idx = 0; Idx < NumElems; ++Idx )
 		{
-			return Value;
+			FString Param, Value;
+			if (StringElements[Idx].Split(ParamValDelimiter, &Param, &Value) && Param == ParameterName)
+			{
+				return Value;
+			}
 		}
 	}
 
