Fixed a typo not pulled accross from build test for AppleHTTP url parameter handling.

diff --git a/Engine/Source/Runtime/Online/HTTP/Private/Apple/AppleHTTP.cpp b/Engine/Source/Runtime/Online/HTTP/Private/Apple/AppleHTTP.cpp
--- a/Engine/Source/Runtime/Online/HTTP/Private/Apple/AppleHTTP.cpp
+++ b/Engine/Source/Runtime/Online/HTTP/Private/Apple/AppleHTTP.cpp
@@ -91,7 +91,7 @@ FString FAppleHttpRequest::GetURLParameter(const FString& ParameterName) const
 	SCOPED_AUTORELEASE_POOL;
 	UE_LOG(LogHttp, Verbose, TEXT("FAppleHttpRequest::GetURLParameter() - %s"), *ParameterName);
 
-	NSRange* ParametersStart = [Request.URL.query rangeOfString:@"?"];
+	NSRange ParametersStart = [Request.URL.query rangeOfString:@"?"];
 	if (ParametersStart.location != NSNotFound && ParametersStart.length > 0)
 	{
 		NSString* ParametersStr = [Request.URL.query substringFromIndex:ParametersStart.location + 1];
