Unshelved from pending changelist '7208189':

PR #5995: HTTP: Fix timeout on waiting for pending requests on exit not working properly (Contributed by [REDACTED_PERSON])

diff --git a/Engine/Source/Runtime/Online/HTTP/Private/HttpManager.cpp b/Engine/Source/Runtime/Online/HTTP/Private/HttpManager.cpp
--- a/Engine/Source/Runtime/Online/HTTP/Private/HttpManager.cpp
+++ b/Engine/Source/Runtime/Online/HTTP/Private/HttpManager.cpp
@@ -142,7 +142,7 @@ void FHttpManager::Flush(bool bShutdown)
 	{
 		const double AppTime = FPlatformTime::Seconds();
 		//UE_LOG(LogHttp, Display, TEXT("Waiting for %0.2f seconds. Limit:%0.2f seconds"), (AppTime - BeginWaitTime), MaxFlushTimeSeconds);
-		if (bShutdown && MaxFlushTimeSeconds > 0 && (GIsRequestingExit || (AppTime - BeginWaitTime > MaxFlushTimeSeconds)))
+		if (bShutdown && MaxFlushTimeSeconds > 0 && (AppTime - BeginWaitTime > MaxFlushTimeSeconds))
 		{
 			UE_LOG(LogHttp, Display, TEXT("Canceling remaining HTTP requests after waiting %0.2f seconds"), (AppTime - BeginWaitTime));
 			for (TArray<TSharedRef<IHttpRequest>>::TIterator It(Requests); It; ++It)
