Guarded use of the xmpp context pointer during Tick handling in FStropheWebsocketConnection.
Some seemingly rare instances of this context pointer being null when it should be valid.

Previous sanity checks for null contexts during construction exist but will allow a nullptr for the context.

diff --git a/Engine/Source/Runtime/Online/XMPP/Private/XmppStrophe/StropheWebsocketConnection.cpp b/Engine/Source/Runtime/Online/XMPP/Private/XmppStrophe/StropheWebsocketConnection.cpp
--- a/Engine/Source/Runtime/Online/XMPP/Private/XmppStrophe/StropheWebsocketConnection.cpp
+++ b/Engine/Source/Runtime/Online/XMPP/Private/XmppStrophe/StropheWebsocketConnection.cpp
@@ -146,7 +146,10 @@ bool FStropheWebsocketConnection::SendStanza(FStropheStanza&& Stanza)
 void FStropheWebsocketConnection::Tick()
 {
 	// Handles timeouts, processes timers
-	xmpp_run_once(Context.GetContextPtr(), 0);
+	if (ensureMsgf(Context.GetContextPtr() != nullptr, TEXT("xmpp_ctx_t is null")))
+	{
+		xmpp_run_once(Context.GetContextPtr(), 0);
+	}
 }
 
 void FStropheWebsocketConnection::Disconnect()
