Fix for Mac TVOS.  libOpus does not work with TVOS.

diff --git a/Engine/Source/Runtime/Online/Voice/Private/VoicePrivate.h b/Engine/Source/Runtime/Online/Voice/Private/VoicePrivate.h
--- a/Engine/Source/Runtime/Online/Voice/Private/VoicePrivate.h
+++ b/Engine/Source/Runtime/Online/Voice/Private/VoicePrivate.h
@@ -24,6 +24,6 @@ THIRD_PARTY_INCLUDES_END
 
 #define ANDROIDVOICE_SUPPORTED_PLATFORMS (PLATFORM_ANDROID && (PLATFORM_ANDROID_ARM || PLATFORM_ANDROID_ARM64 || PLATFORM_ANDROID_X64) && !PLATFORM_LUMIN)
 #define PLATFORM_SUPPORTS_VOICE_CAPTURE (PLATFORM_WINDOWS || PLATFORM_MAC || ANDROIDVOICE_SUPPORTED_PLATFORMS || (PLATFORM_LINUX && VOICE_MODULE_WITH_CAPTURE))
-#define PLATFORM_SUPPORTS_OPUS_CODEC (!PLATFORM_HTML5)
+#define PLATFORM_SUPPORTS_OPUS_CODEC (!PLATFORM_HTML5 && !PLATFORM_TVOS)
 
 // Module includes
