UDN #472564 - CurlHttp merges headers on redirect

store only latest headers when [ moved permanently, temporary/permanent redirect ] detected

#rn

diff --git a/Engine/Source/Runtime/Online/HTTP/Private/Curl/CurlHttp.cpp b/Engine/Source/Runtime/Online/HTTP/Private/Curl/CurlHttp.cpp
--- a/Engine/Source/Runtime/Online/HTTP/Private/Curl/CurlHttp.cpp
+++ b/Engine/Source/Runtime/Online/HTTP/Private/Curl/CurlHttp.cpp
@@ -9,6 +9,7 @@
 #include "Misc/Paths.h"
 #include "Curl/CurlHttpManager.h"
 #include "Misc/ScopeLock.h"
+#include "Internationalization/Regex.h"
 
 #if WITH_LIBCURL
 
@@ -60,6 +61,7 @@ FCurlHttpRequest::FCurlHttpRequest()
 	,	HeaderList(NULL)
 	,	bCanceled(false)
 	,	bCurlRequestCompleted(false)
+	,	bRedirected(false)
 	,	CurlAddToMultiResult(CURLM_OK)
 	,	CurlCompletionResult(CURLE_OK)
 	,	CompletionStatus(EHttpRequestStatus::NotStarted)
@@ -388,7 +390,7 @@ size_t FCurlHttpRequest::ReceiveResponseHeaderCallback(void* Ptr, size_t SizeInB
 			if (Header.Split(TEXT(":"), &HeaderKey, &HeaderValue))
 			{
 				HeaderValue.TrimStartInline();
-				if (!HeaderKey.IsEmpty() && !HeaderValue.IsEmpty())
+				if (!HeaderKey.IsEmpty() && !HeaderValue.IsEmpty() && !bRedirected)
 				{
 					//Store the content length so OnRequestProgress() delegates have something to work with
 					if (HeaderKey == TEXT("Content-Length"))
@@ -398,6 +400,20 @@ size_t FCurlHttpRequest::ReceiveResponseHeaderCallback(void* Ptr, size_t SizeInB
 					Response->NewlyReceivedHeaders.Enqueue(TPair<FString, FString>(MoveTemp(HeaderKey), MoveTemp(HeaderValue)));
 				}
 			}
+			else
+			{
+				const FRegexPattern StatusPattern(TEXT("HTTP/\\d+\\.\\d+ 30\\d")); // \d+ is jic http status version ever goes double digits...
+				FRegexMatcher StatusMatcher(StatusPattern, *Header);
+
+				if (StatusMatcher.FindNext())
+				{
+					bRedirected = true;
+				}
+				else if (Header.IsEmpty()) // an empty line notes the end of the headers
+				{
+					bRedirected = false;
+				}
+			}
 			return HeaderSize;
 		}
 		else
diff --git a/Engine/Source/Runtime/Online/HTTP/Private/Curl/CurlHttp.h b/Engine/Source/Runtime/Online/HTTP/Private/Curl/CurlHttp.h
--- a/Engine/Source/Runtime/Online/HTTP/Private/Curl/CurlHttp.h
+++ b/Engine/Source/Runtime/Online/HTTP/Private/Curl/CurlHttp.h
@@ -326,6 +326,8 @@ private:
 	bool			bCanceled;
 	/** Set to true when request has been completed */
 	bool			bCurlRequestCompleted;
+	/** Set to true when request has "30* Multiple Choices" (e.g. 301 Moved Permanently, 302 temporary redirect, 308 Permanent Redirect, etc.) */
+	bool			bRedirected;
 	/** Set to true if request failed to be added to curl multi */
 	CURLMcode		CurlAddToMultiResult;
 	/** Operation result code as returned by libcurl */
