Adding background http protype system. Enables downloading files while app is suspended on iOS. Other platforms need to be added over time if desired. For now other platforms just wrap basic HTTPRequests under the hood.
#tests iOS / Android

diff --git a/Engine/Config/IOS/IOSEngine.ini b/Engine/Config/IOS/IOSEngine.ini
--- a/Engine/Config/IOS/IOSEngine.ini
+++ b/Engine/Config/IOS/IOSEngine.ini
@@ -55,5 +55,15 @@ PoolSizeVRAMPercentage=70
 [HTTP]
 MaxFlushTimeSeconds=2.0
 
+[BackgroundHttp.iOSSettings]
+;While the app is in the forground, time out if we don't get a response in this timeframe
+BackgroundHttp.ActiveReceiveTimeout=30
+;While the app is in the background, time out if we don't get a response in this timeframe
+BackgroundHttp.BackgroundReceiveTimeout=120
+;Regardless of app state, if the download hasn't finished in this long, time out
+BackgroundHttp.ResourceTimeout=3600
+;If retry data exists for a task, how many times we should try using the Retry data before moving onto a different CDN. ( -1 = unlimited )
+BackgroundHttp.RetryResumeDataLimit=3;
+
 [Plugins]
 +EnabledPlugins=Crashlytics
\ No newline at end of file
diff --git a/Engine/Source/Runtime/ApplicationCore/Private/IOS/IOSAppDelegate.cpp b/Engine/Source/Runtime/ApplicationCore/Private/IOS/IOSAppDelegate.cpp
--- a/Engine/Source/Runtime/ApplicationCore/Private/IOS/IOSAppDelegate.cpp
+++ b/Engine/Source/Runtime/ApplicationCore/Private/IOS/IOSAppDelegate.cpp
@@ -9,6 +9,7 @@
 #include "Async/TaskGraphInterfaces.h"
 #include "GenericPlatform/GenericPlatformChunkInstall.h"
 #include "IOS/IOSPlatformMisc.h"
+#include "IOS/IOSBackgroundURLSessionHandler.h"
 #include "HAL/PlatformStackWalk.h"
 #include "HAL/PlatformProcess.h"
 #include "Misc/OutputDeviceError.h"
@@ -167,6 +168,8 @@ bool FIOSCoreDelegates::PassesPushNotificationFilters(NSDictionary* Payload)
 
 @synthesize savedOpenUrlParameters;
 
+@synthesize BackgroundSessionEventCompleteDelegate;
+
 -(void)dealloc
 {
 #if !UE_BUILD_SHIPPING && !PLATFORM_TVOS
@@ -1301,11 +1304,19 @@ extern double GCStartTime;
     completionHandler(UIBackgroundFetchResultNewData);
 }
 
-- (void)application:(UIApplication *)application handleEventsForBackgroundURLSession:(NSString *)identifier completionHandler:(void (^)(void))completionHandler
+#endif
+
+#if !PLATFORM_TVOS
+- (void)application:(UIApplication *)application handleEventsForBackgroundURLSession:(NSString *)identifier completionHandler:(void(^)(void))completionHandler
 {
-    FString Id(identifier);
-    FCoreDelegates::ApplicationBackgroundSessionEventDelegate.Broadcast(Id);
-    completionHandler();
+    //Save off completionHandler so that a future call to FCoreDelegates::ApplicationBackgroundSessionEventsAllSentDelegate can execute it
+    self.BackgroundSessionEventCompleteDelegate = completionHandler;
+    
+    //Create background session with this identifier if needed to handle these events
+	FString Id(identifier);
+	FBackgroundURLSessionHandler::InitBackgroundSession(Id);
+
+	FCoreDelegates::ApplicationBackgroundSessionEventDelegate.Broadcast(Id);
 }
 #endif
 
diff --git a/Engine/Source/Runtime/ApplicationCore/Private/IOS/IOSBackgroundURLSessionHandler.cpp b/Engine/Source/Runtime/ApplicationCore/Private/IOS/IOSBackgroundURLSessionHandler.cpp
new file mode 100644
--- /dev/null
+++ b/Engine/Source/Runtime/ApplicationCore/Private/IOS/IOSBackgroundURLSessionHandler.cpp
@@ -0,0 +1,202 @@
+// Copyright 1998-2018 Epic Games, Inc. All Rights Reserved.
+
+#include "IOS/IOSBackgroundURLSessionHandler.h"
+#include "Misc/ConfigCacheIni.h"
+#include "Misc/CoreDelegates.h"
+
+#include "HAL/PlatformFile.h"
+#include "HAL/PlatformFilemanager.h"
+
+#include "IOS/IOSAppDelegate.h"
+
+FIOSBackgroundDownloadCoreDelegates::FIOSBackgroundDownload_DidFinishDownloadingToURL FIOSBackgroundDownloadCoreDelegates::OnIOSBackgroundDownload_DidFinishDownloadingToURL;
+FIOSBackgroundDownloadCoreDelegates::FIOSBackgroundDownload_DidWriteData FIOSBackgroundDownloadCoreDelegates::OnIOSBackgroundDownload_DidWriteData;
+FIOSBackgroundDownloadCoreDelegates::FIOSBackgroundDownload_DidCompleteWithError FIOSBackgroundDownloadCoreDelegates::OnIOSBackgroundDownload_DidCompleteWithError;
+FIOSBackgroundDownloadCoreDelegates::FIOSBackgroundDownload_SessionDidFinishAllEvents FIOSBackgroundDownloadCoreDelegates::OnIOSBackgroundDownload_SessionDidFinishAllEvents;
+
+NSURLSession* FBackgroundURLSessionHandler::BackgroundSession = nil;
+FString FBackgroundURLSessionHandler::CachedIdentifierName = FString();
+
+bool FBackgroundURLSessionHandler::InitBackgroundSession(const FString& SessionIdentifier)
+{
+	bool bHasSessionWithIdentifier = false;
+
+	if (!SessionIdentifier.IsEmpty())
+	{
+		//if we don't already have a background session, create it
+		if (nil == BackgroundSession)
+		{
+			//Create actual session
+			{
+				float BackgroundHttpTimeout = 90.0f;
+				GConfig->GetFloat(TEXT("BackgroundHttp.iOSSettings"), TEXT("BackgroundHttp.BackgroundReceiveTimeout"), BackgroundHttpTimeout, GEngineIni);
+
+                float BackgroundHttpResourceTimeout = 3600.f;
+                GConfig->GetFloat(TEXT("BackgroundHttp.iOSSettings"), TEXT("BackgroundHttp.BackgroundHttpResourceTimeout"), BackgroundHttpResourceTimeout, GEngineIni);
+                
+				NSURLSessionConfiguration *config = [NSURLSessionConfiguration backgroundSessionConfigurationWithIdentifier : SessionIdentifier.GetNSString()];
+				config.discretionary = NO;
+				config.sessionSendsLaunchEvents = YES;
+				config.shouldUseExtendedBackgroundIdleMode = YES;
+				config.timeoutIntervalForRequest = BackgroundHttpTimeout;	        //Timeout for requests not getting a response from server at the iOS level. Will still work in the BG.
+				config.timeoutIntervalForResource = BackgroundHttpResourceTimeout;  //Will automatically retry everytime above setting causes a timeout until this many seconds pass. Then we will get
+                                                                                    //an error result. No Error for server not responding until is sent by a background task until this timeout is reached.
+                //Other unused settings we might want to configure
+				//HTTPMaximumConnectionsPerHost
+				//allowsCellularAccess
+                
+                //Can't use on BG session, but may need to consider for any future FG session work:
+				//waitsForConnectivity  -- This is forced true on background sessions
+				//URLSession:taskIsWaitingForConnectivity: method of NSURLSessionTaskDelegate -- Can't respond to this with background sessions since the above is always forced true they never call this delegate
+
+                //NSURLSessionConfiguration* config = [NSURLSessionConfiguration defaultSessionConfiguration];
+                BackgroundSession = [NSURLSession sessionWithConfiguration:config delegate:[[BackgroundDownloadDelegate alloc]init] delegateQueue:nil];
+				[BackgroundSession retain];
+			}
+
+			//test if session created
+			if (nil != BackgroundSession)
+			{
+				CachedIdentifierName = SessionIdentifier;
+				bHasSessionWithIdentifier = true;
+
+				//Make sure our temp directory exists where we will be putting files as they complete as we might get completion callbacks right after we init this session
+				CreateBackgroundSessionWorkingDirectory();
+			}
+		}
+		//we already have one, so just see if names match
+		else
+		{
+			bHasSessionWithIdentifier = CachedIdentifierName.Equals(SessionIdentifier);
+		}
+	}
+
+	return ensureMsgf(bHasSessionWithIdentifier, TEXT("Unable to initialize background session with identifier %s"), *SessionIdentifier);
+}
+
+void FBackgroundURLSessionHandler::ShutdownBackgroundSession(bool bShouldFinishTasksFirst /* = true */)
+{
+	if (nil != BackgroundSession)
+	{
+		if (bShouldFinishTasksFirst)
+		{
+			[BackgroundSession finishTasksAndInvalidate];
+		}
+		else
+		{
+			[BackgroundSession invalidateAndCancel];
+		}
+
+		[BackgroundSession release];
+		BackgroundSession = nil;
+	}
+
+	CachedIdentifierName.Reset();
+}
+
+NSURLSession* FBackgroundURLSessionHandler::GetBackgroundSession()
+{
+	ensureAlwaysMsgf((nullptr != BackgroundSession), TEXT("Call to GetBackgroundSession with no valid BackgroundSession! Make sure InitBackgroundSession is called first!"));
+	return BackgroundSession;
+}
+
+void FBackgroundURLSessionHandler::CreateBackgroundSessionWorkingDirectory()
+{
+	const FString DirectoryPath = GetBackgroundSessionWorkingDirectoryPath();
+    
+	if (!DirectoryPath.IsEmpty())
+	{
+        IPlatformFile& PlatformFile = FPlatformFileManager::Get().GetPlatformFile();
+        PlatformFile.CreateDirectory(*DirectoryPath);
+        
+        if (!PlatformFile.DirectoryExists(*DirectoryPath))
+        {
+            NSString* PathAsNSString = DirectoryPath.GetNSString();
+            //May not yet have UE_Log
+            NSLog(@"Warning: Error in CreateBackgroundSessionWorkingDirectory: %s", [PathAsNSString UTF8String]);
+        }
+	}
+}
+
+FString FBackgroundURLSessionHandler::GetBackgroundSessionWorkingDirectoryPath()
+{
+	return FPaths::Combine(FPlatformMisc::GamePersistentDownloadDir(), TEXT("BackgroundHttp"));
+}
+
+const FString FBackgroundURLSessionHandler::GetTemporaryFilePathFromURL(const FString& URL)
+{
+	static FString BackgroundHttpDir = GetBackgroundSessionWorkingDirectoryPath();
+	const FString FileName = FPaths::MakeValidFileName(URL);
+
+	return FPaths::Combine(BackgroundHttpDir, FileName);
+}
+
+@implementation BackgroundDownloadDelegate
+	
+-(void)URLSession:(NSURLSession *)session downloadTask : (NSURLSessionDownloadTask *)downloadTask didFinishDownloadingToURL : (NSURL *)location;
+{
+	//Go ahead and try and move the file out of IOS system temporary storage and into our own background download specific temp storage that won't be nuked by iOS
+	//We do this here because we need to handle this case even if nothing has subscribed to the delegate yet
+	
+	NSURL* destinationURL = nil;
+	NSError* downloadError = [downloadTask error];
+    const bool didTaskError = (downloadError != nullptr);
+    
+    FString TaskURL([[[downloadTask currentRequest] URL] absoluteString]);
+    FString DestinationPath = FBackgroundURLSessionHandler::GetTemporaryFilePathFromURL(TaskURL);
+    
+	if (!didTaskError && (nil != location))
+	{
+        //Convert UFS path to external path so we can pass it into iOS function to move iOS temp file
+        IPlatformFile& PlatformFile = FPlatformFileManager::Get().GetPlatformFile();
+        FString ConvertedPath = PlatformFile.ConvertToAbsolutePathForExternalAppForWrite(*DestinationPath);
+        
+		destinationURL = [NSURL fileURLWithPath:ConvertedPath.GetNSString()];
+
+		[[NSFileManager defaultManager] moveItemAtURL:location toURL:destinationURL error:&downloadError];
+	}
+    
+	FIOSBackgroundDownloadCoreDelegates::OnIOSBackgroundDownload_DidFinishDownloadingToURL.Broadcast(downloadTask, downloadError, DestinationPath);
+}
+
+-(void)URLSession:(NSURLSession *)session downloadTask : (NSURLSessionDownloadTask *)downloadTask didWriteData:(int64_t)bytesWritten totalBytesWritten:(int64_t)totalBytesWritten totalBytesExpectedToWrite:(int64_t)totalBytesExpectedToWrite;
+{	
+	FIOSBackgroundDownloadCoreDelegates::OnIOSBackgroundDownload_DidWriteData.Broadcast(downloadTask, bytesWritten, totalBytesWritten, totalBytesExpectedToWrite);
+}
+
+-(void)URLSession:(NSURLSession *)session task : (NSURLSessionTask *)task didCompleteWithError : (NSError *)error;
+{
+    NSString* TaskURLString = [[[task currentRequest] URL] absoluteString];
+    NSLog(@"Task ending with error. TaskURL:%s", [TaskURLString UTF8String]);
+
+    FIOSBackgroundDownloadCoreDelegates::OnIOSBackgroundDownload_DidCompleteWithError.Broadcast(task, error);
+}
+
+-(void)URLSessionDidFinishEventsForBackgroundURLSession:(NSURLSession *)session;
+{
+	//First process other systems doing things for all events processing on URL Session
+	FIOSBackgroundDownloadCoreDelegates::OnIOSBackgroundDownload_SessionDidFinishAllEvents.Broadcast(session);
+
+    IOSAppDelegate* AppDelegate = [IOSAppDelegate GetDelegate];
+    
+    if ((nullptr != AppDelegate) && (nil != AppDelegate.BackgroundSessionEventCompleteDelegate))
+    {
+        //make a local copy of the completion handler to capture with the block
+        void(^completionHandler)() = AppDelegate.BackgroundSessionEventCompleteDelegate;
+        [completionHandler retain];
+        
+        //set the completion handler to nil to protect from anything else trying to use it
+        AppDelegate.BackgroundSessionEventCompleteDelegate = nil;
+        
+       //Need to issue our stored completion callback on main thread as its a UI callback
+        [[NSOperationQueue mainQueue] addOperationWithBlock:^{
+            //call the copied completion handler
+            completionHandler();
+            
+            //Release local copy so we don't leak the completion handler
+            [completionHandler release];
+        }];
+    }
+}
+
+@end
diff --git a/Engine/Source/Runtime/ApplicationCore/Public/IOS/IOSAppDelegate.h b/Engine/Source/Runtime/ApplicationCore/Public/IOS/IOSAppDelegate.h
--- a/Engine/Source/Runtime/ApplicationCore/Public/IOS/IOSAppDelegate.h
+++ b/Engine/Source/Runtime/ApplicationCore/Public/IOS/IOSAppDelegate.h
@@ -186,6 +186,9 @@ UITextFieldDelegate>
 @property (atomic) bool bHasStarted;
 - (void)ToggleSuspend:(bool)bSuspend;
 
+
+@property (nonatomic, copy) void(^BackgroundSessionEventCompleteDelegate)();
+
 static void interruptionListener(void* ClientData, UInt32 Interruption);
 
 -(UIWindow*)window;
diff --git a/Engine/Source/Runtime/ApplicationCore/Public/IOS/IOSBackgroundURLSessionHandler.h b/Engine/Source/Runtime/ApplicationCore/Public/IOS/IOSBackgroundURLSessionHandler.h
new file mode 100644
--- /dev/null
+++ b/Engine/Source/Runtime/ApplicationCore/Public/IOS/IOSBackgroundURLSessionHandler.h
@@ -0,0 +1,57 @@
+// Copyright 1998-2018 Epic Games, Inc. All Rights Reserved.
+
+#pragma once
+
+#include "CoreTypes.h"
+#include "CoreMinimal.h"
+
+//Call backs called by the bellow FBackgroundURLSessionHandler so higher-level systems can respond to task updates.
+class APPLICATIONCORE_API FIOSBackgroundDownloadCoreDelegates
+{
+public:
+    DECLARE_MULTICAST_DELEGATE_ThreeParams(FIOSBackgroundDownload_DidFinishDownloadingToURL, NSURLSessionDownloadTask*, NSError*, const FString&);
+    DECLARE_MULTICAST_DELEGATE_FourParams(FIOSBackgroundDownload_DidWriteData, NSURLSessionDownloadTask*, int64_t /*Bytes Written Since Last Call */, int64_t /*Total Bytes Written */, int64_t /*Total Bytes Expedted To Write */);
+    DECLARE_MULTICAST_DELEGATE_TwoParams(FIOSBackgroundDownload_DidCompleteWithError, NSURLSessionTask*, NSError*);
+    DECLARE_MULTICAST_DELEGATE_OneParam(FIOSBackgroundDownload_SessionDidFinishAllEvents, NSURLSession*);
+    
+	static FIOSBackgroundDownload_DidFinishDownloadingToURL OnIOSBackgroundDownload_DidFinishDownloadingToURL;
+	static FIOSBackgroundDownload_DidWriteData OnIOSBackgroundDownload_DidWriteData;
+	static FIOSBackgroundDownload_DidCompleteWithError OnIOSBackgroundDownload_DidCompleteWithError;
+	static FIOSBackgroundDownload_SessionDidFinishAllEvents OnIOSBackgroundDownload_SessionDidFinishAllEvents;
+};
+
+//Interface for wrapping a NSURLSession configured to support background downloading of NSURLSessionDownloadTasks.
+//This exists here as we can have to re-associate with our background session after app launch and need to re-associate with downloads
+//right away before the HttpModule is loaded.
+class APPLICATIONCORE_API FBackgroundURLSessionHandler
+{
+public:
+	// Initializes a BackgroundSession with the given identifier. If the current background session already exists, returns true if the identifier matches. False if identifier doesn't match or if the session fails to create.
+	static bool InitBackgroundSession(const FString& SessionIdentifier);
+
+	//bShouldInvalidateExistingTasks determines if the session cancells all outstanding tasks immediately and cancels the session immediately or waits for them to finish and then invalidates the session
+	static void ShutdownBackgroundSession(bool bShouldFinishTasksFirst = true);
+
+	//Gets a pointer to the current background session
+	static NSURLSession* GetBackgroundSession();
+
+	//Gets a path stored in an FString of the working directory any completed temp files will be downloaded
+	static FString GetBackgroundSessionWorkingDirectoryPath();
+
+	/**
+	* Function that takes in a URL and figures out the location we should use as the temp storage URL
+	*
+	* @return FString to use as the TempFilePath
+	*/
+	static const FString GetTemporaryFilePathFromURL(const FString& URL);
+
+	static void CreateBackgroundSessionWorkingDirectory();
+private:
+	static NSURLSession* BackgroundSession;
+	static FString CachedIdentifierName;
+};
+
+
+//Delegate object associated with our above NSURLSession
+@interface BackgroundDownloadDelegate : NSObject<NSURLSessionDownloadDelegate>
+@end
diff --git a/Engine/Source/Runtime/IOS/IOSLocalNotification/Private/IOSLocalNotification.cpp b/Engine/Source/Runtime/IOS/IOSLocalNotification/Private/IOSLocalNotification.cpp
--- a/Engine/Source/Runtime/IOS/IOSLocalNotification/Private/IOSLocalNotification.cpp
+++ b/Engine/Source/Runtime/IOS/IOSLocalNotification/Private/IOSLocalNotification.cpp
@@ -175,51 +175,73 @@ void FIOSLocalNotificationService::ScheduleLocalNotificationAtTime(const FDateTi
 {
 #if !PLATFORM_TVOS
 #if __IPHONE_OS_VERSION_MIN_REQUIRED >= __IPHONE_10_0
-	UNMutableNotificationContent* Content = FIOSLocalNotificationModule::CreateNotificationContent(Title, Body, Action, ActivationEvent, 1);
-	UNCalendarNotificationTrigger* Trigger = FIOSLocalNotificationModule::CreateCalendarNotificationTrigger(FireDateTime);
-	
-	UNUserNotificationCenter *Center = [UNUserNotificationCenter currentNotificationCenter];
-	
-	FString NotId = Title.ToString() + FString::FromInt(NotificationNumber);
-	NSString* NotificationIdentifier = [NSString stringWithFString:NotId];
-	
-	UNNotificationRequest *Request = [UNNotificationRequest requestWithIdentifier:NotificationIdentifier content:Content trigger:Trigger];
-	
-	[Center addNotificationRequest:Request withCompletionHandler:^(NSError * _Nullable error) {
-		if (error != nil) {
-			UE_LOG(LogIOSLocalNotification, Warning, TEXT("Error scheduling notification: %s"), *NotId);
-		}
-	}];
+    
+    //Create local copies of these for the block to capture
+    FDateTime FireDateTimeCopy = FireDateTime;
+    FText TitleCopy = Title;
+    FText BodyCopy = Body;
+    FText ActionCopy = Action;
+    FString ActivationEventCopy = ActivationEvent;
+    
+	//have to schedule notification on main thread queue
+	dispatch_async(dispatch_get_main_queue(), ^{
+        UNMutableNotificationContent* Content = FIOSLocalNotificationModule::CreateNotificationContent(TitleCopy, BodyCopy, ActionCopy, ActivationEventCopy, 1);
+        UNCalendarNotificationTrigger* Trigger = FIOSLocalNotificationModule::CreateCalendarNotificationTrigger(FireDateTimeCopy);
+        
+        UNUserNotificationCenter *Center = [UNUserNotificationCenter currentNotificationCenter];
+        
+        FString NotId = TitleCopy.ToString() + FString::FromInt(NotificationNumber);
+        NSString* NotificationIdentifier = [NSString stringWithFString : NotId];
+        
+        UNNotificationRequest *Request = [UNNotificationRequest requestWithIdentifier:NotificationIdentifier content:Content trigger:Trigger];
+        
+        [Center addNotificationRequest : Request withCompletionHandler : ^ (NSError * _Nullable error) {
+            if (error != nil) {
+                UE_LOG(LogIOSLocalNotification, Warning, TEXT("Error scheduling notification: %s"), *NotId);
+            }
+        }];
+    });
 #else
-	UILocalNotification *localNotif = FIOSLocalNotificationModule::CreateLocalNotification(FireDateTime, LocalTime, ActivationEvent);
-	if (localNotif == nil)
-		return;
+    
+    //Create local copies of these for the block to capture
+    FDateTime FireDateTimeCopy = FireDateTime;
+    FText TitleCopy = Title;
+    FText BodyCopy = Body;
+    FText ActionCopy = Action;
+    FString ActivationEventCopy = ActivationEvent;
+    
+	//have to schedule notification on main thread queue
+	dispatch_async(dispatch_get_main_queue(), ^{
+		UILocalNotification *localNotif = FIOSLocalNotificationModule::CreateLocalNotification(FireDateTimeCopy, LocalTime, ActivationEventCopy);
+		if (localNotif == nil)
+			return;
 
-	NSString*	alertBody = [NSString stringWithFString : Body.ToString()];
-	if (alertBody != nil)
-	{
-		localNotif.alertBody = alertBody;
-	}
+		NSString*	alertBody = [NSString stringWithFString : BodyCopy.ToString()];
+		if (alertBody != nil)
+		{
+			localNotif.alertBody = alertBody;
+		}
 
-	NSString*	alertAction = [NSString stringWithFString:Action.ToString()];
-	if(alertAction != nil)
-	{
-		localNotif.alertAction = alertAction;
-	}
-	
-	if([IOSAppDelegate GetDelegate].OSVersion >= 8.2f)
-	{
-		NSString*	alertTitle = [NSString stringWithFString:Title.ToString()];
-		if(alertTitle != nil)
+		NSString*	alertAction = [NSString stringWithFString : ActionCopy.ToString()];
+		if (alertAction != nil)
 		{
-			localNotif.alertTitle = alertTitle;
+			localNotif.alertAction = alertAction;
 		}
-	}
-	
-	localNotif.soundName = UILocalNotificationDefaultSoundName;
-	localNotif.applicationIconBadgeNumber = 1;
 
-	[[UIApplication sharedApplication] scheduleLocalNotification:localNotif];
+		if ([IOSAppDelegate GetDelegate].OSVersion >= 8.2f)
+		{
+			NSString*	alertTitle = [NSString stringWithFString : TitleCopy.ToString()];
+			if (alertTitle != nil)
+			{
+				localNotif.alertTitle = alertTitle;
+			}
+		}
+
+		localNotif.soundName = UILocalNotificationDefaultSoundName;
+		localNotif.applicationIconBadgeNumber = 1;
+
+		[[UIApplication sharedApplication] scheduleLocalNotification:localNotif];
+    });
 #endif
 #endif
 }
diff --git a/Engine/Source/Runtime/Online/BackgroundHTTP/BackgroundHTTP.Build.cs b/Engine/Source/Runtime/Online/BackgroundHTTP/BackgroundHTTP.Build.cs
new file mode 100644
--- /dev/null
+++ b/Engine/Source/Runtime/Online/BackgroundHTTP/BackgroundHTTP.Build.cs
@@ -0,0 +1,32 @@
+// Copyright 1998-2018 Epic Games, Inc. All Rights Reserved.
+
+using UnrealBuildTool;
+
+public class BackgroundHTTP : ModuleRules
+{
+    public BackgroundHTTP(ReadOnlyTargetRules Target) : base(Target)
+    {
+        PublicIncludePaths.AddRange(
+            new string[] {
+                "Runtime/Online/HTTP/Public",
+            }
+        );
+
+        PrivateIncludePaths.AddRange(
+            new string[] {
+				"Runtime/Online/HTTP/Private",
+			}
+        );
+
+        PublicDependencyModuleNames.AddRange(
+            new string[]
+            {
+                "Core",
+                "CoreUObject",
+                "ApplicationCore",
+                "Engine",
+                "HTTP",
+            }
+        );
+    }
+}
diff --git a/Engine/Source/Runtime/Online/BackgroundHTTP/Private/BackgroundHttpManagerImpl.cpp b/Engine/Source/Runtime/Online/BackgroundHTTP/Private/BackgroundHttpManagerImpl.cpp
new file mode 100644
--- /dev/null
+++ b/Engine/Source/Runtime/Online/BackgroundHTTP/Private/BackgroundHttpManagerImpl.cpp
@@ -0,0 +1,208 @@
+// Copyright 1998-2018 Epic Games, Inc. All Rights Reserved.
+#include "BackgroundHttpManagerImpl.h"
+
+#include "HAL/PlatformFileManager.h"
+
+DEFINE_LOG_CATEGORY(LogBackgroundHttpManager);
+
+FBackgroundHttpManagerImpl::FBackgroundHttpManagerImpl()
+	: PendingStartRequests()
+	, PendingRequestLock()
+	, ActiveRequests()
+	, ActiveRequestLock()
+	, NumCurrentlyActiveRequests(0)
+{
+}
+
+FBackgroundHttpManagerImpl::~FBackgroundHttpManagerImpl()
+{
+}
+
+void FBackgroundHttpManagerImpl::Initialize()
+{
+}
+
+void FBackgroundHttpManagerImpl::Shutdown()
+{
+	//Pending Requests Clear
+	{
+		FRWScopeLock ScopeLock(PendingRequestLock, SLT_Write);
+		PendingStartRequests.Empty();
+	}
+
+	//Active Requests Clear
+	{
+		FRWScopeLock ScopeLock(ActiveRequestLock, SLT_Write);
+		ActiveRequests.Empty();
+		NumCurrentlyActiveRequests = 0;
+	}
+}
+
+void FBackgroundHttpManagerImpl::CleanUpTemporaryFiles()
+{
+	UE_LOG(LogBackgroundHttpManager, Log, TEXT("Cleaning Up Temporary Files"));
+
+	TArray<FString> FilesToDelete;
+
+	//Default implementation is to just delete everything in the Root Folder non-recursively.
+	IPlatformFile& PlatformFile = FPlatformFileManager::Get().GetPlatformFile();
+	PlatformFile.FindFiles(FilesToDelete, *FPlatformBackgroundHttp::GetTemporaryRootPath(), nullptr);
+
+	for (const FString& File : FilesToDelete)
+	{
+		UE_LOG(LogBackgroundHttpManager, Log, TEXT("Deleting File:%s"), *File);
+		const bool bDidDelete = PlatformFile.DeleteFile(*File);
+
+		if (!bDidDelete)
+		{
+			UE_LOG(LogBackgroundHttpManager, Warning, TEXT("Failure to Delete Temp File:%s"), *File);
+		}
+	}
+}
+
+void FBackgroundHttpManagerImpl::AddRequest(const FBackgroundHttpRequestPtr Request)
+{
+	UE_LOG(LogBackgroundHttpManager, Log, TEXT("AddRequest Called - RequestID:%s"), *Request->GetRequestID());
+
+	//If we don't associate with any existing requests, go into our pending list. These will be moved into the ActiveRequest list during our Tick
+	if (!AssociateWithAnyExistingRequest(Request))
+	{
+		FRWScopeLock ScopeLock(PendingRequestLock, SLT_Write);
+		PendingStartRequests.Add(Request);
+
+		UE_LOG(LogBackgroundHttpManager, Log, TEXT("Adding BackgroundHttpRequest to PendingStartRequests - RequestID:%s"), *Request->GetRequestID());
+	}
+}
+
+void FBackgroundHttpManagerImpl::RemoveRequest(const FBackgroundHttpRequestPtr Request)
+{
+	int NumRequestsRemoved = 0;
+
+	//Check if this request was in active list first
+	{
+		FRWScopeLock ScopeLock(ActiveRequestLock, SLT_Write);
+		NumRequestsRemoved = ActiveRequests.RemoveSwap(Request);
+
+		//If we removed an active request, lets decrement the NumCurrentlyActiveRequests accordingly
+		if (NumRequestsRemoved != 0)
+		{
+			NumCurrentlyActiveRequests = NumCurrentlyActiveRequests - NumRequestsRemoved;
+		}
+	}
+
+	//Only search the PendingRequestList if we didn't remove it in our ActiveRequest List
+	if (NumRequestsRemoved == 0)
+	{
+		FRWScopeLock ScopeLock(PendingRequestLock, SLT_Write);
+		NumRequestsRemoved = PendingStartRequests.RemoveSwap(Request);
+	}
+
+	UE_LOG(LogBackgroundHttpManager, Log, TEXT("FGenericPlatformBackgroundHttpManager::RemoveRequest Called - RequestID:%s | NumRequestsActuallyRemoved:%d | NumCurrentlyActiveRequests:%d"), *Request->GetRequestID(), NumRequestsRemoved, NumCurrentlyActiveRequests);
+}
+
+bool FBackgroundHttpManagerImpl::AssociateWithAnyExistingRequest(const FBackgroundHttpRequestPtr Request)
+{
+	bool bDidAssociateWithExistingRequest = false;
+
+	FString ExistingFilePath;
+	int64 ExistingFileSize;
+	if (CheckForExistingCompletedDownload(Request, ExistingFilePath, ExistingFileSize))
+	{
+		FBackgroundHttpResponsePtr NewResponseWithExistingFile = FPlatformBackgroundHttp::ConstructBackgroundResponse(EHttpResponseCodes::Ok, ExistingFilePath);
+		if (ensureAlwaysMsgf(NewResponseWithExistingFile.IsValid(), TEXT("Failure to create FBackgroundHttpResponsePtr in FPlatformBackgroundHttp::ConstructBackgroundResponse! Can not associate new download with found finished download!")))
+		{
+			bDidAssociateWithExistingRequest = true;
+			UE_LOG(LogBackgroundHttpManager, Log, TEXT("Found existing background task to associate with! RequestID:%s | ExistingFileSize:%lld | ExistingFilePath:%s"), *Request->GetRequestID(), ExistingFileSize, *ExistingFilePath);
+
+			//First send progress update for the file size so anything monitoring this download knows we are about to update this progress
+			Request->OnProgressUpdated().ExecuteIfBound(Request, ExistingFileSize, ExistingFileSize);
+
+			//Now complete with this completed response data
+			Request->CompleteWithExistingResponseData(NewResponseWithExistingFile);
+		}
+	}
+
+	return bDidAssociateWithExistingRequest;
+}
+
+bool FBackgroundHttpManagerImpl::CheckForExistingCompletedDownload(const FBackgroundHttpRequestPtr Request, FString& ExistingFilePathOut, int64& ExistingFileSizeOut)
+{
+	bool bDidFindExistingDownload = false;
+
+	IPlatformFile& PlatformFile = FPlatformFileManager::Get().GetPlatformFile();
+
+	const TArray<FString>& URLList = Request->GetURLList();
+	for (const FString& URL : URLList)
+	{
+		const FString& FileDestination = FPlatformBackgroundHttp::GetTemporaryFilePathFromURL(URL);
+		if (PlatformFile.FileExists(*FileDestination))
+		{
+			bDidFindExistingDownload = true;
+			ExistingFilePathOut = FileDestination;
+
+			ExistingFileSizeOut = PlatformFile.FileSize(*FileDestination);
+			break;
+		}
+	}	
+
+	return bDidFindExistingDownload;
+}
+
+bool FBackgroundHttpManagerImpl::Tick(float DeltaTime)
+{
+	ActivatePendingRequests();
+
+	//Keep ticking in all cases, so just return true
+	return true;
+}
+
+void FBackgroundHttpManagerImpl::ActivatePendingRequests()
+{
+	TArray<FBackgroundHttpRequestPtr> RequestsStartingThisTick;
+
+	//Handle Populating RequestsStartingThisTick from PendingStartRequests
+	{
+		FRWScopeLock ScopeLock(PendingRequestLock, SLT_Write);
+
+		if (PendingStartRequests.Num() > 0)
+		{
+			UE_LOG(LogBackgroundHttpManager, Verbose, TEXT("Populating Requests to Start from PendingStartRequests - PlatformMaxActiveDownloads:%d | NumCurrentlyActiveRequests:%d | NumPendingStartRequests:%d"), FPlatformBackgroundHttp::GetPlatformMaxActiveDownloads(), NumCurrentlyActiveRequests, PendingStartRequests.Num());
+
+			//See how many more requests we can process and only do anything if we can still process more
+			const int NumRequestsWeCanProcess = (FPlatformBackgroundHttp::GetPlatformMaxActiveDownloads() - NumCurrentlyActiveRequests);
+			if (NumRequestsWeCanProcess > 0)
+			{
+				for (int RequestAddedCount = 0; RequestAddedCount < NumRequestsWeCanProcess; ++RequestAddedCount)
+				{
+					//Don't do anything once we are out of PendingStartRequests
+					if (PendingStartRequests.Num() > 0)
+					{
+						//Lets just always use the first element and then remove it
+						FBackgroundHttpRequestPtr& PendingRequest = PendingStartRequests[0];
+						RequestsStartingThisTick.Add(PendingRequest);
+						PendingStartRequests.RemoveAtSwap(0);
+
+						++NumCurrentlyActiveRequests;
+					}
+					else
+					{
+						break;
+					}
+				}
+			}
+		}
+	}
+
+	UE_LOG(LogBackgroundHttpManager, Verbose, TEXT("Starting %d Requests From PendingStartRequests Queue"), RequestsStartingThisTick.Num());
+
+	//Now actually add request to Active List and call Handle 
+	for (FBackgroundHttpRequestPtr& RequestToStart : RequestsStartingThisTick)
+	{
+		//Actually move request to Active list now
+		FRWScopeLock ScopeLock(ActiveRequestLock, SLT_Write);
+		ActiveRequests.Add(RequestToStart);
+
+		//Call Handle for that task to now kick itself off
+		RequestToStart->HandleDelayedProcess();
+	}
+}
diff --git a/Engine/Source/Runtime/Online/BackgroundHTTP/Private/BackgroundHttpModule.cpp b/Engine/Source/Runtime/Online/BackgroundHTTP/Private/BackgroundHttpModule.cpp
new file mode 100644
--- /dev/null
+++ b/Engine/Source/Runtime/Online/BackgroundHTTP/Private/BackgroundHttpModule.cpp
@@ -0,0 +1,52 @@
+#include "BackgroundHttpModule.h"
+#include "BackgroundHttp.h"
+#include "Modules/ModuleManager.h"
+#include "Interfaces/IHttpRequest.h"
+
+
+IMPLEMENT_MODULE(FBackgroundHttpModule, BackgroundHTTP);
+
+FBackgroundHttpModule* FBackgroundHttpModule::Singleton = NULL;
+
+void FBackgroundHttpModule::StartupModule()
+{
+	Singleton = this;
+
+	FPlatformBackgroundHttp::Initialize();
+
+	BackgroundHttpManager = FPlatformBackgroundHttp::CreatePlatformBackgroundHttpManager();
+	if (ensureAlwaysMsgf(BackgroundHttpManager.IsValid(), TEXT("Invalid result from FPlatformBackgroundHttp::CreatePlatformBackgroundHttpManager! Ensure this platform has an implementation supplied!")))
+	{
+		BackgroundHttpManager->Initialize();
+	}
+}
+
+void FBackgroundHttpModule::ShutdownModule()
+{
+	FPlatformBackgroundHttp::Shutdown();
+
+	BackgroundHttpManager.Reset();
+	Singleton = nullptr;
+}
+
+FBackgroundHttpModule& FBackgroundHttpModule::Get()
+{
+	if (Singleton == NULL)
+	{
+		check(IsInGameThread());
+		FModuleManager::LoadModuleChecked<FBackgroundHttpModule>("BackgroundHTTP");
+	}
+	check(Singleton != NULL);
+	return *Singleton;
+}
+
+FBackgroundHttpRequestPtr FBackgroundHttpModule::CreateBackgroundRequest()
+{
+	// Create the platform specific Background Http request instance
+	return FPlatformBackgroundHttp::ConstructBackgroundRequest();
+}
+
+FBackgroundHttpManagerPtr FBackgroundHttpModule::GetBackgroundHttpManager()
+{
+	return BackgroundHttpManager;
+}
diff --git a/Engine/Source/Runtime/Online/BackgroundHTTP/Private/BackgroundHttpNotificationObject.cpp b/Engine/Source/Runtime/Online/BackgroundHTTP/Private/BackgroundHttpNotificationObject.cpp
new file mode 100644
--- /dev/null
+++ b/Engine/Source/Runtime/Online/BackgroundHTTP/Private/BackgroundHttpNotificationObject.cpp
@@ -0,0 +1,55 @@
+// Copyright 1998-2018 Epic Games, Inc. All Rights Reserved.
+#include "BackgroundHttpNotificationObject.h"
+
+#include "LocalNotification.h"
+
+FBackgroundHttpNotificationObject::FBackgroundHttpNotificationObject(FText InNotificationTitle, FText InNotificationBody, FText InNotificationAction, const FString& InNotificationActivationString, bool InNotifyOnlyOnFullSuccess)
+	: NotificationTitle(InNotificationTitle)
+    , NotificationAction(InNotificationAction)
+    , NotificationBody(InNotificationBody)
+	, NotificationActivationString(InNotificationActivationString)
+	, bNotifyOnlyOnFullSuccess(InNotifyOnlyOnFullSuccess)
+	, NumFailedDownloads(0)
+    , PlatformNotificationService(nullptr)
+{
+	if (GConfig)
+	{
+		FString ModuleName;
+		GConfig->GetString(TEXT("LocalNotification"), TEXT("DefaultPlatformService"), ModuleName, GEngineIni);
+
+		if (ModuleName.Len() > 0)
+		{
+			// load the module by name from the .ini
+			if (ILocalNotificationModule* Module = FModuleManager::LoadModulePtr<ILocalNotificationModule>(*ModuleName))
+			{
+				PlatformNotificationService = Module->GetLocalNotificationService();
+			}
+		}
+	}
+}
+
+FBackgroundHttpNotificationObject::~FBackgroundHttpNotificationObject()
+{
+	if (nullptr != PlatformNotificationService)
+	{
+		if (!bNotifyOnlyOnFullSuccess || (NumFailedDownloads == 0))
+		{
+			//make a notification 1 second from now
+			FDateTime TargetTime = FDateTime::Now();
+			TargetTime += FTimespan::FromSeconds(1);
+
+			if (nullptr != PlatformNotificationService)
+			{
+				PlatformNotificationService->ScheduleLocalNotificationAtTime(TargetTime, true, NotificationTitle, NotificationBody, NotificationAction, NotificationActivationString);
+			}
+		}
+	}
+}
+
+void FBackgroundHttpNotificationObject::NotifyOfDownloadResult(bool bWasSuccess)
+{
+	if (!bWasSuccess)
+	{
+		FPlatformAtomics::InterlockedIncrement(&NumFailedDownloads);
+	}
+}
diff --git a/Engine/Source/Runtime/Online/BackgroundHTTP/Private/BackgroundHttpRequestImpl.cpp b/Engine/Source/Runtime/Online/BackgroundHTTP/Private/BackgroundHttpRequestImpl.cpp
new file mode 100644
--- /dev/null
+++ b/Engine/Source/Runtime/Online/BackgroundHTTP/Private/BackgroundHttpRequestImpl.cpp
@@ -0,0 +1,118 @@
+#include "BackgroundHttpRequestImpl.h"
+
+#include "BackgroundHttpModule.h"
+#include "Interfaces/IBackgroundHttpResponse.h"
+
+DEFINE_LOG_CATEGORY(LogBackgroundHttpRequest);
+
+FBackgroundHttpRequestImpl::FBackgroundHttpRequestImpl()
+	: DownloadCompleteNotificationObject(nullptr)
+	, Response(nullptr)
+	, URLList()
+	, RequestID()
+	, NumberOfTotalRetries(0)
+	, HttpRequestCompleteDelegate()
+	, HttpProgressUpdateDelegate()
+{
+}
+
+bool FBackgroundHttpRequestImpl::ProcessRequest()
+{
+	UE_LOG(LogBackgroundHttpRequest, Display, TEXT("Processing Request - RequestID:%s"), *GetRequestID());
+	FBackgroundHttpModule::Get().GetBackgroundHttpManager()->AddRequest(SharedThis(this));
+
+	return true;
+}
+
+void FBackgroundHttpRequestImpl::OnBackgroundDownloadComplete()
+{
+	FBackgroundHttpModule::Get().GetBackgroundHttpManager()->RemoveRequest(SharedThis(this));
+
+	//Determine if this was a success or not
+	FBackgroundHttpResponsePtr SetResponse = GetResponse();
+	const bool bWasSuccess = SetResponse.IsValid() ? EHttpResponseCodes::IsOk(SetResponse->GetResponseCode()) : false;
+
+	UE_LOG(LogBackgroundHttpRequest, Display, TEXT("Download Complete - RequestID:%s | bWasSuccess:%d "), *GetRequestID(), (int)(bWasSuccess));
+
+	//First, send a delegate out for this request completing
+	OnProcessRequestComplete().ExecuteIfBound(SharedThis(this), bWasSuccess);
+
+	//Second notify our download complete notification of the success/failure of this request. Then remove our reference
+	//so a notification will send without having to delete this request. Do this after Complete Delegate so we can ensure
+	//whatever kicked this off doesn't want to react to this download complete before we send this notification (IE: Kick off another download
+	//using the same NotificationObject, etc.)
+	NotifyNotificationObjectOfComplete(bWasSuccess);
+}
+
+void FBackgroundHttpRequestImpl::NotifyNotificationObjectOfComplete(bool bWasSuccess)
+{	
+	if (DownloadCompleteNotificationObject.IsValid())
+	{
+		const int32 SharedRefCount = DownloadCompleteNotificationObject.GetSharedReferenceCount();
+		UE_LOG(LogBackgroundHttpRequest, Display, TEXT("Removing Reference to DownloadCompleteNotificationObject - bWasSuccess:%d | CurrentSharedRefCount:%d"), (int)(bWasSuccess), SharedRefCount);
+
+		DownloadCompleteNotificationObject->NotifyOfDownloadResult(bWasSuccess);
+		DownloadCompleteNotificationObject.Reset();
+	}
+}
+
+void FBackgroundHttpRequestImpl::SetURLAsList(const TArray<FString>& URLs, int NumRetriesIn)
+{
+	NumberOfTotalRetries = NumRetriesIn;
+	for (const FString& URL : URLs)
+	{
+		URLList.Add(URL);
+	}
+}
+
+const TArray<FString>& FBackgroundHttpRequestImpl::GetURLList() const
+{
+	return URLList;
+}
+
+void FBackgroundHttpRequestImpl::SetCompleteNotification(FBackgroundHttpNotificationObjectPtr DownloadCompleteNotificationObjectIn)
+{
+	DownloadCompleteNotificationObject = DownloadCompleteNotificationObjectIn;
+}
+
+void FBackgroundHttpRequestImpl::CompleteWithExistingResponseData(FBackgroundHttpResponsePtr BackgroundResponse)
+{
+	Response = BackgroundResponse;
+
+	const bool bHasValidResponse = Response.IsValid();	
+	UE_LOG(LogBackgroundHttpRequest, Display, TEXT("Completing Download With Existing Response Data - RequestID:%s | bHasValidResponse:%d"), *GetRequestID(), (int)(bHasValidResponse));
+
+	OnBackgroundDownloadComplete();
+}
+
+FBackgroundHttpRequestCompleteDelegate& FBackgroundHttpRequestImpl::OnProcessRequestComplete()
+{
+	return HttpRequestCompleteDelegate;
+}
+
+FBackgroundHttpProgressUpdateDelegate& FBackgroundHttpRequestImpl::OnProgressUpdated()
+{
+	return HttpProgressUpdateDelegate;
+}
+
+const FBackgroundHttpResponsePtr FBackgroundHttpRequestImpl::GetResponse() const
+{
+	return Response;
+}
+
+const FString& FBackgroundHttpRequestImpl::GetRequestID() const
+{
+	return RequestID;
+}
+
+void FBackgroundHttpRequestImpl::SetRequestID(const FString& NewRequestID)
+{
+	RequestID = NewRequestID;
+}
+
+bool FBackgroundHttpRequestImpl::HandleDelayedProcess()
+{
+	//By default we don't provide an implementation for this. If this is called, it should be overriden by the platform specific
+	//BackgroundHttpRequest if the platform expects to make use of this.
+	return ensureAlwaysMsgf(false, TEXT("Platform expects an implementation of HandleDelayedProcess on the BackgroundHttpRequest, but none found!"));
+}
\ No newline at end of file
diff --git a/Engine/Source/Runtime/Online/BackgroundHTTP/Private/BackgroundHttpResponseImpl.cpp b/Engine/Source/Runtime/Online/BackgroundHTTP/Private/BackgroundHttpResponseImpl.cpp
new file mode 100644
--- /dev/null
+++ b/Engine/Source/Runtime/Online/BackgroundHTTP/Private/BackgroundHttpResponseImpl.cpp
@@ -0,0 +1,24 @@
+// Copyright 1998-2018 Epic Games, Inc. All Rights Reserved.
+
+#include "BackgroundHttpResponseImpl.h"
+
+#include "Interfaces/IHttpResponse.h"
+
+#include "HAL/PlatformApplicationMisc.h"
+#include "Misc/Paths.h"
+
+FBackgroundHttpResponseImpl::FBackgroundHttpResponseImpl()
+	: TempContentFilePath()
+	, ResponseCode(EHttpResponseCodes::Unknown)
+{
+}
+
+int32 FBackgroundHttpResponseImpl::GetResponseCode() const
+{
+	return ResponseCode;
+}
+
+const FString& FBackgroundHttpResponseImpl::GetTempContentFilePath() const
+{
+	return TempContentFilePath;
+}
\ No newline at end of file
diff --git a/Engine/Source/Runtime/Online/BackgroundHTTP/Private/GenericPlatform/GenericPlatformBackgroundHttp.cpp b/Engine/Source/Runtime/Online/BackgroundHTTP/Private/GenericPlatform/GenericPlatformBackgroundHttp.cpp
new file mode 100644
--- /dev/null
+++ b/Engine/Source/Runtime/Online/BackgroundHTTP/Private/GenericPlatform/GenericPlatformBackgroundHttp.cpp
@@ -0,0 +1,47 @@
+// Copyright 1998-2018 Epic Games, Inc. All Rights Reserved.
+#include "GenericPlatform/GenericPlatformBackgroundHttp.h"
+#include "GenericPlatform/GenericPlatformBackgroundHttpManager.h"
+#include "GenericPlatform/GenericPlatformBackgroundHttpRequest.h"
+#include "GenericPlatform/GenericPlatformBackgroundHttpResponse.h"
+
+void FGenericPlatformBackgroundHttp::Initialize()
+{
+}
+
+void FGenericPlatformBackgroundHttp::Shutdown()
+{
+}
+
+FBackgroundHttpManagerPtr FGenericPlatformBackgroundHttp::CreatePlatformBackgroundHttpManager()
+{
+    return MakeShared<FGenericPlatformBackgroundHttpManager, ESPMode::ThreadSafe>();
+}
+
+FBackgroundHttpRequestPtr FGenericPlatformBackgroundHttp::ConstructBackgroundRequest()
+{
+	return MakeShared<FGenericPlatformBackgroundHttpRequest, ESPMode::ThreadSafe>();
+}
+
+FBackgroundHttpResponsePtr FGenericPlatformBackgroundHttp::ConstructBackgroundResponse(int32 ResponseCode, const FString& TempFilePath)
+{
+	return MakeShared<FGenericPlatformBackgroundHttpResponse, ESPMode::ThreadSafe>(ResponseCode, TempFilePath);
+}
+
+const FString FGenericPlatformBackgroundHttp::GetTemporaryFilePathFromURL(const FString& URL)
+{
+	//For now, just save any file in the root BackgroundHttp Tempt directory under the sanitized URL
+	const FString FileName = FPaths::MakeValidFileName(URL);
+	return FPaths::Combine(GetTemporaryRootPath(), FileName);
+}
+
+const FString& FGenericPlatformBackgroundHttp::GetTemporaryRootPath()
+{
+	static FString BackgroundHttpDir = FPaths::Combine(FPlatformMisc::GamePersistentDownloadDir(), TEXT("BackgroundHttpTemp"));
+	return BackgroundHttpDir;
+}
+
+int FGenericPlatformBackgroundHttp::GetPlatformMaxActiveDownloads()
+{
+	//In general, lets just use 4 active downloads at once
+	return 4;
+}
diff --git a/Engine/Source/Runtime/Online/BackgroundHTTP/Private/GenericPlatform/GenericPlatformBackgroundHttpManager.cpp b/Engine/Source/Runtime/Online/BackgroundHTTP/Private/GenericPlatform/GenericPlatformBackgroundHttpManager.cpp
new file mode 100644
--- /dev/null
+++ b/Engine/Source/Runtime/Online/BackgroundHTTP/Private/GenericPlatform/GenericPlatformBackgroundHttpManager.cpp
@@ -0,0 +1,5 @@
+#include "GenericPlatform/GenericPlatformBackgroundHttpManager.h"
+
+FGenericPlatformBackgroundHttpManager::~FGenericPlatformBackgroundHttpManager()
+{
+}
\ No newline at end of file
diff --git a/Engine/Source/Runtime/Online/BackgroundHTTP/Private/GenericPlatform/GenericPlatformBackgroundHttpRequest.cpp b/Engine/Source/Runtime/Online/BackgroundHTTP/Private/GenericPlatform/GenericPlatformBackgroundHttpRequest.cpp
new file mode 100644
--- /dev/null
+++ b/Engine/Source/Runtime/Online/BackgroundHTTP/Private/GenericPlatform/GenericPlatformBackgroundHttpRequest.cpp
@@ -0,0 +1,154 @@
+// Copyright 1998-2018 Epic Games, Inc. All Rights Reserved.
+#include "GenericPlatform/GenericPlatformBackgroundHttpRequest.h"
+#include "GenericPlatform/GenericPlatformBackgroundHttpResponse.h"
+#include "Interfaces/IHttpResponse.h"
+#include "Interfaces/IHttpRequest.h"
+#include "HttpModule.h"
+#include "BackgroundHttpModule.h"
+
+FGenericPlatformBackgroundHttpRequest::FGenericPlatformBackgroundHttpRequest()
+	: RequestWrapper(nullptr)
+{
+}
+
+bool FGenericPlatformBackgroundHttpRequest::HandleDelayedProcess()
+{
+	//Our request may already have a response result from associating with another already completed request. If so, lets not spin up our RequestWrapper.
+	if (!GetResponse().IsValid())
+	{
+		if (!RequestWrapper.IsValid())
+		{
+			RequestWrapper = MakeShareable<FGenericPlatformBackgroundHttpWrapper>(new FGenericPlatformBackgroundHttpWrapper(SharedThis(this), NumberOfTotalRetries));
+		}
+
+		if (ensureAlwaysMsgf(RequestWrapper.IsValid(), TEXT("Failure creating FGenericPlatformBackgroundHttpWrapper! Can not process BackgroundHttp request! RequestID:%s"), *GetRequestID()))
+		{
+			//Start request wrapper
+			RequestWrapper->MakeRequest();
+		}
+	}
+
+	return RequestWrapper.IsValid() || GetResponse().IsValid();
+}
+
+void FGenericPlatformBackgroundHttpRequest::FGenericPlatformBackgroundHttpRequest::CompleteWithExistingResponseData(FBackgroundHttpResponsePtr BackgroundResponse)
+{
+	//Remove our request wrapper immediately so the Response cleans up ASAP instead of the user having to delete this task to clean up the Http Response.
+	RequestWrapper.Reset();
+
+	FBackgroundHttpRequestImpl::CompleteWithExistingResponseData(BackgroundResponse);
+}
+
+FGenericPlatformBackgroundHttpRequest::FGenericPlatformBackgroundHttpWrapper::FGenericPlatformBackgroundHttpWrapper(FBackgroundHttpRequestPtr Request, int MaxRetriesToAttempt)
+	: OriginalRequest(Request)
+	, CurrentRetryNumber(0)
+	, MaxRetries(MaxRetriesToAttempt)
+	, LastProgressUpdateBytes(0)
+{
+}
+
+FGenericPlatformBackgroundHttpRequest::FGenericPlatformBackgroundHttpWrapper::~FGenericPlatformBackgroundHttpWrapper()
+{
+	CleanUpHttpRequest();
+}
+
+void FGenericPlatformBackgroundHttpRequest::FGenericPlatformBackgroundHttpWrapper::MakeRequest()
+{
+	if (ensureAlwaysMsgf(OriginalRequest.IsValid(), TEXT("Call to MakeRequest in BackgroundHttpWrapper without associating it to a valid BackgroundHttpRequest! Should never happen!")))
+	{
+		//If we have an existing HttpRequest, lets clear it out before creating a new one
+		CleanUpHttpRequest();
+
+		//Create new Request
+		HttpRequest = FHttpModule::Get().CreateRequest();
+		HttpRequest->OnProcessRequestComplete().BindRaw(this, &FGenericPlatformBackgroundHttpRequest::FGenericPlatformBackgroundHttpWrapper::HttpRequestComplete);
+		HttpRequest->OnRequestProgress().BindRaw(this, &FGenericPlatformBackgroundHttpRequest::FGenericPlatformBackgroundHttpWrapper::UpdateHttpProgress);
+
+		const FString& RequestURL = GetURLForCurrentRetryNumber();
+		if (!RequestURL.IsEmpty())
+		{
+			HttpRequest->SetURL(RequestURL);
+			HttpRequest->ProcessRequest();
+		}
+		else
+		{
+			UE_LOG(LogBackgroundHttpRequest, Display, TEXT("No valid URL for Request. (No valid URL, or out of Retries) - RequestID:%s | RetryNumber:%d | MaxRetries:%d"), *OriginalRequest->GetRequestID(), CurrentRetryNumber, MaxRetries);
+
+			//We don't have a valid URL, so just error immediately
+			HttpRequestComplete(HttpRequest, nullptr, false);
+		}
+	}
+}
+
+void FGenericPlatformBackgroundHttpRequest::FGenericPlatformBackgroundHttpWrapper::CleanUpHttpRequest()
+{
+	if (HttpRequest.IsValid())
+	{
+		HttpRequest->OnProcessRequestComplete().Unbind();
+		HttpRequest->OnRequestProgress().Unbind();
+		HttpRequest->CancelRequest();
+
+		HttpRequest.Reset();
+	}
+}
+
+const FString FGenericPlatformBackgroundHttpRequest::FGenericPlatformBackgroundHttpWrapper::GetURLForCurrentRetryNumber()
+{
+	FString ReturnedURL = FString();
+
+	int URLIndexToUse = CurrentRetryNumber;
+
+	//Don't find a valid URL if we are already passed our retry count
+	if (URLIndexToUse <= MaxRetries)
+	{
+		const TArray<FString>& URLList = OriginalRequest->GetURLList();
+		if (URLList.Num() > 0)
+		{
+			//If our number is too high, lets loop back around to the start of the list
+			while ((URLIndexToUse >= 0) && (!URLList.IsValidIndex(URLIndexToUse)))
+			{
+				URLIndexToUse -= URLList.Num();
+			}
+		}
+
+		if (URLList.IsValidIndex(URLIndexToUse))
+		{
+			ReturnedURL = URLList[URLIndexToUse];
+		}
+	}
+
+	return ReturnedURL;
+}
+
+void FGenericPlatformBackgroundHttpRequest::FGenericPlatformBackgroundHttpWrapper::HttpRequestComplete(FHttpRequestPtr HttpRequestIn, FHttpResponsePtr HttpResponse, bool bSuccess)
+{
+	//Create a new FGenericPlatformBackgroundHttpRequest from the HttpRequest/Response
+	FBackgroundHttpResponsePtr ConstructedResponse = MakeShareable(new FGenericPlatformBackgroundHttpResponse(HttpRequestIn, HttpResponse, bSuccess));
+
+	//Let the FGenericPlatformBackgroundHttpResponse drive if this succeeded or not
+	const bool bRequestWasSuccessful = EHttpResponseCodes::IsOk(ConstructedResponse->GetResponseCode());
+
+	UE_LOG(LogBackgroundHttpRequest, Display, TEXT("Underlying HttpRequest complete - RequestID:%s | bRequestWasSuccessfull:%d | CurrentRetryNumber:%d | MaxRetries:%d"), *OriginalRequest->GetRequestID(), (int)(bRequestWasSuccessful), CurrentRetryNumber, MaxRetries);
+
+	//Attempt a retry on failure
+	if (!bRequestWasSuccessful)
+	{
+		++CurrentRetryNumber;
+		MakeRequest();
+	}
+	else
+	{
+		//Pass constructed response to our BackgroundRequest so it can complete itself using this response
+		OriginalRequest->CompleteWithExistingResponseData(ConstructedResponse);
+	}
+}
+
+void FGenericPlatformBackgroundHttpRequest::FGenericPlatformBackgroundHttpWrapper::UpdateHttpProgress(FHttpRequestPtr UnderlyingHttpRequest, int32 BytesSent, int32 BytesReceived)
+{
+	const int32 ByteDifference = (LastProgressUpdateBytes > 0) ? LastProgressUpdateBytes - BytesReceived : BytesReceived;
+	LastProgressUpdateBytes = ByteDifference;
+	
+	UE_LOG(LogBackgroundHttpRequest, VeryVerbose, TEXT("HttpRequest Progress Update- RequestID:%s | BytesSent: %d | BytesReceived:%d | BytesReceivedSinceLastUpdate:%d"), *OriginalRequest->GetRequestID(), BytesSent, BytesReceived, ByteDifference);
+	
+	OriginalRequest->OnProgressUpdated().ExecuteIfBound(OriginalRequest, BytesReceived, LastProgressUpdateBytes);
+}
diff --git a/Engine/Source/Runtime/Online/BackgroundHTTP/Private/GenericPlatform/GenericPlatformBackgroundHttpResponse.cpp b/Engine/Source/Runtime/Online/BackgroundHTTP/Private/GenericPlatform/GenericPlatformBackgroundHttpResponse.cpp
new file mode 100644
--- /dev/null
+++ b/Engine/Source/Runtime/Online/BackgroundHTTP/Private/GenericPlatform/GenericPlatformBackgroundHttpResponse.cpp
@@ -0,0 +1,36 @@
+// Copyright 1998-2018 Epic Games, Inc. All Rights Reserved.
+#include "GenericPlatform/GenericPlatformBackgroundHttpResponse.h"
+#include "GenericPlatform/GenericPlatformBackgroundHttp.h"
+
+#include "Interfaces/IHttpResponse.h"
+#include "Misc/FileHelper.h"
+
+FGenericPlatformBackgroundHttpResponse::FGenericPlatformBackgroundHttpResponse(int32 InResponseCode, const FString& InTempContentFilePath)
+{
+	ResponseCode = InResponseCode;
+	TempContentFilePath = InTempContentFilePath;
+}
+
+FGenericPlatformBackgroundHttpResponse::FGenericPlatformBackgroundHttpResponse(FHttpRequestPtr HttpRequestIn, FHttpResponsePtr HttpResponse, bool bSuccess)
+{
+	//Don't bother making a response out of a failed HttpRequest
+	if (bSuccess && HttpResponse.IsValid())
+	{
+		//Copy HTTP Response Code if we have one, otherwise it will be Unknown
+		ResponseCode = HttpResponse->GetResponseCode();
+
+		//Only try to copy results to a temp file if it as an ok result
+		if (EHttpResponseCodes::IsOk(ResponseCode))
+		{
+			const TArray<uint8>& ContentBuffer = HttpResponse->GetContent();
+			if (ContentBuffer.Num() > 0)
+			{
+				const FString RequestURL = HttpRequestIn->GetURL();
+				const FString& FileDestination = FPlatformBackgroundHttp::GetTemporaryFilePathFromURL(RequestURL);
+
+				FFileHelper::SaveArrayToFile(ContentBuffer, *FileDestination);
+				TempContentFilePath = FileDestination;
+			}
+		}
+	}
+}
\ No newline at end of file
diff --git a/Engine/Source/Runtime/Online/BackgroundHTTP/Private/IOS/ApplePlatformBackgroundHttp.cpp b/Engine/Source/Runtime/Online/BackgroundHTTP/Private/IOS/ApplePlatformBackgroundHttp.cpp
new file mode 100644
--- /dev/null
+++ b/Engine/Source/Runtime/Online/BackgroundHTTP/Private/IOS/ApplePlatformBackgroundHttp.cpp
@@ -0,0 +1,48 @@
+// Copyright 1998-2018 Epic Games, Inc. All Rights Reserved.
+#include "IOS/ApplePlatformBackgroundHttp.h"
+#include "IOS/ApplePlatformBackgroundHttpManager.h"
+#include "IOS/ApplePlatformBackgroundHttpRequest.h"
+#include "IOS/ApplePlatformBackgroundHttpResponse.h"
+#include "IOS/IOSBackgroundURLSessionHandler.h"
+
+void FApplePlatformBackgroundHttp::Initialize()
+{
+	const FString DefaultIdentifier = TEXT("com.epicgames.backgroundhttp");
+	ensureAlwaysMsgf(FBackgroundURLSessionHandler::InitBackgroundSession(DefaultIdentifier), TEXT("Failure to create a background download session with identifier %s"), *DefaultIdentifier);
+}
+
+void FApplePlatformBackgroundHttp::Shutdown()
+{
+	FBackgroundURLSessionHandler::ShutdownBackgroundSession();
+}
+
+FBackgroundHttpManagerPtr FApplePlatformBackgroundHttp::CreatePlatformBackgroundHttpManager()
+{
+    return MakeShared<FApplePlatformBackgroundHttpManager, ESPMode::ThreadSafe>();
+}
+
+FBackgroundHttpRequestPtr FApplePlatformBackgroundHttp::ConstructBackgroundRequest()
+{
+	return MakeShared<FApplePlatformBackgroundHttpRequest, ESPMode::ThreadSafe>();
+}
+
+FBackgroundHttpResponsePtr FApplePlatformBackgroundHttp::ConstructBackgroundResponse(int32 ResponseCode, const FString& TempFilePath)
+{
+	return MakeShared<FApplePlatformBackgroundHttpResponse, ESPMode::ThreadSafe>(ResponseCode, TempFilePath);
+}
+
+const FString FApplePlatformBackgroundHttp::GetTemporaryFilePathFromURL(const FString& URL)
+{
+	return FBackgroundURLSessionHandler::GetTemporaryFilePathFromURL(URL);
+}
+
+const FString& FApplePlatformBackgroundHttp::GetTemporaryRootPath()
+{
+    static FString BackgroundHttpDir = FPaths::Combine(FPlatformMisc::GamePersistentDownloadDir(), TEXT("BackgroundHttpTemp"));
+    return BackgroundHttpDir;
+}
+
+int FApplePlatformBackgroundHttp::GetPlatformMaxActiveDownloads()
+{
+	return 4;
+}
diff --git a/Engine/Source/Runtime/Online/BackgroundHTTP/Private/IOS/ApplePlatformBackgroundHttpManager.cpp b/Engine/Source/Runtime/Online/BackgroundHTTP/Private/IOS/ApplePlatformBackgroundHttpManager.cpp
new file mode 100644
--- /dev/null
+++ b/Engine/Source/Runtime/Online/BackgroundHTTP/Private/IOS/ApplePlatformBackgroundHttpManager.cpp
@@ -0,0 +1,726 @@
+#include "IOS/ApplePlatformBackgroundHttpManager.h"
+#include "IOS/ApplePlatformBackgroundHttpRequest.h"
+#include "IOS/ApplePlatformBackgroundHttpResponse.h"
+
+#include "HAL/PlatformFilemanager.h"
+#include "HAL/PlatformAtomics.h"
+
+#include "IOS/IOSBackgroundURLSessionHandler.h"
+
+FApplePlatformBackgroundHttpManager::~FApplePlatformBackgroundHttpManager()
+{
+    [UnAssociatedTasks release];
+    UnAssociatedTasks = nullptr;
+    
+    CleanUpNSURLSessionResponseDelegates();
+}
+
+FString FApplePlatformBackgroundHttpManager::BackgroundSessionIdentifier = FString("");
+float FApplePlatformBackgroundHttpManager::ActiveTimeOutSetting = 30.0f;
+int FApplePlatformBackgroundHttpManager::RetryResumeDataLimitSetting = -1.f;
+
+FApplePlatformBackgroundHttpManager::FApplePlatformBackgroundHttpManager()
+	: bHasFinishedPopulatingUnassociatedTasks(false)
+    , bIsInBackground(false)
+    , bIsIteratingThroughSessionTasks(false)
+    , RequestsPendingRemove()
+{
+}
+
+void FApplePlatformBackgroundHttpManager::Initialize()
+{
+	UnAssociatedTasks = [[NSMutableDictionary alloc] init];
+	PopulateUnAssociatedTasks();
+
+    GConfig->GetFloat(TEXT("BackgroundHttp.iOSSettings"), TEXT("BackgroundHttp.ActiveReceiveTimeout"), ActiveTimeOutSetting, GEngineIni);
+	GConfig->GetInt(TEXT("BackgroundHttp.iOSSettings"), TEXT("BackgroundHttp.RetryResumeDataLimit"), RetryResumeDataLimitSetting, GEngineIni);
+
+	SetupNSURLSessionResponseDelegates();
+}
+
+void FApplePlatformBackgroundHttpManager::PopulateUnAssociatedTasks()
+{
+	if (ensureAlwaysMsgf((nullptr != UnAssociatedTasks), TEXT("Call to PopulateUnAssociatedTasks without initializing UnAssociatedTasks Dictionary!")))
+	{
+		NSURLSession* BackgroundDownloadSession = FBackgroundURLSessionHandler::GetBackgroundSession();
+		if (ensureAlwaysMsgf((nullptr != BackgroundDownloadSession), TEXT("Invalid Background Download NSURLSession during AppleBackgroundHttp Init! Should have already Initialized the NSURLSession by this point!")))
+		{
+			[BackgroundDownloadSession getAllTasksWithCompletionHandler : ^ (NSArray<__kindof NSURLSessionTask*> *tasks)
+			{
+				//Store all existing tasks by their URL
+				for (id task in tasks)
+				{
+					[UnAssociatedTasks setObject:task forKey:[[[task currentRequest] URL] absoluteString]];
+				}
+
+				bHasFinishedPopulatingUnassociatedTasks = true;
+			}];
+		}
+	}
+}
+
+void FApplePlatformBackgroundHttpManager::PauseAllUnassociatedTasks()
+{
+    for (id Key in UnAssociatedTasks)
+    {
+        NSURLSessionDownloadTask* Task = (NSURLSessionDownloadTask*)([UnAssociatedTasks objectForKey:Key]);
+        if (nullptr != Task)
+        {
+            if ([Task state] == NSURLSessionTaskStateRunning)
+            {
+                [Task suspend];
+            }
+        }
+    }
+}
+
+void FApplePlatformBackgroundHttpManager::UnpauseAllUnassociatedTasks()
+{
+    for (id Key in UnAssociatedTasks)
+    {
+        NSURLSessionDownloadTask* Task = (NSURLSessionDownloadTask*)([UnAssociatedTasks objectForKey:Key]);
+        if (nullptr != Task)
+        {
+            if ([Task state] == NSURLSessionTaskStateSuspended)
+            {
+                [Task resume];
+            }
+        }
+    }
+}
+
+void FApplePlatformBackgroundHttpManager::Shutdown()
+{
+	[UnAssociatedTasks release];
+	UnAssociatedTasks = nullptr;
+
+	CleanUpNSURLSessionResponseDelegates();
+    FBackgroundURLSessionHandler::ShutdownBackgroundSession();
+}
+
+void FApplePlatformBackgroundHttpManager::AddRequest(const FBackgroundHttpRequestPtr Request)
+{
+	UE_LOG(LogBackgroundHttpManager, Verbose, TEXT("AddRequest Called - RequestID:%s"), *Request->GetRequestID());
+    
+    //See if our request is an AppleBackgroundHttpRequest so we can do more detailed checks on it.
+    FAppleBackgroundHttpRequestPtr AppleRequest = StaticCastSharedPtr<FApplePlatformBackgroundHttpRequest>(Request);
+    if (ensureAlwaysMsgf(AppleRequest.IsValid(), TEXT("Adding a non-Apple background request to our Apple Background Http Manager!")))
+    {
+        GenerateURLMapEntriesForRequest(AppleRequest);
+    }
+    
+	if (!AssociateWithAnyExistingRequest(Request))
+	{
+        //ensures above if this is false, but need to re-check it before adding tasks
+        if (AppleRequest.IsValid())
+        {
+            StartRequest(AppleRequest);
+        }
+        
+		//Just always add un-associated requests to ActiveRequests in the Apple implementation.
+		FRWScopeLock ScopeLock(ActiveRequestLock, SLT_Write);
+		ActiveRequests.Add(Request);
+	}
+}
+
+void FApplePlatformBackgroundHttpManager::GenerateURLMapEntriesForRequest( FAppleBackgroundHttpRequestPtr Request)
+{
+    FRWScopeLock ScopeLock(URLToRequestMapLock, SLT_Write);
+    for (const FString& URL : Request->GetURLList())
+    {
+        FBackgroundHttpURLMappedRequestPtr& FoundRequest = URLToRequestMap.FindOrAdd(URL);
+        
+        const bool bRequestAlreadyExistsForURL = ((FoundRequest.IsValid()) && (Request != FoundRequest));
+        if (ensureAlwaysMsgf(!bRequestAlreadyExistsForURL, TEXT("URL is represented by 2 different Requests! Immediately completing new request with error.")))
+        {
+            FoundRequest = Request;
+        }
+        else
+        {
+            FBackgroundHttpResponsePtr NewResponse = FPlatformBackgroundHttp::ConstructBackgroundResponse(EHttpResponseCodes::Unknown, FString());
+            Request->CompleteWithExistingResponseData(NewResponse);
+        }
+    }
+}
+
+void FApplePlatformBackgroundHttpManager::StartRequest(FAppleBackgroundHttpRequestPtr Request)
+{
+    //Just count it as a retry that won't increment the retry counter before giving us the URL as our RetryCount 0 should start this up.
+    RetryRequest(Request,false, false, nullptr);
+}
+
+void FApplePlatformBackgroundHttpManager::RemoveRequest(const FBackgroundHttpRequestPtr Request)
+{
+	FAppleBackgroundHttpRequestPtr AppleRequest = StaticCastSharedPtr<FApplePlatformBackgroundHttpRequest>(Request);
+	if (AppleRequest.IsValid())
+	{
+		RemoveSessionTasksForRequest(AppleRequest);
+	}
+	
+    RequestsPendingRemove.Add(Request);
+}
+
+void FApplePlatformBackgroundHttpManager::DeletePendingRemoveRequests()
+{
+    for (const FBackgroundHttpRequestPtr& Request : RequestsPendingRemove)
+    {
+        FBackgroundHttpManagerImpl::RemoveRequest(Request);
+    }
+    
+    RequestsPendingRemove.Empty();
+}
+
+void FApplePlatformBackgroundHttpManager::RemoveSessionTasksForRequest(FAppleBackgroundHttpRequestPtr Request)
+{
+    //@TODO TRoss: Should implement this
+}
+
+bool FApplePlatformBackgroundHttpManager::AssociateWithAnyExistingRequest(const FBackgroundHttpRequestPtr Request)
+{
+	if (!bHasFinishedPopulatingUnassociatedTasks)
+	{
+		//@TODO: TRoss, might want to look at trying to associate these after bHasFinishedPopulatingUnassociatedTasks. At this point I don't think its really worth it as it SHOULD be done before we get here, however
+		//the PopulateUnAssociatedTasks() function has an asynch component so it could technically be unfinished with some tight timing.
+		UE_LOG(LogBackgroundHttpManager, Warning, TEXT("Call to AssociateWithAnyExistingRequest before we have finished populating unassociated tasks! Might have an unassociated task for this request that we won't associate with."));
+	}
+	
+	//First check if we Associate using the base implementation.
+	bool bDidAssociateWithExistingRequest = FBackgroundHttpManagerImpl::AssociateWithAnyExistingRequest(Request);
+
+	//Didn't associate, so lets check more Apple specific stuff
+	if (!bDidAssociateWithExistingRequest)
+	{
+		//See if our request is an AppleBackgroundHttpRequest so we can do more detailed checks on it.
+		FAppleBackgroundHttpRequestPtr AppleRequest = StaticCastSharedPtr<FApplePlatformBackgroundHttpRequest>(Request);
+		if (AppleRequest.IsValid())
+		{
+			bDidAssociateWithExistingRequest = CheckForExistingUnAssociatedTask(AppleRequest);
+		}
+	}
+
+	return bDidAssociateWithExistingRequest;
+}
+
+bool FApplePlatformBackgroundHttpManager::CheckForExistingUnAssociatedTask(const FAppleBackgroundHttpRequestPtr Request)
+{
+    bool bDidFindExistingTask = false;
+    
+	if (ensureAlwaysMsgf(Request.IsValid(), TEXT("CheckForExistingUnAssociatedTask called with invalid Request!")))
+	{
+		const TArray<FString>& URLList = Request->GetURLList();
+		for (const FString& URL : URLList)
+		{
+			NSURLSessionTask* FoundTask = [UnAssociatedTasks valueForKey:URL.GetNSString()];
+			if (nullptr != FoundTask)
+			{
+                if ([FoundTask state] != NSURLSessionTaskStateCanceling)
+                {
+                    UE_LOG(LogBackgroundHttpManager, Display, TEXT("Existing UnAssociateTask found for Request! Attempting to Associate! RequestId:%s"), *(Request->GetRequestID()));
+                    Request->AssociateWithTask(FoundTask);
+                    
+                    bDidFindExistingTask = true;
+                }
+			}
+		}
+	}
+    
+    return bDidFindExistingTask;
+}
+
+void FApplePlatformBackgroundHttpManager::SetupNSURLSessionResponseDelegates()
+{
+	OnApp_EnteringBackgroundHandle = FCoreDelegates::ApplicationWillEnterBackgroundDelegate.AddRaw(this, &FApplePlatformBackgroundHttpManager::OnApp_EnteringBackground);
+	OnApp_EnteringForegroundHandle = FCoreDelegates::ApplicationHasEnteredForegroundDelegate.AddRaw(this, &FApplePlatformBackgroundHttpManager::OnApp_EnteringForeground);
+	OnTask_DidFinishDownloadingToURLHandle = FIOSBackgroundDownloadCoreDelegates::OnIOSBackgroundDownload_DidFinishDownloadingToURL.AddRaw(this, &FApplePlatformBackgroundHttpManager::OnTask_DidFinishDownloadingToURL);
+	OnTask_DidWriteDataHandle = FIOSBackgroundDownloadCoreDelegates::OnIOSBackgroundDownload_DidWriteData.AddRaw(this, &FApplePlatformBackgroundHttpManager::OnTask_DidWriteData);
+	OnTask_DidCompleteWithErrorHandle = FIOSBackgroundDownloadCoreDelegates::OnIOSBackgroundDownload_DidCompleteWithError.AddRaw(this, &FApplePlatformBackgroundHttpManager::OnTask_DidCompleteWithError);
+	OnSession_SessionDidFinishAllEventsHandle = FIOSBackgroundDownloadCoreDelegates::OnIOSBackgroundDownload_SessionDidFinishAllEvents.AddRaw(this, &FApplePlatformBackgroundHttpManager::OnSession_SessionDidFinishAllEvents);
+}
+
+void FApplePlatformBackgroundHttpManager::CleanUpNSURLSessionResponseDelegates()
+{
+	FCoreDelegates::ApplicationWillEnterBackgroundDelegate.Remove(OnApp_EnteringBackgroundHandle);
+	FCoreDelegates::ApplicationHasEnteredForegroundDelegate.Remove(OnApp_EnteringForegroundHandle);
+	FIOSBackgroundDownloadCoreDelegates::OnIOSBackgroundDownload_DidFinishDownloadingToURL.Remove(OnTask_DidFinishDownloadingToURLHandle);
+	FIOSBackgroundDownloadCoreDelegates::OnIOSBackgroundDownload_DidWriteData.Remove(OnTask_DidWriteDataHandle);
+	FIOSBackgroundDownloadCoreDelegates::OnIOSBackgroundDownload_DidCompleteWithError.Remove(OnTask_DidCompleteWithErrorHandle);
+	FIOSBackgroundDownloadCoreDelegates::OnIOSBackgroundDownload_SessionDidFinishAllEvents.Remove(OnSession_SessionDidFinishAllEventsHandle);
+}
+
+void FApplePlatformBackgroundHttpManager::OnApp_EnteringForeground()
+{
+    FPlatformAtomics::InterlockedExchange(&bIsInBackground,false);
+	PauseAllActiveRequests();
+}
+
+void FApplePlatformBackgroundHttpManager::OnApp_EnteringBackground()
+{
+	FPlatformAtomics::InterlockedExchange(&bIsInBackground,true);
+	UnpauseAllPausedRequests();
+}
+
+void FApplePlatformBackgroundHttpManager::PauseAllActiveRequests()
+{
+    UE_LOG(LogBackgroundHttpManager, Display, TEXT("Attempting to Pause All Active Tasks"));
+    
+	NSURLSession* BackgroundDownloadSession = FBackgroundURLSessionHandler::GetBackgroundSession();
+	if (nullptr != BackgroundDownloadSession)
+	{
+		[BackgroundDownloadSession getTasksWithCompletionHandler:^(NSArray<__kindof NSURLSessionDataTask*>* DataTasks, NSArray<__kindof NSURLSessionUploadTask*>* UploadTasks, NSArray<__kindof NSURLSessionDownloadTask*>* DownloadTasks)
+		{
+			for (NSURLSessionDownloadTask* DownloadTask : DownloadTasks)
+			{
+				if ([DownloadTask state] == NSURLSessionTaskStateRunning)
+				{
+                    FString TaskURL = [[[DownloadTask currentRequest] URL] absoluteString];
+                    UE_LOG(LogBackgroundHttpManager, Display, TEXT("Pausing Task for URL:%s"), *TaskURL);
+                    
+					[DownloadTask suspend];
+				}
+			}
+            
+            //Reset our active requests to 0 now that we are coming back to the BG
+            FPlatformAtomics::InterlockedExchange(&NumCurrentlyActiveRequests, 0);
+		}];
+	}
+}
+
+void FApplePlatformBackgroundHttpManager::UnpauseAllPausedRequests()
+{
+    UE_LOG(LogBackgroundHttpManager, Display, TEXT("Attempting to Resume All Active Tasks"));
+    
+	NSURLSession* BackgroundDownloadSession = FBackgroundURLSessionHandler::GetBackgroundSession();
+	if (nullptr != BackgroundDownloadSession)
+	{
+		[BackgroundDownloadSession getTasksWithCompletionHandler : ^ (NSArray<__kindof NSURLSessionDataTask*>* DataTasks, NSArray<__kindof NSURLSessionUploadTask*>* UploadTasks, NSArray<__kindof NSURLSessionDownloadTask*>* DownloadTasks)
+		{
+			for (NSURLSessionDownloadTask* DownloadTask : DownloadTasks)
+			{
+				if ([DownloadTask state] == NSURLSessionTaskStateSuspended)
+				{
+                    FString TaskURL = [[[DownloadTask currentRequest] URL] absoluteString];
+                    UE_LOG(LogBackgroundHttpManager, Display, TEXT("Resuming Task for URL:%s"), *TaskURL);
+                    
+					[DownloadTask resume];
+				}
+			}
+		}];
+	}
+}
+
+void FApplePlatformBackgroundHttpManager::OnTask_DidFinishDownloadingToURL(NSURLSessionDownloadTask* Task, NSError* Error, const FString& TempFilePath)
+{
+	FString TaskURL = [[[Task currentRequest] URL] absoluteString];
+    const int ErrorCode = [Error code];
+    const FString ErrorDescription = [Error localizedDescription];
+    
+    IPlatformFile& PlatformFile = FPlatformFileManager::Get().GetPlatformFile();
+	const bool bFileExists = PlatformFile.FileExists(*TempFilePath);
+	
+    UE_LOG(LogBackgroundHttpManager, Display, TEXT("Recieved Background Session Callback for URL:%s | bFileExists:%d | ErrorCode:%d | ErrorDescription:%s | Location:%s"), *TaskURL, (int)(bFileExists), ErrorCode, *ErrorDescription, *TempFilePath);
+
+	//Find request for this task and mark it complete
+	{
+        FRWScopeLock ScopeLock(URLToRequestMapLock, SLT_ReadOnly);
+        FBackgroundHttpURLMappedRequestPtr* WeakRequestInMap = URLToRequestMap.Find(TaskURL);
+        FAppleBackgroundHttpRequestPtr FoundRequest = (nullptr != WeakRequestInMap) ? WeakRequestInMap->Pin() : nullptr;
+        
+        if (FoundRequest.IsValid())
+		{
+			FoundRequest->SetRequestAsSuccess(TempFilePath);
+		}
+        
+		UE_LOG(LogBackgroundHttpManager, Display, TEXT("Attempt To Mark Task Complete -- URL:%s | bDidFindTask:%d"), *TaskURL, (int)(FoundRequest.IsValid()));
+	}
+}
+
+void FApplePlatformBackgroundHttpManager::FinishRequest(FAppleBackgroundHttpRequestPtr Request)
+{
+    //Make sure we another thread hasn't already finished this request
+    bool bHasAlreadyFinishedRequest = FPlatformAtomics::InterlockedExchange(&(Request->bHasAlreadyFinishedRequest), true);
+    if (!bHasAlreadyFinishedRequest)
+    {
+        //by default we will be finishing this request in this function, but some errors might prompt a retry out of this function
+        bool bIsRequestActuallyFinished = true;
+        
+        if (ensureAlwaysMsgf(Request.IsValid(), TEXT("Call to FinishRequest with invalid request!")))
+        {
+            const FString& TempFilePath = Request->CompletedTempDownloadLocation;
+            IPlatformFile& PlatformFile = FPlatformFileManager::Get().GetPlatformFile();
+            const bool bFileExists = PlatformFile.FileExists(*TempFilePath);
+
+            int ResponseCode = bFileExists ? EHttpResponseCodes::Created : EHttpResponseCodes::Unknown;
+
+            if (bFileExists)
+            {
+                UE_LOG(LogBackgroundHttpManager,Display, TEXT("Task Completed Successfully. RequestID:%s TempFileLocation:%s"), *(Request->GetRequestID()), *TempFilePath);
+                FBackgroundHttpResponsePtr NewResponse = FPlatformBackgroundHttp::ConstructBackgroundResponse(ResponseCode, *TempFilePath);
+                Request->CompleteWithExistingResponseData(NewResponse);
+            }
+            else
+            {
+                volatile bool bDidFail = FPlatformAtomics::AtomicRead(&(Request->bIsFailed));
+                
+                //Unexpected case where we didn't find a valid download, but we thought this completed
+                //successfully. Handle this unexpected failure
+                if (!bDidFail)
+                {
+                    UE_LOG(LogBackgroundHttpManager,Warning, TEXT("Task finished downloading, but finished temp file was not found! -- RequestId:%s | TempFileLocation:%s"), *(Request->GetRequestID()), *TempFilePath);
+                    
+                    //Mark our download as not completed as we hit an error so that we don't just keep trying to call FinishRequest
+                    FPlatformAtomics::InterlockedExchange(&(Request->bIsCompleted), false);
+                    FPlatformAtomics::InterlockedExchange(&(Request->bHasAlreadyFinishedRequest), false);
+                    
+                    //Just cancel the task. This will lead to it getting a callback to OnTask_DidCompleteWithError where we will re-create it
+                    Request->CancelActiveTask();
+                    
+                    bIsRequestActuallyFinished = false;
+                }
+                //Expected case where we failed, but expected to fail
+                else
+                {
+                    UE_LOG(LogBackgroundHttpManager, Display, TEXT("Task failed due to being out of retries. -- RequestID:%s"), *(Request->GetRequestID()));
+                    
+                    FBackgroundHttpResponsePtr NewResponse = FPlatformBackgroundHttp::ConstructBackgroundResponse(ResponseCode, TEXT(""));
+                    Request->CompleteWithExistingResponseData(NewResponse);
+                }
+            }
+            
+            //If we are actually finishing this request, lets decrement our NumCurrentlyActiveRequests counter
+            if (bIsRequestActuallyFinished)
+            {
+                //See if we are in the BG or FG. In the BG, we don't want to decrement NumCurrentlyActiveRequests
+                const bool bWasFinishedInBG = FPlatformAtomics::AtomicRead(&bIsInBackground);
+                if (!bWasFinishedInBG)
+                {
+                    int NumActualRequests = FPlatformAtomics::InterlockedDecrement(&NumCurrentlyActiveRequests);
+
+                    //Sanity check that our data is valid. Shouldn't ever trip if everything is working as intended.
+                    const bool bNumActualRequestsValid = ((NumActualRequests > 0) && (NumActualRequests <= FPlatformBackgroundHttp::GetPlatformMaxActiveDownloads()));
+                    ensureMsgf(bNumActualRequestsValid, TEXT("Number of Requests we think are active is invalid! -- NumCurrentlyActiveRequests:%d"), NumActualRequests);
+                }
+            }
+        }
+    }
+    else
+    {
+        UE_LOG(LogBackgroundHttpManager,Display, TEXT("Not finishing Request as its already sending a finish notification -- RequestID:%s"), *(Request->GetRequestID()));
+    }
+}
+
+void FApplePlatformBackgroundHttpManager::RetryRequest(FAppleBackgroundHttpRequestPtr Request, bool bShouldIncreaseRetryCount, bool bShouldStartImmediately, NSData* RetryData)
+{
+	NSURLSessionDownloadTask* NewTask = nullptr;
+
+	if (ensureAlwaysMsgf(Request.IsValid(), TEXT("Call to RetryRequest with an invalid request!")))
+	{
+
+		NSURLSession* BackgroundDownloadSession = FBackgroundURLSessionHandler::GetBackgroundSession();
+		if (ensureAlwaysMsgf((nullptr != BackgroundDownloadSession), TEXT("Invalid Background Download NSURLSession during RetryRequest! Should have already Initialized the NSURLSession by this point!")))
+		{
+			//First, lets see if we should base this task of existing RetryData
+			const bool bShouldUseRetryData = ShouldUseRequestRetryData(Request, RetryData);
+			if (bShouldUseRetryData)
+			{
+				UE_LOG(LogBackgroundHttpManager, Display, TEXT("Resuming Task With Resume Data -- RequestID:%s | RetryData Length:%d"), *(Request->GetRequestID()), [RetryData length]);
+				NewTask = [BackgroundDownloadSession downloadTaskWithResumeData:RetryData];
+			}
+
+			//If not retry data, lets try and just retry on the next CDN
+			else
+			{
+                //Since we created a new task instead of using retry data, reset resume data's retry count on the request
+                Request->ResumeDataRetryCount.Reset();
+                
+				const FString& NewRetryURL = Request->GetURLForRetry(bShouldIncreaseRetryCount);
+				const bool bShouldStartNewRequest = !NewRetryURL.IsEmpty();
+				if (bShouldStartNewRequest)
+				{
+					NSURL* URL = [NSURL URLWithString:NewRetryURL.GetNSString()];
+					NewTask = [BackgroundDownloadSession downloadTaskWithURL:URL];
+				}
+			}
+
+			if (nullptr != NewTask)
+			{
+				Request->AssociateWithTask(NewTask);
+
+				//If we are in BG or flagged this as an immediate start, resume right now without waiting for the FG tick
+				volatile bool bCopyOfBGState = FPlatformAtomics::AtomicRead(&bIsInBackground);
+				if (bCopyOfBGState || bShouldStartImmediately)
+				{
+					Request->ActivateUnderlyingTask();
+				}
+
+				UE_LOG(LogBackgroundHttpManager, Display, TEXT("Created Task for Request -- RequestID:%s | bStartImmediately:%d | bIsAppInBG:%d"), *(Request->GetRequestID()), (int)bShouldStartImmediately, (int)bCopyOfBGState);
+
+				//Always set our bWasTaskStartedInBG flag on our Request so we will know if we need to restart this task next FG Tick.
+				FPlatformAtomics::InterlockedExchange(&(Request->bWasTaskStartedInBG), bCopyOfBGState);
+			}
+			else
+			{
+				UE_LOG(LogBackgroundHttpManager, Display, TEXT("Marking Request Failed. Out of Retries -- RequestID:%s | bShouldUseRetryData:%d"), *(Request->GetRequestID()), (int)bShouldUseRetryData);
+				Request->SetRequestAsFailed();
+			}
+		}
+	}
+}
+
+bool FApplePlatformBackgroundHttpManager::ShouldUseRequestRetryData(FAppleBackgroundHttpRequestPtr Request, NSData* RetryData) const
+{
+	bool bShouldUseData = false;
+
+	if (ensureAlwaysMsgf(Request.IsValid(), TEXT("Call to ShouldUseRequestRetryData with an invalid request!")))
+	{
+		if (IsRetryDataValid(RetryData))
+		{
+			const int CurrentResumeDataRetryCount = Request->ResumeDataRetryCount.Increment();
+			if ((RetryResumeDataLimitSetting < 0) || (CurrentResumeDataRetryCount <= RetryResumeDataLimitSetting))
+			{
+				bShouldUseData = true;
+			}
+		}
+	}
+    
+    return bShouldUseData;
+}
+
+bool FApplePlatformBackgroundHttpManager::IsRetryDataValid(NSData* RetryData) const
+{
+	return ((nullptr != RetryData) && ([RetryData length] > 0));
+}
+
+void FApplePlatformBackgroundHttpManager::OnTask_DidWriteData(NSURLSessionDownloadTask* Task, int64_t BytesWrittenSinceLastCall, int64_t TotalBytesWritten, int64_t TotalBytesExpectedToWrite)
+{
+    if (ensureAlwaysMsgf((nullptr != Task), TEXT("Call to DidWriteData with invalid Task!")))
+    {
+        FString TaskURL = [[[Task currentRequest] URL] absoluteString];
+        
+        //Find task and update it's download progress
+        {
+            FRWScopeLock ScopeLock(URLToRequestMapLock, SLT_ReadOnly);
+            FBackgroundHttpURLMappedRequestPtr* WeakRequestInMap = URLToRequestMap.Find(TaskURL);
+            FAppleBackgroundHttpRequestPtr FoundRequest = (nullptr != WeakRequestInMap) ? WeakRequestInMap->Pin() : nullptr;
+            
+            if (FoundRequest.IsValid())
+            {
+                if (FoundRequest->DownloadProgress < TotalBytesWritten)
+                {
+                    int64 DownloadProgress = FPlatformAtomics::AtomicRead(&(FoundRequest->DownloadProgress));
+                    UE_LOG(LogBackgroundHttpManager, Verbose, TEXT("Updating Task Progress! -- RequestID:%s | Current Progress:%lld | New Progress:%lld"), *(FoundRequest->GetRequestID()), DownloadProgress, TotalBytesWritten);
+                }
+                else
+                {
+                    UE_LOG(LogBackgroundHttpManager, Warning, TEXT("Download Progress tried to go down not up unexpectidly! This could mean a task was unknowingly duplicated! -- RequestID:%s | Current Progress:%lld | New Progress:%lld"), *(FoundRequest->GetRequestID()),FoundRequest->DownloadProgress, TotalBytesWritten);
+                }
+
+                FoundRequest->UpdateDownloadProgress(TotalBytesWritten, BytesWrittenSinceLastCall);
+            }
+        }
+    }
+}
+
+void FApplePlatformBackgroundHttpManager::OnTask_DidCompleteWithError(NSURLSessionTask* Task, NSError* Error)
+{
+	if (ensureAlwaysMsgf((nullptr != Task), TEXT("Call to OnTask_DidCompleteWithError delegate with an invalid task!")))
+	{
+		FString TaskURL = [[[Task currentRequest] URL] absoluteString];
+		const bool bDidCompleteWithError = (nullptr != Error);
+		const int ErrorCode = [Error code];
+        const FString ErrorDescription = [Error localizedDescription];
+        
+		NSData* ResumeData = Error ? [Error.userInfo objectForKey:NSURLSessionDownloadTaskResumeData] : nullptr;
+		const bool bHasResumeData = (ResumeData && ([ResumeData  length] > 0));
+
+		NSNumber* CancelledReasonKey = [Error.userInfo objectForKey:NSURLErrorBackgroundTaskCancelledReasonKey];
+		int CancelledReasonInt = (nullptr != CancelledReasonKey) ? [CancelledReasonKey intValue] : -1;
+
+		FString DebugRetryOverrideReason;
+
+		if (bDidCompleteWithError)
+		{
+            FRWScopeLock ScopeLock(URLToRequestMapLock, SLT_ReadOnly);
+            FBackgroundHttpURLMappedRequestPtr* WeakRequestInMap = URLToRequestMap.Find(TaskURL);
+            FAppleBackgroundHttpRequestPtr FoundRequest = (nullptr != WeakRequestInMap) ? WeakRequestInMap->Pin() : nullptr;
+            const bool bDidFindValidRequest = FoundRequest.IsValid();
+         
+            //by default increase error count. Special cases below will overwrite this
+            bool bShouldRetryIncreaseRetryCount = true;
+            
+            //If we don't have internet, we don't want to move through our CDNs, but rather chain recreate download tasks until we regain internet
+            if ([Error code] == NSURLErrorNotConnectedToInternet)
+            {
+                bShouldRetryIncreaseRetryCount = false;
+                DebugRetryOverrideReason = TEXT("Not Connected To Internet");
+            }
+            
+            UE_LOG(LogBackgroundHttpManager, Display, TEXT("DidCompleteWithError for Task. -- URL:%s | bDidFindVaildRequest:%d | bDidCompleteWithError:%d | ErrorCode:%d | bHasResumeData:%d | CancelledReasonKey:%d | RetryOverrideReason:%s | bShouldRetryIncreaseRetryCount:%d | ErrorDescription:%s"), *TaskURL, (int)bDidFindValidRequest, (int)bDidCompleteWithError, ErrorCode, (int)bHasResumeData, CancelledReasonInt, *DebugRetryOverrideReason, (int)bShouldRetryIncreaseRetryCount, *ErrorDescription);
+            
+			if (bDidFindValidRequest)
+			{
+                RetryRequest(FoundRequest, bShouldRetryIncreaseRetryCount, true, ResumeData);
+			}
+			else
+			{
+				//This can legitimately happen in the case of UnAssociatedTasks completing, so don't error
+				UE_LOG(LogBackgroundHttpManager, Display, TEXT("No Request Found for Errored Task -- TaskURL:%s"), *TaskURL);
+			}
+		}
+	}
+}
+
+void FApplePlatformBackgroundHttpManager::OnSession_SessionDidFinishAllEvents(NSURLSession* Session)
+{
+	//@TODO: TRoss, For now not using this. Will be using this to track how long we take in BG Handling BG Downloads in analytics
+    UE_LOG(LogBackgroundHttpManager, Verbose, TEXT("NSURLSession done sending background events"));
+}
+
+bool FApplePlatformBackgroundHttpManager::Tick(float DeltaTime)
+{
+    TickRequests(DeltaTime);
+    TickTasks(DeltaTime);
+    TickUnassociatedTasks(DeltaTime);
+    
+    //Always keep ticking
+    return true;
+}
+
+void FApplePlatformBackgroundHttpManager::TickRequests(float DeltaTime)
+{
+    //First lets go through all our Requests to see if we need to complete or recreate any requests
+    {
+        FRWScopeLock ScopeLock(ActiveRequestLock, SLT_Write);
+        for (FBackgroundHttpRequestPtr& Request : ActiveRequests)
+        {
+            FAppleBackgroundHttpRequestPtr AppleRequest = StaticCastSharedPtr<FApplePlatformBackgroundHttpRequest>(Request);
+            if (ensureAlwaysMsgf(AppleRequest.IsValid(), TEXT("Invalid Request Pointer in ActiveRequests list!")))
+            {
+                const bool bIsTaskActive = AppleRequest->IsUnderlyingTaskActive();
+                const bool bIsTaskComplete = AppleRequest->IsTaskComplete();
+                const bool bWasStartedInBG = FPlatformAtomics::AtomicRead(&(AppleRequest->bWasTaskStartedInBG));
+                
+                UE_LOG(LogBackgroundHttpManager, VeryVerbose, TEXT("Checking Status of Request on Tick -- RequestID::%s | bIsTaskComplete:%d | bWasStartedInBG:%d"), *(AppleRequest->GetRequestID()), (int)bIsTaskComplete, (int)bWasStartedInBG);
+                
+                if (bIsTaskComplete)
+                {
+                    UE_LOG(LogBackgroundHttpManager, Display, TEXT("Calling FinishRequest On -- RequestID::%s | bIsTaskComplete:%d | bWasStartedInBG:%d"), *(AppleRequest->GetRequestID()), (int)bIsTaskComplete, (int)bWasStartedInBG);
+                    FinishRequest(AppleRequest);
+                }
+                else if (bWasStartedInBG)
+                {
+                    UE_LOG(LogBackgroundHttpManager, Display, TEXT("Cancelling Request Created In BG To Re-Create In FG -- RequestID:%s"), *(AppleRequest->GetRequestID()));
+                    
+                    //Just cancel the task. This will lead to it getting a callback to OnTask_DidCompleteWithError where we will re-create it
+                    //We want to recreate any task spun up in the background as it will not respect our session settings if created in BG.
+                    AppleRequest->CancelActiveTask();
+                }
+                else if (bIsTaskActive)
+                {
+                    const bool bShouldTimeOut = AppleRequest->TickTimeOutTimer(DeltaTime);
+                    if (bShouldTimeOut)
+                    {
+                        UE_LOG(LogBackgroundHttpManager, Display, TEXT("Timing out Request Due To Lack of Server Response -- RequestID:%s"), *(AppleRequest->GetRequestID()));
+                        
+                        //Just cancel the task and let the OnTask_DidCompleteWithError callback handle retrying it if appropriate.
+                        AppleRequest->CancelActiveTask();
+                    }
+                }
+            }
+        }
+    }
+    
+    //Now that we have gone through and finished all the requests, go ahead and delete any pending removes
+    DeletePendingRemoveRequests();
+}
+
+void FApplePlatformBackgroundHttpManager::TickTasks(float DeltaTime)
+{
+    //Go through all our session's tasks and unpause as many tasks as our platform's limit allows
+    
+    NSURLSession* BackgroundDownloadSession = FBackgroundURLSessionHandler::GetBackgroundSession();
+    if (nullptr != BackgroundDownloadSession)
+    {
+        //Make sure we only queue up 1 tick's worth of task parsing at a time as this is an async callback
+        const bool IsIterateAlreadyQueued = FPlatformAtomics::InterlockedExchange(&bIsIteratingThroughSessionTasks, true);
+        if (!IsIterateAlreadyQueued)
+        {
+            [BackgroundDownloadSession getAllTasksWithCompletionHandler:^(NSArray<__kindof NSURLSessionTask*> *tasks)
+             {
+                 //Check to make sure we have room for more tasks to be active first
+                 int CurrentCount = FPlatformAtomics::AtomicRead(&NumCurrentlyActiveRequests);
+                 if (CurrentCount < FPlatformBackgroundHttp::GetPlatformMaxActiveDownloads())
+                 {
+                     //Go through our tasks and try and activate as many as possible
+                     for (NSURLSessionTask* task in tasks)
+                     {
+                         //Only resume state if the task is suspended, not intersted in already active, canceling, or completed tasks here
+                         if ([task state] == NSURLSessionTaskStateSuspended)
+                         {
+                             //Check to make sure we didn't have another task increment our count over the limit since we checked above
+                             int NewRequestCount = FPlatformAtomics::InterlockedIncrement(&NumCurrentlyActiveRequests);
+                             if (NewRequestCount <= FPlatformBackgroundHttp::GetPlatformMaxActiveDownloads())
+                             {
+                                 FString TaskURL = [[[task currentRequest] URL] absoluteString];
+                                 
+                                 //Try and find Request in map that matches this Task
+                                 {
+                                     FRWScopeLock ScopeLock(URLToRequestMapLock, SLT_ReadOnly);
+                                     FBackgroundHttpURLMappedRequestPtr* WeakRequestInMap = URLToRequestMap.Find(TaskURL);
+                                     FAppleBackgroundHttpRequestPtr FoundRequest = (nullptr != WeakRequestInMap) ? WeakRequestInMap->Pin() : nullptr;
+                                     
+                                     if (FoundRequest.IsValid())
+                                     {
+                                         UE_LOG(LogBackgroundHttpManager, Display, TEXT("Activating Task For Requets -- RequestID:%s | TaskURL:%s | CurrentlyActiveRequests:%d"), *(FoundRequest->GetRequestID()), *TaskURL, NewRequestCount);
+                                         FoundRequest->ActivateUnderlyingTask();
+                                     }
+                                     else
+                                     {
+                                         UE_LOG(LogBackgroundHttpManager, Display, TEXT("Skipping Activating Task as there is no associated Request. Once a Request associates with this task, it can then be activated. -- TaskURL:%s"), *TaskURL);
+                                         
+                                         //Don't activate and remove our increment from above because something put us over the limit before we resumed
+                                         FPlatformAtomics::InterlockedDecrement(&NumCurrentlyActiveRequests);
+                                     }
+                                 }
+                             }
+                             else
+                             {
+                                 FString TaskURL = [[[task currentRequest] URL] absoluteString];
+                                 UE_LOG(LogBackgroundHttpManager, Log, TEXT("Task failed to activate as we passed the platform max from another task before we could resume. Task -- TaskURL:%s | CurrentlyActiveRequests:%d"), *TaskURL, NewRequestCount);
+                                 
+                                 //Don't activate and remove our increment from above because something put us over the limit before we resumed
+                                 FPlatformAtomics::InterlockedDecrement(&NumCurrentlyActiveRequests);
+                                 break;
+                             }
+                             
+                             //We now have enough requests queued up that we can stop looking for more
+                             if (NewRequestCount >= FPlatformBackgroundHttp::GetPlatformMaxActiveDownloads())
+                             {
+                                 break;
+                             }
+                         }
+                     }
+                 }
+                 
+                 //If this wasn't true when we exited this async callback, we have a logic error somewhere and have queued this async callback twice. Should always be true as this async callback ends
+                 //Also setting this to false so that future ticks will call this callback again
+                 const bool bWasLogicError = FPlatformAtomics::InterlockedExchange(&bIsIteratingThroughSessionTasks, false);
+                 ensureAlwaysMsgf(bWasLogicError, TEXT("Leaving Tick's getAllTasksWithCompletionHandler and bIsIteratingThroughSessionTasks was false before we finished. Something else has set bIsIteratingThroughSessionTasks during our callback!"));
+             }];
+        }
+    }
+}
+
+void FApplePlatformBackgroundHttpManager::TickUnassociatedTasks(float DeltaTime)
+{
+    //If we don't have anything queued, lets resume any un-associated tasks
+    int CurrentCount = FPlatformAtomics::AtomicRead(&NumCurrentlyActiveRequests);
+    if (CurrentCount == 0)
+    {
+        UnpauseAllUnassociatedTasks();
+    }
+    else
+    {
+        //we have something queued, lets pause unassociated tasks
+        PauseAllUnassociatedTasks();
+    }
+}
diff --git a/Engine/Source/Runtime/Online/BackgroundHTTP/Private/IOS/ApplePlatformBackgroundHttpRequest.cpp b/Engine/Source/Runtime/Online/BackgroundHTTP/Private/IOS/ApplePlatformBackgroundHttpRequest.cpp
new file mode 100644
--- /dev/null
+++ b/Engine/Source/Runtime/Online/BackgroundHTTP/Private/IOS/ApplePlatformBackgroundHttpRequest.cpp
@@ -0,0 +1,171 @@
+// Copyright 1998-2018 Epic Games, Inc. All Rights Reserved.
+#include "IOS/ApplePlatformBackgroundHttpRequest.h"
+#include "Interfaces/IBackgroundHttpResponse.h"
+#include "BackgroundHttpModule.h"
+
+FApplePlatformBackgroundHttpRequest::FApplePlatformBackgroundHttpRequest()
+    : CompletedTempDownloadLocation()
+    , ActiveTimeOutTimer(30.f)
+    , RetryCount(0)
+	, ResumeDataRetryCount(0)
+    , FirstTask(nullptr)
+    , bIsTaskActive(false)
+    , bIsCompleted(false)
+    , bIsFailed(false)
+    , bWasTaskStartedInBG(false)
+    , bHasAlreadyFinishedRequest(false)
+    , DownloadProgress(0)
+{    
+}
+
+void FApplePlatformBackgroundHttpRequest::CompleteWithExistingResponseData(FBackgroundHttpResponsePtr BackgroundResponse)
+{
+    if (ensureAlwaysMsgf(BackgroundResponse.IsValid(), TEXT("Call to CompleteWithExistingResponseData with an invalid response!")))
+    {
+        FBackgroundHttpRequestImpl::CompleteWithExistingResponseData(BackgroundResponse);
+        CompleteRequest_Internal(true, BackgroundResponse->GetTempContentFilePath());
+    }
+}
+
+void FApplePlatformBackgroundHttpRequest::SetRequestAsSuccess(const FString& CompletedTempDownloadLocationIn)
+{
+    CompleteRequest_Internal(true, CompletedTempDownloadLocationIn);
+}
+
+void FApplePlatformBackgroundHttpRequest::SetRequestAsFailed()
+{
+    CompleteRequest_Internal(false, FString());
+}
+
+void FApplePlatformBackgroundHttpRequest::CompleteRequest_Internal(bool bWasRequestSuccess, const FString& CompletedTempDownloadLocationIn)
+{
+    UE_LOG(LogBackgroundHttpRequest, Display, TEXT("Marking Request Complete -- RequestID:%s | bWasRequestSuccess:%d | CompletedTempDownloadLocation:%s"), *GetRequestID(), (int)bWasRequestSuccess, *CompletedTempDownloadLocationIn);
+    
+    FPlatformAtomics::InterlockedExchange(&bIsTaskActive, false);
+	FPlatformAtomics::InterlockedExchange(&bIsCompleted, true);
+	FPlatformAtomics::InterlockedExchange(&bIsFailed, !bWasRequestSuccess);
+    
+    if (!CompletedTempDownloadLocationIn.IsEmpty())
+    {
+        CompletedTempDownloadLocation = CompletedTempDownloadLocationIn;
+    }
+    
+	NotifyNotificationObjectOfComplete(bWasRequestSuccess);
+}
+
+const FString& FApplePlatformBackgroundHttpRequest::GetURLForRetry(bool bShouldIncrementRetryCountFirst)
+{
+	const int NewRetryCount = bShouldIncrementRetryCountFirst ? RetryCount.Increment() : RetryCount.GetValue();
+	
+	//If we are out of Retries, just send an empty string
+	if (NewRetryCount > NumberOfTotalRetries)
+	{
+		UE_LOG(LogBackgroundHttpRequest, Display, TEXT("GetURLForRetry is out of Retries for Request -- RequestID:%s"), *GetRequestID());
+        
+        static FString EmptyResponse = TEXT("");
+        return EmptyResponse;
+	}
+	//Still have remaining retries
+	else
+	{
+		const int URLIndex = NewRetryCount % URLList.Num();
+		const FString& URLToReturn = URLList[URLIndex];
+
+		UE_LOG(LogBackgroundHttpRequest, Display, TEXT("GetURLForRetry found valid URL for current retry -- RequestID:%s | NewRetryCount:%d | URLToReturn:%s"), *GetRequestID(), NewRetryCount, *URLToReturn);
+        return URLToReturn;
+    }
+}
+
+void FApplePlatformBackgroundHttpRequest::ResetProgressTracking()
+{
+    FPlatformAtomics::InterlockedExchange(&DownloadProgress, 0);
+}
+
+void FApplePlatformBackgroundHttpRequest::ActivateUnderlyingTask()
+{
+    volatile FTaskNode* ActiveTaskNode = FirstTask;
+    if (ensureAlwaysMsgf((nullptr != ActiveTaskNode), TEXT("Call to ActivateUnderlyingTask with an invalid node! Need to create underlying task(and node) before activating!")))
+    {
+        NSURLSessionTask* UnderlyingTask = ActiveTaskNode->OurTask;
+        if (ensureAlwaysMsgf((nullptr != UnderlyingTask), TEXT("Call to ActivateUnderlyingTask with an invalid task! Need to create underlying task before activating!")))
+        {
+            FString TaskURL = [[[UnderlyingTask currentRequest] URL] absoluteString];
+            UE_LOG(LogBackgroundHttpRequest, Display, TEXT("Activating Task for Request -- RequestID:%s | TaskURL:%s"), *GetRequestID(), *TaskURL);
+        
+            FPlatformAtomics::InterlockedExchange(&bIsTaskActive, true);
+            [UnderlyingTask resume];
+            
+            ResetTimeOutTimer();
+            ResetProgressTracking();
+        }
+    }
+}
+
+bool FApplePlatformBackgroundHttpRequest::IsUnderlyingTaskActive()
+{
+    return FPlatformAtomics::AtomicRead(&bIsTaskActive);
+}
+
+bool FApplePlatformBackgroundHttpRequest::TickTimeOutTimer(float DeltaTime)
+{
+    ActiveTimeOutTimer -= DeltaTime;
+    return (ActiveTimeOutTimer <= 0.f) ? true : false;
+}
+
+void FApplePlatformBackgroundHttpRequest::ResetTimeOutTimer()
+{
+    ActiveTimeOutTimer = FApplePlatformBackgroundHttpManager::ActiveTimeOutSetting;
+}
+
+void FApplePlatformBackgroundHttpRequest::AssociateWithTask(NSURLSessionTask* ExistingTask)
+{
+	if (ensureAlwaysMsgf((nullptr != ExistingTask), TEXT("Call to AssociateWithTask with an invalid Task! RequestID:%s"), *GetRequestID()))
+	{
+		volatile FTaskNode* NewNode = new FTaskNode();
+		NewNode->OurTask = ExistingTask;
+		
+		//Add a count to our task's reference list so it doesn't get deleted while in our Request's task list
+		[ExistingTask retain];
+
+        //Swap our new node and the first one in the list
+		NewNode->NextNode = (FTaskNode*)FPlatformAtomics::InterlockedExchangePtr((void**)(&FirstTask), (void*)NewNode);
+
+		FString TaskURL = [[[ExistingTask currentRequest] URL] absoluteString];
+		UE_LOG(LogBackgroundHttpRequest, Display, TEXT("Associated Request With New Task -- RequestID:%s | TaskURL:%s"), *GetRequestID(), *TaskURL);
+        
+        ResetTimeOutTimer();
+        ResetProgressTracking();
+	}
+}
+
+void FApplePlatformBackgroundHttpRequest::CancelActiveTask()
+{
+    volatile FTaskNode* TaskNodeWeAreCancelling = FirstTask;
+	if (nullptr != TaskNodeWeAreCancelling)
+	{
+        if (nullptr != TaskNodeWeAreCancelling->OurTask)
+        {
+            FString TaskURL = [[[TaskNodeWeAreCancelling->OurTask currentRequest] URL] absoluteString];
+            UE_LOG(LogBackgroundHttpRequest, Display, TEXT("Cancelling Task -- RequestID:%s | TaskURL:%s"), *GetRequestID(), *TaskURL);
+
+            [TaskNodeWeAreCancelling->OurTask cancel];
+        }
+	}
+}
+
+void FApplePlatformBackgroundHttpRequest::UpdateDownloadProgress(int64_t TotalDownloaded,int64_t DownloadedSinceLastUpdate)
+{
+    UE_LOG(LogBackgroundHttpRequest, VeryVerbose, TEXT("Request Update Progress -- RequestID:%s | OldProgress:%lld | NewProgress:%lld | ProgressSinceLastUpdate:%lld"), *GetRequestID(), DownloadProgress, TotalDownloaded, DownloadedSinceLastUpdate);
+	
+    FPlatformAtomics::AtomicStore(&DownloadProgress, TotalDownloaded);
+    ResetTimeOutTimer();
+    
+    OnProgressUpdated().ExecuteIfBound(SharedThis(this), TotalDownloaded, DownloadedSinceLastUpdate);
+}
+
+bool FApplePlatformBackgroundHttpRequest::IsTaskComplete() const
+{
+    const bool bDidRequestFail = FPlatformAtomics::AtomicRead(&bIsFailed);
+    const bool bDidRequestComplete = FPlatformAtomics::AtomicRead(&bIsCompleted);
+    return (bDidRequestFail || bDidRequestComplete);
+}
diff --git a/Engine/Source/Runtime/Online/BackgroundHTTP/Private/IOS/ApplePlatformBackgroundHttpResponse.cpp b/Engine/Source/Runtime/Online/BackgroundHTTP/Private/IOS/ApplePlatformBackgroundHttpResponse.cpp
new file mode 100644
--- /dev/null
+++ b/Engine/Source/Runtime/Online/BackgroundHTTP/Private/IOS/ApplePlatformBackgroundHttpResponse.cpp
@@ -0,0 +1,9 @@
+// Copyright 1998-2018 Epic Games, Inc. All Rights Reserved.
+
+#include "IOS/ApplePlatformBackgroundHttpResponse.h"
+
+FApplePlatformBackgroundHttpResponse::FApplePlatformBackgroundHttpResponse(int32 ResponseCodeIn, const FString& TempContentFilePathIn)
+{
+    TempContentFilePath = TempContentFilePathIn;
+    ResponseCode = ResponseCodeIn;
+}
diff --git a/Engine/Source/Runtime/Online/BackgroundHTTP/Public/Android/AndroidPlatformBackgroundHttp.h b/Engine/Source/Runtime/Online/BackgroundHTTP/Public/Android/AndroidPlatformBackgroundHttp.h
new file mode 100644
--- /dev/null
+++ b/Engine/Source/Runtime/Online/BackgroundHTTP/Public/Android/AndroidPlatformBackgroundHttp.h
@@ -0,0 +1,6 @@
+#pragma once
+
+#include "GenericPlatform/GenericPlatformBackgroundHttp.h"
+
+//Currently no implementation on Android
+typedef FGenericPlatformBackgroundHttp FPlatformBackgroundHttp;
\ No newline at end of file
diff --git a/Engine/Source/Runtime/Online/BackgroundHTTP/Public/BackgroundHttp.h b/Engine/Source/Runtime/Online/BackgroundHTTP/Public/BackgroundHttp.h
new file mode 100644
--- /dev/null
+++ b/Engine/Source/Runtime/Online/BackgroundHTTP/Public/BackgroundHttp.h
@@ -0,0 +1,13 @@
+// Copyright 1998-2018 Epic Games, Inc. All Rights Reserved.
+
+#pragma once
+
+// Module dependencies
+
+#include "CoreMinimal.h"
+#include "BackgroundHttpModule.h"
+
+// Module includes
+#include "PlatformBackgroundHttp.h"
+
+BACKGROUNDHTTP_API DECLARE_LOG_CATEGORY_EXTERN(LogBackgroundHttp, Display, All);
diff --git a/Engine/Source/Runtime/Online/BackgroundHTTP/Public/BackgroundHttpManagerImpl.h b/Engine/Source/Runtime/Online/BackgroundHTTP/Public/BackgroundHttpManagerImpl.h
new file mode 100644
--- /dev/null
+++ b/Engine/Source/Runtime/Online/BackgroundHTTP/Public/BackgroundHttpManagerImpl.h
@@ -0,0 +1,52 @@
+// Copyright 1998-2018 Epic Games, Inc. All Rights Reserved.
+
+#pragma once
+
+#include "Interfaces/IBackgroundHttpManager.h"
+#include "Interfaces/IBackgroundHttpResponse.h"
+
+#include "Containers/Ticker.h"
+
+DECLARE_LOG_CATEGORY_EXTERN(LogBackgroundHttpManager, Log, All)
+
+/**
+ * Contains implementation of some common functions that don't vary between implementation
+ */
+class FBackgroundHttpManagerImpl 
+	: public IBackgroundHttpManager
+	, public FTickerObjectBase
+{
+public:
+	FBackgroundHttpManagerImpl();
+	virtual ~FBackgroundHttpManagerImpl();
+
+	virtual void AddRequest(const FBackgroundHttpRequestPtr Request) override;
+	virtual void RemoveRequest(const FBackgroundHttpRequestPtr Request) override;
+
+	virtual void Initialize() override;
+	virtual void Shutdown() override;
+
+	virtual void CleanUpTemporaryFiles() override;
+
+	//FTickerObjectBase implementation
+	virtual bool Tick(float DeltaTime) override;
+
+protected:
+	virtual bool AssociateWithAnyExistingRequest(const FBackgroundHttpRequestPtr Request) override;
+
+	virtual bool CheckForExistingCompletedDownload(const FBackgroundHttpRequestPtr Request, FString& ExistingFilePathOut, int64& ExistingFileSizeOut);
+	virtual void ActivatePendingRequests();
+
+protected:
+
+	/** List of Background Http requests that we have called AddRequest on, but have not yet started due to platform active download limits **/
+	TArray<FBackgroundHttpRequestPtr> PendingStartRequests;
+	FRWLock PendingRequestLock;
+
+	/** List of Background Http requests that are actively being processed **/
+	TArray<FBackgroundHttpRequestPtr> ActiveRequests;
+	FRWLock ActiveRequestLock;
+
+	/** Count of how many requests we have active **/
+	volatile int NumCurrentlyActiveRequests;
+};
diff --git a/Engine/Source/Runtime/Online/BackgroundHTTP/Public/BackgroundHttpModule.h b/Engine/Source/Runtime/Online/BackgroundHTTP/Public/BackgroundHttpModule.h
new file mode 100644
--- /dev/null
+++ b/Engine/Source/Runtime/Online/BackgroundHTTP/Public/BackgroundHttpModule.h
@@ -0,0 +1,59 @@
+// Copyright 1998-2018 Epic Games, Inc. All Rights Reserved.
+
+#pragma once
+
+#include "CoreMinimal.h"
+#include "Misc/CoreMisc.h"
+#include "Interfaces/IBackgroundHttpRequest.h"
+#include "Interfaces/IBackgroundHttpManager.h"
+#include "Modules/ModuleInterface.h"
+
+
+/**
+ * Module for Http Background request implementations
+ */
+class FBackgroundHttpModule : 
+	public IModuleInterface
+{
+public:
+	/**
+	 * Singleton-like access to this module's interface.  This is just for convenience!
+	 * Beware of calling this during the shutdown phase, though.  Your module might have been unloaded already.
+	 *
+	 * @return Returns singleton instance, loading the module on demand if needed
+	 */
+	BACKGROUNDHTTP_API static FBackgroundHttpModule& Get();
+
+	/**
+	* Creates a BackgroundRequest through the appropriate platform-layer.
+	*
+	* @return SharedRef to the new BackgroundRequest.
+	*/
+	BACKGROUNDHTTP_API static FBackgroundHttpRequestPtr CreateBackgroundRequest();
+
+	/**
+	* Gets the current BackgroundHttpManager from the appropriate platform layer.
+	*
+	* @return SharedPtr to the BackgroundHttpManager.
+	*/
+	BACKGROUNDHTTP_API FBackgroundHttpManagerPtr GetBackgroundHttpManager();
+
+public:
+	// IModuleInterface
+	/**
+	 * Called when Background Http module is loaded
+	 */
+	virtual void StartupModule() override;
+
+	/**
+	 * Called when Background Http module is unloaded
+	 */
+	virtual void ShutdownModule() override;
+
+private:
+	/** Keeps track of Background Http requests while they are being processed */
+	FBackgroundHttpManagerPtr BackgroundHttpManager;
+
+	/** singleton for the module while loaded and available */
+	static FBackgroundHttpModule* Singleton;
+};
diff --git a/Engine/Source/Runtime/Online/BackgroundHTTP/Public/BackgroundHttpNotificationObject.h b/Engine/Source/Runtime/Online/BackgroundHTTP/Public/BackgroundHttpNotificationObject.h
new file mode 100644
--- /dev/null
+++ b/Engine/Source/Runtime/Online/BackgroundHTTP/Public/BackgroundHttpNotificationObject.h
@@ -0,0 +1,31 @@
+// Copyright 1998-2018 Epic Games, Inc. All Rights Reserved.
+#pragma once
+
+#include "CoreMinimal.h"
+
+
+class BACKGROUNDHTTP_API FBackgroundHttpNotificationObject
+	: public TSharedFromThis<FBackgroundHttpNotificationObject, ESPMode::ThreadSafe>
+{
+public:
+	FBackgroundHttpNotificationObject(FText InNotificationTitle, FText InNotificationAction, FText InNotificationBody, const FString& InNotificationActivationString, bool InNotifyOnlyOnFullSuccess);
+	~FBackgroundHttpNotificationObject();
+
+	void NotifyOfDownloadResult(bool bWasSuccess);
+
+private:
+	FText NotificationTitle;
+	FText NotificationAction;
+	FText NotificationBody;
+	FString NotificationActivationString;
+
+	bool bNotifyOnlyOnFullSuccess;
+	volatile int32 NumFailedDownloads;
+
+	class ILocalNotificationService* PlatformNotificationService;
+
+	//No default constructor
+	FBackgroundHttpNotificationObject() {}
+};
+
+typedef TSharedPtr<FBackgroundHttpNotificationObject, ESPMode::ThreadSafe> FBackgroundHttpNotificationObjectPtr;
diff --git a/Engine/Source/Runtime/Online/BackgroundHTTP/Public/BackgroundHttpRequestImpl.h b/Engine/Source/Runtime/Online/BackgroundHTTP/Public/BackgroundHttpRequestImpl.h
new file mode 100644
--- /dev/null
+++ b/Engine/Source/Runtime/Online/BackgroundHTTP/Public/BackgroundHttpRequestImpl.h
@@ -0,0 +1,46 @@
+// Copyright 1998-2018 Epic Games, Inc. All Rights Reserved.
+
+#pragma once
+
+#include "BackgroundHttpNotificationObject.h"
+#include "Interfaces/IBackgroundHttpRequest.h"
+
+DECLARE_LOG_CATEGORY_EXTERN(LogBackgroundHttpRequest, Log, All)
+
+/**
+ * Contains implementation of some common functions that don't have to vary between implementation
+ */
+class BACKGROUNDHTTP_API FBackgroundHttpRequestImpl 
+	: public IBackgroundHttpRequest
+{
+public:
+	FBackgroundHttpRequestImpl();
+	virtual ~FBackgroundHttpRequestImpl() {}
+
+	//This should be called from the platform level when a BG download finishes.
+	virtual void OnBackgroundDownloadComplete();
+
+	// IBackgroundHttpRequest
+	virtual bool ProcessRequest() override;
+	virtual void SetURLAsList(const TArray<FString>& URLs, int NumRetriesToAttempt) override;
+	virtual const TArray<FString>& GetURLList() const override;
+	virtual void SetCompleteNotification(FBackgroundHttpNotificationObjectPtr DownloadCompleteNotificationObjectIn) override;
+	virtual void CompleteWithExistingResponseData(FBackgroundHttpResponsePtr BackgroundResponse) override;
+	virtual FBackgroundHttpRequestCompleteDelegate& OnProcessRequestComplete() override;
+	virtual FBackgroundHttpProgressUpdateDelegate& OnProgressUpdated() override;
+	virtual const FBackgroundHttpResponsePtr GetResponse() const override;
+	virtual const FString& GetRequestID() const override;
+	virtual void SetRequestID(const FString& NewRequestID) override;
+	virtual bool HandleDelayedProcess() override;
+
+	virtual void NotifyNotificationObjectOfComplete(bool bWasSuccess);
+protected:
+	TSharedPtr<FBackgroundHttpNotificationObject, ESPMode::ThreadSafe> DownloadCompleteNotificationObject;
+	FBackgroundHttpResponsePtr Response;
+	TArray<FString> URLList;
+	FString RequestID;
+	int NumberOfTotalRetries;
+
+	FBackgroundHttpRequestCompleteDelegate HttpRequestCompleteDelegate;
+	FBackgroundHttpProgressUpdateDelegate HttpProgressUpdateDelegate;
+};
\ No newline at end of file
diff --git a/Engine/Source/Runtime/Online/BackgroundHTTP/Public/BackgroundHttpResponseImpl.h b/Engine/Source/Runtime/Online/BackgroundHTTP/Public/BackgroundHttpResponseImpl.h
new file mode 100644
--- /dev/null
+++ b/Engine/Source/Runtime/Online/BackgroundHTTP/Public/BackgroundHttpResponseImpl.h
@@ -0,0 +1,24 @@
+// Copyright 1998-2018 Epic Games, Inc. All Rights Reserved.
+
+#pragma once
+
+#include "Interfaces/IBackgroundHttpResponse.h"
+
+/**
+ * Contains implementation of some common functions that don't vary between implementation
+ */
+class FBackgroundHttpResponseImpl 
+	: public IBackgroundHttpResponse
+{
+public:
+	FBackgroundHttpResponseImpl();
+	virtual ~FBackgroundHttpResponseImpl() {}
+
+	//IHttpBackgroundResponse
+	virtual int32 GetResponseCode() const;
+	virtual const FString& GetTempContentFilePath() const;
+
+protected:	
+	FString TempContentFilePath;
+	int32 ResponseCode;
+};
diff --git a/Engine/Source/Runtime/Online/BackgroundHTTP/Public/GenericPlatform/GenericPlatformBackgroundHttp.h b/Engine/Source/Runtime/Online/BackgroundHTTP/Public/GenericPlatform/GenericPlatformBackgroundHttp.h
new file mode 100644
--- /dev/null
+++ b/Engine/Source/Runtime/Online/BackgroundHTTP/Public/GenericPlatform/GenericPlatformBackgroundHttp.h
@@ -0,0 +1,74 @@
+// Copyright 1998-2018 Epic Games, Inc. All Rights Reserved.
+
+
+#pragma once
+
+#include "CoreMinimal.h"
+
+#include "Interfaces/IBackgroundHttpRequest.h"
+
+class IBackgroundHttpManager;
+
+/**
+ * Generic version of Background Http implementations for platforms that don't need a special implementation
+ * Intended usage is to use FPlatformBackgroundHttp instead of FGenericPlatformHttp
+ * On platforms without a specific implementation, you should still use FPlatformBackgroundHttp and it will call into these functions
+ */
+class BACKGROUNDHTTP_API FGenericPlatformBackgroundHttp
+{
+public:
+	/**
+	 * Platform initialization step
+	 */
+	static void Initialize();
+
+	/**
+	 * Platform shutdown step
+	 */
+	static void Shutdown();
+
+	/**
+	 * Creates a platform-specific Background HTTP manager.
+	 * Un-implemented platforms should create a FGenericPlatformBackgroundHttpManager
+	 */
+	static FBackgroundHttpManagerPtr CreatePlatformBackgroundHttpManager();
+
+	/**
+	 * Creates a new Background Http request instance for the current platform
+	 * that will continue to download when the application is in the background
+	 *
+	 * @return request object
+	 */
+	static FBackgroundHttpRequestPtr ConstructBackgroundRequest();
+
+	/**
+	 * Creates a new Background Http Response instance for the current platform
+	 * This normally is called by the request and associated with itself.
+	 *
+	 * @return response object
+	 */
+	static FBackgroundHttpResponsePtr ConstructBackgroundResponse(int32 ResponseCode, const FString& TempFilePath);
+
+	/**
+	* Function that takes in a URL and figures out the location we should use as the temp storage URL
+	*
+	* @param URL FString URL of the request that we need to find the end location for.
+	*
+	* @return FString File Path the temp file would be located at for the supplied URL.
+	*/
+	static const FString GetTemporaryFilePathFromURL(const FString& URL);
+
+	/**
+	* Function that returns the root path where all our temporary files are stored on this platform
+	*
+	* @return FString File path that is the root path of our temp files for background http work on this platform.
+	*/
+	static const FString& GetTemporaryRootPath();
+
+	/**
+	* Function that returns how many active BackgroundHttpRequests we should have actively downloading at once. 
+	*
+	* @return int Number of downloads we should have active at once.
+	*/
+	static int GetPlatformMaxActiveDownloads();
+};
diff --git a/Engine/Source/Runtime/Online/BackgroundHTTP/Public/GenericPlatform/GenericPlatformBackgroundHttpManager.h b/Engine/Source/Runtime/Online/BackgroundHTTP/Public/GenericPlatform/GenericPlatformBackgroundHttpManager.h
new file mode 100644
--- /dev/null
+++ b/Engine/Source/Runtime/Online/BackgroundHTTP/Public/GenericPlatform/GenericPlatformBackgroundHttpManager.h
@@ -0,0 +1,20 @@
+// Copyright 1998-2018 Epic Games, Inc. All Rights Reserved.
+
+#pragma once
+
+#include "CoreMinimal.h"
+
+#include "BackgroundHttpManagerImpl.h"
+#include "Containers/Ticker.h"
+
+#include "GenericPlatform/GenericPlatformBackgroundHttpRequest.h"
+
+/**
+ * Manages Background Http request that are currently being processed if no platform specific implementation has been made. 
+ */
+class BACKGROUNDHTTP_API FGenericPlatformBackgroundHttpManager
+	: public FBackgroundHttpManagerImpl
+{
+public:
+	virtual ~FGenericPlatformBackgroundHttpManager();
+};
\ No newline at end of file
diff --git a/Engine/Source/Runtime/Online/BackgroundHTTP/Public/GenericPlatform/GenericPlatformBackgroundHttpRequest.h b/Engine/Source/Runtime/Online/BackgroundHTTP/Public/GenericPlatform/GenericPlatformBackgroundHttpRequest.h
new file mode 100644
--- /dev/null
+++ b/Engine/Source/Runtime/Online/BackgroundHTTP/Public/GenericPlatform/GenericPlatformBackgroundHttpRequest.h
@@ -0,0 +1,68 @@
+// Copyright 1998-2018 Epic Games, Inc. All Rights Reserved.
+
+#pragma once
+
+#include "Interfaces/IBackgroundHttpRequest.h"
+#include "Interfaces/IHttpRequest.h"
+
+#include "BackgroundHttpRequestImpl.h"
+
+
+/**
+ * Contains implementation of some common functions that don't vary between implementation
+ */
+class BACKGROUNDHTTP_API FGenericPlatformBackgroundHttpRequest 
+	: public FBackgroundHttpRequestImpl
+{
+public:
+
+	FGenericPlatformBackgroundHttpRequest();
+	virtual ~FGenericPlatformBackgroundHttpRequest() {}
+
+	//IHttpBackgroundHttpRequest
+	virtual bool HandleDelayedProcess() override;
+	virtual void CompleteWithExistingResponseData(FBackgroundHttpResponsePtr BackgroundResponse) override;
+
+protected:
+	//In the default implementation we actually use HTTPRequests instead of HttpBackgroundRequests.
+	//This class handles wrapping some of the HttpBackgroundRequest only functionality (such as request lists
+	//and notification objects) so they work with non-background HttpRequests
+	class FGenericPlatformBackgroundHttpWrapper
+	{
+	public:
+		FGenericPlatformBackgroundHttpWrapper(FBackgroundHttpRequestPtr Request, int MaxRetriesToAttempt);
+		~FGenericPlatformBackgroundHttpWrapper();
+
+		void MakeRequest();
+		void HttpRequestComplete(FHttpRequestPtr HttpRequestIn, FHttpResponsePtr HttpResponse, bool bSuccess);
+
+	private:
+		void UpdateHttpProgress(FHttpRequestPtr UnderlyingHttpRequest, int32 BytesSent, int32 BytesReceived);
+
+		void CleanUpHttpRequest();
+
+		// Gets the current URL the HTTP Request should use by comparing the URLList in the background request to the current retry number.
+		// Returns empty FString if we are "out of" retries (CurrentRetryNumber > MaxRetries) or if no valid URL.
+		const FString GetURLForCurrentRetryNumber();
+
+		//Don't want to provide default implementation. Must provide an IHttpBackgroundRequest to wrap in public constructor
+		FGenericPlatformBackgroundHttpWrapper() {}
+
+		//Reference to the creating request
+		FBackgroundHttpRequestPtr OriginalRequest;
+
+		//Current Http Request being processed by this wrapper
+		TSharedPtr<IHttpRequest> HttpRequest;
+
+		//Tracking retry number we are currently on
+		int32 CurrentRetryNumber;
+
+		//Max number of retries we should do
+		int32 MaxRetries;
+
+		//How many bytes we had last time we sent progress updates
+		int32 LastProgressUpdateBytes;
+	};
+
+	TSharedPtr<FGenericPlatformBackgroundHttpWrapper> RequestWrapper;
+};
diff --git a/Engine/Source/Runtime/Online/BackgroundHTTP/Public/GenericPlatform/GenericPlatformBackgroundHttpResponse.h b/Engine/Source/Runtime/Online/BackgroundHTTP/Public/GenericPlatform/GenericPlatformBackgroundHttpResponse.h
new file mode 100644
--- /dev/null
+++ b/Engine/Source/Runtime/Online/BackgroundHTTP/Public/GenericPlatform/GenericPlatformBackgroundHttpResponse.h
@@ -0,0 +1,20 @@
+// Copyright 1998-2018 Epic Games, Inc. All Rights Reserved.
+
+#pragma once
+
+#include "BackgroundHttpResponseImpl.h"
+#include "Interfaces/IHttpRequest.h"
+
+/**
+ * Contains implementation of some common functions that don't vary between implementation
+ */
+class FGenericPlatformBackgroundHttpResponse
+	: public FBackgroundHttpResponseImpl
+{
+public:
+	//Provide constructor to create an already 
+	FGenericPlatformBackgroundHttpResponse(int32 ResponseCode, const FString& TempContentFilePath);
+
+	//Provide constuctor to create a request from underlying HttpRequest information
+	FGenericPlatformBackgroundHttpResponse(FHttpRequestPtr HttpRequestIn, FHttpResponsePtr HttpResponse, bool bSuccess);
+};
\ No newline at end of file
diff --git a/Engine/Source/Runtime/Online/BackgroundHTTP/Public/HTML5/HTML5PlatformBackgroundHttp.h b/Engine/Source/Runtime/Online/BackgroundHTTP/Public/HTML5/HTML5PlatformBackgroundHttp.h
new file mode 100644
--- /dev/null
+++ b/Engine/Source/Runtime/Online/BackgroundHTTP/Public/HTML5/HTML5PlatformBackgroundHttp.h
@@ -0,0 +1,6 @@
+#pragma once
+
+#include "GenericPlatform/GenericPlatformBackgroundHttp.h"
+
+//Currently no implementation on HTML5
+typedef FGenericPlatformBackgroundHttp FPlatformBackgroundHttp;
\ No newline at end of file
diff --git a/Engine/Source/Runtime/Online/BackgroundHTTP/Public/IOS/ApplePlatformBackgroundHttp.h b/Engine/Source/Runtime/Online/BackgroundHTTP/Public/IOS/ApplePlatformBackgroundHttp.h
new file mode 100644
--- /dev/null
+++ b/Engine/Source/Runtime/Online/BackgroundHTTP/Public/IOS/ApplePlatformBackgroundHttp.h
@@ -0,0 +1,74 @@
+// Copyright 1998-2018 Epic Games, Inc. All Rights Reserved.
+
+
+#pragma once
+
+#include "CoreMinimal.h"
+
+#include "Interfaces/IBackgroundHttpRequest.h"
+#include "Interfaces/IBackgroundHttpManager.h"
+#include "Interfaces/IBackgroundHttpResponse.h"
+
+/**
+ * Apple specific Background Http implementations
+ * Intended usage is to use FPlatformBackgroundHttp
+ */
+class BACKGROUNDHTTP_API FApplePlatformBackgroundHttp
+{
+public:
+	/**
+	 * Platform initialization step
+	 */
+	static void Initialize();
+
+	/**
+	 * Platform shutdown step
+	 */
+	static void Shutdown();
+
+	/**
+	 * Creates a platform-specific Background HTTP manager.
+	 * Un-implemented platforms should create a FGenericPlatformBackgroundHttpManager
+	 */
+	static FBackgroundHttpManagerPtr CreatePlatformBackgroundHttpManager();
+
+	/**
+	 * Creates a new Background Http request instance for the current platform
+	 * that will continue to download when the application is in the background
+	 *
+	 * @return request object
+	 */
+	static FBackgroundHttpRequestPtr ConstructBackgroundRequest();
+
+	/**
+	 * Creates a new Background Http Response instance for the current platform
+	 * This normally is called by the request and associated with itself.
+	 *
+	 * @return response object
+	 */
+	static FBackgroundHttpResponsePtr ConstructBackgroundResponse(int32 ResponseCode, const FString& TempFilePath);
+
+	/**
+	* Function that takes in a URL and figures out the location we should use as the temp storage URL
+	*
+	* @return FString to use as the TempFilePath
+	*/
+	static const FString GetTemporaryFilePathFromURL(const FString& URL);
+
+    /**
+     * Function that returns the root path where all our temporary files are stored on this platform
+     *
+     * @return FString File path that is the root path of our temp files for background http work on this platform.
+     */
+    static const FString& GetTemporaryRootPath();
+    
+	/**
+	* Function that returns how many active BackgroundHttpRequests we should have actively downloading at once.
+	*
+	* @return int Number of downloads we should have active at once.
+	*/
+	static int GetPlatformMaxActiveDownloads();
+};
+
+//Setup Platform Implementation calls to point here at the Apple Implementation
+typedef FApplePlatformBackgroundHttp FPlatformBackgroundHttp;
diff --git a/Engine/Source/Runtime/Online/BackgroundHTTP/Public/IOS/ApplePlatformBackgroundHttpManager.h b/Engine/Source/Runtime/Online/BackgroundHTTP/Public/IOS/ApplePlatformBackgroundHttpManager.h
new file mode 100644
--- /dev/null
+++ b/Engine/Source/Runtime/Online/BackgroundHTTP/Public/IOS/ApplePlatformBackgroundHttpManager.h
@@ -0,0 +1,113 @@
+// Copyright 1998-2018 Epic Games, Inc. All Rights Reserved.
+
+#pragma once
+
+#include "CoreMinimal.h"
+
+#include "BackgroundHttpManagerImpl.h"
+#include "Interfaces/IBackgroundHttpRequest.h"
+#include "IOS/ApplePlatformBackgroundHttpRequest.h"
+
+typedef TWeakPtr<FApplePlatformBackgroundHttpRequest,ESPMode::ThreadSafe> FBackgroundHttpURLMappedRequestPtr;
+
+/**
+ * Manages Background Http request that are currently being processed if we are on an Apple Platform
+ */
+class BACKGROUNDHTTP_API FApplePlatformBackgroundHttpManager
+	: public FBackgroundHttpManagerImpl
+{
+public:
+	virtual void Initialize() override;
+	virtual void Shutdown() override;
+	virtual void AddRequest(const FBackgroundHttpRequestPtr Request) override;
+	virtual void RemoveRequest(const FBackgroundHttpRequestPtr Request) override;
+
+	/**
+	* Constructor
+	*/
+	FApplePlatformBackgroundHttpManager();
+
+	/**
+	 * Destructor
+	 */
+	virtual ~FApplePlatformBackgroundHttpManager();
+
+    //FTickerObjectBase implementation
+    virtual bool Tick(float DeltaTime) override;
+    
+    //This is the identifier we expect to be associated with our NSURLDownloadSession for background downloads
+    static FString BackgroundSessionIdentifier;
+    
+    //This value is pulled from .ini settings and is used to time out requests when we are active in the foreground
+    static float ActiveTimeOutSetting;
+    
+	//This value is pulled from .ini setting and is used to limit the number of times a request will try and use the same
+	//resume data before just ignoring that data and moving onto another CDN retry
+	static int RetryResumeDataLimitSetting;
+
+protected:
+	virtual bool AssociateWithAnyExistingRequest(const FBackgroundHttpRequestPtr Request) override;
+	
+	bool CheckForExistingUnAssociatedTask(const FAppleBackgroundHttpRequestPtr Request);
+
+    void DeletePendingRemoveRequests();
+    
+protected:
+	//This dictionary is used to hold tasks that already existed on our Background Session when our BackgroundHttpManager was initialized. See PopulateUnAssociatedTasks and CheckForExistingUnAssociatedTask
+	NSMutableDictionary<NSString*, NSURLSessionDownloadTask*>* UnAssociatedTasks;
+
+	//Map to hold the associated BackgroundHttpRequest for any given URL. Multiple URLs will end up pointing to the same request in this list as it stores
+    TMap<const FString, FBackgroundHttpURLMappedRequestPtr> URLToRequestMap;
+	FRWLock URLToRequestMapLock;
+
+private:
+	//Checks for tasks that already exist on the BackgroundSession. Should only be called during Initialize. These tasks should really only exist 
+	//due to previous application sessions starting background downloads that finished after the app closed. (IE: If the user started a patch and then our app was terminated by the OS while 
+	//suspended. On next launch we will have re-associated those existing downloads to our background download session, but don't yet have a FBackgroundHttpRequest for that task.)
+	//See AssociateWithAnyExistingRequest for how these Tasks are eventually claimed by an FBackgroundHttpRequest
+	void PopulateUnAssociatedTasks();
+
+	void SetupNSURLSessionResponseDelegates();
+	void CleanUpNSURLSessionResponseDelegates();
+
+	void PauseAllActiveRequests();
+	void UnpauseAllPausedRequests();
+
+    void PauseAllUnassociatedTasks();
+    void UnpauseAllUnassociatedTasks();
+    
+    void TickRequests(float DeltaTime);
+    void TickTasks(float DeltaTime);
+    void TickUnassociatedTasks(float DeltaTime);
+    
+    void StartRequest(FAppleBackgroundHttpRequestPtr Request);
+	void FinishRequest(FAppleBackgroundHttpRequestPtr Request);
+	void RetryRequest(FAppleBackgroundHttpRequestPtr Request, bool bShouldIncreaseRetryCount, bool bShouldStartImmediately, NSData* RetryData);
+    void RemoveSessionTasksForRequest(FAppleBackgroundHttpRequestPtr Request);
+    void GenerateURLMapEntriesForRequest(FAppleBackgroundHttpRequestPtr Request);
+    
+	bool IsRetryDataValid(NSData* RetryData) const;
+	bool ShouldUseRequestRetryData(FAppleBackgroundHttpRequestPtr Request, NSData* RetryData) const;
+
+	//Delegates to respond to Background Session Events
+	void OnApp_EnteringForeground();
+	void OnApp_EnteringBackground();
+	void OnTask_DidFinishDownloadingToURL(NSURLSessionDownloadTask* Task, NSError* Error, const FString& TempFilePath);
+	void OnTask_DidWriteData(NSURLSessionDownloadTask* Task, int64_t BytesWrittenSinceLastCall, int64_t TotalBytesWritten, int64_t TotalBytesExpectedToWrite);
+	void OnTask_DidCompleteWithError(NSURLSessionTask* Task, NSError* Error);
+	void OnSession_SessionDidFinishAllEvents(NSURLSession* Session);
+
+	FDelegateHandle OnApp_EnteringForegroundHandle;
+	FDelegateHandle OnApp_EnteringBackgroundHandle;
+	FDelegateHandle OnTask_DidFinishDownloadingToURLHandle;
+	FDelegateHandle OnTask_DidWriteDataHandle;
+	FDelegateHandle OnTask_DidCompleteWithErrorHandle;
+	FDelegateHandle OnSession_SessionDidFinishAllEventsHandle;
+
+	volatile int bHasFinishedPopulatingUnassociatedTasks;
+	volatile int bIsInBackground;
+    volatile int bIsIteratingThroughSessionTasks;
+    
+    //Need to store requests to remove at the end of our tick as we find them during iteration
+    TArray<FBackgroundHttpRequestPtr> RequestsPendingRemove;
+};
diff --git a/Engine/Source/Runtime/Online/BackgroundHTTP/Public/IOS/ApplePlatformBackgroundHttpRequest.h b/Engine/Source/Runtime/Online/BackgroundHTTP/Public/IOS/ApplePlatformBackgroundHttpRequest.h
new file mode 100644
--- /dev/null
+++ b/Engine/Source/Runtime/Online/BackgroundHTTP/Public/IOS/ApplePlatformBackgroundHttpRequest.h
@@ -0,0 +1,65 @@
+// Copyright 1998-2018 Epic Games, Inc. All Rights Reserved.
+
+#pragma once
+
+#include "Interfaces/IBackgroundHttpRequest.h"
+#include "Interfaces/IHttpRequest.h"
+
+#include "BackgroundHttpRequestImpl.h"
+
+/**
+ * Contains implementation of Apple specific background http requests
+ */
+class BACKGROUNDHTTP_API FApplePlatformBackgroundHttpRequest 
+	: public FBackgroundHttpRequestImpl
+{
+public:
+	FApplePlatformBackgroundHttpRequest();
+	virtual ~FApplePlatformBackgroundHttpRequest() {}
+
+	virtual void CompleteWithExistingResponseData(FBackgroundHttpResponsePtr BackgroundResponse) override;
+
+    virtual bool IsTaskComplete() const;
+
+private:
+    //Super simple linked list with volatile pointers to next element
+    struct FTaskNode
+    {
+        NSURLSessionTask* OurTask;
+        volatile FTaskNode* NextNode;
+    };
+    
+private:
+    const FString& GetURLForRetry(bool bShouldIncrementRetryCountFirst);
+    void AssociateWithTask(NSURLSessionTask* ExistingTask);
+	void SetRequestAsSuccess(const FString& CompletedTempDownloadLocation);
+    void SetRequestAsFailed();
+    void CompleteRequest_Internal(bool bWasRequestSuccess, const FString& CompletedTempDownloadLocation);
+	void CancelActiveTask();
+	void UpdateDownloadProgress(int64_t TotalDownloaded, int64_t DownloadedSinceLastUpdate);
+    void ResetProgressTracking();
+    
+    void ActivateUnderlyingTask();
+    bool IsUnderlyingTaskActive();
+    bool TickTimeOutTimer(float DeltaTime);
+    void ResetTimeOutTimer();
+    
+    FString CompletedTempDownloadLocation;
+    volatile float ActiveTimeOutTimer;
+    
+	FThreadSafeCounter RetryCount;
+	FThreadSafeCounter ResumeDataRetryCount;
+
+	volatile FTaskNode* FirstTask;
+
+    volatile int32 bIsTaskActive;
+	volatile int32 bIsCompleted;
+	volatile int32 bIsFailed;
+	volatile int32 bWasTaskStartedInBG;
+    volatile int32 bHasAlreadyFinishedRequest;
+	volatile int64 DownloadProgress;
+    
+    friend class FApplePlatformBackgroundHttpManager;
+};
+
+typedef TSharedPtr<class FApplePlatformBackgroundHttpRequest, ESPMode::ThreadSafe> FAppleBackgroundHttpRequestPtr;
diff --git a/Engine/Source/Runtime/Online/BackgroundHTTP/Public/IOS/ApplePlatformBackgroundHttpResponse.h b/Engine/Source/Runtime/Online/BackgroundHTTP/Public/IOS/ApplePlatformBackgroundHttpResponse.h
new file mode 100644
--- /dev/null
+++ b/Engine/Source/Runtime/Online/BackgroundHTTP/Public/IOS/ApplePlatformBackgroundHttpResponse.h
@@ -0,0 +1,17 @@
+// Copyright 1998-2018 Epic Games, Inc. All Rights Reserved.
+
+#pragma once
+
+#include "BackgroundHttpResponseImpl.h"
+#include "Interfaces/IBackgroundHttpRequest.h"
+
+/**
+ * Contains implementation of some common functions that don't vary between implementation
+ */
+class FApplePlatformBackgroundHttpResponse
+	: public FBackgroundHttpResponseImpl
+{
+public:
+    //Provide constructor to create an already
+    FApplePlatformBackgroundHttpResponse(int32 ResponseCode, const FString& TempContentFilePath);
+};
diff --git a/Engine/Source/Runtime/Online/BackgroundHTTP/Public/Interfaces/IBackgroundHttpManager.h b/Engine/Source/Runtime/Online/BackgroundHTTP/Public/Interfaces/IBackgroundHttpManager.h
new file mode 100644
--- /dev/null
+++ b/Engine/Source/Runtime/Online/BackgroundHTTP/Public/Interfaces/IBackgroundHttpManager.h
@@ -0,0 +1,58 @@
+// Copyright 1998-2018 Epic Games, Inc. All Rights Reserved.
+
+#pragma once
+
+#include "BackgroundHttp.h"
+
+typedef TSharedPtr<class IBackgroundHttpManager, ESPMode::ThreadSafe> FBackgroundHttpManagerPtr;
+
+class BACKGROUNDHTTP_API IBackgroundHttpManager
+{
+public:
+	/**
+	 * Initialize
+	 */
+	virtual void Initialize() = 0;
+
+	/**
+	 * Shutdown
+	 */
+	virtual void Shutdown() = 0;
+
+	/**
+	 * Adds a Background Http request instance to the manager for tracking
+	 * Manager should always have a list of requests currently being processed
+	 *
+	 * @param Request - the request object to add
+	 */
+	virtual void AddRequest(const FBackgroundHttpRequestPtr Request) = 0;
+
+	/**
+	 * Removes a Background Http request instance from the manager
+	 * Presumably it is done being processed
+	 *
+	 * @param Request - the request object to remove
+	 */
+	virtual void RemoveRequest(const FBackgroundHttpRequestPtr Request) = 0;
+
+	/**
+	* Function to remove all temporary files used by the system to store completed downloads.
+	* This should be called when there are no active background downloads that we care
+	* about to regain disk space from background downloads that may no longer be relevant (IE: downloads that were never moved after completion)
+	* but are still on disk.
+	*
+	*/
+	virtual void CleanUpTemporaryFiles() = 0;
+
+protected:
+	/**
+	* Designed to be called internally by AddRequest to associate our incoming request with any 
+	* previously completed background downloads (might have completed before this app launch, or 
+	* carried over from a previous application launch and be running without any other information.
+	*
+	* @param Request - The request to check for any existing request matches on.
+	*
+	* @return bool - true if we found an existing request and associated with it. false otherwise.
+	*/
+	virtual bool AssociateWithAnyExistingRequest(const FBackgroundHttpRequestPtr Request) = 0;
+};
diff --git a/Engine/Source/Runtime/Online/BackgroundHTTP/Public/Interfaces/IBackgroundHttpRequest.h b/Engine/Source/Runtime/Online/BackgroundHTTP/Public/Interfaces/IBackgroundHttpRequest.h
new file mode 100644
--- /dev/null
+++ b/Engine/Source/Runtime/Online/BackgroundHTTP/Public/Interfaces/IBackgroundHttpRequest.h
@@ -0,0 +1,131 @@
+// Copyright 1998-2018 Epic Games, Inc. All Rights Reserved.
+
+#pragma once
+
+#include "CoreMinimal.h"
+#include "BackgroundHttpNotificationObject.h"
+
+//Included as we use EHttpResponseCode
+#include "Interfaces/IHttpResponse.h"
+
+
+typedef TSharedPtr<class IBackgroundHttpRequest, ESPMode::ThreadSafe> FBackgroundHttpRequestPtr;
+typedef TSharedPtr<class IBackgroundHttpResponse, ESPMode::ThreadSafe> FBackgroundHttpResponsePtr;
+typedef TSharedPtr<class IBackgroundHttpManager, ESPMode::ThreadSafe> FBackgroundHttpManagerPtr;
+
+/**
+ * Delegate called when a Background Http request completes
+ *
+ * @param Request original Background Http request that started things
+ * @param Response response received from the server if a successful connection was established
+ */
+DECLARE_DELEGATE_TwoParams(FBackgroundHttpRequestCompleteDelegate, FBackgroundHttpRequestPtr /*Request*/, bool /*bWasSuccessful*/);
+
+/**
+* Delegate Called when a Background Http request updates its progress
+*
+* @param Request Background Http Request that is updating its progress
+* @param TotalBytesWritten Amount of data we have written so far. Total Figure.
+* @param BytesWrittenSinceLastUpdate Amount of data we have written since our last call to this delegate.
+*/
+DECLARE_DELEGATE_ThreeParams(FBackgroundHttpProgressUpdateDelegate, FBackgroundHttpRequestPtr /*Request*/, int32 /*TotalBytesWritten*/, int32 /*BytesWrittenSinceLastUpdate*/);
+
+/**
+ * Interface for Http requests (created using FHttpFactory)
+ */
+class IBackgroundHttpRequest 
+	: public TSharedFromThis<IBackgroundHttpRequest, ESPMode::ThreadSafe>
+{
+public:
+	/**
+	* Sets up a list of URLs to automatically fall through as each one fails.
+	* Must be called before ProcessRequest.
+	*
+	* @param URLs list of URLs. Eg: download.epicgames.com/downloadfilehere.txt, download2.epicgames.com/downloadfilehere.txt, download3.epicgames/downloadfilehere.txt
+	* @param NumRetriesToAttempt How many times we want to fall through and try different URLS. Loops back to the beginning of the list if NumRetries > URLs.Num()
+	*
+	*/
+	virtual void SetURLAsList(const TArray<FString>& URLs, int NumRetriesToAttempt) = 0;
+
+	/**
+	* Gets the current URL List that this background request is currently processing.
+	*
+	* @return List of URLs as FStrings
+	*
+	*/
+	virtual const TArray<FString>& GetURLList() const = 0;
+
+	/**
+	* Sets an FHTTPRequestDownloadNotificationObject to be mapped to this HTTPRequest.
+	* This request will keep the particular FHttpRequestDownloadNotificationObject reference until it completes.
+	* To use this function make sure you keep a reference to the supplied DownloadNotificationObject until you have created all the different
+	* IHTTPRequest that you would like to use it. Then delete your reference. Once all references are removed (by each Request completing)
+	* a callback set on the DownloadCompleteDelegateObject will be called
+	*
+	* @param DownloadCompleteDelegateObject The particular download complete delegate object
+	*
+	*/
+	virtual void SetCompleteNotification(FBackgroundHttpNotificationObjectPtr DownloadCompleteNotificationObject) = 0;
+
+	/**
+	* Function used to complete an IHttpBackgroundRequest from an external source, passing it in a pre-existing response data.
+	*
+	* @param BackgroundReponse, A SharedRef to the already existing IHttpBackgroundResponse we want to base this tasks' Response off of.
+	*/
+	virtual void CompleteWithExistingResponseData(FBackgroundHttpResponsePtr BackgroundResponse) = 0;
+
+	/**
+	* Delegate called when the request is complete. See FBackgroundHttpRequestCompleteDelegate
+	*/
+	virtual FBackgroundHttpRequestCompleteDelegate& OnProcessRequestComplete() = 0;
+
+	/**
+	* Delegate called when the request has a progress update.
+	*/
+	virtual FBackgroundHttpProgressUpdateDelegate& OnProgressUpdated() = 0;
+
+	/**
+	* Called to begin processing the request.
+	* OnProcessRequestComplete delegate is always called when the request completes or on error if it is bound.
+	* A request can be re-used but not while still being processed.
+	*
+	* @return if the request was successfully started.
+	*/
+	virtual bool ProcessRequest() = 0;
+
+	/**
+	* Called by certain platform's implementation when we have to wait for the BackgroundHttpManager / PlatformBackgroundHttp to do some work
+	* before we can finish our ProcessRequest call. Should only be called by different platform layers.
+	*
+	* @return if the request was successfully handled.
+	*/
+	virtual bool HandleDelayedProcess() = 0;
+
+	/**
+	* Get the associated Response
+	*
+	* @return the response
+	*/
+	virtual const FBackgroundHttpResponsePtr GetResponse() const = 0;
+
+	/**
+	* Gets the associated RequestID for this BackgroundDownload.
+	*
+	* @return The RequestID
+	*/
+	virtual const FString& GetRequestID() const = 0;
+
+	/**
+	* Sets the associated RequestID for this BackgroundDownload. Useful as we associate a background downlaod with multiple
+	* URLs, so this provides an easier way to identify the download with 1 string instead of checking multiple URLs.
+	*
+	* @param NewRequestID FString to set the Request ID to.
+	*/
+	virtual void SetRequestID(const FString& NewRequestID) = 0;
+
+	/** 
+	* Destructor for overrides 
+	*/
+	virtual ~IBackgroundHttpRequest() = default;
+};
+
diff --git a/Engine/Source/Runtime/Online/BackgroundHTTP/Public/Interfaces/IBackgroundHttpResponse.h b/Engine/Source/Runtime/Online/BackgroundHTTP/Public/Interfaces/IBackgroundHttpResponse.h
new file mode 100644
--- /dev/null
+++ b/Engine/Source/Runtime/Online/BackgroundHTTP/Public/Interfaces/IBackgroundHttpResponse.h
@@ -0,0 +1,34 @@
+// Copyright 1998-2018 Epic Games, Inc. All Rights Reserved.
+
+#pragma once
+
+#include "CoreMinimal.h"
+
+/**
+ * Interface for Http Background Responses
+ */
+class BACKGROUNDHTTP_API IBackgroundHttpResponse
+{
+public:
+	/**
+	 * Gets the response code returned by the requested server.
+	 * See EHttpResponseCodes for known response codes. Same as HttpRequest (non-background)
+	 *
+	 * @return the response code.
+	 */
+	virtual int32 GetResponseCode() const = 0;
+
+	/**
+	* Gets the file name that our temporary content file will be stored in. HttpBackgroundRequests
+	* Download their files to a temp location and this can be used to get at the temporary file on disk and move it to your desired location.
+	* 
+	* @return FString representing the FilePath to the temporary file.
+	*/
+	virtual const FString& GetTempContentFilePath() const = 0;
+
+	/** 
+	 * Destructor for overrides 
+	 */
+	virtual ~IBackgroundHttpResponse() = default;
+};
+
diff --git a/Engine/Source/Runtime/Online/BackgroundHTTP/Public/Mac/MacPlatformBackgroundHttp.h b/Engine/Source/Runtime/Online/BackgroundHTTP/Public/Mac/MacPlatformBackgroundHttp.h
new file mode 100644
--- /dev/null
+++ b/Engine/Source/Runtime/Online/BackgroundHTTP/Public/Mac/MacPlatformBackgroundHttp.h
@@ -0,0 +1,6 @@
+#pragma once
+
+#include "GenericPlatform/GenericPlatformBackgroundHttp.h"
+
+//Currently no implementation on Mac
+typedef FGenericPlatformBackgroundHttp FPlatformBackgroundHttp;
\ No newline at end of file
diff --git a/Engine/Source/Runtime/Online/BackgroundHTTP/Public/PlatformBackgroundHttp.h b/Engine/Source/Runtime/Online/BackgroundHTTP/Public/PlatformBackgroundHttp.h
new file mode 100644
--- /dev/null
+++ b/Engine/Source/Runtime/Online/BackgroundHTTP/Public/PlatformBackgroundHttp.h
@@ -0,0 +1,28 @@
+// Copyright 1998-2018 Epic Games, Inc. All Rights Reserved.
+#pragma once
+
+#include "CoreMinimal.h"
+
+//On all un-implemented platforms, just uses GenericPlatform/GenericPlatformBackgroundHttp.h
+
+#if PLATFORM_WINDOWS
+#include "Windows/WindowsPlatformBackgroundHttp.h"
+#elif PLATFORM_PS4
+#include "PS4/PS4PlatformBackgroundHttp.h"
+#elif PLATFORM_XBOXONE
+#include "XboxOne/XboxOnePlatformBackgroundHttp.h"
+#elif PLATFORM_MAC
+#include "Mac/MacPlatformBackgroundHttp.h" // These 2 are a little confusing right now. Only thing holding back Mac is an implementation
+#elif PLATFORM_IOS                         // of FBackgroundURLSessionHandler. For now everything is setup generically "Apple" so that
+#include "IOS/ApplePlatformBackgroundHttp.h" // we can transition easier if we decide to implement this on Mac. So for now ApplePlatformX is actually only iOS in background http
+#elif PLATFORM_ANDROID
+#include "Android/AndroidPlatformBackgroundHttp.h"
+#elif PLATFORM_HTML5
+#include "HTML5/HTML5PlatformBackgroundHttp.h"
+#elif PLATFORM_UNIX
+#include "Unix/UnixPlatformBackgroundHttp.h"
+#elif PLATFORM_SWITCH
+#include "Switch/SwitchPlatformBackgroundHttp.h"
+#else
+#error Unknown platform
+#endif
diff --git a/Engine/Source/Runtime/Online/BackgroundHTTP/Public/Unix/UnixPlatformBackgroundHttp.h b/Engine/Source/Runtime/Online/BackgroundHTTP/Public/Unix/UnixPlatformBackgroundHttp.h
new file mode 100644
--- /dev/null
+++ b/Engine/Source/Runtime/Online/BackgroundHTTP/Public/Unix/UnixPlatformBackgroundHttp.h
@@ -0,0 +1,6 @@
+#pragma once
+
+#include "GenericPlatform/GenericPlatformBackgroundHttp.h"
+
+//Currently no implementation on Unix
+typedef FGenericPlatformBackgroundHttp FPlatformBackgroundHttp;
\ No newline at end of file
diff --git a/Engine/Source/Runtime/Online/BackgroundHTTP/Public/Windows/WindowsPlatformBackgroundHttp.h b/Engine/Source/Runtime/Online/BackgroundHTTP/Public/Windows/WindowsPlatformBackgroundHttp.h
new file mode 100644
--- /dev/null
+++ b/Engine/Source/Runtime/Online/BackgroundHTTP/Public/Windows/WindowsPlatformBackgroundHttp.h
@@ -0,0 +1,6 @@
+#pragma once
+
+#include "GenericPlatform/GenericPlatformBackgroundHttp.h"
+
+//Currently no implementation on Windows
+typedef FGenericPlatformBackgroundHttp FPlatformBackgroundHttp;
\ No newline at end of file
