Added ability to cancel/pause/resume background http requests at the request level
#tests iOS Client, Windows Editor Client

diff --git a/Engine/Source/Runtime/Online/BackgroundHTTP/Private/BackgroundHttpRequestImpl.cpp b/Engine/Source/Runtime/Online/BackgroundHTTP/Private/BackgroundHttpRequestImpl.cpp
--- a/Engine/Source/Runtime/Online/BackgroundHTTP/Private/BackgroundHttpRequestImpl.cpp
+++ b/Engine/Source/Runtime/Online/BackgroundHTTP/Private/BackgroundHttpRequestImpl.cpp
@@ -24,6 +24,26 @@ bool FBackgroundHttpRequestImpl::ProcessRequest()
 	return true;
 }
 
+void FBackgroundHttpRequestImpl::CancelRequest()
+{
+	UE_LOG(LogBackgroundHttpRequest, Display, TEXT("Cancelling Request - RequestID:%s"), *GetRequestID());
+	FBackgroundHttpModule::Get().GetBackgroundHttpManager()->RemoveRequest(SharedThis(this));
+}
+
+void FBackgroundHttpRequestImpl::PauseRequest()
+{
+    //for now a pause is just wrapping a cancel in the general case
+    UE_LOG(LogBackgroundHttpRequest, Display, TEXT("Pausing Request (through cancel) - RequestID:%s"), *GetRequestID());
+    CancelRequest();
+}
+
+void FBackgroundHttpRequestImpl::ResumeRequest()
+{
+    //for now a resume is just wrapping a restart in the general case
+    UE_LOG(LogBackgroundHttpRequest, Display, TEXT("Pausing Request (through restart) - RequestID:%s"), *GetRequestID());
+    ProcessRequest();
+}
+
 void FBackgroundHttpRequestImpl::OnBackgroundDownloadComplete()
 {
 	FBackgroundHttpModule::Get().GetBackgroundHttpManager()->RemoveRequest(SharedThis(this));
@@ -115,4 +135,4 @@ bool FBackgroundHttpRequestImpl::HandleDelayedProcess()
 	//By default we don't provide an implementation for this. If this is called, it should be overriden by the platform specific
 	//BackgroundHttpRequest if the platform expects to make use of this.
 	return ensureAlwaysMsgf(false, TEXT("Platform expects an implementation of HandleDelayedProcess on the BackgroundHttpRequest, but none found!"));
-}
\ No newline at end of file
+}
diff --git a/Engine/Source/Runtime/Online/BackgroundHTTP/Private/GenericPlatform/GenericPlatformBackgroundHttpRequest.cpp b/Engine/Source/Runtime/Online/BackgroundHTTP/Private/GenericPlatform/GenericPlatformBackgroundHttpRequest.cpp
--- a/Engine/Source/Runtime/Online/BackgroundHTTP/Private/GenericPlatform/GenericPlatformBackgroundHttpRequest.cpp
+++ b/Engine/Source/Runtime/Online/BackgroundHTTP/Private/GenericPlatform/GenericPlatformBackgroundHttpRequest.cpp
@@ -31,6 +31,12 @@ bool FGenericPlatformBackgroundHttpRequest::HandleDelayedProcess()
 	return RequestWrapper.IsValid() || GetResponse().IsValid();
 }
 
+void FGenericPlatformBackgroundHttpRequest::CancelRequest()
+{
+	FBackgroundHttpRequestImpl::CancelRequest();
+	RequestWrapper.Reset();
+}
+
 void FGenericPlatformBackgroundHttpRequest::FGenericPlatformBackgroundHttpRequest::CompleteWithExistingResponseData(FBackgroundHttpResponsePtr BackgroundResponse)
 {
 	//Remove our request wrapper immediately so the Response cleans up ASAP instead of the user having to delete this task to clean up the Http Response.
diff --git a/Engine/Source/Runtime/Online/BackgroundHTTP/Private/IOS/ApplePlatformBackgroundHttpManager.cpp b/Engine/Source/Runtime/Online/BackgroundHTTP/Private/IOS/ApplePlatformBackgroundHttpManager.cpp
--- a/Engine/Source/Runtime/Online/BackgroundHTTP/Private/IOS/ApplePlatformBackgroundHttpManager.cpp
+++ b/Engine/Source/Runtime/Online/BackgroundHTTP/Private/IOS/ApplePlatformBackgroundHttpManager.cpp
@@ -5,6 +5,10 @@
 #include "HAL/PlatformFilemanager.h"
 #include "HAL/PlatformAtomics.h"
 
+#include "Misc/ConfigCacheIni.h"
+#include "Misc/CoreDelegates.h"
+#include "Misc/ScopeRWLock.h"
+
 #include "IOS/IOSBackgroundURLSessionHandler.h"
 
 FApplePlatformBackgroundHttpManager::~FApplePlatformBackgroundHttpManager()
@@ -143,6 +147,21 @@ void FApplePlatformBackgroundHttpManager::GenerateURLMapEntriesForRequest( FAppl
     }
 }
 
+void FApplePlatformBackgroundHttpManager::RemoveURLMapEntriesForRequest(FAppleBackgroundHttpRequestPtr Request)
+{
+    FRWScopeLock ScopeLock(URLToRequestMapLock, SLT_Write);
+    for (const FString& URL : Request->GetURLList())
+    {
+        FBackgroundHttpURLMappedRequestPtr& FoundRequest = URLToRequestMap.FindOrAdd(URL);
+        
+        if (FoundRequest == Request)
+        {
+            UE_LOG(LogBackgroundHttpManager, Display, TEXT("Removing URL Entry -- RequestID:%s | URL:%s"), *Request->GetRequestID(), *URL);
+            URLToRequestMap.Remove(URL);
+        }
+    }
+}
+
 void FApplePlatformBackgroundHttpManager::StartRequest(FAppleBackgroundHttpRequestPtr Request)
 {
     //Just count it as a retry that won't increment the retry counter before giving us the URL as our RetryCount 0 should start this up.
@@ -172,7 +191,11 @@ void FApplePlatformBackgroundHttpManager::DeletePendingRemoveRequests()
 
 void FApplePlatformBackgroundHttpManager::RemoveSessionTasksForRequest(FAppleBackgroundHttpRequestPtr Request)
 {
-    //@TODO TRoss: Should implement this
+    //First remove map entries. That way we won't send a completion handler when we cancel
+    RemoveURLMapEntriesForRequest(Request);
+    
+    //Now cancel our active task
+    Request->CancelActiveTask();
 }
 
 bool FApplePlatformBackgroundHttpManager::AssociateWithAnyExistingRequest(const FBackgroundHttpRequestPtr Request)
@@ -250,16 +273,16 @@ void FApplePlatformBackgroundHttpManager::CleanUpNSURLSessionResponseDelegates()
 void FApplePlatformBackgroundHttpManager::OnApp_EnteringForeground()
 {
     FPlatformAtomics::InterlockedExchange(&bIsInBackground,false);
-	PauseAllActiveRequests();
+	PauseAllActiveTasks();
 }
 
 void FApplePlatformBackgroundHttpManager::OnApp_EnteringBackground()
 {
 	FPlatformAtomics::InterlockedExchange(&bIsInBackground,true);
-	UnpauseAllPausedRequests();
+	ResumeAllTasks();
 }
 
-void FApplePlatformBackgroundHttpManager::PauseAllActiveRequests()
+void FApplePlatformBackgroundHttpManager::PauseAllActiveTasks()
 {
     UE_LOG(LogBackgroundHttpManager, Display, TEXT("Attempting to Pause All Active Tasks"));
     
@@ -279,13 +302,13 @@ void FApplePlatformBackgroundHttpManager::PauseAllActiveRequests()
 				}
 			}
             
-            //Reset our active requests to 0 now that we are coming back to the BG
+            //Reset our active requests to 0 now that we are pausing everything
             FPlatformAtomics::InterlockedExchange(&NumCurrentlyActiveRequests, 0);
 		}];
 	}
 }
 
-void FApplePlatformBackgroundHttpManager::UnpauseAllPausedRequests()
+void FApplePlatformBackgroundHttpManager::ResumeAllTasks()
 {
     UE_LOG(LogBackgroundHttpManager, Display, TEXT("Attempting to Resume All Active Tasks"));
     
@@ -596,6 +619,7 @@ void FApplePlatformBackgroundHttpManager::TickRequests(float DeltaTime)
             if (ensureAlwaysMsgf(AppleRequest.IsValid(), TEXT("Invalid Request Pointer in ActiveRequests list!")))
             {
                 const bool bIsTaskActive = AppleRequest->IsUnderlyingTaskActive();
+                const bool bIsTaskPaused = AppleRequest->IsUnderlyingTaskPaused();
                 const bool bIsTaskComplete = AppleRequest->IsTaskComplete();
                 const bool bWasStartedInBG = FPlatformAtomics::AtomicRead(&(AppleRequest->bWasTaskStartedInBG));
                 
@@ -614,7 +638,7 @@ void FApplePlatformBackgroundHttpManager::TickRequests(float DeltaTime)
                     //We want to recreate any task spun up in the background as it will not respect our session settings if created in BG.
                     AppleRequest->CancelActiveTask();
                 }
-                else if (bIsTaskActive)
+                else if (bIsTaskActive && !bIsTaskPaused)
                 {
                     const bool bShouldTimeOut = AppleRequest->TickTimeOutTimer(DeltaTime);
                     if (bShouldTimeOut)
@@ -668,6 +692,8 @@ void FApplePlatformBackgroundHttpManager::TickTasks(float DeltaTime)
                                      FBackgroundHttpURLMappedRequestPtr* WeakRequestInMap = URLToRequestMap.Find(TaskURL);
                                      FAppleBackgroundHttpRequestPtr FoundRequest = (nullptr != WeakRequestInMap) ? WeakRequestInMap->Pin() : nullptr;
                                      
+                                     const bool bIsPaused = FoundRequest.IsValid() ? FoundRequest->IsUnderlyingTaskPaused() : false;
+                                     
                                      if (FoundRequest.IsValid())
                                      {
                                          UE_LOG(LogBackgroundHttpManager, Display, TEXT("Activating Task For Requets -- RequestID:%s | TaskURL:%s | CurrentlyActiveRequests:%d"), *(FoundRequest->GetRequestID()), *TaskURL, NewRequestCount);
@@ -675,7 +701,7 @@ void FApplePlatformBackgroundHttpManager::TickTasks(float DeltaTime)
                                      }
                                      else
                                      {
-                                         UE_LOG(LogBackgroundHttpManager, Display, TEXT("Skipping Activating Task as there is no associated Request. Once a Request associates with this task, it can then be activated. -- TaskURL:%s"), *TaskURL);
+                                         UE_LOG(LogBackgroundHttpManager, Display, TEXT("Skipping Activating Task as there is no associated Request or Request is paused. Once a Request associates with this task, it can then be activated. -- TaskURL:%s | bIsPaused:%d"), *TaskURL, (int)bIsPaused);
                                          
                                          //Don't activate and remove our increment from above because something put us over the limit before we resumed
                                          FPlatformAtomics::InterlockedDecrement(&NumCurrentlyActiveRequests);
diff --git a/Engine/Source/Runtime/Online/BackgroundHTTP/Private/IOS/ApplePlatformBackgroundHttpRequest.cpp b/Engine/Source/Runtime/Online/BackgroundHTTP/Private/IOS/ApplePlatformBackgroundHttpRequest.cpp
--- a/Engine/Source/Runtime/Online/BackgroundHTTP/Private/IOS/ApplePlatformBackgroundHttpRequest.cpp
+++ b/Engine/Source/Runtime/Online/BackgroundHTTP/Private/IOS/ApplePlatformBackgroundHttpRequest.cpp
@@ -1,5 +1,6 @@
 // Copyright 1998-2018 Epic Games, Inc. All Rights Reserved.
 #include "IOS/ApplePlatformBackgroundHttpRequest.h"
+#include "IOS/ApplePlatformBackgroundHttpManager.h"
 #include "Interfaces/IBackgroundHttpResponse.h"
 #include "BackgroundHttpModule.h"
 
@@ -10,6 +11,7 @@ FApplePlatformBackgroundHttpRequest::FApplePlatformBackgroundHttpRequest()
 	, ResumeDataRetryCount(0)
     , FirstTask(nullptr)
     , bIsTaskActive(false)
+    , bIsTaskPaused(false)
     , bIsCompleted(false)
     , bIsFailed(false)
     , bWasTaskStartedInBG(false)
@@ -93,6 +95,8 @@ void FApplePlatformBackgroundHttpRequest::ActivateUnderlyingTask()
             UE_LOG(LogBackgroundHttpRequest, Display, TEXT("Activating Task for Request -- RequestID:%s | TaskURL:%s"), *GetRequestID(), *TaskURL);
         
             FPlatformAtomics::InterlockedExchange(&bIsTaskActive, true);
+            FPlatformAtomics::InterlockedExchange(&bIsTaskPaused, false);
+
             [UnderlyingTask resume];
             
             ResetTimeOutTimer();
@@ -101,11 +105,37 @@ void FApplePlatformBackgroundHttpRequest::ActivateUnderlyingTask()
     }
 }
 
+void FApplePlatformBackgroundHttpRequest::PauseUnderlyingTask()
+{
+    volatile FTaskNode* ActiveTaskNode = FirstTask;
+    if (ensureAlwaysMsgf((nullptr != ActiveTaskNode), TEXT("Call to PauseUnderlyingTask with an invalid node! Need to create underlying task(and node) before trying to pause!")))
+    {
+        NSURLSessionTask* UnderlyingTask = ActiveTaskNode->OurTask;
+        if (ensureAlwaysMsgf((nullptr != UnderlyingTask), TEXT("Call to PauseUnderlyingTask with an invalid task! Need to create underlying task before trying to pause!")))
+        {
+            FString TaskURL = [[[UnderlyingTask currentRequest] URL] absoluteString];
+            UE_LOG(LogBackgroundHttpRequest, Display, TEXT("Pausing Task for Request -- RequestID:%s | TaskURL:%s"), *GetRequestID(), *TaskURL);
+            
+            FPlatformAtomics::InterlockedExchange(&bIsTaskActive, false);
+            FPlatformAtomics::InterlockedExchange(&bIsTaskPaused, true);
+            [UnderlyingTask suspend];
+            
+            ResetTimeOutTimer();
+            ResetProgressTracking();
+        }
+    }
+}
+
 bool FApplePlatformBackgroundHttpRequest::IsUnderlyingTaskActive()
 {
     return FPlatformAtomics::AtomicRead(&bIsTaskActive);
 }
 
+bool FApplePlatformBackgroundHttpRequest::IsUnderlyingTaskPaused()
+{
+    return FPlatformAtomics::AtomicRead(&bIsTaskPaused);
+}
+
 bool FApplePlatformBackgroundHttpRequest::TickTimeOutTimer(float DeltaTime)
 {
     ActiveTimeOutTimer -= DeltaTime;
@@ -138,6 +168,16 @@ void FApplePlatformBackgroundHttpRequest::AssociateWithTask(NSURLSessionTask* Ex
 	}
 }
 
+void FApplePlatformBackgroundHttpRequest::PauseRequest()
+{
+    PauseUnderlyingTask();
+}
+
+void FApplePlatformBackgroundHttpRequest::ResumeRequest()
+{
+    ActivateUnderlyingTask();
+}
+
 void FApplePlatformBackgroundHttpRequest::CancelActiveTask()
 {
     volatile FTaskNode* TaskNodeWeAreCancelling = FirstTask;
diff --git a/Engine/Source/Runtime/Online/BackgroundHTTP/Public/BackgroundHttpRequestImpl.h b/Engine/Source/Runtime/Online/BackgroundHTTP/Public/BackgroundHttpRequestImpl.h
--- a/Engine/Source/Runtime/Online/BackgroundHTTP/Public/BackgroundHttpRequestImpl.h
+++ b/Engine/Source/Runtime/Online/BackgroundHTTP/Public/BackgroundHttpRequestImpl.h
@@ -22,7 +22,10 @@ public:
 
 	// IBackgroundHttpRequest
 	virtual bool ProcessRequest() override;
-	virtual void SetURLAsList(const TArray<FString>& URLs, int NumRetriesToAttempt) override;
+	virtual void CancelRequest() override;
+    virtual void PauseRequest() override;
+    virtual void ResumeRequest() override;
+    virtual void SetURLAsList(const TArray<FString>& URLs, int NumRetriesToAttempt) override;
 	virtual const TArray<FString>& GetURLList() const override;
 	virtual void SetCompleteNotification(FBackgroundHttpNotificationObjectPtr DownloadCompleteNotificationObjectIn) override;
 	virtual void CompleteWithExistingResponseData(FBackgroundHttpResponsePtr BackgroundResponse) override;
@@ -43,4 +46,4 @@ protected:
 
 	FBackgroundHttpRequestCompleteDelegate HttpRequestCompleteDelegate;
 	FBackgroundHttpProgressUpdateDelegate HttpProgressUpdateDelegate;
-};
\ No newline at end of file
+};
diff --git a/Engine/Source/Runtime/Online/BackgroundHTTP/Public/GenericPlatform/GenericPlatformBackgroundHttpRequest.h b/Engine/Source/Runtime/Online/BackgroundHTTP/Public/GenericPlatform/GenericPlatformBackgroundHttpRequest.h
--- a/Engine/Source/Runtime/Online/BackgroundHTTP/Public/GenericPlatform/GenericPlatformBackgroundHttpRequest.h
+++ b/Engine/Source/Runtime/Online/BackgroundHTTP/Public/GenericPlatform/GenericPlatformBackgroundHttpRequest.h
@@ -22,6 +22,7 @@ public:
 	//IHttpBackgroundHttpRequest
 	virtual bool HandleDelayedProcess() override;
 	virtual void CompleteWithExistingResponseData(FBackgroundHttpResponsePtr BackgroundResponse) override;
+	virtual void CancelRequest() override;
 
 protected:
 	//In the default implementation we actually use HTTPRequests instead of HttpBackgroundRequests.
diff --git a/Engine/Source/Runtime/Online/BackgroundHTTP/Public/IOS/ApplePlatformBackgroundHttpManager.h b/Engine/Source/Runtime/Online/BackgroundHTTP/Public/IOS/ApplePlatformBackgroundHttpManager.h
--- a/Engine/Source/Runtime/Online/BackgroundHTTP/Public/IOS/ApplePlatformBackgroundHttpManager.h
+++ b/Engine/Source/Runtime/Online/BackgroundHTTP/Public/IOS/ApplePlatformBackgroundHttpManager.h
@@ -70,8 +70,9 @@ private:
 	void SetupNSURLSessionResponseDelegates();
 	void CleanUpNSURLSessionResponseDelegates();
 
-	void PauseAllActiveRequests();
-	void UnpauseAllPausedRequests();
+    //These are used internally to pause/resume things at the task level
+	void PauseAllActiveTasks();
+	void ResumeAllTasks();
 
     void PauseAllUnassociatedTasks();
     void UnpauseAllUnassociatedTasks();
@@ -85,6 +86,7 @@ private:
 	void RetryRequest(FAppleBackgroundHttpRequestPtr Request, bool bShouldIncreaseRetryCount, bool bShouldStartImmediately, NSData* RetryData);
     void RemoveSessionTasksForRequest(FAppleBackgroundHttpRequestPtr Request);
     void GenerateURLMapEntriesForRequest(FAppleBackgroundHttpRequestPtr Request);
+    void RemoveURLMapEntriesForRequest(FAppleBackgroundHttpRequestPtr Request);
     
 	bool IsRetryDataValid(NSData* RetryData) const;
 	bool ShouldUseRequestRetryData(FAppleBackgroundHttpRequestPtr Request, NSData* RetryData) const;
diff --git a/Engine/Source/Runtime/Online/BackgroundHTTP/Public/IOS/ApplePlatformBackgroundHttpRequest.h b/Engine/Source/Runtime/Online/BackgroundHTTP/Public/IOS/ApplePlatformBackgroundHttpRequest.h
--- a/Engine/Source/Runtime/Online/BackgroundHTTP/Public/IOS/ApplePlatformBackgroundHttpRequest.h
+++ b/Engine/Source/Runtime/Online/BackgroundHTTP/Public/IOS/ApplePlatformBackgroundHttpRequest.h
@@ -18,9 +18,12 @@ public:
 	virtual ~FApplePlatformBackgroundHttpRequest() {}
 
 	virtual void CompleteWithExistingResponseData(FBackgroundHttpResponsePtr BackgroundResponse) override;
-
+    virtual void PauseRequest() override;
+    virtual void ResumeRequest() override;
+    
     virtual bool IsTaskComplete() const;
 
+    
 private:
     //Super simple linked list with volatile pointers to next element
     struct FTaskNode
@@ -41,6 +44,8 @@ private:
     
     void ActivateUnderlyingTask();
     bool IsUnderlyingTaskActive();
+    void PauseUnderlyingTask();
+    bool IsUnderlyingTaskPaused();
     bool TickTimeOutTimer(float DeltaTime);
     void ResetTimeOutTimer();
     
@@ -53,6 +58,7 @@ private:
 	volatile FTaskNode* FirstTask;
 
     volatile int32 bIsTaskActive;
+    volatile int32 bIsTaskPaused;
 	volatile int32 bIsCompleted;
 	volatile int32 bIsFailed;
 	volatile int32 bWasTaskStartedInBG;
diff --git a/Engine/Source/Runtime/Online/BackgroundHTTP/Public/Interfaces/IBackgroundHttpRequest.h b/Engine/Source/Runtime/Online/BackgroundHTTP/Public/Interfaces/IBackgroundHttpRequest.h
--- a/Engine/Source/Runtime/Online/BackgroundHTTP/Public/Interfaces/IBackgroundHttpRequest.h
+++ b/Engine/Source/Runtime/Online/BackgroundHTTP/Public/Interfaces/IBackgroundHttpRequest.h
@@ -96,11 +96,28 @@ public:
 	/**
 	* Called by certain platform's implementation when we have to wait for the BackgroundHttpManager / PlatformBackgroundHttp to do some work
 	* before we can finish our ProcessRequest call. Should only be called by different platform layers.
+	* 
+	* NOTE: Should really only be called by the BackgroundHttpManager! You are probably looking for ProcessRequest.
 	*
 	* @return if the request was successfully handled.
 	*/
 	virtual bool HandleDelayedProcess() = 0;
 
+	/**
+	 * Called to cancel a request that is still being processed
+	 */
+	virtual void CancelRequest() = 0;
+
+    /**
+    * Called to pause a request that is still being processed
+    */
+    virtual void PauseRequest() = 0;
+    
+    /**
+    * Called to resume a request that was previously paused
+    */
+    virtual void ResumeRequest() = 0;
+    
 	/**
 	* Get the associated Response
 	*
