Fixing Clang compile errors.
#tests Windows Editor Compile

diff --git a/Engine/Source/Runtime/Online/BackgroundHTTP/Private/BackgroundHttpManagerImpl.cpp b/Engine/Source/Runtime/Online/BackgroundHTTP/Private/BackgroundHttpManagerImpl.cpp
--- a/Engine/Source/Runtime/Online/BackgroundHTTP/Private/BackgroundHttpManagerImpl.cpp
+++ b/Engine/Source/Runtime/Online/BackgroundHTTP/Private/BackgroundHttpManagerImpl.cpp
@@ -1,7 +1,11 @@
 // Copyright 1998-2018 Epic Games, Inc. All Rights Reserved.
 #include "BackgroundHttpManagerImpl.h"
 
-#include "HAL/PlatformFileManager.h"
+#include "HAL/PlatformFilemanager.h"
+#include "HAL/PlatformAtomics.h"
+#include "HAL/PlatformFile.h"
+
+#include "Misc/ScopeRWLock.h"
 
 DEFINE_LOG_CATEGORY(LogBackgroundHttpManager);
 
diff --git a/Engine/Source/Runtime/Online/BackgroundHTTP/Private/BackgroundHttpNotificationObject.cpp b/Engine/Source/Runtime/Online/BackgroundHTTP/Private/BackgroundHttpNotificationObject.cpp
--- a/Engine/Source/Runtime/Online/BackgroundHTTP/Private/BackgroundHttpNotificationObject.cpp
+++ b/Engine/Source/Runtime/Online/BackgroundHTTP/Private/BackgroundHttpNotificationObject.cpp
@@ -1,6 +1,11 @@
 // Copyright 1998-2018 Epic Games, Inc. All Rights Reserved.
 #include "BackgroundHttpNotificationObject.h"
 
+#include "Misc/ConfigCacheIni.h"
+
+#include "Modules/ModuleInterface.h"
+#include "Modules/ModuleManager.h"
+
 #include "LocalNotification.h"
 
 FBackgroundHttpNotificationObject::FBackgroundHttpNotificationObject(FText InNotificationTitle, FText InNotificationBody, FText InNotificationAction, const FString& InNotificationActivationString, bool InNotifyOnlyOnFullSuccess)
diff --git a/Engine/Source/Runtime/Online/BackgroundHTTP/Private/GenericPlatform/GenericPlatformBackgroundHttp.cpp b/Engine/Source/Runtime/Online/BackgroundHTTP/Private/GenericPlatform/GenericPlatformBackgroundHttp.cpp
--- a/Engine/Source/Runtime/Online/BackgroundHTTP/Private/GenericPlatform/GenericPlatformBackgroundHttp.cpp
+++ b/Engine/Source/Runtime/Online/BackgroundHTTP/Private/GenericPlatform/GenericPlatformBackgroundHttp.cpp
@@ -4,6 +4,8 @@
 #include "GenericPlatform/GenericPlatformBackgroundHttpRequest.h"
 #include "GenericPlatform/GenericPlatformBackgroundHttpResponse.h"
 
+#include "Misc/Paths.h"
+
 void FGenericPlatformBackgroundHttp::Initialize()
 {
 }
diff --git a/Engine/Source/Runtime/Online/BackgroundHTTP/Private/GenericPlatform/GenericPlatformBackgroundHttpResponse.cpp b/Engine/Source/Runtime/Online/BackgroundHTTP/Private/GenericPlatform/GenericPlatformBackgroundHttpResponse.cpp
--- a/Engine/Source/Runtime/Online/BackgroundHTTP/Private/GenericPlatform/GenericPlatformBackgroundHttpResponse.cpp
+++ b/Engine/Source/Runtime/Online/BackgroundHTTP/Private/GenericPlatform/GenericPlatformBackgroundHttpResponse.cpp
@@ -2,6 +2,8 @@
 #include "GenericPlatform/GenericPlatformBackgroundHttpResponse.h"
 #include "GenericPlatform/GenericPlatformBackgroundHttp.h"
 
+#include "PlatformBackgroundHttp.h"
+
 #include "Interfaces/IHttpResponse.h"
 #include "Misc/FileHelper.h"
 
