Fixing bug in BackgroundHttp generic platform implementation where download progress was tracked incorrectly.
Fix shadowed member variable error.
#test iosClient, Windows Editor Client

diff --git a/Engine/Source/Runtime/Online/BackgroundHTTP/Private/GenericPlatform/GenericPlatformBackgroundHttpRequest.cpp b/Engine/Source/Runtime/Online/BackgroundHTTP/Private/GenericPlatform/GenericPlatformBackgroundHttpRequest.cpp
--- a/Engine/Source/Runtime/Online/BackgroundHTTP/Private/GenericPlatform/GenericPlatformBackgroundHttpRequest.cpp
+++ b/Engine/Source/Runtime/Online/BackgroundHTTP/Private/GenericPlatform/GenericPlatformBackgroundHttpRequest.cpp
@@ -74,6 +74,9 @@ void FGenericPlatformBackgroundHttpRequest::FGenericPlatformBackgroundHttpWrappe
 		const FString& RequestURL = GetURLForCurrentRetryNumber();
 		if (!RequestURL.IsEmpty())
 		{
+			//Reset this value as we are restarting our progress
+			LastProgressUpdateBytes = 0;
+
 			HttpRequest->SetURL(RequestURL);
 			HttpRequest->ProcessRequest();
 		}
@@ -108,19 +111,19 @@ const FString FGenericPlatformBackgroundHttpRequest::FGenericPlatformBackgroundH
 	//Don't find a valid URL if we are already passed our retry count
 	if (URLIndexToUse <= MaxRetries)
 	{
-		const TArray<FString>& URLList = OriginalRequest->GetURLList();
-		if (URLList.Num() > 0)
+		const TArray<FString>& RequestURLList = OriginalRequest->GetURLList();
+		if (RequestURLList.Num() > 0)
 		{
 			//If our number is too high, lets loop back around to the start of the list
-			while ((URLIndexToUse >= 0) && (!URLList.IsValidIndex(URLIndexToUse)))
+			while ((URLIndexToUse >= 0) && (!RequestURLList.IsValidIndex(URLIndexToUse)))
 			{
-				URLIndexToUse -= URLList.Num();
+				URLIndexToUse -= RequestURLList.Num();
 			}
 		}
 
-		if (URLList.IsValidIndex(URLIndexToUse))
+		if (RequestURLList.IsValidIndex(URLIndexToUse))
 		{
-			ReturnedURL = URLList[URLIndexToUse];
+			ReturnedURL = RequestURLList[URLIndexToUse];
 		}
 	}
 
@@ -152,10 +155,11 @@ void FGenericPlatformBackgroundHttpRequest::FGenericPlatformBackgroundHttpWrappe
 
 void FGenericPlatformBackgroundHttpRequest::FGenericPlatformBackgroundHttpWrapper::UpdateHttpProgress(FHttpRequestPtr UnderlyingHttpRequest, int32 BytesSent, int32 BytesReceived)
 {
-	const int32 ByteDifference = (LastProgressUpdateBytes > 0) ? LastProgressUpdateBytes - BytesReceived : BytesReceived;
-	LastProgressUpdateBytes = ByteDifference;
-	
+	const int32 ByteDifference = (LastProgressUpdateBytes > 0) ? BytesReceived - LastProgressUpdateBytes : BytesReceived;
+	ensureAlwaysMsgf((ByteDifference > 0), TEXT("Invalid Byte Difference in UpdateHttpProgress -- ByteDifference:%d | LastProgressUpdateBytes:%d | BytesSent:%d | BytesReceived:%d"), ByteDifference, LastProgressUpdateBytes, BytesSent, BytesReceived);
+	LastProgressUpdateBytes = BytesReceived;
+
 	UE_LOG(LogBackgroundHttpRequest, VeryVerbose, TEXT("HttpRequest Progress Update- RequestID:%s | BytesSent: %d | BytesReceived:%d | BytesReceivedSinceLastUpdate:%d"), *OriginalRequest->GetRequestID(), BytesSent, BytesReceived, ByteDifference);
 	
-	OriginalRequest->OnProgressUpdated().ExecuteIfBound(OriginalRequest, BytesReceived, LastProgressUpdateBytes);
+	OriginalRequest->OnProgressUpdated().ExecuteIfBound(OriginalRequest, BytesReceived, ByteDifference);
 }
