Fixing incremental compile error.
#tests iOS compile

diff --git a/Engine/Source/Runtime/Online/BackgroundHTTP/Private/IOS/ApplePlatformBackgroundHttp.cpp b/Engine/Source/Runtime/Online/BackgroundHTTP/Private/IOS/ApplePlatformBackgroundHttp.cpp
--- a/Engine/Source/Runtime/Online/BackgroundHTTP/Private/IOS/ApplePlatformBackgroundHttp.cpp
+++ b/Engine/Source/Runtime/Online/BackgroundHTTP/Private/IOS/ApplePlatformBackgroundHttp.cpp
@@ -6,6 +6,8 @@
 #include "IOS/ApplePlatformBackgroundHttpResponse.h"
 #include "IOS/IOSBackgroundURLSessionHandler.h"
 
+#include "Misc/Paths.h"
+
 void FApplePlatformBackgroundHttp::Initialize()
 {
 	const FString DefaultIdentifier = TEXT("com.epicgames.backgroundhttp");
