Fix crash in FVoiceDelayBuffer::PushSample when FVoiceCaptureWindows::Init is initialized with an empty device name

diff --git a/Engine/Source/Runtime/Online/Voice/Private/Windows/VoiceCaptureWindows.cpp b/Engine/Source/Runtime/Online/Voice/Private/Windows/VoiceCaptureWindows.cpp
--- a/Engine/Source/Runtime/Online/Voice/Private/Windows/VoiceCaptureWindows.cpp
+++ b/Engine/Source/Runtime/Online/Voice/Private/Windows/VoiceCaptureWindows.cpp
@@ -104,11 +104,6 @@ bool FVoiceCaptureWindows::Init(const FString& DeviceName, int32 SampleRate, int
 		return false;
 	}
 	
-	if (DeviceName.IsEmpty())
-	{
-		return CreateCaptureBuffer(VoiceDev->DefaultVoiceCaptureDevice.DeviceName, SampleRate, NumChannels);
-	}
-
    	// init the sample counter to 0 on init
 	SampleCounter = 0; 
 	CachedSampleStart = 0;
@@ -127,7 +122,7 @@ bool FVoiceCaptureWindows::Init(const FString& DeviceName, int32 SampleRate, int
 
 	bIsMicActive = false;
 
-	return CreateCaptureBuffer(DeviceName, SampleRate, NumChannels);
+	return CreateCaptureBuffer(DeviceName.IsEmpty() ? VoiceDev->DefaultVoiceCaptureDevice.DeviceName : DeviceName, SampleRate, NumChannels);
 }
 
 bool FVoiceCaptureWindows::CreateCaptureBuffer(const FString& DeviceName, int32 SampleRate, int32 NumChannels)
