iOS Background Http Improvements:
- Improved handling of UnAssociatedTasks. Removing some bugs where they would stall and no longer update.
- Fixed bug with how many tasks we thought were active while in the foreground being inaccurate.
- Added ensures to make sure we don't violate thread safety in our callbacks.
- Minor logging changes
#tests ios client

diff --git a/Engine/Source/Runtime/ApplicationCore/Private/IOS/IOSBackgroundURLSessionHandler.cpp b/Engine/Source/Runtime/ApplicationCore/Private/IOS/IOSBackgroundURLSessionHandler.cpp
--- a/Engine/Source/Runtime/ApplicationCore/Private/IOS/IOSBackgroundURLSessionHandler.cpp
+++ b/Engine/Source/Runtime/ApplicationCore/Private/IOS/IOSBackgroundURLSessionHandler.cpp
@@ -152,7 +152,8 @@ const FString FBackgroundURLSessionHandler::GetTemporaryFilePathFromURL(const FS
         FString ConvertedPath = PlatformFile.ConvertToAbsolutePathForExternalAppForWrite(*DestinationPath);
         
 		destinationURL = [NSURL fileURLWithPath:ConvertedPath.GetNSString()];
-
+        NSLog(@"Attempting to move iOS temp file to BackgroundHttp temp folder | iOSLocation:%s | DestinationLocation:%s", [[location path] UTF8String], [DestinationPath.GetNSString() UTF8String]);
+        
 		[[NSFileManager defaultManager] moveItemAtURL:location toURL:destinationURL error:&downloadError];
 	}
     
diff --git a/Engine/Source/Runtime/Online/BackgroundHTTP/Private/IOS/ApplePlatformBackgroundHttpManager.cpp b/Engine/Source/Runtime/Online/BackgroundHTTP/Private/IOS/ApplePlatformBackgroundHttpManager.cpp
--- a/Engine/Source/Runtime/Online/BackgroundHTTP/Private/IOS/ApplePlatformBackgroundHttpManager.cpp
+++ b/Engine/Source/Runtime/Online/BackgroundHTTP/Private/IOS/ApplePlatformBackgroundHttpManager.cpp
@@ -30,6 +30,7 @@ FApplePlatformBackgroundHttpManager::FApplePlatformBackgroundHttpManager()
     , bIsInBackground(false)
     , bIsIteratingThroughSessionTasks(false)
     , RequestsPendingRemove()
+    , NumCurrentlyActiveTasks(0)
 {
 }
 
@@ -56,7 +57,11 @@ void FApplePlatformBackgroundHttpManager::PopulateUnAssociatedTasks()
 				//Store all existing tasks by their URL
 				for (id task in tasks)
 				{
-					[UnAssociatedTasks setObject:task forKey:[[[task currentRequest] URL] absoluteString]];
+                    //Only grab running tasks in this list.
+                    if ([task state] == NSURLSessionTaskStateRunning)
+                    {
+                        [UnAssociatedTasks setObject:task forKey:[[[task currentRequest] URL] absoluteString]];
+                    }
 				}
 
 				bHasFinishedPopulatingUnassociatedTasks = true;
@@ -110,20 +115,23 @@ void FApplePlatformBackgroundHttpManager::AddRequest(const FBackgroundHttpReques
     
     //See if our request is an AppleBackgroundHttpRequest so we can do more detailed checks on it.
     FAppleBackgroundHttpRequestPtr AppleRequest = StaticCastSharedPtr<FApplePlatformBackgroundHttpRequest>(Request);
-    if (ensureAlwaysMsgf(AppleRequest.IsValid(), TEXT("Adding a non-Apple background request to our Apple Background Http Manager!")))
+    if (ensureAlwaysMsgf(AppleRequest.IsValid(), TEXT("Adding a non-Apple background request to our Apple Background Http Manager! This is not supported or expected!")))
     {
         GenerateURLMapEntriesForRequest(AppleRequest);
     }
+    else
+    {
+        //just error out if we don't have the expected AppleRequest
+        return;
+    }
     
 	if (!AssociateWithAnyExistingRequest(Request))
 	{
-        //ensures above if this is false, but need to re-check it before adding tasks
-        if (AppleRequest.IsValid())
+        if (!AssociateWithAnyExistingUnAssociatedTasks(Request))
         {
             StartRequest(AppleRequest);
         }
         
-		//Just always add un-associated requests to ActiveRequests in the Apple implementation.
 		FRWScopeLock ScopeLock(ActiveRequestLock, SLT_Write);
 		ActiveRequests.Add(Request);
 	}
@@ -183,6 +191,9 @@ void FApplePlatformBackgroundHttpManager::RemoveRequest(const FBackgroundHttpReq
 
 void FApplePlatformBackgroundHttpManager::DeletePendingRemoveRequests()
 {
+    //Don't want to do this when background tasks might be using our request
+    ensureAlwaysMsgf(IsInGameThread(), TEXT("Called from un-expected thread! Potential error in an implementation of background downloads!"));
+    
     for (const FBackgroundHttpRequestPtr& Request : RequestsPendingRemove)
     {
         FBackgroundHttpManagerImpl::RemoveRequest(Request);
@@ -200,30 +211,25 @@ void FApplePlatformBackgroundHttpManager::RemoveSessionTasksForRequest(FAppleBac
     Request->CancelActiveTask();
 }
 
-bool FApplePlatformBackgroundHttpManager::AssociateWithAnyExistingRequest(const FBackgroundHttpRequestPtr Request)
+bool FApplePlatformBackgroundHttpManager::AssociateWithAnyExistingUnAssociatedTasks(const FBackgroundHttpRequestPtr Request)
 {
-	if (!bHasFinishedPopulatingUnassociatedTasks)
-	{
-		//@TODO: TRoss, might want to look at trying to associate these after bHasFinishedPopulatingUnassociatedTasks. At this point I don't think its really worth it as it SHOULD be done before we get here, however
-		//the PopulateUnAssociatedTasks() function has an asynch component so it could technically be unfinished with some tight timing.
-		UE_LOG(LogBackgroundHttpManager, Warning, TEXT("Call to AssociateWithAnyExistingRequest before we have finished populating unassociated tasks! Might have an unassociated task for this request that we won't associate with."));
-	}
-	
-	//First check if we Associate using the base implementation.
-	bool bDidAssociateWithExistingRequest = FBackgroundHttpManagerImpl::AssociateWithAnyExistingRequest(Request);
-
-	//Didn't associate, so lets check more Apple specific stuff
-	if (!bDidAssociateWithExistingRequest)
-	{
-		//See if our request is an AppleBackgroundHttpRequest so we can do more detailed checks on it.
-		FAppleBackgroundHttpRequestPtr AppleRequest = StaticCastSharedPtr<FApplePlatformBackgroundHttpRequest>(Request);
-		if (AppleRequest.IsValid())
-		{
-			bDidAssociateWithExistingRequest = CheckForExistingUnAssociatedTask(AppleRequest);
-		}
-	}
-
-	return bDidAssociateWithExistingRequest;
+    bool bDidAssociateWithUnAssociatedTask = false;
+    
+    if (!bHasFinishedPopulatingUnassociatedTasks)
+    {
+        //@TODO: TRoss, might want to look at trying to associate these after bHasFinishedPopulatingUnassociatedTasks. At this point I don't think its really worth it as it SHOULD be done before we get here, however
+        //the PopulateUnAssociatedTasks() function has an asynch component so it could technically be unfinished with some tight timing.
+        UE_LOG(LogBackgroundHttpManager, Warning, TEXT("Call to AssociateWithAnyExistingRequest before we have finished populating unassociated tasks! Might have an unassociated task for this request that we won't associate with."));
+    }
+        
+    //See if our request is an AppleBackgroundHttpRequest so we can do more detailed checks on it.
+    FAppleBackgroundHttpRequestPtr AppleRequest = StaticCastSharedPtr<FApplePlatformBackgroundHttpRequest>(Request);
+    if (AppleRequest.IsValid())
+    {
+        bDidAssociateWithUnAssociatedTask = CheckForExistingUnAssociatedTask(AppleRequest);
+    }
+    
+    return bDidAssociateWithUnAssociatedTask;
 }
 
 bool FApplePlatformBackgroundHttpManager::CheckForExistingUnAssociatedTask(const FAppleBackgroundHttpRequestPtr Request)
@@ -241,9 +247,19 @@ bool FApplePlatformBackgroundHttpManager::CheckForExistingUnAssociatedTask(const
                 if ([FoundTask state] != NSURLSessionTaskStateCanceling)
                 {
                     UE_LOG(LogBackgroundHttpManager, Display, TEXT("Existing UnAssociateTask found for Request! Attempting to Associate! RequestId:%s"), *(Request->GetRequestID()));
+                    
+                    //Associate with task so that our Request takes over ownership of this task so we can remove it from our UnAssociated Tasks list without it getting GC'd
                     Request->AssociateWithTask(FoundTask);
+                    [UnAssociatedTasks removeObjectForKey:URL.GetNSString()];
+                    
+                    //Always set our bWasTaskStartedInBG flag on our Request as true in the UnAssociated case as we don't know when it was really started
+                    FPlatformAtomics::InterlockedExchange(&(Request->bWasTaskStartedInBG), true);
+                    
+                    //Suspend task in case it was running so that we can adhere to our desired platform max tasks
+                    [FoundTask suspend];
                     
                     bDidFindExistingTask = true;
+                    break;
                 }
 			}
 		}
@@ -274,8 +290,8 @@ void FApplePlatformBackgroundHttpManager::CleanUpNSURLSessionResponseDelegates()
 
 void FApplePlatformBackgroundHttpManager::OnApp_EnteringForeground()
 {
-    FPlatformAtomics::InterlockedExchange(&bIsInBackground,false);
 	PauseAllActiveTasks();
+    FPlatformAtomics::InterlockedExchange(&bIsInBackground,false);
 }
 
 void FApplePlatformBackgroundHttpManager::OnApp_EnteringBackground()
@@ -305,7 +321,7 @@ void FApplePlatformBackgroundHttpManager::PauseAllActiveTasks()
 			}
             
             //Reset our active requests to 0 now that we are pausing everything
-            FPlatformAtomics::InterlockedExchange(&NumCurrentlyActiveRequests, 0);
+            FPlatformAtomics::InterlockedExchange(&NumCurrentlyActiveTasks, 0);
 		}];
 	}
 }
@@ -344,23 +360,35 @@ void FApplePlatformBackgroundHttpManager::OnTask_DidFinishDownloadingToURL(NSURL
 	
     UE_LOG(LogBackgroundHttpManager, Display, TEXT("Recieved Background Session Callback for URL:%s | bFileExists:%d | ErrorCode:%d | ErrorDescription:%s | Location:%s"), *TaskURL, (int)(bFileExists), ErrorCode, *ErrorDescription, *TempFilePath);
 
-	//Find request for this task and mark it complete
-	{
-        FRWScopeLock ScopeLock(URLToRequestMapLock, SLT_ReadOnly);
-        FBackgroundHttpURLMappedRequestPtr* WeakRequestInMap = URLToRequestMap.Find(TaskURL);
-        FAppleBackgroundHttpRequestPtr FoundRequest = (nullptr != WeakRequestInMap) ? WeakRequestInMap->Pin() : nullptr;
-        
-        if (FoundRequest.IsValid())
-		{
-			FoundRequest->SetRequestAsSuccess(TempFilePath);
-		}
-        
-		UE_LOG(LogBackgroundHttpManager, Display, TEXT("Attempt To Mark Task Complete -- URL:%s | bDidFindTask:%d"), *TaskURL, (int)(FoundRequest.IsValid()));
-	}
+    if (bFileExists)
+    {
+        //Find request for this task and mark it complete
+        {
+            FRWScopeLock ScopeLock(URLToRequestMapLock, SLT_ReadOnly);
+            FBackgroundHttpURLMappedRequestPtr* WeakRequestInMap = URLToRequestMap.Find(TaskURL);
+            FAppleBackgroundHttpRequestPtr FoundRequest = (nullptr != WeakRequestInMap) ? WeakRequestInMap->Pin() : nullptr;
+            
+            if (FoundRequest.IsValid())
+            {
+                FoundRequest->SetRequestAsSuccess(TempFilePath);
+            }
+            
+            UE_LOG(LogBackgroundHttpManager, Display, TEXT("Attempt To Mark Task Complete -- URL:%s | bDidFindTask:%d"), *TaskURL, (int)(FoundRequest.IsValid()));
+        }
+    }
+    else
+    {
+        //Forward to the OnCompleteWithError as we don't have our finished file!
+        UE_LOG(LogBackgroundHttpManager, Display, TEXT("File Not Found For DidFinishDownloadingToURL. Transitioning to DidCompleteWithError -- TaskURL:%s | ErrorCode:%d | ErrorDescription:%s | Location:%s"), *TaskURL, ErrorCode, *ErrorDescription, *TempFilePath);
+        OnTask_DidCompleteWithError(Task, Error);
+    }
 }
 
 void FApplePlatformBackgroundHttpManager::FinishRequest(FAppleBackgroundHttpRequestPtr Request)
 {
+    //We should only come into here from the GameThread so that if we send out a complete event our delegate subscribers don't have to worry about being thread-safe unnessecarily
+    ensureAlwaysMsgf(IsInGameThread(), TEXT("Called from un-expected thread! Potential error in an implementation of background downloads!"));
+    
     //Make sure we another thread hasn't already finished this request
     bool bHasAlreadyFinishedRequest = FPlatformAtomics::InterlockedExchange(&(Request->bHasAlreadyFinishedRequest), true);
     if (!bHasAlreadyFinishedRequest)
@@ -387,10 +415,10 @@ void FApplePlatformBackgroundHttpManager::FinishRequest(FAppleBackgroundHttpRequ
                 volatile bool bDidFail = FPlatformAtomics::AtomicRead(&(Request->bIsFailed));
                 
                 //Unexpected case where we didn't find a valid download, but we thought this completed
-                //successfully. Handle this unexpected failure
+                //successfully. Handle this unexpected failure by trying to retry the task.
                 if (!bDidFail)
                 {
-                    UE_LOG(LogBackgroundHttpManager,Warning, TEXT("Task finished downloading, but finished temp file was not found! -- RequestId:%s | TempFileLocation:%s"), *(Request->GetRequestID()), *TempFilePath);
+                    UE_LOG(LogBackgroundHttpManager,Error, TEXT("Task finished downloading, but finished temp file was not found! -- RequestId:%s | TempFileLocation:%s"), *(Request->GetRequestID()), *TempFilePath);
                     
                     //Mark our download as not completed as we hit an error so that we don't just keep trying to call FinishRequest
                     FPlatformAtomics::InterlockedExchange(&(Request->bIsCompleted), false);
@@ -404,25 +432,28 @@ void FApplePlatformBackgroundHttpManager::FinishRequest(FAppleBackgroundHttpRequ
                 //Expected case where we failed, but expected to fail
                 else
                 {
-                    UE_LOG(LogBackgroundHttpManager, Display, TEXT("Task failed due to being out of retries. -- RequestID:%s"), *(Request->GetRequestID()));
+                    UE_LOG(LogBackgroundHttpManager, Display, TEXT("Task failed completely -- RequestID:%s"), *(Request->GetRequestID()));
                     
                     FBackgroundHttpResponsePtr NewResponse = FPlatformBackgroundHttp::ConstructBackgroundResponse(ResponseCode, TEXT(""));
                     Request->CompleteWithExistingResponseData(NewResponse);
                 }
             }
             
-            //If we are actually finishing this request, lets decrement our NumCurrentlyActiveRequests counter
+            //If we are actually finishing this request, lets decrement our NumCurrentlyActiveTasks counter
             if (bIsRequestActuallyFinished)
             {
-                //See if we are in the BG or FG. In the BG, we don't want to decrement NumCurrentlyActiveRequests
+                //See if we are in the BG or FG. In the BG, we don't want to decrement NumCurrentlyActiveTasks
                 const bool bWasFinishedInBG = FPlatformAtomics::AtomicRead(&bIsInBackground);
                 if (!bWasFinishedInBG)
                 {
-                    int NumActualRequests = FPlatformAtomics::InterlockedDecrement(&NumCurrentlyActiveRequests);
+                    int NumActualTasks = FPlatformAtomics::InterlockedDecrement(&NumCurrentlyActiveTasks);
 
                     //Sanity check that our data is valid. Shouldn't ever trip if everything is working as intended.
-                    const bool bNumActualRequestsValid = ((NumActualRequests > 0) && (NumActualRequests <= FPlatformBackgroundHttp::GetPlatformMaxActiveDownloads()));
-                    ensureMsgf(bNumActualRequestsValid, TEXT("Number of Requests we think are active is invalid! -- NumCurrentlyActiveRequests:%d"), NumActualRequests);
+                    const bool bNumActualTasksIsValid = ((NumActualTasks >= 0) && (NumActualTasks <= FPlatformBackgroundHttp::GetPlatformMaxActiveDownloads()));
+                    
+                    UE_LOG(LogBackgroundHttpManager,Display, TEXT("Finishing Request lowering Task Count: %d"), NumActualTasks);
+                    
+                    ensureMsgf(bNumActualTasksIsValid, TEXT("Number of Requests we think are active is invalid! -- NumCurrentlyActiveTasks:%d"), NumActualTasks);
                 }
             }
         }
@@ -602,6 +633,8 @@ void FApplePlatformBackgroundHttpManager::OnSession_SessionDidFinishAllEvents(NS
 
 bool FApplePlatformBackgroundHttpManager::Tick(float DeltaTime)
 {
+    ensureAlwaysMsgf(IsInGameThread(), TEXT("Called from un-expected thread! Potential error in an implementation of background downloads!"));
+    
     TickRequests(DeltaTime);
     TickTasks(DeltaTime);
     TickUnassociatedTasks(DeltaTime);
@@ -612,6 +645,8 @@ bool FApplePlatformBackgroundHttpManager::Tick(float DeltaTime)
 
 void FApplePlatformBackgroundHttpManager::TickRequests(float DeltaTime)
 {
+    ensureAlwaysMsgf(IsInGameThread(), TEXT("Called from un-expected thread! Potential error in an implementation of background downloads!"));
+    
     //First lets go through all our Requests to see if we need to complete or recreate any requests
     {
         FRWScopeLock ScopeLock(ActiveRequestLock, SLT_Write);
@@ -624,6 +659,7 @@ void FApplePlatformBackgroundHttpManager::TickRequests(float DeltaTime)
                 const bool bIsTaskPaused = AppleRequest->IsUnderlyingTaskPaused();
                 const bool bIsTaskComplete = AppleRequest->IsTaskComplete();
                 const bool bWasStartedInBG = FPlatformAtomics::AtomicRead(&(AppleRequest->bWasTaskStartedInBG));
+                const bool bIsPendingCancel = FPlatformAtomics::AtomicRead(&(AppleRequest->bIsPendingCancel));
                 
                 UE_LOG(LogBackgroundHttpManager, VeryVerbose, TEXT("Checking Status of Request on Tick -- RequestID::%s | bIsTaskComplete:%d | bWasStartedInBG:%d"), *(AppleRequest->GetRequestID()), (int)bIsTaskComplete, (int)bWasStartedInBG);
                 
@@ -632,15 +668,18 @@ void FApplePlatformBackgroundHttpManager::TickRequests(float DeltaTime)
                     UE_LOG(LogBackgroundHttpManager, Display, TEXT("Calling FinishRequest On -- RequestID::%s | bIsTaskComplete:%d | bWasStartedInBG:%d"), *(AppleRequest->GetRequestID()), (int)bIsTaskComplete, (int)bWasStartedInBG);
                     FinishRequest(AppleRequest);
                 }
-                else if (bWasStartedInBG)
+                else if (bWasStartedInBG && !bIsPendingCancel)
                 {
                     UE_LOG(LogBackgroundHttpManager, Display, TEXT("Cancelling Request Created In BG To Re-Create In FG -- RequestID:%s"), *(AppleRequest->GetRequestID()));
                     
+                    //reset to false so we don't run this twice while waiting on recreation
+                    FPlatformAtomics::InterlockedExchange(&(AppleRequest->bWasTaskStartedInBG), false);
+                    
                     //Just cancel the task. This will lead to it getting a callback to OnTask_DidCompleteWithError where we will re-create it
                     //We want to recreate any task spun up in the background as it will not respect our session settings if created in BG.
                     AppleRequest->CancelActiveTask();
                 }
-                else if (bIsTaskActive && !bIsTaskPaused)
+                else if (bIsTaskActive && !bIsTaskPaused && !bIsPendingCancel)
                 {
                     const bool bShouldTimeOut = AppleRequest->TickTimeOutTimer(DeltaTime);
                     if (bShouldTimeOut)
@@ -650,6 +689,12 @@ void FApplePlatformBackgroundHttpManager::TickRequests(float DeltaTime)
                         //Just cancel the task and let the OnTask_DidCompleteWithError callback handle retrying it if appropriate.
                         AppleRequest->CancelActiveTask();
                     }
+                    
+                    AppleRequest->SendDownloadProgressUpdate();
+                }
+                else
+                {
+                    AppleRequest->SendDownloadProgressUpdate();
                 }
             }
         }
@@ -661,6 +706,8 @@ void FApplePlatformBackgroundHttpManager::TickRequests(float DeltaTime)
 
 void FApplePlatformBackgroundHttpManager::TickTasks(float DeltaTime)
 {
+    ensureAlwaysMsgf(IsInGameThread(), TEXT("Called from un-expected thread! Potential error in an implementation of background downloads!"));
+    
     //Go through all our session's tasks and unpause as many tasks as our platform's limit allows
     
     NSURLSession* BackgroundDownloadSession = FBackgroundURLSessionHandler::GetBackgroundSession();
@@ -673,7 +720,7 @@ void FApplePlatformBackgroundHttpManager::TickTasks(float DeltaTime)
             [BackgroundDownloadSession getAllTasksWithCompletionHandler:^(NSArray<__kindof NSURLSessionTask*> *tasks)
              {
                  //Check to make sure we have room for more tasks to be active first
-                 int CurrentCount = FPlatformAtomics::AtomicRead(&NumCurrentlyActiveRequests);
+                 int CurrentCount = FPlatformAtomics::AtomicRead(&NumCurrentlyActiveTasks);
                  if (CurrentCount < FPlatformBackgroundHttp::GetPlatformMaxActiveDownloads())
                  {
                      //Go through our tasks and try and activate as many as possible
@@ -683,7 +730,9 @@ void FApplePlatformBackgroundHttpManager::TickTasks(float DeltaTime)
                          if ([task state] == NSURLSessionTaskStateSuspended)
                          {
                              //Check to make sure we didn't have another task increment our count over the limit since we checked above
-                             int NewRequestCount = FPlatformAtomics::InterlockedIncrement(&NumCurrentlyActiveRequests);
+                             int NewRequestCount = FPlatformAtomics::InterlockedIncrement(&NumCurrentlyActiveTasks);
+                             UE_LOG(LogBackgroundHttpManager,Verbose, TEXT("Incrementing Task Count: %d"), NewRequestCount);
+                             
                              if (NewRequestCount <= FPlatformBackgroundHttp::GetPlatformMaxActiveDownloads())
                              {
                                  FString TaskURL = [[[task currentRequest] URL] absoluteString];
@@ -698,7 +747,7 @@ void FApplePlatformBackgroundHttpManager::TickTasks(float DeltaTime)
                                      
                                      if (FoundRequest.IsValid())
                                      {
-                                         UE_LOG(LogBackgroundHttpManager, Display, TEXT("Activating Task For Requets -- RequestID:%s | TaskURL:%s | CurrentlyActiveRequests:%d"), *(FoundRequest->GetRequestID()), *TaskURL, NewRequestCount);
+                                         UE_LOG(LogBackgroundHttpManager, Display, TEXT("Manager Calling to Active Task For Request -- RequestID:%s | TaskURL:%s | CurrentlyActiveRequests:%d"), *(FoundRequest->GetRequestID()), *TaskURL, NewRequestCount);
                                          FoundRequest->ActivateUnderlyingTask();
                                      }
                                      else
@@ -706,17 +755,17 @@ void FApplePlatformBackgroundHttpManager::TickTasks(float DeltaTime)
                                          UE_LOG(LogBackgroundHttpManager, Display, TEXT("Skipping Activating Task as there is no associated Request or Request is paused. Once a Request associates with this task, it can then be activated. -- TaskURL:%s | bIsPaused:%d"), *TaskURL, (int)bIsPaused);
                                          
                                          //Don't activate and remove our increment from above because something put us over the limit before we resumed
-                                         FPlatformAtomics::InterlockedDecrement(&NumCurrentlyActiveRequests);
+                                         FPlatformAtomics::InterlockedDecrement(&NumCurrentlyActiveTasks);
                                      }
                                  }
                              }
                              else
-                             {
+                             {                                 
                                  FString TaskURL = [[[task currentRequest] URL] absoluteString];
                                  UE_LOG(LogBackgroundHttpManager, Log, TEXT("Task failed to activate as we passed the platform max from another task before we could resume. Task -- TaskURL:%s | CurrentlyActiveRequests:%d"), *TaskURL, NewRequestCount);
                                  
                                  //Don't activate and remove our increment from above because something put us over the limit before we resumed
-                                 FPlatformAtomics::InterlockedDecrement(&NumCurrentlyActiveRequests);
+                                 FPlatformAtomics::InterlockedDecrement(&NumCurrentlyActiveTasks);
                                  break;
                              }
                              
@@ -740,8 +789,10 @@ void FApplePlatformBackgroundHttpManager::TickTasks(float DeltaTime)
 
 void FApplePlatformBackgroundHttpManager::TickUnassociatedTasks(float DeltaTime)
 {
+    ensureAlwaysMsgf(IsInGameThread(), TEXT("Called from un-expected thread! Potential error in an implementation of background downloads!"));
+    
     //If we don't have anything queued, lets resume any un-associated tasks
-    int CurrentCount = FPlatformAtomics::AtomicRead(&NumCurrentlyActiveRequests);
+    int CurrentCount = FPlatformAtomics::AtomicRead(&NumCurrentlyActiveTasks);
     if (CurrentCount == 0)
     {
         UnpauseAllUnassociatedTasks();
diff --git a/Engine/Source/Runtime/Online/BackgroundHTTP/Private/IOS/ApplePlatformBackgroundHttpRequest.cpp b/Engine/Source/Runtime/Online/BackgroundHTTP/Private/IOS/ApplePlatformBackgroundHttpRequest.cpp
--- a/Engine/Source/Runtime/Online/BackgroundHTTP/Private/IOS/ApplePlatformBackgroundHttpRequest.cpp
+++ b/Engine/Source/Runtime/Online/BackgroundHTTP/Private/IOS/ApplePlatformBackgroundHttpRequest.cpp
@@ -17,7 +17,9 @@ FApplePlatformBackgroundHttpRequest::FApplePlatformBackgroundHttpRequest()
     , bIsFailed(false)
     , bWasTaskStartedInBG(false)
     , bHasAlreadyFinishedRequest(false)
+    , bIsPendingCancel(false)
     , DownloadProgress(0)
+    , DownloadProgressSinceLastUpdateSent(0)
 {    
 }
 
@@ -47,6 +49,7 @@ void FApplePlatformBackgroundHttpRequest::CompleteRequest_Internal(bool bWasRequ
     FPlatformAtomics::InterlockedExchange(&bIsTaskActive, false);
 	FPlatformAtomics::InterlockedExchange(&bIsCompleted, true);
 	FPlatformAtomics::InterlockedExchange(&bIsFailed, !bWasRequestSuccess);
+    FPlatformAtomics::InterlockedExchange(&bIsPendingCancel, false);
     
     if (!CompletedTempDownloadLocationIn.IsEmpty())
     {
@@ -97,7 +100,8 @@ void FApplePlatformBackgroundHttpRequest::ActivateUnderlyingTask()
         
             FPlatformAtomics::InterlockedExchange(&bIsTaskActive, true);
             FPlatformAtomics::InterlockedExchange(&bIsTaskPaused, false);
-
+            FPlatformAtomics::InterlockedExchange(&bIsPendingCancel, false);
+            
             [UnderlyingTask resume];
             
             ResetTimeOutTimer();
@@ -118,7 +122,9 @@ void FApplePlatformBackgroundHttpRequest::PauseUnderlyingTask()
             UE_LOG(LogBackgroundHttpRequest, Display, TEXT("Pausing Task for Request -- RequestID:%s | TaskURL:%s"), *GetRequestID(), *TaskURL);
             
             FPlatformAtomics::InterlockedExchange(&bIsTaskActive, false);
+            FPlatformAtomics::InterlockedExchange(&bIsPendingCancel, false);
             FPlatformAtomics::InterlockedExchange(&bIsTaskPaused, true);
+            
             [UnderlyingTask suspend];
             
             ResetTimeOutTimer();
@@ -164,6 +170,8 @@ void FApplePlatformBackgroundHttpRequest::AssociateWithTask(NSURLSessionTask* Ex
 		FString TaskURL = [[[ExistingTask currentRequest] URL] absoluteString];
 		UE_LOG(LogBackgroundHttpRequest, Display, TEXT("Associated Request With New Task -- RequestID:%s | TaskURL:%s"), *GetRequestID(), *TaskURL);
         
+        FPlatformAtomics::InterlockedExchange(&bIsPendingCancel, false);
+        
         ResetTimeOutTimer();
         ResetProgressTracking();
 	}
@@ -188,7 +196,9 @@ void FApplePlatformBackgroundHttpRequest::CancelActiveTask()
         {
             FString TaskURL = [[[TaskNodeWeAreCancelling->OurTask currentRequest] URL] absoluteString];
             UE_LOG(LogBackgroundHttpRequest, Display, TEXT("Cancelling Task -- RequestID:%s | TaskURL:%s"), *GetRequestID(), *TaskURL);
-
+            
+            FPlatformAtomics::InterlockedExchange(&bIsPendingCancel, true);
+            
             [TaskNodeWeAreCancelling->OurTask cancel];
         }
 	}
@@ -199,9 +209,23 @@ void FApplePlatformBackgroundHttpRequest::UpdateDownloadProgress(int64_t TotalDo
     UE_LOG(LogBackgroundHttpRequest, VeryVerbose, TEXT("Request Update Progress -- RequestID:%s | OldProgress:%lld | NewProgress:%lld | ProgressSinceLastUpdate:%lld"), *GetRequestID(), DownloadProgress, TotalDownloaded, DownloadedSinceLastUpdate);
 	
     FPlatformAtomics::AtomicStore(&DownloadProgress, TotalDownloaded);
+    FPlatformAtomics::InterlockedAdd(&DownloadProgressSinceLastUpdateSent, DownloadedSinceLastUpdate);
+    
     ResetTimeOutTimer();
+}
+
+void FApplePlatformBackgroundHttpRequest::SendDownloadProgressUpdate()
+{
+    volatile int64 DownloadProgressCopy = FPlatformAtomics::AtomicRead(&DownloadProgress);
     
-    OnProgressUpdated().ExecuteIfBound(SharedThis(this), TotalDownloaded, DownloadedSinceLastUpdate);
+    //Don't send any updates if we haven't updated anything since we last sent an update
+    if (DownloadProgressCopy > 0)
+    {
+        //Reset our DownloadProgressSinceLastUpdateSent to 0 now that we are sending a progress update
+        volatile int64 DownloadProgressSinceLastUpdateSentCopy = FPlatformAtomics::InterlockedExchange(&DownloadProgressSinceLastUpdateSent, 0);
+        
+        OnProgressUpdated().ExecuteIfBound(SharedThis(this), DownloadProgressCopy, DownloadProgressSinceLastUpdateSentCopy);
+    }
 }
 
 bool FApplePlatformBackgroundHttpRequest::IsTaskComplete() const
diff --git a/Engine/Source/Runtime/Online/BackgroundHTTP/Public/IOS/ApplePlatformBackgroundHttpManager.h b/Engine/Source/Runtime/Online/BackgroundHTTP/Public/IOS/ApplePlatformBackgroundHttpManager.h
--- a/Engine/Source/Runtime/Online/BackgroundHTTP/Public/IOS/ApplePlatformBackgroundHttpManager.h
+++ b/Engine/Source/Runtime/Online/BackgroundHTTP/Public/IOS/ApplePlatformBackgroundHttpManager.h
@@ -46,8 +46,8 @@ public:
 	static int RetryResumeDataLimitSetting;
 
 protected:
-	virtual bool AssociateWithAnyExistingRequest(const FBackgroundHttpRequestPtr Request) override;
-	
+	bool AssociateWithAnyExistingUnAssociatedTasks(const FBackgroundHttpRequestPtr Request);
+        
 	bool CheckForExistingUnAssociatedTask(const FAppleBackgroundHttpRequestPtr Request);
 
     void DeletePendingRemoveRequests();
@@ -112,4 +112,7 @@ private:
     
     //Need to store requests to remove at the end of our tick as we find them during iteration
     TArray<FBackgroundHttpRequestPtr> RequestsPendingRemove;
+    
+    /** On iOS we need to track how many Tasks we have active. This is to replace the default implementations NumCurrentlyActiveRequests **/
+    volatile int NumCurrentlyActiveTasks;
 };
diff --git a/Engine/Source/Runtime/Online/BackgroundHTTP/Public/IOS/ApplePlatformBackgroundHttpRequest.h b/Engine/Source/Runtime/Online/BackgroundHTTP/Public/IOS/ApplePlatformBackgroundHttpRequest.h
--- a/Engine/Source/Runtime/Online/BackgroundHTTP/Public/IOS/ApplePlatformBackgroundHttpRequest.h
+++ b/Engine/Source/Runtime/Online/BackgroundHTTP/Public/IOS/ApplePlatformBackgroundHttpRequest.h
@@ -39,7 +39,11 @@ private:
     void SetRequestAsFailed();
     void CompleteRequest_Internal(bool bWasRequestSuccess, const FString& CompletedTempDownloadLocation);
 	void CancelActiveTask();
-	void UpdateDownloadProgress(int64_t TotalDownloaded, int64_t DownloadedSinceLastUpdate);
+	
+    //on iOS we need to delay when we send progress updates as they can come from a background thread
+    //save off results in UpdateDownloadProgress to send later when we get a SendDownloadProgressUpdate
+    void UpdateDownloadProgress(int64_t TotalDownloaded, int64_t DownloadedSinceLastUpdate);
+    void SendDownloadProgressUpdate();
     void ResetProgressTracking();
     
     void ActivateUnderlyingTask();
@@ -63,7 +67,10 @@ private:
 	volatile int32 bIsFailed;
 	volatile int32 bWasTaskStartedInBG;
     volatile int32 bHasAlreadyFinishedRequest;
+    volatile int32 bIsPendingCancel;
+    
 	volatile int64 DownloadProgress;
+    volatile int64 DownloadProgressSinceLastUpdateSent;
     
     friend class FApplePlatformBackgroundHttpManager;
 };
