#xmpp
Add a construct method for FXmppUserJid from a formatted JID

diff --git a/Engine/Source/Runtime/Online/XMPP/Private/XmppConnection.cpp b/Engine/Source/Runtime/Online/XMPP/Private/XmppConnection.cpp
--- a/Engine/Source/Runtime/Online/XMPP/Private/XmppConnection.cpp
+++ b/Engine/Source/Runtime/Online/XMPP/Private/XmppConnection.cpp
@@ -15,6 +15,33 @@ enum class XMPP_RESOURCE_VERSION : uint8
 	LATEST = VERSION_PLUSONE - 1
 };
 
+FXmppUserJid FXmppUserJid::FromFullJid(const FString& JidString)
+{
+	FString User;
+	FString Domain;
+	FString Resource;
+
+	FString DomainAndResource;
+	if (JidString.Split(TEXT("@"), &User, &DomainAndResource, ESearchCase::CaseSensitive, ESearchDir::FromStart))
+	{
+		if (!DomainAndResource.Split(TEXT("/"), &Domain, &Resource, ESearchCase::CaseSensitive, ESearchDir::FromEnd))
+		{
+			// If we don't have a resource, the domain is all of DomainAndResource
+			Domain = MoveTemp(DomainAndResource);
+		}
+	}
+	else
+	{
+		if (!JidString.Split(TEXT("/"), &Domain, &Resource, ESearchCase::CaseSensitive, ESearchDir::FromEnd))
+		{
+			// If we don't have a resource, we only have the Domain in the JidString
+			Domain = JidString;
+		}
+	}
+
+	return FXmppUserJid(MoveTemp(User), MoveTemp(Domain), MoveTemp(Resource));
+}
+
 bool FXmppUserJid::ParseResource(const FString& InResource, FString& OutAppId, FString& OutPlatform, FString& OutPlatformUserId)
 {
 	OutAppId.Empty();
diff --git a/Engine/Source/Runtime/Online/XMPP/Private/XmppStrophe/XmppStrophe.cpp b/Engine/Source/Runtime/Online/XMPP/Private/XmppStrophe/XmppStrophe.cpp
--- a/Engine/Source/Runtime/Online/XMPP/Private/XmppStrophe/XmppStrophe.cpp
+++ b/Engine/Source/Runtime/Online/XMPP/Private/XmppStrophe/XmppStrophe.cpp
@@ -34,29 +34,7 @@ FString FXmppStrophe::JidToString(const FXmppUserJid& UserJid)
 
 FXmppUserJid FXmppStrophe::JidFromString(const FString& JidString)
 {
-	FString User;
-	FString Domain;
-	FString Resource;
-
-	FString DomainAndResource;
-	if (JidString.Split(TEXT("@"), &User, &DomainAndResource, ESearchCase::CaseSensitive, ESearchDir::FromStart))
-	{
-		if (!DomainAndResource.Split(TEXT("/"), &Domain, &Resource, ESearchCase::CaseSensitive, ESearchDir::FromEnd))
-		{
-			// If we don't have a resource, the domain is all of DomainAndResource
-			Domain = MoveTemp(DomainAndResource);
-		}
-	}
-	else
-	{
-		if (!JidString.Split(TEXT("/"), &Domain, &Resource, ESearchCase::CaseSensitive, ESearchDir::FromEnd))
-		{
-			// If we don't have a resource, we only have the Domain in the JidString
-			Domain = JidString;
-		}
-	}
-
-	return FXmppUserJid(MoveTemp(User), MoveTemp(Domain), MoveTemp(Resource));
+	return FXmppUserJid::FromFullJid(JidString);
 }
 
 FXmppUserJid FXmppStrophe::JidFromStropheString(const char* StropheJidString)
diff --git a/Engine/Source/Runtime/Online/XMPP/Public/XmppConnection.h b/Engine/Source/Runtime/Online/XMPP/Public/XmppConnection.h
--- a/Engine/Source/Runtime/Online/XMPP/Public/XmppConnection.h
+++ b/Engine/Source/Runtime/Online/XMPP/Public/XmppConnection.h
@@ -103,6 +103,15 @@ public:
 	{
 	}
 
+	/**
+	 * Create an FXmppUserJid from a formatted JID
+	 * e.g., user@domain/resource
+	 *
+	 * @param JidString the formatted JID
+	 * @return FXmppUserJid populated from the formatted JID
+	 */
+	static FXmppUserJid FromFullJid(const FString& JidString);
+
 	/** unique id for the user */
 	FString Id;
 	/** domain user has access to */
