Background Http Improvements:
- Fixed bug where BackgroundHttp temp files were not being cleaned up by calls to the CleanUpTemporaryFiles on the  BackgroundHttpManager were not correctly cleaning temp files on iOS.
- Improved handling of BackgroundHttp downloads that were in progress at app start by improving UnAssociatedTask handling.
- Fixed bug where Apple BackgroundHttp could end up queueing duplicate download tasks for the same file in some threading cases.
- Made general logging improvements and converted some NSLogs and Warnings into ensures to better track bad behavior.
#tests iOS Client

diff --git a/Engine/Source/Runtime/ApplicationCore/Private/IOS/IOSBackgroundURLSessionHandler.cpp b/Engine/Source/Runtime/ApplicationCore/Private/IOS/IOSBackgroundURLSessionHandler.cpp
--- a/Engine/Source/Runtime/ApplicationCore/Private/IOS/IOSBackgroundURLSessionHandler.cpp
+++ b/Engine/Source/Runtime/ApplicationCore/Private/IOS/IOSBackgroundURLSessionHandler.cpp
@@ -102,9 +102,9 @@ NSURLSession* FBackgroundURLSessionHandler::GetBackgroundSession()
 
 void FBackgroundURLSessionHandler::CreateBackgroundSessionWorkingDirectory()
 {
-	const FString DirectoryPath = GetBackgroundSessionWorkingDirectoryPath();
+	const FString& DirectoryPath = GetBackgroundSessionWorkingDirectoryPath();
     
-	if (!DirectoryPath.IsEmpty())
+	if (ensureAlwaysMsgf(!DirectoryPath.IsEmpty(), TEXT("Invalid empty BackgroundSessionWorkingDirectoryPath!")))
 	{
         IPlatformFile& PlatformFile = FPlatformFileManager::Get().GetPlatformFile();
         PlatformFile.CreateDirectory(*DirectoryPath);
@@ -113,19 +113,20 @@ void FBackgroundURLSessionHandler::CreateBackgroundSessionWorkingDirectory()
         {
             NSString* PathAsNSString = DirectoryPath.GetNSString();
             //May not yet have UE_Log
-            NSLog(@"Warning: Error in CreateBackgroundSessionWorkingDirectory: %s", [PathAsNSString UTF8String]);
+            ensureAlwaysMsgf(false, TEXT("Warning: Error in CreateBackgroundSessionWorkingDirectory: %s"), [PathAsNSString UTF8String]);
         }
 	}
 }
 
-FString FBackgroundURLSessionHandler::GetBackgroundSessionWorkingDirectoryPath()
+const FString& FBackgroundURLSessionHandler::GetBackgroundSessionWorkingDirectoryPath()
 {
-	return FPaths::Combine(FPlatformMisc::GamePersistentDownloadDir(), TEXT("BackgroundHttp"));
+    static FString BackgroundHttpDir = FPaths::Combine(FPlatformMisc::GamePersistentDownloadDir(), TEXT("BackgroundHttp"));
+	return BackgroundHttpDir;
 }
 
 const FString FBackgroundURLSessionHandler::GetTemporaryFilePathFromURL(const FString& URL)
 {
-	static FString BackgroundHttpDir = GetBackgroundSessionWorkingDirectoryPath();
+	const FString& BackgroundHttpDir = GetBackgroundSessionWorkingDirectoryPath();
 	const FString FileName = FPaths::MakeValidFileName(URL);
 
 	return FPaths::Combine(BackgroundHttpDir, FileName);
diff --git a/Engine/Source/Runtime/ApplicationCore/Public/IOS/IOSBackgroundURLSessionHandler.h b/Engine/Source/Runtime/ApplicationCore/Public/IOS/IOSBackgroundURLSessionHandler.h
--- a/Engine/Source/Runtime/ApplicationCore/Public/IOS/IOSBackgroundURLSessionHandler.h
+++ b/Engine/Source/Runtime/ApplicationCore/Public/IOS/IOSBackgroundURLSessionHandler.h
@@ -36,7 +36,7 @@ public:
 	static NSURLSession* GetBackgroundSession();
 
 	//Gets a path stored in an FString of the working directory any completed temp files will be downloaded
-	static FString GetBackgroundSessionWorkingDirectoryPath();
+	static const FString& GetBackgroundSessionWorkingDirectoryPath();
 
 	/**
 	* Function that takes in a URL and figures out the location we should use as the temp storage URL
diff --git a/Engine/Source/Runtime/Online/BackgroundHTTP/Private/IOS/ApplePlatformBackgroundHttp.cpp b/Engine/Source/Runtime/Online/BackgroundHTTP/Private/IOS/ApplePlatformBackgroundHttp.cpp
--- a/Engine/Source/Runtime/Online/BackgroundHTTP/Private/IOS/ApplePlatformBackgroundHttp.cpp
+++ b/Engine/Source/Runtime/Online/BackgroundHTTP/Private/IOS/ApplePlatformBackgroundHttp.cpp
@@ -41,8 +41,7 @@ const FString FApplePlatformBackgroundHttp::GetTemporaryFilePathFromURL(const FS
 
 const FString& FApplePlatformBackgroundHttp::GetTemporaryRootPath()
 {
-    static FString BackgroundHttpDir = FPaths::Combine(FPlatformMisc::GamePersistentDownloadDir(), TEXT("BackgroundHttpTemp"));
-    return BackgroundHttpDir;
+    return FBackgroundURLSessionHandler::GetBackgroundSessionWorkingDirectoryPath();
 }
 
 int FApplePlatformBackgroundHttp::GetPlatformMaxActiveDownloads()
diff --git a/Engine/Source/Runtime/Online/BackgroundHTTP/Private/IOS/ApplePlatformBackgroundHttpManager.cpp b/Engine/Source/Runtime/Online/BackgroundHTTP/Private/IOS/ApplePlatformBackgroundHttpManager.cpp
--- a/Engine/Source/Runtime/Online/BackgroundHTTP/Private/IOS/ApplePlatformBackgroundHttpManager.cpp
+++ b/Engine/Source/Runtime/Online/BackgroundHTTP/Private/IOS/ApplePlatformBackgroundHttpManager.cpp
@@ -57,11 +57,7 @@ void FApplePlatformBackgroundHttpManager::PopulateUnAssociatedTasks()
 				//Store all existing tasks by their URL
 				for (id task in tasks)
 				{
-                    //Only grab running tasks in this list.
-                    if ([task state] == NSURLSessionTaskStateRunning)
-                    {
-                        [UnAssociatedTasks setObject:task forKey:[[[task currentRequest] URL] absoluteString]];
-                    }
+					[UnAssociatedTasks setObject:task forKey:[[[task currentRequest] URL] absoluteString]];
 				}
 
 				bHasFinishedPopulatingUnassociatedTasks = true;
@@ -134,6 +130,12 @@ void FApplePlatformBackgroundHttpManager::AddRequest(const FBackgroundHttpReques
         
 		FRWScopeLock ScopeLock(ActiveRequestLock, SLT_Write);
 		ActiveRequests.Add(Request);
+
+		//Increment our underlying FBackgroundHttpManagerImpl tracker for active requests as we
+		//don't implement the method it uses to increase this number.
+		// NOTE: We don't make use of this number in Apple Platform functions as all requests are "Active" but their
+		// underlying Task might not be, see NumCurrentlyActiveTasks instead to track how many current Tasks are downloading data.
+		++NumCurrentlyActiveRequests;
 	}
 }
 
@@ -166,7 +168,7 @@ void FApplePlatformBackgroundHttpManager::RemoveURLMapEntriesForRequest(FAppleBa
         
         if (FoundRequest == Request)
         {
-            UE_LOG(LogBackgroundHttpManager, Display, TEXT("Removing URL Entry -- RequestID:%s | URL:%s"), *Request->GetRequestID(), *URL);
+            UE_LOG(LogBackgroundHttpManager, Verbose, TEXT("Removing URL Entry -- RequestDebugID:%s | URL:%s"), *Request->GetRequestDebugID(), *URL);
             URLToRequestMap.Remove(URL);
         }
     }
@@ -244,23 +246,27 @@ bool FApplePlatformBackgroundHttpManager::CheckForExistingUnAssociatedTask(const
 			NSURLSessionTask* FoundTask = [UnAssociatedTasks valueForKey:URL.GetNSString()];
 			if (nullptr != FoundTask)
 			{
-                if ([FoundTask state] != NSURLSessionTaskStateCanceling)
-                {
-                    UE_LOG(LogBackgroundHttpManager, Display, TEXT("Existing UnAssociateTask found for Request! Attempting to Associate! RequestId:%s"), *(Request->GetRequestID()));
-                    
-                    //Associate with task so that our Request takes over ownership of this task so we can remove it from our UnAssociated Tasks list without it getting GC'd
-                    Request->AssociateWithTask(FoundTask);
-                    [UnAssociatedTasks removeObjectForKey:URL.GetNSString()];
-                    
-                    //Always set our bWasTaskStartedInBG flag on our Request as true in the UnAssociated case as we don't know when it was really started
-                    FPlatformAtomics::InterlockedExchange(&(Request->bWasTaskStartedInBG), true);
+                UE_LOG(LogBackgroundHttpManager, Display, TEXT("Existing UnAssociateTask found for Request! Attempting to Associate! RequestDebugID:%s"), *(Request->GetRequestDebugID()));
                     
-                    //Suspend task in case it was running so that we can adhere to our desired platform max tasks
-                    [FoundTask suspend];
-                    
-                    bDidFindExistingTask = true;
-                    break;
-                }
+                //Associate with task so that our Request takes over ownership of this task so we can remove it from our UnAssociated Tasks list without it getting GC'd
+				if (Request->AssociateWithTask(FoundTask))
+				{
+					//Always set our bWasTaskStartedInBG flag on our Request as true in the UnAssociated case as we don't know when it was really started
+					FPlatformAtomics::InterlockedExchange(&(Request->bWasTaskStartedInBG), true);
+
+					//Suspend task in case it was running so that we can adhere to our desired platform max tasks
+					[FoundTask suspend];
+
+					bDidFindExistingTask = true;
+					break;
+				}
+				else
+				{
+					UE_LOG(LogBackgroundHttpManager, Display, TEXT("UnAssociatedTask for request found, but failed to Associate with Task! -- RequestDebugID:%s | URL:%s"), *(Request->GetRequestDebugID()), *URL);
+				}
+
+				//Still want to remove UnAssociatedTask even though we didn't use it as something else can now be downloading this data and we do not want duplicates
+				[UnAssociatedTasks removeObjectForKey : URL.GetNSString()];
 			}
 		}
 	}
@@ -314,7 +320,9 @@ void FApplePlatformBackgroundHttpManager::PauseAllActiveTasks()
 				if ([DownloadTask state] == NSURLSessionTaskStateRunning)
 				{
                     FString TaskURL = [[[DownloadTask currentRequest] URL] absoluteString];
-                    UE_LOG(LogBackgroundHttpManager, Display, TEXT("Pausing Task for URL:%s"), *TaskURL);
+                    int TaskIdentifier = (int)[DownloadTask taskIdentifier];
+                    
+                    UE_LOG(LogBackgroundHttpManager, Display, TEXT("Pausing Task for URL:%s | TaskIdentifier:%d"), *TaskURL, TaskIdentifier);
                     
 					[DownloadTask suspend];
 				}
@@ -340,7 +348,9 @@ void FApplePlatformBackgroundHttpManager::ResumeAllTasks()
 				if ([DownloadTask state] == NSURLSessionTaskStateSuspended)
 				{
                     FString TaskURL = [[[DownloadTask currentRequest] URL] absoluteString];
-                    UE_LOG(LogBackgroundHttpManager, Display, TEXT("Resuming Task for URL:%s"), *TaskURL);
+                    int TaskIdentifier = (int)[DownloadTask taskIdentifier];
+                    
+                    UE_LOG(LogBackgroundHttpManager, Display, TEXT("Resuming Task for URL:%s | TaskIdentifier:%d"), *TaskURL, TaskIdentifier);
                     
 					[DownloadTask resume];
 				}
@@ -352,13 +362,15 @@ void FApplePlatformBackgroundHttpManager::ResumeAllTasks()
 void FApplePlatformBackgroundHttpManager::OnTask_DidFinishDownloadingToURL(NSURLSessionDownloadTask* Task, NSError* Error, const FString& TempFilePath)
 {
 	FString TaskURL = [[[Task currentRequest] URL] absoluteString];
+    int TaskIdentifier = (int)[Task taskIdentifier];
+    
     const int ErrorCode = [Error code];
     const FString ErrorDescription = [Error localizedDescription];
     
     IPlatformFile& PlatformFile = FPlatformFileManager::Get().GetPlatformFile();
 	const bool bFileExists = PlatformFile.FileExists(*TempFilePath);
 	
-    UE_LOG(LogBackgroundHttpManager, Display, TEXT("Recieved Background Session Callback for URL:%s | bFileExists:%d | ErrorCode:%d | ErrorDescription:%s | Location:%s"), *TaskURL, (int)(bFileExists), ErrorCode, *ErrorDescription, *TempFilePath);
+    UE_LOG(LogBackgroundHttpManager, Display, TEXT("Recieved Background Session Callback for URL:%s | TaskIdentifier:%d | bFileExists:%d | ErrorCode:%d | ErrorDescription:%s | Location:%s"), *TaskURL, TaskIdentifier, (int)(bFileExists), ErrorCode, *ErrorDescription, *TempFilePath);
 
     if (bFileExists)
     {
@@ -373,13 +385,13 @@ void FApplePlatformBackgroundHttpManager::OnTask_DidFinishDownloadingToURL(NSURL
                 FoundRequest->SetRequestAsSuccess(TempFilePath);
             }
             
-            UE_LOG(LogBackgroundHttpManager, Display, TEXT("Attempt To Mark Task Complete -- URL:%s | bDidFindTask:%d"), *TaskURL, (int)(FoundRequest.IsValid()));
+            UE_LOG(LogBackgroundHttpManager, Display, TEXT("Attempt To Mark Task Complete -- URL:%s | TaskIdentifier:%d |bDidFindTask:%d"), *TaskURL, TaskIdentifier, (int)(FoundRequest.IsValid()));
         }
     }
     else
     {
         //Forward to the OnCompleteWithError as we don't have our finished file!
-        UE_LOG(LogBackgroundHttpManager, Display, TEXT("File Not Found For DidFinishDownloadingToURL. Transitioning to DidCompleteWithError -- TaskURL:%s | ErrorCode:%d | ErrorDescription:%s | Location:%s"), *TaskURL, ErrorCode, *ErrorDescription, *TempFilePath);
+        UE_LOG(LogBackgroundHttpManager, Display, TEXT("File Not Found For DidFinishDownloadingToURL. Transitioning to DidCompleteWithError -- TaskURL:%s | TaskIdentifier:%d| ErrorCode:%d | ErrorDescription:%s | Location:%s"), *TaskURL, TaskIdentifier, ErrorCode, *ErrorDescription, *TempFilePath);
         OnTask_DidCompleteWithError(Task, Error);
     }
 }
@@ -406,7 +418,7 @@ void FApplePlatformBackgroundHttpManager::FinishRequest(FAppleBackgroundHttpRequ
 
             if (bFileExists)
             {
-                UE_LOG(LogBackgroundHttpManager,Display, TEXT("Task Completed Successfully. RequestID:%s TempFileLocation:%s"), *(Request->GetRequestID()), *TempFilePath);
+                UE_LOG(LogBackgroundHttpManager,Display, TEXT("Task Completed Successfully. RequestDebugID:%s TempFileLocation:%s"), *(Request->GetRequestDebugID()), *TempFilePath);
                 FBackgroundHttpResponsePtr NewResponse = FPlatformBackgroundHttp::ConstructBackgroundResponse(ResponseCode, *TempFilePath);
                 Request->CompleteWithExistingResponseData(NewResponse);
             }
@@ -418,7 +430,7 @@ void FApplePlatformBackgroundHttpManager::FinishRequest(FAppleBackgroundHttpRequ
                 //successfully. Handle this unexpected failure by trying to retry the task.
                 if (!bDidFail)
                 {
-                    UE_LOG(LogBackgroundHttpManager,Error, TEXT("Task finished downloading, but finished temp file was not found! -- RequestId:%s | TempFileLocation:%s"), *(Request->GetRequestID()), *TempFilePath);
+                    UE_LOG(LogBackgroundHttpManager,Error, TEXT("Task finished downloading, but finished temp file was not found! -- RequestDebugID:%s | TempFileLocation:%s"), *(Request->GetRequestDebugID()), *TempFilePath);
                     
                     //Mark our download as not completed as we hit an error so that we don't just keep trying to call FinishRequest
                     FPlatformAtomics::InterlockedExchange(&(Request->bIsCompleted), false);
@@ -432,7 +444,7 @@ void FApplePlatformBackgroundHttpManager::FinishRequest(FAppleBackgroundHttpRequ
                 //Expected case where we failed, but expected to fail
                 else
                 {
-                    UE_LOG(LogBackgroundHttpManager, Display, TEXT("Task failed completely -- RequestID:%s"), *(Request->GetRequestID()));
+                    UE_LOG(LogBackgroundHttpManager, Display, TEXT("Task failed completely -- RequestDebugID:%s"), *(Request->GetRequestDebugID()));
                     
                     FBackgroundHttpResponsePtr NewResponse = FPlatformBackgroundHttp::ConstructBackgroundResponse(ResponseCode, TEXT(""));
                     Request->CompleteWithExistingResponseData(NewResponse);
@@ -460,7 +472,7 @@ void FApplePlatformBackgroundHttpManager::FinishRequest(FAppleBackgroundHttpRequ
     }
     else
     {
-        UE_LOG(LogBackgroundHttpManager,Display, TEXT("Not finishing Request as its already sending a finish notification -- RequestID:%s"), *(Request->GetRequestID()));
+        UE_LOG(LogBackgroundHttpManager,Display, TEXT("Not finishing Request as its already sending a finish notification -- RequestDebugID:%s"), *(Request->GetRequestDebugID()));
     }
 }
 
@@ -478,12 +490,12 @@ void FApplePlatformBackgroundHttpManager::RetryRequest(FAppleBackgroundHttpReque
 			const bool bShouldUseRetryData = ShouldUseRequestRetryData(Request, RetryData);
 			if (bShouldUseRetryData)
 			{
-				UE_LOG(LogBackgroundHttpManager, Display, TEXT("Resuming Task With Resume Data -- RequestID:%s | RetryData Length:%d"), *(Request->GetRequestID()), [RetryData length]);
+				UE_LOG(LogBackgroundHttpManager, Display, TEXT("Resuming Task With Resume Data -- RequestDebugID:%s | RetryData Length:%d"), *(Request->GetRequestDebugID()), [RetryData length]);
 				NewTask = [BackgroundDownloadSession downloadTaskWithResumeData:RetryData];
 			}
 
 			//If not retry data, lets try and just retry on the next CDN
-			else
+			if (nullptr == NewTask)
 			{
                 //Since we created a new task instead of using retry data, reset resume data's retry count on the request
                 Request->ResumeDataRetryCount.Reset();
@@ -508,14 +520,14 @@ void FApplePlatformBackgroundHttpManager::RetryRequest(FAppleBackgroundHttpReque
 					Request->ActivateUnderlyingTask();
 				}
 
-				UE_LOG(LogBackgroundHttpManager, Display, TEXT("Created Task for Request -- RequestID:%s | bStartImmediately:%d | bIsAppInBG:%d"), *(Request->GetRequestID()), (int)bShouldStartImmediately, (int)bCopyOfBGState);
+				UE_LOG(LogBackgroundHttpManager, Display, TEXT("Created Task for Request -- RequestDebugID:%s | bStartImmediately:%d | bIsAppInBG:%d"), *(Request->GetRequestDebugID()), (int)bShouldStartImmediately, (int)bCopyOfBGState);
 
 				//Always set our bWasTaskStartedInBG flag on our Request so we will know if we need to restart this task next FG Tick.
 				FPlatformAtomics::InterlockedExchange(&(Request->bWasTaskStartedInBG), bCopyOfBGState);
 			}
 			else
 			{
-				UE_LOG(LogBackgroundHttpManager, Display, TEXT("Marking Request Failed. Out of Retries -- RequestID:%s | bShouldUseRetryData:%d"), *(Request->GetRequestID()), (int)bShouldUseRetryData);
+				UE_LOG(LogBackgroundHttpManager, Display, TEXT("Marking Request Failed. Out of Retries -- RequestDebugID:%s | bShouldUseRetryData:%d"), *(Request->GetRequestDebugID()), (int)bShouldUseRetryData);
 				Request->SetRequestAsFailed();
 			}
 		}
@@ -551,6 +563,7 @@ void FApplePlatformBackgroundHttpManager::OnTask_DidWriteData(NSURLSessionDownlo
     if (ensureAlwaysMsgf((nullptr != Task), TEXT("Call to DidWriteData with invalid Task!")))
     {
         FString TaskURL = [[[Task currentRequest] URL] absoluteString];
+        int TaskIdentifier = (int)[Task taskIdentifier];
         
         //Find task and update it's download progress
         {
@@ -563,11 +576,11 @@ void FApplePlatformBackgroundHttpManager::OnTask_DidWriteData(NSURLSessionDownlo
                 if (FoundRequest->DownloadProgress < TotalBytesWritten)
                 {
                     int64 DownloadProgress = FPlatformAtomics::AtomicRead(&(FoundRequest->DownloadProgress));
-                    UE_LOG(LogBackgroundHttpManager, Verbose, TEXT("Updating Task Progress! -- RequestID:%s | Current Progress:%lld | New Progress:%lld"), *(FoundRequest->GetRequestID()), DownloadProgress, TotalBytesWritten);
+                    UE_LOG(LogBackgroundHttpManager, Verbose, TEXT("Updating Task Progress! -- RequestDebugID:%s | TaskIdentifier:%d | Current Progress:%lld | New Progress:%lld"), *(FoundRequest->GetRequestDebugID()), TaskIdentifier, DownloadProgress, TotalBytesWritten);
                 }
                 else
                 {
-                    UE_LOG(LogBackgroundHttpManager, Warning, TEXT("Download Progress tried to go down not up unexpectidly! This could mean a task was unknowingly duplicated! -- RequestID:%s | Current Progress:%lld | New Progress:%lld"), *(FoundRequest->GetRequestID()),FoundRequest->DownloadProgress, TotalBytesWritten);
+                    ensureAlwaysMsgf(false, TEXT("Download Progress tried to go down not up unexpectidly! This could mean a task was unknowingly duplicated! -- RequestDebugID:%s | TaskIdentifier:%d | Current Progress:%lld | New Progress:%lld"), *(FoundRequest->GetRequestDebugID()), TaskIdentifier, FoundRequest->DownloadProgress, TotalBytesWritten);
                 }
 
                 FoundRequest->UpdateDownloadProgress(TotalBytesWritten, BytesWrittenSinceLastCall);
@@ -581,6 +594,8 @@ void FApplePlatformBackgroundHttpManager::OnTask_DidCompleteWithError(NSURLSessi
 	if (ensureAlwaysMsgf((nullptr != Task), TEXT("Call to OnTask_DidCompleteWithError delegate with an invalid task!")))
 	{
 		FString TaskURL = [[[Task currentRequest] URL] absoluteString];
+        int TaskIdentifier = (int)[Task taskIdentifier];
+        
 		const bool bDidCompleteWithError = (nullptr != Error);
 		const int ErrorCode = [Error code];
         const FString ErrorDescription = [Error localizedDescription];
@@ -593,6 +608,7 @@ void FApplePlatformBackgroundHttpManager::OnTask_DidCompleteWithError(NSURLSessi
 
 		FString DebugRetryOverrideReason;
 
+		//We still come into the function when tasks complete successfully. Only handle actual errors
 		if (bDidCompleteWithError)
 		{
             FRWScopeLock ScopeLock(URLToRequestMapLock, SLT_ReadOnly);
@@ -609,17 +625,17 @@ void FApplePlatformBackgroundHttpManager::OnTask_DidCompleteWithError(NSURLSessi
                 bShouldRetryIncreaseRetryCount = false;
                 DebugRetryOverrideReason = TEXT("Not Connected To Internet");
             }
-            
-            UE_LOG(LogBackgroundHttpManager, Display, TEXT("DidCompleteWithError for Task. -- URL:%s | bDidFindVaildRequest:%d | bDidCompleteWithError:%d | ErrorCode:%d | bHasResumeData:%d | CancelledReasonKey:%d | RetryOverrideReason:%s | bShouldRetryIncreaseRetryCount:%d | ErrorDescription:%s"), *TaskURL, (int)bDidFindValidRequest, (int)bDidCompleteWithError, ErrorCode, (int)bHasResumeData, CancelledReasonInt, *DebugRetryOverrideReason, (int)bShouldRetryIncreaseRetryCount, *ErrorDescription);
+
+            UE_LOG(LogBackgroundHttpManager, Display, TEXT("DidCompleteWithError for Task. -- URL:%s | TaskIdentifier:%d | bDidFindVaildRequest:%d | bDidCompleteWithError:%d | ErrorCode:%d | bHasResumeData:%d | CancelledReasonKey:%d | RetryOverrideReason:%s | bShouldRetryIncreaseRetryCount:%d | ErrorDescription:%s"), *TaskURL, TaskIdentifier, (int)bDidFindValidRequest, (int)bDidCompleteWithError, ErrorCode, (int)bHasResumeData, CancelledReasonInt, *DebugRetryOverrideReason, (int)bShouldRetryIncreaseRetryCount, *ErrorDescription);
             
 			if (bDidFindValidRequest)
 			{
-                RetryRequest(FoundRequest, bShouldRetryIncreaseRetryCount, true, ResumeData);
+				RetryRequest(FoundRequest, bShouldRetryIncreaseRetryCount, true, ResumeData);
 			}
 			else
 			{
-				//This can legitimately happen in the case of UnAssociatedTasks completing, so don't error
-				UE_LOG(LogBackgroundHttpManager, Display, TEXT("No Request Found for Errored Task -- TaskURL:%s"), *TaskURL);
+                //This can be a valid case because of UnAssociatedTasks, so don't error here
+                UE_LOG(LogBackgroundHttpManager, Display, TEXT("No request for completing task! -- TaskURL:%s | TaskIdentifier:%d"), *TaskURL, TaskIdentifier);
 			}
 		}
 	}
@@ -661,16 +677,16 @@ void FApplePlatformBackgroundHttpManager::TickRequests(float DeltaTime)
                 const bool bWasStartedInBG = FPlatformAtomics::AtomicRead(&(AppleRequest->bWasTaskStartedInBG));
                 const bool bIsPendingCancel = FPlatformAtomics::AtomicRead(&(AppleRequest->bIsPendingCancel));
                 
-                UE_LOG(LogBackgroundHttpManager, VeryVerbose, TEXT("Checking Status of Request on Tick -- RequestID::%s | bIsTaskComplete:%d | bWasStartedInBG:%d"), *(AppleRequest->GetRequestID()), (int)bIsTaskComplete, (int)bWasStartedInBG);
+                UE_LOG(LogBackgroundHttpManager, VeryVerbose, TEXT("Checking Status of Request on Tick -- RequestDebugID::%s | bIsTaskComplete:%d | bWasStartedInBG:%d"), *(AppleRequest->GetRequestDebugID()), (int)bIsTaskComplete, (int)bWasStartedInBG);
                 
                 if (bIsTaskComplete)
                 {
-                    UE_LOG(LogBackgroundHttpManager, Display, TEXT("Calling FinishRequest On -- RequestID::%s | bIsTaskComplete:%d | bWasStartedInBG:%d"), *(AppleRequest->GetRequestID()), (int)bIsTaskComplete, (int)bWasStartedInBG);
+                    UE_LOG(LogBackgroundHttpManager, Display, TEXT("Calling FinishRequest On -- RequestDebugID::%s | bIsTaskComplete:%d | bWasStartedInBG:%d"), *(AppleRequest->GetRequestDebugID()), (int)bIsTaskComplete, (int)bWasStartedInBG);
                     FinishRequest(AppleRequest);
                 }
                 else if (bWasStartedInBG && !bIsPendingCancel)
                 {
-                    UE_LOG(LogBackgroundHttpManager, Display, TEXT("Cancelling Request Created In BG To Re-Create In FG -- RequestID:%s"), *(AppleRequest->GetRequestID()));
+                    UE_LOG(LogBackgroundHttpManager, Display, TEXT("Cancelling Request Created In BG To Re-Create In FG -- RequestDebugID:%s"), *(AppleRequest->GetRequestDebugID()));
                     
                     //reset to false so we don't run this twice while waiting on recreation
                     FPlatformAtomics::InterlockedExchange(&(AppleRequest->bWasTaskStartedInBG), false);
@@ -684,7 +700,7 @@ void FApplePlatformBackgroundHttpManager::TickRequests(float DeltaTime)
                     const bool bShouldTimeOut = AppleRequest->TickTimeOutTimer(DeltaTime);
                     if (bShouldTimeOut)
                     {
-                        UE_LOG(LogBackgroundHttpManager, Display, TEXT("Timing out Request Due To Lack of Server Response -- RequestID:%s"), *(AppleRequest->GetRequestID()));
+                        UE_LOG(LogBackgroundHttpManager, Display, TEXT("Timing out Request Due To Lack of Server Response -- RequestDebugID:%s"), *(AppleRequest->GetRequestDebugID()));
                         
                         //Just cancel the task and let the OnTask_DidCompleteWithError callback handle retrying it if appropriate.
                         AppleRequest->CancelActiveTask();
@@ -736,6 +752,7 @@ void FApplePlatformBackgroundHttpManager::TickTasks(float DeltaTime)
                              if (NewRequestCount <= FPlatformBackgroundHttp::GetPlatformMaxActiveDownloads())
                              {
                                  FString TaskURL = [[[task currentRequest] URL] absoluteString];
+                                 int TaskIdentifier = (int)[task taskIdentifier];
                                  
                                  //Try and find Request in map that matches this Task
                                  {
@@ -747,12 +764,12 @@ void FApplePlatformBackgroundHttpManager::TickTasks(float DeltaTime)
                                      
                                      if (FoundRequest.IsValid())
                                      {
-                                         UE_LOG(LogBackgroundHttpManager, Display, TEXT("Manager Calling to Active Task For Request -- RequestID:%s | TaskURL:%s | CurrentlyActiveRequests:%d"), *(FoundRequest->GetRequestID()), *TaskURL, NewRequestCount);
+                                         UE_LOG(LogBackgroundHttpManager, Display, TEXT("Manager Calling to Active Task For Request -- RequestDebugID:%s | TaskURL:%s | TaskIdentifier:%d | CurrentlyActiveRequests:%d"), *(FoundRequest->GetRequestDebugID()), *TaskURL, TaskIdentifier, NewRequestCount);
                                          FoundRequest->ActivateUnderlyingTask();
                                      }
                                      else
                                      {
-                                         UE_LOG(LogBackgroundHttpManager, Display, TEXT("Skipping Activating Task as there is no associated Request or Request is paused. Once a Request associates with this task, it can then be activated. -- TaskURL:%s | bIsPaused:%d"), *TaskURL, (int)bIsPaused);
+                                         UE_LOG(LogBackgroundHttpManager, Display, TEXT("Skipping Activating Task as there is no associated Request or Request is paused. Once a Request associates with this task, it can then be activated. -- TaskURL:%s | TaskIdentifier:%d| bIsPaused:%d"), *TaskURL, TaskIdentifier, (int)bIsPaused);
                                          
                                          //Don't activate and remove our increment from above because something put us over the limit before we resumed
                                          FPlatformAtomics::InterlockedDecrement(&NumCurrentlyActiveTasks);
@@ -762,7 +779,9 @@ void FApplePlatformBackgroundHttpManager::TickTasks(float DeltaTime)
                              else
                              {                                 
                                  FString TaskURL = [[[task currentRequest] URL] absoluteString];
-                                 UE_LOG(LogBackgroundHttpManager, Log, TEXT("Task failed to activate as we passed the platform max from another task before we could resume. Task -- TaskURL:%s | CurrentlyActiveRequests:%d"), *TaskURL, NewRequestCount);
+                                 int TaskIdentifier = (int)[task taskIdentifier];
+                                 
+                                 UE_LOG(LogBackgroundHttpManager, Log, TEXT("Task failed to activate as we passed the platform max from another task before we could resume. Task -- TaskURL:%s | TaskIdentifier:%d | CurrentlyActiveRequests:%d"), *TaskURL, TaskIdentifier, NewRequestCount);
                                  
                                  //Don't activate and remove our increment from above because something put us over the limit before we resumed
                                  FPlatformAtomics::InterlockedDecrement(&NumCurrentlyActiveTasks);
diff --git a/Engine/Source/Runtime/Online/BackgroundHTTP/Private/IOS/ApplePlatformBackgroundHttpRequest.cpp b/Engine/Source/Runtime/Online/BackgroundHTTP/Private/IOS/ApplePlatformBackgroundHttpRequest.cpp
--- a/Engine/Source/Runtime/Online/BackgroundHTTP/Private/IOS/ApplePlatformBackgroundHttpRequest.cpp
+++ b/Engine/Source/Runtime/Online/BackgroundHTTP/Private/IOS/ApplePlatformBackgroundHttpRequest.cpp
@@ -11,10 +11,12 @@ FApplePlatformBackgroundHttpRequest::FApplePlatformBackgroundHttpRequest()
     , RetryCount(0)
 	, ResumeDataRetryCount(0)
     , FirstTask(nullptr)
+	, FirstTaskIdentifier(0)
     , bIsTaskActive(false)
     , bIsTaskPaused(false)
     , bIsCompleted(false)
     , bIsFailed(false)
+	, bIsRequestSwitchingTasks(false)
     , bWasTaskStartedInBG(false)
     , bHasAlreadyFinishedRequest(false)
     , bIsPendingCancel(false)
@@ -44,7 +46,7 @@ void FApplePlatformBackgroundHttpRequest::SetRequestAsFailed()
 
 void FApplePlatformBackgroundHttpRequest::CompleteRequest_Internal(bool bWasRequestSuccess, const FString& CompletedTempDownloadLocationIn)
 {
-    UE_LOG(LogBackgroundHttpRequest, Display, TEXT("Marking Request Complete -- RequestID:%s | bWasRequestSuccess:%d | CompletedTempDownloadLocation:%s"), *GetRequestID(), (int)bWasRequestSuccess, *CompletedTempDownloadLocationIn);
+    UE_LOG(LogBackgroundHttpRequest, Display, TEXT("Marking Request Complete -- RequestDebugID:%s | bWasRequestSuccess:%d | CompletedTempDownloadLocation:%s"), *GetRequestDebugID(), (int)bWasRequestSuccess, *CompletedTempDownloadLocationIn);
     
     FPlatformAtomics::InterlockedExchange(&bIsTaskActive, false);
 	FPlatformAtomics::InterlockedExchange(&bIsCompleted, true);
@@ -66,7 +68,7 @@ const FString& FApplePlatformBackgroundHttpRequest::GetURLForRetry(bool bShouldI
 	//If we are out of Retries, just send an empty string
 	if (NewRetryCount > NumberOfTotalRetries)
 	{
-		UE_LOG(LogBackgroundHttpRequest, Display, TEXT("GetURLForRetry is out of Retries for Request -- RequestID:%s"), *GetRequestID());
+		UE_LOG(LogBackgroundHttpRequest, Display, TEXT("GetURLForRetry is out of Retries for Request -- RequestDebugID:%s"), *GetRequestDebugID());
         
         static FString EmptyResponse = TEXT("");
         return EmptyResponse;
@@ -77,7 +79,7 @@ const FString& FApplePlatformBackgroundHttpRequest::GetURLForRetry(bool bShouldI
 		const int URLIndex = NewRetryCount % URLList.Num();
 		const FString& URLToReturn = URLList[URLIndex];
 
-		UE_LOG(LogBackgroundHttpRequest, Display, TEXT("GetURLForRetry found valid URL for current retry -- RequestID:%s | NewRetryCount:%d | URLToReturn:%s"), *GetRequestID(), NewRetryCount, *URLToReturn);
+		UE_LOG(LogBackgroundHttpRequest, Display, TEXT("GetURLForRetry found valid URL for current retry -- RequestDebugID:%s | NewRetryCount:%d | URLToReturn:%s"), *GetRequestDebugID(), NewRetryCount, *URLToReturn);
         return URLToReturn;
     }
 }
@@ -96,7 +98,9 @@ void FApplePlatformBackgroundHttpRequest::ActivateUnderlyingTask()
         if (ensureAlwaysMsgf((nullptr != UnderlyingTask), TEXT("Call to ActivateUnderlyingTask with an invalid task! Need to create underlying task before activating!")))
         {
             FString TaskURL = [[[UnderlyingTask currentRequest] URL] absoluteString];
-            UE_LOG(LogBackgroundHttpRequest, Display, TEXT("Activating Task for Request -- RequestID:%s | TaskURL:%s"), *GetRequestID(), *TaskURL);
+            int TaskIdentifier = (int)[UnderlyingTask taskIdentifier];
+            
+            UE_LOG(LogBackgroundHttpRequest, Display, TEXT("Activating Task for Request -- RequestDebugID:%s | TaskIdentifier:%d | TaskURL:%s"), *GetRequestDebugID(), TaskIdentifier, *TaskURL);
         
             FPlatformAtomics::InterlockedExchange(&bIsTaskActive, true);
             FPlatformAtomics::InterlockedExchange(&bIsTaskPaused, false);
@@ -119,7 +123,9 @@ void FApplePlatformBackgroundHttpRequest::PauseUnderlyingTask()
         if (ensureAlwaysMsgf((nullptr != UnderlyingTask), TEXT("Call to PauseUnderlyingTask with an invalid task! Need to create underlying task before trying to pause!")))
         {
             FString TaskURL = [[[UnderlyingTask currentRequest] URL] absoluteString];
-            UE_LOG(LogBackgroundHttpRequest, Display, TEXT("Pausing Task for Request -- RequestID:%s | TaskURL:%s"), *GetRequestID(), *TaskURL);
+            int TaskIdentifier = (int)[UnderlyingTask taskIdentifier];
+            
+            UE_LOG(LogBackgroundHttpRequest, Display, TEXT("Pausing Task for Request -- RequestDebugID:%s | TaskIdentifier:%d | TaskURL:%s"), *GetRequestDebugID(), TaskIdentifier, *TaskURL);
             
             FPlatformAtomics::InterlockedExchange(&bIsTaskActive, false);
             FPlatformAtomics::InterlockedExchange(&bIsPendingCancel, false);
@@ -154,27 +160,48 @@ void FApplePlatformBackgroundHttpRequest::ResetTimeOutTimer()
     ActiveTimeOutTimer = FApplePlatformBackgroundHttpManager::ActiveTimeOutSetting;
 }
 
-void FApplePlatformBackgroundHttpRequest::AssociateWithTask(NSURLSessionTask* ExistingTask)
+bool FApplePlatformBackgroundHttpRequest::AssociateWithTask(NSURLSessionTask* ExistingTask)
 {
-	if (ensureAlwaysMsgf((nullptr != ExistingTask), TEXT("Call to AssociateWithTask with an invalid Task! RequestID:%s"), *GetRequestID()))
+	bool bDidAssociate = false;
+
+	if (ensureAlwaysMsgf((nullptr != ExistingTask), TEXT("Call to AssociateWithTask with an invalid Task! RequestDebugID:%s"), *GetRequestDebugID()))
 	{
-		volatile FTaskNode* NewNode = new FTaskNode();
-		NewNode->OurTask = ExistingTask;
-		
-		//Add a count to our task's reference list so it doesn't get deleted while in our Request's task list
-		[ExistingTask retain];
+		const FString TaskURL = [[[ExistingTask currentRequest] URL] absoluteString];
+        int TaskIdentifier = (int)[ExistingTask taskIdentifier];
+        
+		const bool bWasAlreadySwitching = FPlatformAtomics::InterlockedExchange(&bIsRequestSwitchingTasks, true);
+		if (!bWasAlreadySwitching)
+		{
+			volatile FTaskNode* NewNode = new FTaskNode();
+			NewNode->OurTask = ExistingTask;
 
-        //Swap our new node and the first one in the list
-		NewNode->NextNode = (FTaskNode*)FPlatformAtomics::InterlockedExchangePtr((void**)(&FirstTask), (void*)NewNode);
+			//Add a count to our task's reference list so it doesn't get deleted while in our Request's task list
+			[ExistingTask retain];
 
-		FString TaskURL = [[[ExistingTask currentRequest] URL] absoluteString];
-		UE_LOG(LogBackgroundHttpRequest, Display, TEXT("Associated Request With New Task -- RequestID:%s | TaskURL:%s"), *GetRequestID(), *TaskURL);
-        
-        FPlatformAtomics::InterlockedExchange(&bIsPendingCancel, false);
-        
-        ResetTimeOutTimer();
-        ResetProgressTracking();
+			//Swap our new node and the first one in the list
+			NewNode->NextNode = (FTaskNode*)FPlatformAtomics::InterlockedExchangePtr((void**)(&FirstTask), (void*)NewNode);
+
+			//Save off our first task's identifier and our CombinedRequestID that includes it
+			FirstTaskIdentifier = (int)[(FirstTask->OurTask) taskIdentifier];
+            CombinedRequestID = FString::Printf(TEXT("%d.%s"), FirstTaskIdentifier, *RequestID);
+            
+            UE_LOG(LogBackgroundHttpRequest, Display, TEXT("Associated Request With New Task -- RequestDebugID:%s | TaskIdentifier:%d | TaskURL:%s"), *GetRequestDebugID(), TaskIdentifier, *TaskURL);
+
+			FPlatformAtomics::InterlockedExchange(&bIsPendingCancel, false);
+
+			ResetTimeOutTimer();
+			ResetProgressTracking();
+
+			bDidAssociate = true;
+            FPlatformAtomics::InterlockedExchange(&bIsRequestSwitchingTasks, false);
+		}
+		else
+		{
+            UE_LOG(LogBackgroundHttpRequest, Display, TEXT("Failed to Associate Request with new Task as there was already a pending AssociateWithTask running! -- RequestDebugID:%s | TaskIdentifier:%d | TaskURL:%s"), *GetRequestDebugID(), TaskIdentifier, *TaskURL);
+		}
 	}
+
+	return bDidAssociate;
 }
 
 void FApplePlatformBackgroundHttpRequest::PauseRequest()
@@ -195,7 +222,9 @@ void FApplePlatformBackgroundHttpRequest::CancelActiveTask()
         if (nullptr != TaskNodeWeAreCancelling->OurTask)
         {
             FString TaskURL = [[[TaskNodeWeAreCancelling->OurTask currentRequest] URL] absoluteString];
-            UE_LOG(LogBackgroundHttpRequest, Display, TEXT("Cancelling Task -- RequestID:%s | TaskURL:%s"), *GetRequestID(), *TaskURL);
+            int TaskIdentifier = (int)[TaskNodeWeAreCancelling->OurTask taskIdentifier];
+            
+            UE_LOG(LogBackgroundHttpRequest, Display, TEXT("Cancelling Task -- RequestDebugID:%s | TaskIdentifier:%d | TaskURL:%s"), *GetRequestDebugID(), TaskIdentifier, *TaskURL);
             
             FPlatformAtomics::InterlockedExchange(&bIsPendingCancel, true);
             
@@ -206,7 +235,7 @@ void FApplePlatformBackgroundHttpRequest::CancelActiveTask()
 
 void FApplePlatformBackgroundHttpRequest::UpdateDownloadProgress(int64_t TotalDownloaded,int64_t DownloadedSinceLastUpdate)
 {
-    UE_LOG(LogBackgroundHttpRequest, VeryVerbose, TEXT("Request Update Progress -- RequestID:%s | OldProgress:%lld | NewProgress:%lld | ProgressSinceLastUpdate:%lld"), *GetRequestID(), DownloadProgress, TotalDownloaded, DownloadedSinceLastUpdate);
+    UE_LOG(LogBackgroundHttpRequest, VeryVerbose, TEXT("Request Update Progress -- RequestDebugID:%s | OldProgress:%lld | NewProgress:%lld | ProgressSinceLastUpdate:%lld"), *GetRequestDebugID(), DownloadProgress, TotalDownloaded, DownloadedSinceLastUpdate);
 	
     FPlatformAtomics::AtomicStore(&DownloadProgress, TotalDownloaded);
     FPlatformAtomics::InterlockedAdd(&DownloadProgressSinceLastUpdateSent, DownloadedSinceLastUpdate);
@@ -228,6 +257,12 @@ void FApplePlatformBackgroundHttpRequest::SendDownloadProgressUpdate()
     }
 }
 
+const FString& FApplePlatformBackgroundHttpRequest::GetRequestDebugID() const
+{
+    //We use CombinedRequestID to append a TaskIdentifier on the end of the RequestID. If we have one, use that. Otherwise fallback on what is already set
+    return CombinedRequestID.IsEmpty() ? RequestID : CombinedRequestID;
+}
+
 bool FApplePlatformBackgroundHttpRequest::IsTaskComplete() const
 {
     const bool bDidRequestFail = FPlatformAtomics::AtomicRead(&bIsFailed);
diff --git a/Engine/Source/Runtime/Online/BackgroundHTTP/Public/IOS/ApplePlatformBackgroundHttpRequest.h b/Engine/Source/Runtime/Online/BackgroundHTTP/Public/IOS/ApplePlatformBackgroundHttpRequest.h
--- a/Engine/Source/Runtime/Online/BackgroundHTTP/Public/IOS/ApplePlatformBackgroundHttpRequest.h
+++ b/Engine/Source/Runtime/Online/BackgroundHTTP/Public/IOS/ApplePlatformBackgroundHttpRequest.h
@@ -23,7 +23,11 @@ public:
     
     virtual bool IsTaskComplete() const;
 
-    
+    //Used to provide some extra debug information over normal GetRequestID()
+    //Returns string in format of X.Y where X is the underlying Task Identifier if set and Y is what was set in the SetRequestID() call if this reqeust has associated with a task.
+    //Returns the same as GetRequestID() if no task has been associated yet.
+	const FString& GetRequestDebugID() const;
+
 private:
     //Super simple linked list with volatile pointers to next element
     struct FTaskNode
@@ -34,7 +38,7 @@ private:
     
 private:
     const FString& GetURLForRetry(bool bShouldIncrementRetryCountFirst);
-    void AssociateWithTask(NSURLSessionTask* ExistingTask);
+    bool AssociateWithTask(NSURLSessionTask* ExistingTask);
 	void SetRequestAsSuccess(const FString& CompletedTempDownloadLocation);
     void SetRequestAsFailed();
     void CompleteRequest_Internal(bool bWasRequestSuccess, const FString& CompletedTempDownloadLocation);
@@ -60,18 +64,22 @@ private:
 	FThreadSafeCounter ResumeDataRetryCount;
 
 	volatile FTaskNode* FirstTask;
+	
+    volatile int FirstTaskIdentifier;
+    FString CombinedRequestID;
 
     volatile int32 bIsTaskActive;
     volatile int32 bIsTaskPaused;
 	volatile int32 bIsCompleted;
 	volatile int32 bIsFailed;
+	volatile int32 bIsRequestSwitchingTasks;
 	volatile int32 bWasTaskStartedInBG;
     volatile int32 bHasAlreadyFinishedRequest;
     volatile int32 bIsPendingCancel;
     
 	volatile int64 DownloadProgress;
     volatile int64 DownloadProgressSinceLastUpdateSent;
-    
+
     friend class FApplePlatformBackgroundHttpManager;
 };
 
