Linux: small voice settings refactor

diff --git a/Engine/Source/Runtime/Engine/Private/VoiceConfig.cpp b/Engine/Source/Runtime/Engine/Private/VoiceConfig.cpp
--- a/Engine/Source/Runtime/Engine/Private/VoiceConfig.cpp
+++ b/Engine/Source/Runtime/Engine/Private/VoiceConfig.cpp
@@ -23,7 +23,9 @@ static TAutoConsoleVariable<float> CVarVoiceSilenceDetectionThreshold(TEXT("voic
 
 int32 UVOIPStatics::GetVoiceSampleRate()
 {
-#if USE_DEFAULT_VOICE_SAMPLE_RATE
+#if PLATFORM_UNIX
+	return PLATFORM_LINUX ? 16000 : 48000;
+#elif USE_DEFAULT_VOICE_SAMPLE_RATE
 	return (int32) 16000;
 #else
 	static bool bRetrievedSampleRate = false;
diff --git a/Engine/Source/Runtime/Online/Voice/Private/VoicePrivate.h b/Engine/Source/Runtime/Online/Voice/Private/VoicePrivate.h
--- a/Engine/Source/Runtime/Online/Voice/Private/VoicePrivate.h
+++ b/Engine/Source/Runtime/Online/Voice/Private/VoicePrivate.h
@@ -23,7 +23,7 @@ THIRD_PARTY_INCLUDES_END
 #endif // PLATFORM_WINDOWS
 
 #define ANDROIDVOICE_SUPPORTED_PLATFORMS (PLATFORM_ANDROID && (PLATFORM_ANDROID_ARM || PLATFORM_ANDROID_ARM64 || PLATFORM_ANDROID_X64) && !PLATFORM_LUMIN)
-#define PLATFORM_SUPPORTS_VOICE_CAPTURE (PLATFORM_WINDOWS || PLATFORM_MAC || ANDROIDVOICE_SUPPORTED_PLATFORMS || (PLATFORM_LINUX && VOICE_MODULE_WITH_CAPTURE))
+#define PLATFORM_SUPPORTS_VOICE_CAPTURE (PLATFORM_WINDOWS || PLATFORM_MAC || ANDROIDVOICE_SUPPORTED_PLATFORMS || (PLATFORM_UNIX && VOICE_MODULE_WITH_CAPTURE))
 #define PLATFORM_SUPPORTS_OPUS_CODEC (!PLATFORM_HTML5 && !PLATFORM_TVOS)
 
 // Module includes
