Reuse XMPP connection between MCP:DefaultInstance and MCP:EMBEDDED OSS instances for a logged in user

New config for identity MCP implementation to reuse connection:
[OnlineSubsystemMcp.OnlineIdentityMcp]
bShouldReuseXmpp=true

When bShouldReuseXmpp is enabled callilng FOnlineIdentityMcp::LoginXmpp will attempt to acquire an existing XMPP connection and notify other OSS instances that it has done so
When calling FOnlineIdentityMcp::LogoutXmpp the XMPP connection is relinquished without actually loging out of it

Calling login on an existing XMPP connection will no longer error and will succeed.
This will trigger reconnect handling on XMPP and still trigger OnLoginComplete/OnLoginChanged on the XMPP connection
Presence info that is cached by existing connection is rebroadcast when reconnecting XMPP.  This allows OSS presence implementation to get matching presence when reusing connection.

Added exec command "AUTH EXCHANGE <OSS Name>" for testing login via exchange code between OSS instances
Added exec command "AUTH LOGOUT" for testing account logout (similar to "AUTH KILL" but logs out immediately
Added exec command "XMPP DUMP" for listing current connections and their state

#tests PC -kairosonly flow, and with exec commands for logging in via embedded OSS exchange code and via XMPP connection logout

diff --git a/Engine/Source/Runtime/Online/XMPP/Private/XmppJingle/XmppConnectionJingle.cpp b/Engine/Source/Runtime/Online/XMPP/Private/XmppJingle/XmppConnectionJingle.cpp
--- a/Engine/Source/Runtime/Online/XMPP/Private/XmppJingle/XmppConnectionJingle.cpp
+++ b/Engine/Source/Runtime/Online/XMPP/Private/XmppJingle/XmppConnectionJingle.cpp
@@ -719,9 +719,14 @@ void FXmppConnectionJingle::Login(const FString& UserId, const FString& Password
 		{
 			ErrorStr = TEXT("Still processing last logout");
 		}
-		else if (LoginState == EXmppLoginStatus::LoggedIn)
+		else if (LoginStatus == EXmppLoginStatus::LoggedIn)
 		{
-			ErrorStr = TEXT("Already logged in");
+			UE_LOG(LogXmpp, Log, TEXT("Reusing existing connection for Jid=%s"), *GetUserJid().ToDebugString());
+
+			OnLoginComplete().Broadcast(GetUserJid(), true, FString());
+			OnLoginChanged().Broadcast(GetUserJid(), EXmppLoginStatus::LoggedIn);
+			ReconnectLogin();
+			return;
 		}
 		else
 		{
@@ -797,4 +802,10 @@ const FXmppUserJid& FXmppConnectionJingle::GetUserJid() const
 	return UserJid;
 }
 
+void FXmppConnectionJingle::DumpState() const
+{
+	UE_LOG(LogXmpp, Warning, TEXT("JID=%s Status=%s"), *GetUserJid().ToDebugString(), EXmppLoginStatus::ToString(GetLoginStatus()));
+}
+
+
 #endif //WITH_XMPP_JINGLE
diff --git a/Engine/Source/Runtime/Online/XMPP/Private/XmppJingle/XmppConnectionJingle.h b/Engine/Source/Runtime/Online/XMPP/Private/XmppJingle/XmppConnectionJingle.h
--- a/Engine/Source/Runtime/Online/XMPP/Private/XmppJingle/XmppConnectionJingle.h
+++ b/Engine/Source/Runtime/Online/XMPP/Private/XmppJingle/XmppConnectionJingle.h
@@ -44,6 +44,8 @@ public:
 	virtual IXmppMultiUserChatPtr MultiUserChat() override;
 	virtual IXmppChatPtr PrivateChat() override;
 
+	virtual void DumpState() const override;
+
 	// FTickerObjectBase
 
 	virtual bool Tick(float DeltaTime) override;
diff --git a/Engine/Source/Runtime/Online/XMPP/Private/XmppModule.cpp b/Engine/Source/Runtime/Online/XMPP/Private/XmppModule.cpp
--- a/Engine/Source/Runtime/Online/XMPP/Private/XmppModule.cpp
+++ b/Engine/Source/Runtime/Online/XMPP/Private/XmppModule.cpp
@@ -650,6 +650,16 @@ bool FXmppModule::HandleXmppCommand( const TCHAR* Cmd, FOutputDevice& Ar )
 		}
 		return true;
 	}
+	else if (FParse::Command(&Cmd, TEXT("Dump")))
+	{
+		UE_LOG(LogXmpp, Warning, TEXT("Dumping %d XMPP connections: "), ActiveConnections.Num());
+		for (const auto& Pair : ActiveConnections)
+		{
+			const TSharedRef<IXmppConnection>& Connection = Pair.Value;
+			Connection->DumpState();
+		}
+		return true;
+	}
 
 	return false;
 }
diff --git a/Engine/Source/Runtime/Online/XMPP/Private/XmppNull.cpp b/Engine/Source/Runtime/Online/XMPP/Private/XmppNull.cpp
--- a/Engine/Source/Runtime/Online/XMPP/Private/XmppNull.cpp
+++ b/Engine/Source/Runtime/Online/XMPP/Private/XmppNull.cpp
@@ -44,6 +44,8 @@ public:
 	virtual IXmppMultiUserChatPtr MultiUserChat() override { return NULL; }
 	virtual IXmppChatPtr PrivateChat() override { return NULL; }
 
+	virtual void DumpState() const override {}
+
 	// FXmppConnectionNull
 
 	FXmppConnectionNull() {};
diff --git a/Engine/Source/Runtime/Online/XMPP/Private/XmppStrophe/XmppConnectionStrophe.cpp b/Engine/Source/Runtime/Online/XMPP/Private/XmppStrophe/XmppConnectionStrophe.cpp
--- a/Engine/Source/Runtime/Online/XMPP/Private/XmppStrophe/XmppConnectionStrophe.cpp
+++ b/Engine/Source/Runtime/Online/XMPP/Private/XmppStrophe/XmppConnectionStrophe.cpp
@@ -63,7 +63,12 @@ void FXmppConnectionStrophe::Login(const FString& UserId, const FString& Auth)
 		}
 		else if (LoginStatus == EXmppLoginStatus::LoggedIn)
 		{
-			ErrorStr = TEXT("Already logged in");
+			UE_LOG(LogXmpp, Log, TEXT("Reusing existing connection for Jid=%s"), *GetUserJid().ToDebugString());
+
+			OnLoginComplete().Broadcast(GetUserJid(), true, FString());
+			OnLoginChanged().Broadcast(GetUserJid(), EXmppLoginStatus::LoggedIn);
+			ReconnectLogin();
+			return;
 		}
 		else
 		{
@@ -95,6 +100,16 @@ void FXmppConnectionStrophe::Login(const FString& UserId, const FString& Auth)
 	}
 }
 
+void FXmppConnectionStrophe::ReconnectLogin()
+{
+	MessagesStrophe->OnReconnect();
+	MultiUserChatStrophe->OnReconnect();
+	PingStrophe->OnReconnect();
+	PresenceStrophe->OnReconnect();
+	PrivateChatStrophe->OnReconnect();
+	PubSubStrophe->OnReconnect();
+}
+
 void FXmppConnectionStrophe::Logout()
 {
 	if (StropheThread.IsValid())
@@ -370,4 +385,9 @@ void FXmppConnectionStrophe::ProcessLoginStatusChange(EXmppLoginStatus::Type New
 	}
 }
 
+void FXmppConnectionStrophe::DumpState() const
+{
+	UE_LOG(LogXmpp, Warning, TEXT("JID=%s Status=%s"), *GetUserJid().ToDebugString(), EXmppLoginStatus::ToString(GetLoginStatus()));
+}
+
 #endif
diff --git a/Engine/Source/Runtime/Online/XMPP/Private/XmppStrophe/XmppConnectionStrophe.h b/Engine/Source/Runtime/Online/XMPP/Private/XmppStrophe/XmppConnectionStrophe.h
--- a/Engine/Source/Runtime/Online/XMPP/Private/XmppStrophe/XmppConnectionStrophe.h
+++ b/Engine/Source/Runtime/Online/XMPP/Private/XmppStrophe/XmppConnectionStrophe.h
@@ -60,6 +60,8 @@ public:
 	virtual IXmppChatPtr PrivateChat() override;
 	virtual IXmppPubSubPtr PubSub() override;
 
+	virtual void DumpState() const override;
+
 	// FTickerObjectBase
 	virtual bool Tick(float DeltaTime) override;
 
@@ -82,6 +84,10 @@ public:
 	/** Get the MUC domain that was saved right before we connected */
 	const FString& GetMucDomain() const { return MucDomain; }
 
+private:
+	/** Notify that Login was triggered again while already connected */
+	void ReconnectLogin();
+
 protected:
 	/** XMPP Context to use across this connection */
 	FStropheContext StropheContext;
diff --git a/Engine/Source/Runtime/Online/XMPP/Private/XmppStrophe/XmppMessagesStrophe.cpp b/Engine/Source/Runtime/Online/XMPP/Private/XmppStrophe/XmppMessagesStrophe.cpp
--- a/Engine/Source/Runtime/Online/XMPP/Private/XmppStrophe/XmppMessagesStrophe.cpp
+++ b/Engine/Source/Runtime/Online/XMPP/Private/XmppStrophe/XmppMessagesStrophe.cpp
@@ -35,6 +35,11 @@ void FXmppMessagesStrophe::OnDisconnect()
 	CleanupMessages();
 }
 
+void FXmppMessagesStrophe::OnReconnect()
+{
+
+}
+
 bool FXmppMessagesStrophe::ReceiveStanza(const FStropheStanza& IncomingStanza)
 {
 	if (IncomingStanza.GetName() != Strophe::SN_MESSAGE || // Must be a message
diff --git a/Engine/Source/Runtime/Online/XMPP/Private/XmppStrophe/XmppMessagesStrophe.h b/Engine/Source/Runtime/Online/XMPP/Private/XmppStrophe/XmppMessagesStrophe.h
--- a/Engine/Source/Runtime/Online/XMPP/Private/XmppStrophe/XmppMessagesStrophe.h
+++ b/Engine/Source/Runtime/Online/XMPP/Private/XmppStrophe/XmppMessagesStrophe.h
@@ -23,13 +23,13 @@ public:
 	~FXmppMessagesStrophe();
 
 	// XMPP Thread
-	void OnDisconnect();
 	bool ReceiveStanza(const FStropheStanza& IncomingStanza);
-
 	bool HandleMessageStanza(const FStropheStanza& IncomingStanza);
 	bool HandleMessageErrorStanza(const FStropheStanza& ErrorStanza);
 
 	// Game Thread
+	void OnDisconnect();
+	void OnReconnect();
 
 	// IXmppMessages
 	virtual bool SendMessage(const FXmppUserJid& RecipientId, const FString& Type, const FString& Payload, bool bPayloadIsSerializedJson = false) override;
diff --git a/Engine/Source/Runtime/Online/XMPP/Private/XmppStrophe/XmppMultiUserChatStrophe.cpp b/Engine/Source/Runtime/Online/XMPP/Private/XmppStrophe/XmppMultiUserChatStrophe.cpp
--- a/Engine/Source/Runtime/Online/XMPP/Private/XmppStrophe/XmppMultiUserChatStrophe.cpp
+++ b/Engine/Source/Runtime/Online/XMPP/Private/XmppStrophe/XmppMultiUserChatStrophe.cpp
@@ -34,6 +34,11 @@ void FXmppMultiUserChatStrophe::OnDisconnect()
 	PendingRoomCreateConfigs.Empty();
 }
 
+void FXmppMultiUserChatStrophe::OnReconnect()
+{
+
+}
+
 bool FXmppMultiUserChatStrophe::ReceiveStanza(const FStropheStanza& IncomingStanza)
 {
 	// MUC presence are from our MUC domain
diff --git a/Engine/Source/Runtime/Online/XMPP/Private/XmppStrophe/XmppMultiUserChatStrophe.h b/Engine/Source/Runtime/Online/XMPP/Private/XmppStrophe/XmppMultiUserChatStrophe.h
--- a/Engine/Source/Runtime/Online/XMPP/Private/XmppStrophe/XmppMultiUserChatStrophe.h
+++ b/Engine/Source/Runtime/Online/XMPP/Private/XmppStrophe/XmppMultiUserChatStrophe.h
@@ -147,18 +147,19 @@ public:
 	FXmppMultiUserChatStrophe(class FXmppConnectionStrophe& InConnectionManager);
 	virtual ~FXmppMultiUserChatStrophe();
 
-	void OnDisconnect();
+	// XMPP Thread
 	bool ReceiveStanza(const FStropheStanza& IncomingStanza);
-
 	bool HandlePresenceStanza(const FStropheStanza& IncomingStanza);
 	bool HandlePresenceErrorStanza(const FStropheStanza& IncomingStanza);
-
 	bool HandleGroupChatStanza(const FStropheStanza& IncomingStanza);
 	bool HandleGroupChatErrorStanza(const FStropheStanza& IncomingStanza);
-
 	bool HandleRoomConfigStanza(const FStropheStanza& IncomingStanza);
 	bool HandleRoomConfigErrorStanza(const FStropheStanza& IncomingStanza);
 
+	// Game Thread
+	void OnDisconnect();
+	void OnReconnect();
+
 	// IXmppMultiUserChat
 	virtual bool CreateRoom(const FXmppRoomId& RoomId, const FString& Nickname, const FXmppRoomConfig& RoomConfig) override;
 	virtual bool ConfigureRoom(const FXmppRoomId& RoomId, const FXmppRoomConfig& RoomConfig) override;
diff --git a/Engine/Source/Runtime/Online/XMPP/Private/XmppStrophe/XmppPingStrophe.cpp b/Engine/Source/Runtime/Online/XMPP/Private/XmppStrophe/XmppPingStrophe.cpp
--- a/Engine/Source/Runtime/Online/XMPP/Private/XmppStrophe/XmppPingStrophe.cpp
+++ b/Engine/Source/Runtime/Online/XMPP/Private/XmppStrophe/XmppPingStrophe.cpp
@@ -69,6 +69,11 @@ void FXmppPingStrophe::OnDisconnect()
 	ResetPongTimer();
 }
 
+void FXmppPingStrophe::OnReconnect()
+{
+
+}
+
 bool FXmppPingStrophe::ReceiveStanza(const FStropheStanza& IncomingStanza)
 {
 	// All ping stanzas are IQ
diff --git a/Engine/Source/Runtime/Online/XMPP/Private/XmppStrophe/XmppPingStrophe.h b/Engine/Source/Runtime/Online/XMPP/Private/XmppStrophe/XmppPingStrophe.h
--- a/Engine/Source/Runtime/Online/XMPP/Private/XmppStrophe/XmppPingStrophe.h
+++ b/Engine/Source/Runtime/Online/XMPP/Private/XmppStrophe/XmppPingStrophe.h
@@ -45,6 +45,7 @@ public:
 
 	/** Clear any caches on disconnect */
 	void OnDisconnect();
+	void OnReconnect();
 
 	/** Determine if an incoming stanza is a ping stanza */
 	bool ReceiveStanza(const FStropheStanza& IncomingStanza);
diff --git a/Engine/Source/Runtime/Online/XMPP/Private/XmppStrophe/XmppPresenceStrophe.cpp b/Engine/Source/Runtime/Online/XMPP/Private/XmppStrophe/XmppPresenceStrophe.cpp
--- a/Engine/Source/Runtime/Online/XMPP/Private/XmppStrophe/XmppPresenceStrophe.cpp
+++ b/Engine/Source/Runtime/Online/XMPP/Private/XmppStrophe/XmppPresenceStrophe.cpp
@@ -31,6 +31,17 @@ void FXmppPresenceStrophe::OnDisconnect()
 	CleanupMessages();
 }
 
+void FXmppPresenceStrophe::OnReconnect()
+{
+	// Triggered by login request when already connected
+	// re-broadcast all cached presence entries
+	for (const auto& Pair : RosterMembers)
+	{
+		const TSharedRef<FXmppUserPresence>& Presence = Pair.Value;
+		OnXmppPresenceReceivedDelegate.Broadcast(ConnectionManager.AsShared(), Presence->UserJid, Presence);
+	}
+}
+
 bool FXmppPresenceStrophe::ReceiveStanza(const FStropheStanza& IncomingStanza)
 {
 	if (IncomingStanza.GetName() != Strophe::SN_PRESENCE)
diff --git a/Engine/Source/Runtime/Online/XMPP/Private/XmppStrophe/XmppPresenceStrophe.h b/Engine/Source/Runtime/Online/XMPP/Private/XmppStrophe/XmppPresenceStrophe.h
--- a/Engine/Source/Runtime/Online/XMPP/Private/XmppStrophe/XmppPresenceStrophe.h
+++ b/Engine/Source/Runtime/Online/XMPP/Private/XmppStrophe/XmppPresenceStrophe.h
@@ -24,11 +24,11 @@ public:
 	virtual ~FXmppPresenceStrophe();
 
 	// XMPP Thread
-
-	void OnDisconnect();
 	bool ReceiveStanza(const FStropheStanza& IncomingStanza);
 
 	// Game Thread
+	void OnDisconnect();
+	void OnReconnect();
 
 	// IXmppPresence
 	virtual bool UpdatePresence(const FXmppUserPresence& NewPresence) override;
diff --git a/Engine/Source/Runtime/Online/XMPP/Private/XmppStrophe/XmppPrivateChatStrophe.cpp b/Engine/Source/Runtime/Online/XMPP/Private/XmppStrophe/XmppPrivateChatStrophe.cpp
--- a/Engine/Source/Runtime/Online/XMPP/Private/XmppStrophe/XmppPrivateChatStrophe.cpp
+++ b/Engine/Source/Runtime/Online/XMPP/Private/XmppStrophe/XmppPrivateChatStrophe.cpp
@@ -27,6 +27,11 @@ void FXmppPrivateChatStrophe::OnDisconnect()
 	CleanupMessages();
 }
 
+void FXmppPrivateChatStrophe::OnReconnect()
+{
+
+}
+
 bool FXmppPrivateChatStrophe::ReceiveStanza(const FStropheStanza& IncomingStanza)
 {
 	if (IncomingStanza.GetName() != Strophe::SN_MESSAGE || // Must be a message
diff --git a/Engine/Source/Runtime/Online/XMPP/Private/XmppStrophe/XmppPrivateChatStrophe.h b/Engine/Source/Runtime/Online/XMPP/Private/XmppStrophe/XmppPrivateChatStrophe.h
--- a/Engine/Source/Runtime/Online/XMPP/Private/XmppStrophe/XmppPrivateChatStrophe.h
+++ b/Engine/Source/Runtime/Online/XMPP/Private/XmppStrophe/XmppPrivateChatStrophe.h
@@ -26,10 +26,11 @@ public:
 	virtual ~FXmppPrivateChatStrophe();
 
 	// XMPP Thread
-	void OnDisconnect();
 	bool ReceiveStanza(const FStropheStanza& IncomingStanza);
 
 	// Game Thread
+	void OnDisconnect();
+	void OnReconnect();
 
 	// IXmppChat
 	virtual bool SendChat(const FXmppUserJid& RecipientId, const FString& Message) override;
diff --git a/Engine/Source/Runtime/Online/XMPP/Private/XmppStrophe/XmppPubSubStrophe.cpp b/Engine/Source/Runtime/Online/XMPP/Private/XmppStrophe/XmppPubSubStrophe.cpp
--- a/Engine/Source/Runtime/Online/XMPP/Private/XmppStrophe/XmppPubSubStrophe.cpp
+++ b/Engine/Source/Runtime/Online/XMPP/Private/XmppStrophe/XmppPubSubStrophe.cpp
@@ -16,6 +16,11 @@ void FXmppPubSubStrophe::OnDisconnect()
 
 }
 
+void FXmppPubSubStrophe::OnReconnect()
+{
+
+}
+
 bool FXmppPubSubStrophe::ReceiveStanza(const FStropheStanza& IncomingStanza)
 {
 	return false;
diff --git a/Engine/Source/Runtime/Online/XMPP/Private/XmppStrophe/XmppPubSubStrophe.h b/Engine/Source/Runtime/Online/XMPP/Private/XmppStrophe/XmppPubSubStrophe.h
--- a/Engine/Source/Runtime/Online/XMPP/Private/XmppStrophe/XmppPubSubStrophe.h
+++ b/Engine/Source/Runtime/Online/XMPP/Private/XmppStrophe/XmppPubSubStrophe.h
@@ -21,9 +21,13 @@ public:
 	FXmppPubSubStrophe(class FXmppConnectionStrophe& InConnectionManager);
 	virtual ~FXmppPubSubStrophe() = default;
 
-	void OnDisconnect();
+	// XMPP Thread
 	bool ReceiveStanza(const FStropheStanza& IncomingStanza);
 
+	// Game Thread
+	void OnDisconnect();
+	void OnReconnect();
+
 	// IXmppPubSub
 	virtual bool CreateNode(const FXmppPubSubId& NodeId, const FXmppPubSubConfig& NodeConfig) override;
 	virtual bool ConfigureNode(const FXmppPubSubId& NodeId, const FXmppPubSubConfig& NodeConfig) override;
diff --git a/Engine/Source/Runtime/Online/XMPP/Public/XmppConnection.h b/Engine/Source/Runtime/Online/XMPP/Public/XmppConnection.h
--- a/Engine/Source/Runtime/Online/XMPP/Public/XmppConnection.h
+++ b/Engine/Source/Runtime/Online/XMPP/Public/XmppConnection.h
@@ -303,5 +303,8 @@ public:
 	virtual IXmppMultiUserChatPtr MultiUserChat() = 0;
 	/** @return PrivateChat interface if available. NULL otherwise */
 	virtual IXmppChatPtr PrivateChat() = 0;
+	
+	/** Debug dump of this connection */
+	virtual void DumpState() const = 0;
 };
 
diff --git a/Engine/Source/Runtime/Online/XMPP/Public/XmppModule.h b/Engine/Source/Runtime/Online/XMPP/Public/XmppModule.h
--- a/Engine/Source/Runtime/Online/XMPP/Public/XmppModule.h
+++ b/Engine/Source/Runtime/Online/XMPP/Public/XmppModule.h
@@ -96,6 +96,24 @@ public:
 		return bEnabled;
 	}
 
+	/**
+	 * Delegate callback when a system acquires ownership over an XMPP connection
+	 *
+	 * @param XmppConnection The connection that is acquired
+	 * @param SystemName The name of the system initiating this request
+	 */
+	DECLARE_EVENT_TwoParams(FXmppModule, FOnXmppConnectionAcquired, const TSharedRef<IXmppConnection>& /*XmppConnection*/, const FString& /*SystemName*/);
+	FOnXmppConnectionAcquired OnXmppConnectionAcquired;
+
+	/**
+	 * Delegate callback when a system relinquishes ownership of an XMPP connection
+	 *
+	 * @param XmppConnection The connection that is relinquished
+	 * @param SystemName The name of the system initiating this request
+	 */
+	DECLARE_EVENT_TwoParams(FXmppModule, FOnXmppConnectionRelinquished, const TSharedRef<IXmppConnection>& /*XmppConnection*/, const FString& /*SystemName*/);
+	FOnXmppConnectionRelinquished OnXmppConnectionRelinquished;
+
 private:
 
 	// IModuleInterface
