cis fix

diff --git a/Engine/Source/Runtime/Online/XMPP/Private/XmppJingle/XmppConnectionJingle.cpp b/Engine/Source/Runtime/Online/XMPP/Private/XmppJingle/XmppConnectionJingle.cpp
--- a/Engine/Source/Runtime/Online/XMPP/Private/XmppJingle/XmppConnectionJingle.cpp
+++ b/Engine/Source/Runtime/Online/XMPP/Private/XmppJingle/XmppConnectionJingle.cpp
@@ -719,7 +719,7 @@ void FXmppConnectionJingle::Login(const FString& UserId, const FString& Password
 		{
 			ErrorStr = TEXT("Still processing last logout");
 		}
-		else if (LoginStatus == EXmppLoginStatus::LoggedIn)
+		else if (LoginState == EXmppLoginStatus::LoggedIn)
 		{
 			UE_LOG(LogXmpp, Log, TEXT("Reusing existing connection for Jid=%s"), *GetUserJid().ToDebugString());
 
