cis fix

diff --git a/Engine/Source/Runtime/Online/XMPP/Private/XmppJingle/XmppConnectionJingle.cpp b/Engine/Source/Runtime/Online/XMPP/Private/XmppJingle/XmppConnectionJingle.cpp
--- a/Engine/Source/Runtime/Online/XMPP/Private/XmppJingle/XmppConnectionJingle.cpp
+++ b/Engine/Source/Runtime/Online/XMPP/Private/XmppJingle/XmppConnectionJingle.cpp
@@ -725,7 +725,8 @@ void FXmppConnectionJingle::Login(const FString& UserId, const FString& Password
 
 			OnLoginComplete().Broadcast(GetUserJid(), true, FString());
 			OnLoginChanged().Broadcast(GetUserJid(), EXmppLoginStatus::LoggedIn);
-			ReconnectLogin();
+			//@todo - rebroadcast presence after login triggered on already logged in connection
+			//ReconnectLogin();
 			return;
 		}
 		else
