Add support for the OPTIONS verb when building http server request

diff --git a/Engine/Source/Runtime/Online/HTTPServer/Private/HttpConnectionRequestReadContext.cpp b/Engine/Source/Runtime/Online/HTTPServer/Private/HttpConnectionRequestReadContext.cpp
--- a/Engine/Source/Runtime/Online/HTTPServer/Private/HttpConnectionRequestReadContext.cpp
+++ b/Engine/Source/Runtime/Online/HTTPServer/Private/HttpConnectionRequestReadContext.cpp
@@ -242,6 +242,11 @@ TSharedPtr<FHttpServerRequest> FHttpConnectionRequestReadContext::BuildRequest(c
 	{
 		Request->Verb = EHttpServerRequestVerbs::VERB_PATCH;
 	}
+	else if (0 == RequestVerb.Compare(TEXT("OPTIONS"),
+		ESearchCase::IgnoreCase))
+	{
+		Request->Verb = EHttpServerRequestVerbs::VERB_OPTIONS;
+	}
 	else
 	{
 		// Unknown Verb
diff --git a/Engine/Source/Runtime/Online/HTTPServer/Public/HttpServerRequest.h b/Engine/Source/Runtime/Online/HTTPServer/Public/HttpServerRequest.h
--- a/Engine/Source/Runtime/Online/HTTPServer/Public/HttpServerRequest.h
+++ b/Engine/Source/Runtime/Online/HTTPServer/Public/HttpServerRequest.h
@@ -12,7 +12,8 @@ enum class EHttpServerRequestVerbs : uint16
 	VERB_POST   = 1 << 1,
 	VERB_PUT    = 1 << 2,
 	VERB_PATCH  = 1 << 3,
-	VERB_DELETE = 1 << 4
+	VERB_DELETE = 1 << 4,
+	VERB_OPTIONS = 1 << 5
 };
 
 ENUM_CLASS_FLAGS(EHttpServerRequestVerbs);
