#rnx
#5571

diff --git a/Engine/Source/Runtime/HTML5/HTML5JS/Private/HTML5JavaScriptFx.js b/Engine/Source/Runtime/HTML5/HTML5JS/Private/HTML5JavaScriptFx.js
--- a/Engine/Source/Runtime/HTML5/HTML5JS/Private/HTML5JavaScriptFx.js
+++ b/Engine/Source/Runtime/HTML5/HTML5JS/Private/HTML5JavaScriptFx.js
@@ -157,7 +157,7 @@ var UE_JavascriptLibrary = {
         HEAPU8.set(byteArray, buffer);
 
         if (onload)
-          Runtime.dynCall('viiii', onload, [ctx, buffer, byteArray.length, header_buffer]);
+          Runtime.dynCall('viiiii', onload, [ctx, buffer, byteArray.length, header_buffer, xhr.status]);
         if (freeBuffer) // seems POST reqeusts keeps the buffer
           _free(buffer);
         _free(header_buffer);
diff --git a/Engine/Source/Runtime/HTML5/HTML5JS/Public/HTML5JavaScriptFx.h b/Engine/Source/Runtime/HTML5/HTML5JS/Public/HTML5JavaScriptFx.h
--- a/Engine/Source/Runtime/HTML5/HTML5JS/Public/HTML5JavaScriptFx.h
+++ b/Engine/Source/Runtime/HTML5/HTML5JS/Public/HTML5JavaScriptFx.h
@@ -23,7 +23,7 @@
 					const char* headers,
 					int async, // e.g. onBeforeUnload
 					int freeBuffer,
-					void(*onload)(void*, void*, unsigned, void*),
+					void(*onload)(void*, void*, unsigned, void*, int),
 					void(*onerror)(void*, int, const char*),
 					void(*onprogress)(void*, int, int)
 				);
diff --git a/Engine/Source/Runtime/Online/HTTP/Private/HTML5/HTML5HTTP.cpp b/Engine/Source/Runtime/Online/HTTP/Private/HTML5/HTML5HTTP.cpp
--- a/Engine/Source/Runtime/Online/HTTP/Private/HTML5/HTML5HTTP.cpp
+++ b/Engine/Source/Runtime/Online/HTTP/Private/HTML5/HTML5HTTP.cpp
@@ -195,16 +195,16 @@ void FHTML5HttpRequest::SetVerb(const FString& InVerb)
 	Verb = InVerb.ToUpper();
 }
 
-void FHTML5HttpRequest::StaticReceiveCallback(void *arg, void *buffer, uint32 size, void* httpHeaders)
+void FHTML5HttpRequest::StaticReceiveCallback(void *arg, void *buffer, uint32 size, void* httpHeaders, int httpStatusCode)
 {
 	UE_LOG(LogHttp, Verbose, TEXT("FHTML5HttpRequest::StaticReceiveDataCallback()"));
 
 	FHTML5HttpRequest* Request = reinterpret_cast<FHTML5HttpRequest*>(arg);
 
-	return Request->ReceiveCallback(arg, buffer, size, httpHeaders);
+	return Request->ReceiveCallback(arg, buffer, size, httpHeaders, httpStatusCode);
 }
 
-void FHTML5HttpRequest::ReceiveCallback(void *arg, void *buffer, uint32 size, void* httpHeaders)
+void FHTML5HttpRequest::ReceiveCallback(void *arg, void *buffer, uint32 size, void* httpHeaders, int httpStatusCode)
 {
 	UE_LOG(LogHttp, Verbose, TEXT("FHTML5HttpRequest::ReceiveDataCallback()"));
 	UE_LOG(LogHttp, Verbose, TEXT("Response size: %d"), size);
@@ -255,7 +255,7 @@ void FHTML5HttpRequest::ReceiveCallback(void *arg, void *buffer, uint32 size, vo
 
 		FMemory::Memcpy(static_cast<uint8*>(Response->Payload.GetData()), buffer, size);
 		Response->TotalBytesRead = size;
-		Response->HttpCode = 200;
+		Response->HttpCode = httpStatusCode;
 
 		UE_LOG(LogHttp, Verbose, TEXT("Payload length: %d"), Response->Payload.Num());
 
@@ -481,16 +481,15 @@ void FHTML5HttpRequest::FinishedRequest()
 	// Clean up session/request handles that may have been created
 	CleanupRequest();
 
-	if (Response.IsValid() &&
-		Response->bSucceeded)
+	if (Response.IsValid())
 	{
 		UE_LOG(LogHttp, Verbose, TEXT("%p: request has been successfully processed. HTTP code: %d, content length: %d, actual payload size: %d"), 
 			this, Response->HttpCode, Response->ContentLength, Response->Payload.Num() );
 
 		// Mark last request attempt as completed successfully
-		CompletionStatus = EHttpRequestStatus::Succeeded;
+		CompletionStatus = Response->bSucceeded ? EHttpRequestStatus::Succeeded : EHttpRequestStatus::Failed;
 		// Call delegate with valid request/response objects
-		OnProcessRequestComplete().ExecuteIfBound(SharedThis(this),Response,true);
+		OnProcessRequestComplete().ExecuteIfBound(SharedThis(this),Response,Response->bSucceeded);
 	}
 	else
 	{
diff --git a/Engine/Source/Runtime/Online/HTTP/Private/HTML5/HTML5HTTP.h b/Engine/Source/Runtime/Online/HTTP/Private/HTML5/HTML5HTTP.h
--- a/Engine/Source/Runtime/Online/HTTP/Private/HTML5/HTML5HTTP.h
+++ b/Engine/Source/Runtime/Online/HTTP/Private/HTML5/HTML5HTTP.h
@@ -84,8 +84,8 @@ private:
 	*/
 	void CleanupRequest();
 
-	static void StaticReceiveCallback(void *arg, void *buffer, uint32 size, void* httpHeaders);
-	void ReceiveCallback(void *arg, void *buffer, uint32 size, void* httpHeaders);
+	static void StaticReceiveCallback(void *arg, void *buffer, uint32 size, void* httpHeaders, int httpStatusCode);
+	void ReceiveCallback(void *arg, void *buffer, uint32 size, void* httpHeaders, int httpStatusCode);
 	static void StaticErrorCallback(void* arg, int httpStatusCode, const char* httpStatusText);
 	void ErrorCallback(void* arg, int httpStatusCode, const char* httpStatusText);
 	static void StaticProgressCallback(void* arg, int Loaded, int Total);
