Fix deprecation warning in HttpListener with CreateInternetAddr

#rnx

diff --git a/Engine/Source/Runtime/Online/HTTPServer/Private/HttpListener.cpp b/Engine/Source/Runtime/Online/HTTPServer/Private/HttpListener.cpp
--- a/Engine/Source/Runtime/Online/HTTPServer/Private/HttpListener.cpp
+++ b/Engine/Source/Runtime/Online/HTTPServer/Private/HttpListener.cpp
@@ -58,7 +58,9 @@ bool FHttpListener::StartListening()
 	ListenSocket->SetNonBlocking(true);
 
 	// Bind to Localhost/Caller-defined port
-	TSharedRef<FInternetAddr> LocalhostAddr = SocketSubsystem->CreateInternetAddr(ListenerInterfaceAddress, ListenPort);
+	TSharedRef<FInternetAddr> LocalhostAddr = SocketSubsystem->CreateInternetAddr();
+	LocalhostAddr->SetAnyAddress();
+	LocalhostAddr->SetPort(ListenPort);
 	if (!ListenSocket->Bind(*LocalhostAddr))
 	{
 		UE_LOG(LogHttpListener, Error, 
diff --git a/Engine/Source/Runtime/Online/HTTPServer/Private/HttpListener.h b/Engine/Source/Runtime/Online/HTTPServer/Private/HttpListener.h
--- a/Engine/Source/Runtime/Online/HTTPServer/Private/HttpListener.h
+++ b/Engine/Source/Runtime/Online/HTTPServer/Private/HttpListener.h
@@ -115,6 +115,4 @@ private:
 	static constexpr uint32 ListenerConnectionBacklogSize = 16;
 	/** Maximum send buffer size */
 	static constexpr uint32 ListenerBufferSize = 512 * 1024;
-	/** The interface to bind to */
-	static constexpr uint32 ListenerInterfaceAddress = 0x0;
 };
\ No newline at end of file
