Fix for BackgroundHttpManager resuming paused tasks when moving to background.
Fix for BackgroundHttpManager resuming paused tasks when they associate with an UnAssociatedTask.
Fix for BackgroundHttpManager resuming paused tasks when they were created in the background.
#tests iOS Client

diff --git a/Engine/Source/Runtime/Online/BackgroundHTTP/Private/IOS/ApplePlatformBackgroundHttpManager.cpp b/Engine/Source/Runtime/Online/BackgroundHTTP/Private/IOS/ApplePlatformBackgroundHttpManager.cpp
--- a/Engine/Source/Runtime/Online/BackgroundHTTP/Private/IOS/ApplePlatformBackgroundHttpManager.cpp
+++ b/Engine/Source/Runtime/Online/BackgroundHTTP/Private/IOS/ApplePlatformBackgroundHttpManager.cpp
@@ -323,7 +323,7 @@ void FApplePlatformBackgroundHttpManager::OnApp_EnteringForeground()
 void FApplePlatformBackgroundHttpManager::OnApp_EnteringBackground()
 {
 	FPlatformAtomics::InterlockedExchange(&bIsInBackground,true);
-	ResumeAllTasks();
+	ResumeAllTasksWithoutPausedRequests();
 }
 
 void FApplePlatformBackgroundHttpManager::PauseAllActiveTasks()
@@ -354,25 +354,36 @@ void FApplePlatformBackgroundHttpManager::PauseAllActiveTasks()
 	}
 }
 
-void FApplePlatformBackgroundHttpManager::ResumeAllTasks()
+void FApplePlatformBackgroundHttpManager::ResumeAllTasksWithoutPausedRequests()
 {
-    UE_LOG(LogBackgroundHttpManager, Display, TEXT("Attempting to Resume All Active Tasks"));
+    UE_LOG(LogBackgroundHttpManager, Display, TEXT("Attempting to resume all active tasks that don't have a corresponding request that has paused them"));
     
 	NSURLSession* BackgroundDownloadSession = FBackgroundURLSessionHandler::GetBackgroundSession();
 	if (nullptr != BackgroundDownloadSession)
 	{
 		[BackgroundDownloadSession getTasksWithCompletionHandler : ^ (NSArray<__kindof NSURLSessionDataTask*>* DataTasks, NSArray<__kindof NSURLSessionUploadTask*>* UploadTasks, NSArray<__kindof NSURLSessionDownloadTask*>* DownloadTasks)
 		{
+            FRWScopeLock ScopeLock(URLToRequestMapLock, SLT_ReadOnly);
+            
 			for (NSURLSessionDownloadTask* DownloadTask : DownloadTasks)
 			{
 				if ([DownloadTask state] == NSURLSessionTaskStateSuspended)
 				{
                     FString TaskURL = [[[DownloadTask currentRequest] URL] absoluteString];
                     int TaskIdentifier = (int)[DownloadTask taskIdentifier];
+                    FBackgroundHttpURLMappedRequestPtr* WeakRequestInMap = URLToRequestMap.Find(TaskURL);
+                    FAppleBackgroundHttpRequestPtr FoundRequest = (nullptr != WeakRequestInMap) ? WeakRequestInMap->Pin() : nullptr;
                     
-                    UE_LOG(LogBackgroundHttpManager, Display, TEXT("Resuming Task for URL:%s | TaskIdentifier:%d"), *TaskURL, TaskIdentifier);
-                    
-					[DownloadTask resume];
+                    bool bIsRequestPaused = FoundRequest.IsValid() ? FoundRequest->bIsTaskPaused : false;
+                    if (!bIsRequestPaused)
+                    {
+                        UE_LOG(LogBackgroundHttpManager, Display, TEXT("Resuming Task for URL:%s | TaskIdentifier:%d"), *TaskURL, TaskIdentifier);
+                        [DownloadTask resume];
+                    }
+                    else
+                    {
+                        UE_LOG(LogBackgroundHttpManager, Display, TEXT("NOT Resuming Task for URL as associated request was paused! URL:%s | TaskIdentifier:%d"), *TaskURL, TaskIdentifier);
+                    }
 				}
 			}
 		}];
@@ -533,14 +544,15 @@ void FApplePlatformBackgroundHttpManager::RetryRequest(FAppleBackgroundHttpReque
 			{
 				Request->AssociateWithTask(NewTask);
 
-				//If we are in BG or flagged this as an immediate start, resume right now without waiting for the FG tick
+				//If we are in BG or flagged this as an immediate start, resume right now without waiting for the FG tick.
 				volatile bool bCopyOfBGState = FPlatformAtomics::AtomicRead(&bIsInBackground);
-				if (bCopyOfBGState || bShouldStartImmediately)
+				if (!Request->bIsTaskPaused
+                    && (bCopyOfBGState || bShouldStartImmediately))
 				{
 					Request->ActivateUnderlyingTask();
 				}
 
-				UE_LOG(LogBackgroundHttpManager, Display, TEXT("Created Task for Request -- RequestDebugID:%s | bStartImmediately:%d | bIsAppInBG:%d"), *(Request->GetRequestDebugID()), (int)bShouldStartImmediately, (int)bCopyOfBGState);
+                UE_LOG(LogBackgroundHttpManager, Display, TEXT("Created Task for Request -- RequestDebugID:%s | bStartImmediately:%d | bIsAppInBG:%d | bIsPaused:%d"), *(Request->GetRequestDebugID()), (int)bShouldStartImmediately, (int)bCopyOfBGState, (int)Request->bIsTaskPaused);
 
 				//Always set our bWasTaskStartedInBG flag on our Request so we will know if we need to restart this task next FG Tick.
 				FPlatformAtomics::InterlockedExchange(&(Request->bWasTaskStartedInBG), bCopyOfBGState);
@@ -575,7 +587,7 @@ bool FApplePlatformBackgroundHttpManager::ShouldUseRequestRetryData(FAppleBackgr
 
 bool FApplePlatformBackgroundHttpManager::IsRetryDataValid(NSData* RetryData) const
 {
-	return ((nullptr != RetryData) && ([RetryData length] > 0));
+    return ((nullptr != RetryData) && ([RetryData length] > 0));
 }
 
 void FApplePlatformBackgroundHttpManager::OnTask_DidWriteData(NSURLSessionDownloadTask* Task, int64_t BytesWrittenSinceLastCall, int64_t TotalBytesWritten, int64_t TotalBytesExpectedToWrite)
@@ -645,7 +657,7 @@ void FApplePlatformBackgroundHttpManager::OnTask_DidCompleteWithError(NSURLSessi
                 bShouldRetryIncreaseRetryCount = false;
                 DebugRetryOverrideReason = TEXT("Not Connected To Internet");
             }
-
+            
             UE_LOG(LogBackgroundHttpManager, Display, TEXT("DidCompleteWithError for Task. -- URL:%s | TaskIdentifier:%d | bDidFindVaildRequest:%d | bDidCompleteWithError:%d | ErrorCode:%d | bHasResumeData:%d | CancelledReasonKey:%d | RetryOverrideReason:%s | bShouldRetryIncreaseRetryCount:%d | ErrorDescription:%s"), *TaskURL, TaskIdentifier, (int)bDidFindValidRequest, (int)bDidCompleteWithError, ErrorCode, (int)bHasResumeData, CancelledReasonInt, *DebugRetryOverrideReason, (int)bShouldRetryIncreaseRetryCount, *ErrorDescription);
             
 			if (bDidFindValidRequest)
diff --git a/Engine/Source/Runtime/Online/BackgroundHTTP/Public/IOS/ApplePlatformBackgroundHttpManager.h b/Engine/Source/Runtime/Online/BackgroundHTTP/Public/IOS/ApplePlatformBackgroundHttpManager.h
--- a/Engine/Source/Runtime/Online/BackgroundHTTP/Public/IOS/ApplePlatformBackgroundHttpManager.h
+++ b/Engine/Source/Runtime/Online/BackgroundHTTP/Public/IOS/ApplePlatformBackgroundHttpManager.h
@@ -73,7 +73,7 @@ private:
 
     //These are used internally to pause/resume things at the task level
 	void PauseAllActiveTasks();
-	void ResumeAllTasks();
+	void ResumeAllTasksWithoutPausedRequests();
 
     void PauseAllUnassociatedTasks();
     void UnpauseAllUnassociatedTasks();
