#stomp Fixed shaky reliance on case of command FNames. Now explicitly upper cases all Stomp commands. Also fixed rogue 'auto' redundantly copying (FName,FString) pairs.

diff --git a/Engine/Source/Runtime/Online/Stomp/Private/StompCommand.cpp b/Engine/Source/Runtime/Online/Stomp/Private/StompCommand.cpp
--- a/Engine/Source/Runtime/Online/Stomp/Private/StompCommand.cpp
+++ b/Engine/Source/Runtime/Online/Stomp/Private/StompCommand.cpp
@@ -4,23 +4,25 @@
 
 #if WITH_STOMP
 
-#define DEFINE_COMMAND(Name, Value) const FStompCommand Name ## Command = TEXT(#Value)
+#define DEFINE_COMMAND(Name) const FStompCommand Name ## Command = TEXT(#Name)
 
 const FStompCommand HeartbeatCommand = FStompCommand();
-DEFINE_COMMAND(Connect, CONNECT);
-DEFINE_COMMAND(Connected, CONNECTED);
-DEFINE_COMMAND(Send, SEND);
-DEFINE_COMMAND(Subscribe, SUBSCRIBE);
-DEFINE_COMMAND(Unsubscribe, UNSUBSCRIBE);
-DEFINE_COMMAND(Begin, BEGIN);
-DEFINE_COMMAND(Commit, COMMIT);
-DEFINE_COMMAND(Abort, ABORT);
-DEFINE_COMMAND(Ack, ACK);
-DEFINE_COMMAND(Nack, NACK);
-DEFINE_COMMAND(Disconnect, DISCONNECT);
-DEFINE_COMMAND(Message, MESSAGE);
-DEFINE_COMMAND(Receipt, RECEIPT);
-DEFINE_COMMAND(Error, ERROR);
+
+// note: commands are converted to upper case to send over the wire
+DEFINE_COMMAND(Connect);
+DEFINE_COMMAND(Connected);
+DEFINE_COMMAND(Send);
+DEFINE_COMMAND(Subscribe);
+DEFINE_COMMAND(Unsubscribe);
+DEFINE_COMMAND(Begin);
+DEFINE_COMMAND(Commit);
+DEFINE_COMMAND(Abort);
+DEFINE_COMMAND(Ack);
+DEFINE_COMMAND(Nack);
+DEFINE_COMMAND(Disconnect);
+DEFINE_COMMAND(Message);
+DEFINE_COMMAND(Receipt);
+DEFINE_COMMAND(Error);
 
 #undef DEFINE_COMMAND
 
diff --git a/Engine/Source/Runtime/Online/Stomp/Private/StompFrame.cpp b/Engine/Source/Runtime/Online/Stomp/Private/StompFrame.cpp
--- a/Engine/Source/Runtime/Online/Stomp/Private/StompFrame.cpp
+++ b/Engine/Source/Runtime/Online/Stomp/Private/StompFrame.cpp
@@ -122,20 +122,24 @@ void FStompFrame::Encode(FStompBuffer& Out) const
 		// According to the spec, the CONNECT command should not escape metacharacters for backwards compatibility.
 		bool bShouldEscapeFrameHeader = Command != ConnectCommand;
 
-		FTCHARToUTF8 CommandEncoded(*Command.ToString());
+		FString CommandString = Command.ToString();
+		CommandString.ToUpperInline();
+		FTCHARToUTF8 CommandEncoded(*CommandString);
 		AppendArray(Out, (uint8*)CommandEncoded.Get(), CommandEncoded.Length(), bShouldEscapeFrameHeader);
 		Out.Add('\n');
 
-		for(auto Element : Header)
+		for (const TPair<FName, FString>& Element : Header)
 		{
-			FTCHARToUTF8 HeaderNameEncoded(*Element.Key.ToString().ToLower());
+			FString ElementKeyString = Element.Key.ToString();
+			ElementKeyString.ToLowerInline();
+			FTCHARToUTF8 HeaderNameEncoded(*ElementKeyString);
 			FTCHARToUTF8 HeaderValueEncoded(*Element.Value);
 			AppendArray(Out, (uint8*)HeaderNameEncoded.Get(), HeaderNameEncoded.Length(), bShouldEscapeFrameHeader);
 			Out.Add(':');
 			AppendArray(Out, (uint8*)HeaderValueEncoded.Get(), HeaderValueEncoded.Length(), bShouldEscapeFrameHeader);
 			Out.Add('\n');
-
 		}
+
 		Out.Add('\n');
 		Out.Append(Body);
 		Out.Add('\0');
@@ -169,7 +173,7 @@ void FStompFrame::Decode(const uint8* In, SIZE_T Length)
 
 	if (Index >= Length)
 	{
-		UE_LOG(LogStomp, Warning, TEXT("Stomp command '%s' received without any headers"), *Command.ToString());
+		UE_LOG(LogStomp, Warning, TEXT("Stomp command '%s' received without any headers"), *Command.ToString().ToUpper());
 		return;
 	}
 
