Move IVoiceChat header out of VivoxVoiceChat

diff --git a/Engine/Plugins/Online/VoiceChat/VivoxVoiceChat/Source/Public/VivoxVoiceChat.h b/Engine/Plugins/Online/VoiceChat/VivoxVoiceChat/Source/Public/VivoxVoiceChat.h
--- a/Engine/Plugins/Online/VoiceChat/VivoxVoiceChat/Source/Public/VivoxVoiceChat.h
+++ b/Engine/Plugins/Online/VoiceChat/VivoxVoiceChat/Source/Public/VivoxVoiceChat.h
@@ -14,7 +14,7 @@ THIRD_PARTY_INCLUDES_START
 #include "vivoxclientapi/clientconnection.h"
 THIRD_PARTY_INCLUDES_END
 
-#include "VoiceChat.h"
+#include "Interfaces/VoiceChat.h"
 #include "Stats/Stats.h"
 #include "Misc/CoreMisc.h"
 #include "Logging/LogMacros.h"
diff --git a/Engine/Plugins/Online/VoiceChat/VivoxVoiceChat/Source/Public/VoiceChat.h b/Engine/Plugins/Online/VoiceChat/VivoxVoiceChat/Source/Public/VoiceChat.h
deleted file mode 100644
diff --git a/Engine/Plugins/Online/VoiceChat/VivoxVoiceChat/Source/VivoxVoiceChat.Build.cs b/Engine/Plugins/Online/VoiceChat/VivoxVoiceChat/Source/VivoxVoiceChat.Build.cs
--- a/Engine/Plugins/Online/VoiceChat/VivoxVoiceChat/Source/VivoxVoiceChat.Build.cs
+++ b/Engine/Plugins/Online/VoiceChat/VivoxVoiceChat/Source/VivoxVoiceChat.Build.cs
@@ -18,6 +18,13 @@ namespace UnrealBuildTool.Rules
 				}
 			);
 
+			PublicIncludePathModuleNames.AddRange(
+				new string[]
+				{
+					"Voice"
+				}
+			);
+
 			PrivateDependencyModuleNames.AddRange(
 				new string[]
 				{
diff --git a/Engine/Source/Runtime/Online/Voice/Public/Interfaces/VoiceChat.h b/Engine/Source/Runtime/Online/Voice/Public/Interfaces/VoiceChat.h
new file mode 100644
--- /dev/null
+++ b/Engine/Source/Runtime/Online/Voice/Public/Interfaces/VoiceChat.h
@@ -0,0 +1,569 @@
+// Copyright 1998-2019 Epic Games, Inc. All Rights Reserved.
+
+#pragma once
+
+#include "Delegates/Delegate.h"
+#include "Runtime/Core/Public/Features/IModularFeature.h"
+#include "Runtime/Core/Public/Features/IModularFeatures.h"
+
+struct FVoiceChatResult
+{
+	/** Whether or not we succeeded */
+	bool bSuccess;
+	/** If we failed, a numeric code for the error */
+	int ErrorCode;
+	/** If we failed, a string describing the error suitable for logging */
+	FString Error;
+};
+
+enum class EVoiceChatChannelType
+{
+	/** Non positional/2d audio channel */
+	NonPositional,
+	/** Positional/3d audio channel */
+	Positional,
+	/** Echo channel. Will only ever have one player and will echo anything you say */
+	Echo
+};
+
+enum class EVoiceChatAttenuationModel
+{
+	/** No attenuation is applied. The audio will drop to 0 at MaxDistance */
+	None,
+	/** The attenuation increases in inverse proportion to the distance. The Rolloff is the inverse of the slope of the attenuation curve. */
+	InverseByDistance,
+	/** The attenuation increases in linear proportion to the distance. The Rolloff is the negative slope of the attenuation curve. */
+	LinearByDistance,
+	/** The attenuation increases in inverse proportion to the distance raised to the power of the Rolloff. */
+	ExponentialByDistance
+};
+
+struct FVoiceChatChannel3dProperties
+{
+	/** The model used to determine how loud audio is at different distances */
+	EVoiceChatAttenuationModel AttenuationModel;
+	/** The distance at which the sound will start to attenuate */
+	float MinDistance;
+	/** The distance at which sound will no longer be audible */
+	float MaxDistance;
+	/** How fast the sound attenuates with distance */
+	float Rolloff;
+};
+
+enum class EVoiceChatTransmitMode
+{
+	/** Transmit to none of the channels you are in */
+	None,
+	/** Transmit to all of the channels you are in */
+	All,
+	/** Transmit to a specific channel */
+	Channel
+};
+
+struct FVoiceChatCallStats
+{
+	/** Call length in seconds */
+	double CallLength;
+	/** Minimum measured network latency in seconds. Zero if no measurements made */
+	double LatencyMinMeasuredSeconds;
+	/** Maximum measured network latency in seconds. Zero if no measurements made */
+	double LatencyMaxMeasuredSeconds;
+	/** Average measured network latency in seconds. Zero if no measurements made */
+	double LatencyAverageMeasuredSeconds;
+	/** Total number of packets, both received and lost */
+	int PacketsNumTotal;
+	/** Total number of lost packets */
+	int PacketsNumLost;
+};
+
+DECLARE_DELEGATE_OneParam(FOnVoiceChatConnectCompleteDelegate, const FVoiceChatResult& /* Result */);
+DECLARE_DELEGATE_OneParam(FOnVoiceChatDisconnectCompleteDelegate, const FVoiceChatResult& /* Result */);
+DECLARE_DELEGATE_TwoParams(FOnVoiceChatLoginCompleteDelegate, const FString& /* PlayerName */, const FVoiceChatResult& /* Result */);
+DECLARE_DELEGATE_TwoParams(FOnVoiceChatLogoutCompleteDelegate, const FString& /* PlayerName */, const FVoiceChatResult& /* Result */);
+DECLARE_DELEGATE_TwoParams(FOnVoiceChatChannelJoinCompleteDelegate, const FString& /* ChannelName */, const FVoiceChatResult& /* Result */);
+DECLARE_DELEGATE_TwoParams(FOnVoiceChatChannelLeaveCompleteDelegate, const FString& /* ChannelName */, const FVoiceChatResult& /* Result */);
+
+DECLARE_MULTICAST_DELEGATE(FOnVoiceChatAvailableAudioDevicesChangedDelegate);
+DECLARE_MULTICAST_DELEGATE(FOnVoiceChatReconnectedDelegate);
+DECLARE_MULTICAST_DELEGATE(FOnVoiceChatConnectedDelegate);
+DECLARE_MULTICAST_DELEGATE_OneParam(FOnVoiceChatDisconnectedDelegate, const FVoiceChatResult& /* Reason */);
+DECLARE_MULTICAST_DELEGATE_OneParam(FOnVoiceChatLoggedInDelegate, const FString& /* PlayerName */);
+DECLARE_MULTICAST_DELEGATE_OneParam(FOnVoiceChatLoggedOutDelegate, const FString& /* PlayerName */);
+DECLARE_MULTICAST_DELEGATE_OneParam(FOnVoiceChatCallStatsUpdatedDelegate, const FVoiceChatCallStats& /* CallStats */);
+DECLARE_MULTICAST_DELEGATE_OneParam(FOnVoiceChatChannelJoinedDelegate, const FString& /* ChannelName */);
+DECLARE_MULTICAST_DELEGATE_TwoParams(FOnVoiceChatChannelExitedDelegate, const FString& /* ChannelName */, const FVoiceChatResult& /* Reason */);
+DECLARE_MULTICAST_DELEGATE_TwoParams(FOnVoiceChatPlayerAddedDelegate, const FString& /* ChannelName */, const FString& /* PlayerName */);
+DECLARE_MULTICAST_DELEGATE_TwoParams(FOnVoiceChatPlayerRemovedDelegate, const FString& /* ChannelName */, const FString& /* PlayerName */);
+DECLARE_MULTICAST_DELEGATE_ThreeParams(FOnVoiceChatPlayerTalkingUpdatedDelegate, const FString& /* ChannelName */, const FString& /* PlayerName */, bool /* bIsTalking */);
+DECLARE_MULTICAST_DELEGATE_ThreeParams(FOnVoiceChatPlayerMuteUpdatedDelegate, const FString& /* ChannelName */, const FString& /* PlayerName */, bool /* bIsMuted */);
+DECLARE_MULTICAST_DELEGATE_ThreeParams(FOnVoiceChatPlayerVolumeUpdatedDelegate, const FString& /* ChannelName */, const FString& /* PlayerName */, float /* Volume */);
+
+DECLARE_MULTICAST_DELEGATE_ThreeParams(FOnVoiceChatRecordSamplesAvailableDelegate, TArrayView<const int16> /* PcmSamples */, int /* SampleRate */, int /* Channels */);
+DECLARE_MULTICAST_DELEGATE_ThreeParams(FOnVoiceChatAfterCaptureAudioReadDelegate, TArrayView<int16> /* PcmSamples */, int /* SampleRate */, int /* Channels */);
+DECLARE_MULTICAST_DELEGATE_FourParams(FOnVoiceChatBeforeCaptureAudioSentDelegate, TArrayView<const int16> /* PcmSamples */, int /* SampleRate */, int /* Channels */, bool /* bIsSpeaking */);
+DECLARE_MULTICAST_DELEGATE_FourParams(FOnVoiceChatBeforeRecvAudioRenderedDelegate, TArrayView<const int16> /* PcmSamples */, int /* SampleRate */, int /* Channels */, bool /* bIsSilence */);
+
+class IVoiceChat : public IModularFeature
+{
+public:
+	/**
+	 * Get the IVoiceChat modular feature instance
+	 */
+	static IVoiceChat* Get()
+	{
+		if (IModularFeatures::Get().IsModularFeatureAvailable(GetModularFeatureName()))
+		{
+			return &IModularFeatures::Get().GetModularFeature<IVoiceChat>(GetModularFeatureName());
+		}
+		return nullptr;
+	}
+
+	/**
+	 * Initialize VoiceChat
+	 */
+	virtual bool Initialize() = 0;
+
+	/**
+	 * Uninitialize VoiceChat
+	 */
+	virtual bool Uninitialize() = 0;
+
+	/**
+	 * Is voice chat initialized?
+	 *
+	 * @return true if voice chat is initialized
+	 */
+	virtual bool IsInitialized() const = 0;
+
+	/**
+	 * Set an implementation specific setting
+	 *
+	 * @param Name Name of the config setting to set
+	 * @param Value Value for the setting
+	 */
+	virtual void SetSetting(const FString& Name, const FString& Value) = 0;
+
+	/**
+	 * Get an implementation specific setting
+	 *
+	 * @param Name Name of the config setting to get
+	 *
+	 * @return Current value for the setting
+	 */
+	virtual FString GetSetting(const FString& Name) = 0;
+
+	/**
+	 * Set record volume
+	 *
+	 * @param Volume Value between 0.0 and 1.0
+	 */
+	virtual void SetAudioInputVolume(float Volume) = 0;
+
+	/**
+	 * Set playback volume
+	 *
+	 * @param Volume Value between 0.0 and 1.0
+	 */
+	virtual void SetAudioOutputVolume(float Volume) = 0;
+
+	/**
+	 * Get volume of audio input device
+	 *
+	 * @return Audio input device volume. Will be between 0.0 and 1.0
+	 */
+	virtual float GetAudioInputVolume() const = 0;
+
+	/**
+	 * Get volume of audio output device
+	 *
+	 * @return Audio output device volume. Will be between 0.0 and 1.0
+	 */
+	virtual float GetAudioOutputVolume() const = 0;
+
+	/**
+	 * Mute or unmute the audio input device
+	 *
+	 * @param bIsMuted set to true to mute the device or false to unmute
+	 */
+	virtual void SetAudioInputDeviceMuted(bool bIsMuted) = 0;
+
+	/**
+	 * Mute or unmute the audio output device
+	 *
+	 * @param bIsMuted set to true to mute the device or false to unmute
+	 */
+	virtual void SetAudioOutputDeviceMuted(bool bIsMuted) = 0;
+
+	/**
+	 * Is the input device muted?
+	 *
+	 * @return true if the input device is muted
+	 */
+	virtual bool GetAudioInputDeviceMuted() const = 0;
+
+	/**
+	 * Is the output device muted?
+	 *
+	 * @return true if the input device is muted
+	 */
+	virtual bool GetAudioOutputDeviceMuted() const = 0;
+
+	/**
+	 * Get a list of available audio input devices
+	 *
+	 * @return Array of audio input devices
+	 */
+	virtual TArray<FString> GetAvailableInputDevices() const = 0;
+
+	/**
+	 * Get a list of available audio output devices
+	 *
+	 * @return Array of audio output devices
+	 */
+	virtual TArray<FString> GetAvailableOutputDevices() const = 0;
+
+	/**
+	 * Delegate broadcast whenever the available audio devices change. Any cached values from GetAvailableInputDevices or GetAvailableOutputDevices should be discarded and requeried
+	 */
+	virtual FOnVoiceChatAvailableAudioDevicesChangedDelegate& OnVoiceChatAvailableAudioDevicesChanged() = 0;
+
+	/**
+	 * Set the audio input device to use
+	 *
+	 * @param InputDevice String from GetAvailableInputDevice. If empty, will use the default input device
+	 */
+	virtual void SetInputDevice(const FString& InputDevice) = 0;
+
+	/**
+	 * Set the audio output device to use
+	 *
+	 * @param OutputDevice String from GetAvailableOutputDevice. If empty, will use the default output device
+	 */
+	virtual void SetOutputDevice(const FString& OutputDevice) = 0;
+
+	/**
+	 * Get the audio input device that is being used
+	 */
+	virtual FString GetInputDevice() const = 0;
+
+	/**
+	 * Get the audio output device that is being used
+	 */
+	virtual FString GetOutputDevice() const = 0;
+
+	/**
+	 * Get the default audio input device
+	 */
+	virtual FString GetDefaultInputDevice() const = 0;
+
+	/**
+	 * Get the default audio output device
+	 */
+	virtual FString GetDefaultOutputDevice() const = 0;
+
+	/**
+	 * Connect to a voice server
+	 *
+	 * @param Delegate delegate called once connect completes
+	 */
+	virtual void Connect(const FOnVoiceChatConnectCompleteDelegate& Delegate) = 0;
+
+	/**
+	 * Disconnect from a voice server
+	 *
+	 * @param Delegate delegate called once disconnect completes
+	 */
+	virtual void Disconnect(const FOnVoiceChatDisconnectCompleteDelegate& Delegate) = 0;
+
+	/**
+	* Are we connecting to the voice server?
+	*
+	* @return true if we are connecting to the voice server
+	*/
+	virtual bool IsConnecting() const = 0;
+
+	/**
+	 * Are we connected to the voice server?
+	 *
+	 * @return true if we are connected to the voice server
+	 */
+	virtual bool IsConnected() const = 0;
+
+	/**
+	 * Delegate triggered when we are connected to voice chat
+	 */
+	virtual FOnVoiceChatConnectedDelegate& OnVoiceChatConnected() = 0;
+
+	/**
+	 * Delegate triggered when we are disconnected from voice chat
+	 */
+	virtual FOnVoiceChatDisconnectedDelegate& OnVoiceChatDisconnected() = 0;
+
+	/**
+	 * Delegate triggered when we are reconnected to voice chat in cases where the underlying implementation disconnected
+	 * and then reconnected. This is most commonly seen when an application resumes after being suspended. If the disconnect
+	 * is unexpected, OnVoiceChatDisconnected will be called instead.
+	 */
+	virtual FOnVoiceChatReconnectedDelegate& OnVoiceChatReconnected() = 0;
+
+	/**
+	 * Login to the connected voice server
+	 *
+	 * @param PlatformId Platform user Id for the player logging in
+	 * @param PlayerName Name of the account to login to
+	 * @param Credentials Token or password
+	 * @param Delegate delegate called once login completes
+	 */
+	virtual void Login(FPlatformUserId PlatformId, const FString& PlayerName, const FString& Credentials, const FOnVoiceChatLoginCompleteDelegate& Delegate) = 0;
+
+	/**
+	 * Logout from the connected voice server
+	 *
+	 * @param Delegate delegate called once logout completes
+	 */
+	virtual void Logout(const FOnVoiceChatLogoutCompleteDelegate& Delegate) = 0;
+
+	/**
+	 * Are we logging in?
+	 *
+	 * @return true if we are logging in to the voice server
+	 */
+	virtual bool IsLoggingIn() const = 0;
+
+	/**
+	 * Are we logged in?
+	 *
+	 * @return true if we are logged in to the voice server
+	 */
+	virtual bool IsLoggedIn() const = 0;
+
+	/**
+	 * Delegate triggered when we are logged in to voice chat
+	 */
+	virtual FOnVoiceChatLoggedInDelegate& OnVoiceChatLoggedIn() = 0;
+
+	/**
+	 * Delegate triggered when we are logged out from voice chat
+	 */
+	virtual FOnVoiceChatLoggedOutDelegate& OnVoiceChatLoggedOut() = 0;
+
+	/**
+	 * Get the player name used to log in
+	 *
+	 * @return player name used to log in
+	 */
+	virtual FString GetLoggedInPlayerName() const = 0;
+
+	/**
+	 * Block communication with players
+	 *
+	 * @param PlayerNames List of players to block communication with
+	 */
+	virtual void BlockPlayers(const TArray<FString>& PlayerNames) = 0;
+
+	/**
+	 * Unblock communication with players
+	 *
+	 * @param PlayerNames List of players to unblock communication with
+	 */
+	virtual void UnblockPlayers(const TArray<FString>& PlayerNames) = 0;
+
+	/**
+	 * Join a voice channel
+	 *
+	 * @param ChannelName Channel to join
+	 * @param ChannelCredentials password or token
+	 * @param ChannelType NonPositional, Positional, or Echo
+	 * @param Delegate delegate called once join completes
+	 * @param Channel3dProperties Optional parameter to set 3d/Positional channel properties
+	 */
+	virtual void JoinChannel(const FString& ChannelName, const FString& ChannelCredentials, EVoiceChatChannelType ChannelType, const FOnVoiceChatChannelJoinCompleteDelegate& Delegate, TOptional<FVoiceChatChannel3dProperties> Channel3dProperties = TOptional<FVoiceChatChannel3dProperties>()) = 0;
+
+	/**
+	 * Leave a voice channel
+	 * 
+	 * @param ChannelName Channel to leave
+	 * @param Delegate delegate called once leave completes
+	 */
+	virtual void LeaveChannel(const FString& ChannelName, const FOnVoiceChatChannelLeaveCompleteDelegate& Delegate) = 0;
+
+	/**
+	 * Delegate triggered when we join a voice channel
+	 */
+	virtual FOnVoiceChatChannelJoinedDelegate& OnVoiceChatChannelJoined() = 0;
+
+	/**
+	 * Delegate triggered when we leave a voice channel
+	 */
+	virtual FOnVoiceChatChannelExitedDelegate& OnVoiceChatChannelExited() = 0;
+
+	/**
+	 * Delegate triggered when a call is ended, providing the stats for the call.
+	 */
+	virtual FOnVoiceChatCallStatsUpdatedDelegate& OnVoiceChatCallStatsUpdated() = 0;
+
+	/**
+	 * Set the 3d position of the player
+	 *
+	 * @param ChannelName
+	 * @param SpeakerPosition
+	 * @param ListenerPosition
+	 * @param ListenerForwardDirection
+	 * @param ListenerUpDirection
+	 */
+	virtual void Set3DPosition(const FString& ChannelName, const FVector& SpeakerPosition, const FVector& ListenerPosition, const FVector& ListenerForwardDirection, const FVector& ListenerUpDirection) = 0;
+
+	/**
+	 * Get an array of channels the user is in
+	 *
+	 * @return Array of connected channel names
+	 */
+	virtual TArray<FString> GetChannels() const = 0;
+
+	/**
+	 * Get an array of players in a channel
+	 *
+	 * @param ChannelName Channel to get the list of players for
+	 * @return Array of names of players in Channel
+	 */
+	virtual TArray<FString> GetPlayersInChannel(const FString& ChannelName) const = 0;
+
+	/**
+	 * Get an array of players in a channel
+	 *
+	 * @param ChannelName Channel to get the type of
+	 * @return Channel type
+	 */
+	virtual EVoiceChatChannelType GetChannelType(const FString& ChannelName) const = 0;
+
+	/**
+	 * Delegate triggered when a player is added to a channel
+	 */
+	virtual FOnVoiceChatPlayerAddedDelegate& OnVoiceChatPlayerAdded() = 0;
+
+	/**
+	 * Delegate triggered when a player is removed from a channel
+	 */
+	virtual FOnVoiceChatPlayerRemovedDelegate& OnVoiceChatPlayerRemoved() = 0;
+
+	/**
+	 * @param PlayerName Player to get the talking state of
+	 * @return true if player is talking
+	 */
+	virtual bool IsPlayerTalking(const FString& PlayerName) const = 0;
+
+	/**
+	 * Delegate triggered when a player's talking state is updated
+	 */
+	virtual FOnVoiceChatPlayerTalkingUpdatedDelegate& OnVoiceChatPlayerTalkingUpdated() = 0;
+
+	/**
+	 * Mute or unmute a player
+	 *
+	 * @param PlayerName Player to mute
+	 * @param bMuted true if the player should be muted
+	 */
+	virtual void SetPlayerMuted(const FString& PlayerName, bool bMuted) = 0;
+
+	/**
+	 * @param PlayerName Player to get the mute state of
+	 * @return true if player is muted
+	 */
+	virtual bool IsPlayerMuted(const FString& PlayerName) const = 0;
+
+	/**
+	 * Delegate triggered when a player's muted state is updated, usually as a result of calling SetPlayerMuted
+	 */
+	virtual FOnVoiceChatPlayerMuteUpdatedDelegate& OnVoiceChatPlayerMuteUpdated() = 0;
+
+	/**
+	 * Set player volume
+	 *
+	 * @param PlayerName Player to adjust the volume for
+	 * @param Volume Value between 0.0 and 1.0
+	 */
+	virtual void SetPlayerVolume(const FString& PlayerName, float Volume) = 0;
+
+	/**
+	 * @param PlayerName Player to get the volume of
+	 * @return volume for player
+	 */
+	virtual float GetPlayerVolume(const FString& PlayerName) const = 0;
+
+	/**
+	 * Delegate triggered when a player's volume is updated, usually as a result of calling SetPlayerVolume
+	 */
+	virtual FOnVoiceChatPlayerVolumeUpdatedDelegate& OnVoiceChatPlayerVolumeUpdated() = 0;
+
+	/**
+	 * Transmit to all channels that you are currently in
+	 */
+	virtual void TransmitToAllChannels() = 0;
+
+	/**
+	 * Stop transmitting to all channels that you are currently in. This does not apply to channels you will join.
+	 */
+	virtual void TransmitToNoChannels() = 0;
+
+	/**
+	 * Transmit to a specific channel
+	 *
+	 * @param ChannelName Channel to transmit to
+	 */
+	virtual void TransmitToSpecificChannel(const FString& ChannelName) = 0;
+
+	/**
+	 * Get the current transmit mode
+	 *
+	 * @return Transmit mode. If it is EVoiceChatTransmitMode::Channel, the channel can be retrieved using GetTransmitChannel()
+	 */
+	virtual EVoiceChatTransmitMode GetTransmitMode() const = 0;
+
+	/**
+	 * Get the specific channel we are transmitting to
+	 *
+	 * @return Channel we are transmitting to. Will be empty if GetTransmissionMode() returned a value other than EVoiceChatTransmitMode::Channel
+	 */
+	virtual FString GetTransmitChannel() const = 0;
+
+	/*
+	 * Start recording
+	 *
+	 * @param Delegate This will be called from audio thread.
+	 *
+	 * @return Delegate handle to pass to stop recording
+	 */
+	virtual FDelegateHandle StartRecording(const FOnVoiceChatRecordSamplesAvailableDelegate::FDelegate& Delegate) = 0;
+
+	/*
+	 * Stop recording
+	 *
+	 * @param Handle Delegate handle returned from a StartRecording call
+	 */
+	virtual void StopRecording(FDelegateHandle Handle) = 0;
+
+	/**
+	 * Delegates called from audio thread.
+	 */
+	virtual FDelegateHandle RegisterOnVoiceChatAfterCaptureAudioReadDelegate(const FOnVoiceChatAfterCaptureAudioReadDelegate::FDelegate& Delegate) = 0;
+	virtual void UnregisterOnVoiceChatAfterCaptureAudioReadDelegate(FDelegateHandle Handle) = 0;
+
+	virtual FDelegateHandle RegisterOnVoiceChatBeforeCaptureAudioSentDelegate(const FOnVoiceChatBeforeCaptureAudioSentDelegate::FDelegate& Delegate) = 0;
+	virtual void UnregisterOnVoiceChatBeforeCaptureAudioSentDelegate(FDelegateHandle Handle) = 0;
+
+	virtual FDelegateHandle RegisterOnVoiceChatBeforeRecvAudioRenderedDelegate(const FOnVoiceChatBeforeRecvAudioRenderedDelegate::FDelegate& Delegate) = 0;
+	virtual void UnregisterOnVoiceChatBeforeRecvAudioRenderedDelegate(FDelegateHandle Handle) = 0;
+
+	/**
+	 * Generate Login/Join tokens. Should be used for development purposes only
+	 */
+	virtual FString InsecureGetLoginToken(const FString& PlayerName) = 0;
+	virtual FString InsecureGetJoinToken(const FString& ChannelName, EVoiceChatChannelType ChannelType, TOptional<FVoiceChatChannel3dProperties> Channel3dProperties = TOptional<FVoiceChatChannel3dProperties>()) = 0;
+
+private:
+	static FName GetModularFeatureName()
+	{
+		static FName FeatureName = FName(TEXT("VoiceChat"));
+		return FeatureName;
+	}
+};
