Expose HTTP retry system's throttle time parsing

diff --git a/Engine/Source/Runtime/Online/HTTP/Private/HttpRetrySystem.cpp b/Engine/Source/Runtime/Online/HTTP/Private/HttpRetrySystem.cpp
--- a/Engine/Source/Runtime/Online/HTTP/Private/HttpRetrySystem.cpp
+++ b/Engine/Source/Runtime/Online/HTTP/Private/HttpRetrySystem.cpp
@@ -9,6 +9,53 @@
 #include "HttpManager.h"
 #include "Stats/Stats.h"
 
+namespace FHttpRetrySystem
+{
+	TOptional<double> ReadThrottledTimeFromResponseInSeconds(FHttpResponsePtr Response)
+	{
+		TOptional<double> LockoutPeriod;
+		// Check if there was a Retry-After header
+		if (Response.IsValid())
+		{
+			int32 ResponseCode = Response->GetResponseCode();
+			if (ResponseCode == EHttpResponseCodes::TooManyRequests || ResponseCode == EHttpResponseCodes::ServiceUnavail)
+			{
+				FString RetryAfter = Response->GetHeader(TEXT("Retry-After"));
+				if (!RetryAfter.IsEmpty())
+				{
+					if (RetryAfter.IsNumeric())
+					{
+						// seconds
+						LockoutPeriod.Emplace(FCString::Atof(*RetryAfter));
+					}
+					else
+					{
+						// http date
+						FDateTime UTCServerTime;
+						if (FDateTime::ParseHttpDate(RetryAfter, UTCServerTime))
+						{
+							const FDateTime UTCNow = FDateTime::UtcNow();
+							LockoutPeriod.Emplace((UTCServerTime - UTCNow).GetTotalSeconds());
+						}
+					}
+				}
+				else
+				{
+					FString RateLimitReset = Response->GetHeader(TEXT("X-Rate-Limit-Reset"));
+					if (!RateLimitReset.IsEmpty())
+					{
+						// UTC seconds
+						const FDateTime UTCServerTime = FDateTime::FromUnixTimestamp(FCString::Atoi64(*RateLimitReset));
+						const FDateTime UTCNow = FDateTime::UtcNow();
+						LockoutPeriod.Emplace((UTCServerTime - UTCNow).GetTotalSeconds());
+					}
+				}
+			}
+		}
+		return LockoutPeriod;
+	}
+}
+
 FHttpRetrySystem::FRequest::FRequest(
 	FManager& InManager,
 	const TSharedRef<IHttpRequest>& HttpRequest, 
@@ -223,45 +270,10 @@ bool FHttpRetrySystem::FManager::HasTimedOut(const FHttpRetryRequestEntry& HttpR
 float FHttpRetrySystem::FManager::GetLockoutPeriodSeconds(const FHttpRetryRequestEntry& HttpRetryRequestEntry)
 {
 	float LockoutPeriod = 0.0f;
-
-	// Check if there was a Retry-After header
-	FHttpResponsePtr Response = HttpRetryRequestEntry.Request->GetResponse();
-	if (Response.IsValid())
+	TOptional<double> ResponseLockoutPeriod = FHttpRetrySystem::ReadThrottledTimeFromResponseInSeconds(HttpRetryRequestEntry.Request->GetResponse());
+	if (ResponseLockoutPeriod.IsSet())
 	{
-		int32 ResponseCode = Response->GetResponseCode();
-		if (ResponseCode == EHttpResponseCodes::TooManyRequests || ResponseCode == EHttpResponseCodes::ServiceUnavail)
-		{
-			FString RetryAfter = Response->GetHeader(TEXT("Retry-After"));
-			if (!RetryAfter.IsEmpty())
-			{
-				if (RetryAfter.IsNumeric())
-				{
-					// seconds
-					LockoutPeriod = FCString::Atof(*RetryAfter);
-				}
-				else
-				{
-					// http date
-					FDateTime UTCServerTime;
-					if (FDateTime::ParseHttpDate(RetryAfter, UTCServerTime))
-					{
-						const FDateTime UTCNow = FDateTime::UtcNow();
-						LockoutPeriod = (UTCServerTime - UTCNow).GetTotalSeconds();
-					}
-				}
-			}
-			else
-			{
-				FString RateLimitReset = Response->GetHeader(TEXT("X-Rate-Limit-Reset"));
-				if (!RateLimitReset.IsEmpty())
-				{
-					// UTC seconds
-					const FDateTime UTCServerTime = FDateTime::FromUnixTimestamp(FCString::Atoi64(*RateLimitReset));
-					const FDateTime UTCNow = FDateTime::UtcNow();
-					LockoutPeriod = (UTCServerTime - UTCNow).GetTotalSeconds();
-				}
-			}
-		}
+		LockoutPeriod = static_cast<float>(ResponseLockoutPeriod.GetValue());
 	}
 
 	if (HttpRetryRequestEntry.CurrentRetryCount >= 1)
diff --git a/Engine/Source/Runtime/Online/HTTP/Public/HttpRetrySystem.h b/Engine/Source/Runtime/Online/HTTP/Public/HttpRetrySystem.h
--- a/Engine/Source/Runtime/Online/HTTP/Public/HttpRetrySystem.h
+++ b/Engine/Source/Runtime/Online/HTTP/Public/HttpRetrySystem.h
@@ -47,6 +47,13 @@ namespace FHttpRetrySystem
 		TAtomic<int32> ActiveIndex;
 	};
 	typedef TSharedPtr<FRetryDomains, ESPMode::ThreadSafe> FRetryDomainsPtr;
+
+	/**
+	 * Read the number of seconds a HTTP request is throttled for from the response
+	 * @param Response the HTTP response to read the value from
+	 * @return If found, the number of seconds the request is rate limited for.  If not found, an unset TOptional
+	 */
+	TOptional<double> HTTP_API ReadThrottledTimeFromResponseInSeconds(FHttpResponsePtr Response);
 };
 
 /**
