#stomp Fixed missing const

diff --git a/Engine/Source/Runtime/Online/Stomp/Private/StompClient.cpp b/Engine/Source/Runtime/Online/Stomp/Private/StompClient.cpp
--- a/Engine/Source/Runtime/Online/Stomp/Private/StompClient.cpp
+++ b/Engine/Source/Runtime/Online/Stomp/Private/StompClient.cpp
@@ -77,7 +77,7 @@ void FStompClient::HandleDisconnectCompleted(bool bIsSuccess, const FString& Err
 	WebSocket->Close();
 }
 
-bool FStompClient::IsConnected()
+bool FStompClient::IsConnected() const
 {
 	return bIsConnected;
 }
diff --git a/Engine/Source/Runtime/Online/Stomp/Private/StompClient.h b/Engine/Source/Runtime/Online/Stomp/Private/StompClient.h
--- a/Engine/Source/Runtime/Online/Stomp/Private/StompClient.h
+++ b/Engine/Source/Runtime/Online/Stomp/Private/StompClient.h
@@ -45,7 +45,7 @@ public:
 	/**
 	 * Inquire if this instance is connected to a server.
 	 */
-	virtual bool IsConnected() override;
+	virtual bool IsConnected() const override;
 
 	/**
 	 * Subscribe to an event
diff --git a/Engine/Source/Runtime/Online/Stomp/Public/IStompClient.h b/Engine/Source/Runtime/Online/Stomp/Public/IStompClient.h
--- a/Engine/Source/Runtime/Online/Stomp/Public/IStompClient.h
+++ b/Engine/Source/Runtime/Online/Stomp/Public/IStompClient.h
@@ -37,7 +37,7 @@ public:
 	/**
 	 * Inquire if this instance is connected to a server.
 	 */
-	virtual bool IsConnected() = 0;
+	virtual bool IsConnected() const = 0;
 
 	/**
 	 * Subscribe to an event
