[HttpServer] Per-request logs in HttpConnection are log level Verbose

diff --git a/Engine/Source/Runtime/Online/HTTPServer/Private/HttpConnection.cpp b/Engine/Source/Runtime/Online/HTTPServer/Private/HttpConnection.cpp
--- a/Engine/Source/Runtime/Online/HTTPServer/Private/HttpConnection.cpp
+++ b/Engine/Source/Runtime/Online/HTTPServer/Private/HttpConnection.cpp
@@ -158,7 +158,7 @@ void FHttpConnection::CompleteRead(const TSharedPtr<FHttpServerRequest>& Request
 		TSharedPtr<FHttpConnection> SharedThisPtr = WeakThisPtr.Pin();
 		if (SharedThisPtr.IsValid())
 		{
-			UE_LOG(LogHttpConnection, Log,
+			UE_LOG(LogHttpConnection, Verbose,
 				TEXT("EndProcessRequest\t [%d][%u]-%u : %u"), 
 				OriginPortCapture, ConnectionIdCapture, LastRequestNumberCapture, Response->Code);
 
@@ -178,7 +178,7 @@ void FHttpConnection::ProcessRequest(const TSharedPtr<FHttpServerRequest>& Reque
 {
 	TransferState(EHttpConnectionState::Reading, EHttpConnectionState::AwaitingProcessing);
 
-	UE_LOG(LogHttpConnection, Log,
+	UE_LOG(LogHttpConnection, Verbose,
 		TEXT("BeginProcessRequest\t [%d][%u]-%u : %s"), 
 		OriginPort, ConnectionId, LastRequestNumber, *Request->RelativePath.GetPath());
 
