libCURL crash fix on non-existing URL requests

[RN]
#5571

diff --git a/Engine/Source/Runtime/Online/HTTP/Private/Curl/CurlHttp.cpp b/Engine/Source/Runtime/Online/HTTP/Private/Curl/CurlHttp.cpp
--- a/Engine/Source/Runtime/Online/HTTP/Private/Curl/CurlHttp.cpp
+++ b/Engine/Source/Runtime/Online/HTTP/Private/Curl/CurlHttp.cpp
@@ -1154,11 +1154,11 @@ void FCurlHttpRequest::FinishedRequest()
 				CompletionStatus = EHttpRequestStatus::Failed_ConnectionError;
 			}
 		}
-		// No response since connection failed
-		Response = NULL;
-
 		// Call delegate with failure
-		OnProcessRequestComplete().ExecuteIfBound(SharedThis(this),NULL,false);
+		OnProcessRequestComplete().ExecuteIfBound(SharedThis(this), Response, false);
+
+		//Delegate needs to know about the errors -- so nuke Response (since connection failed) afterwards...
+		Response = NULL;
 	}
 }
 
