Make MaxActiveDownloads configurable

diff --git a/Engine/Config/BaseEngine.ini b/Engine/Config/BaseEngine.ini
--- a/Engine/Config/BaseEngine.ini
+++ b/Engine/Config/BaseEngine.ini
@@ -29,6 +29,7 @@ HttpSendTimeout=30
 [BackgroundHttp]
 ; How many days we will wait to delete any found BackgroundHTTP Temp Files that haven't been claimed. Default is 259200 seconds (3 days). -1 to disable.
 BackgroundHttp.TempFileTimeOutSeconds=259200
+BackgroundHttp.MaxActiveDownloads=4
 
 [WebSockets.LibWebSockets]
 ThreadStackSize=131072
diff --git a/Engine/Source/Runtime/Online/BackgroundHTTP/Private/BackgroundHttpManagerImpl.cpp b/Engine/Source/Runtime/Online/BackgroundHTTP/Private/BackgroundHttpManagerImpl.cpp
--- a/Engine/Source/Runtime/Online/BackgroundHTTP/Private/BackgroundHttpManagerImpl.cpp
+++ b/Engine/Source/Runtime/Online/BackgroundHTTP/Private/BackgroundHttpManagerImpl.cpp
@@ -17,6 +17,7 @@ FBackgroundHttpManagerImpl::FBackgroundHttpManagerImpl()
 	, ActiveRequests()
 	, ActiveRequestLock()
 	, NumCurrentlyActiveRequests(0)
+	, MaxActiveDownloads(4)
 {
 }
 
@@ -27,6 +28,11 @@ FBackgroundHttpManagerImpl::~FBackgroundHttpManagerImpl()
 void FBackgroundHttpManagerImpl::Initialize()
 {
 	ClearAnyTempFilesFromTimeOut();
+
+	// Can't read into an atomic int directly
+	int MaxActiveDownloadsConfig = 4;
+	ensureAlwaysMsgf(GConfig->GetInt(TEXT("BackgroundHttp"), TEXT("BackgroundHttp.MaxActiveDownloads"), MaxActiveDownloadsConfig, GEngineIni), TEXT("No value found for MaxActiveDownloads! Defaulting to 4!"));
+	MaxActiveDownloads = MaxActiveDownloadsConfig;
 }
 
 void FBackgroundHttpManagerImpl::Shutdown()
@@ -113,6 +119,16 @@ void FBackgroundHttpManagerImpl::CleanUpTemporaryFiles()
 	}
 }
 
+int FBackgroundHttpManagerImpl::GetMaxActiveDownloads() const
+{
+	return MaxActiveDownloads;
+}
+
+void FBackgroundHttpManagerImpl::SetMaxActiveDownloads(int InMaxActiveDownloads)
+{
+	MaxActiveDownloads = InMaxActiveDownloads;
+}
+
 void FBackgroundHttpManagerImpl::AddRequest(const FBackgroundHttpRequestPtr Request)
 {
 	UE_LOG(LogBackgroundHttpManager, Log, TEXT("AddRequest Called - RequestID:%s"), *Request->GetRequestID());
@@ -219,10 +235,10 @@ void FBackgroundHttpManagerImpl::ActivatePendingRequests()
 
 		if (PendingStartRequests.Num() > 0)
 		{
-			UE_LOG(LogBackgroundHttpManager, Verbose, TEXT("Populating Requests to Start from PendingStartRequests - PlatformMaxActiveDownloads:%d | NumCurrentlyActiveRequests:%d | NumPendingStartRequests:%d"), FPlatformBackgroundHttp::GetPlatformMaxActiveDownloads(), NumCurrentlyActiveRequests, PendingStartRequests.Num());
+			UE_LOG(LogBackgroundHttpManager, Verbose, TEXT("Populating Requests to Start from PendingStartRequests - MaxActiveDownloads:%d | NumCurrentlyActiveRequests:%d | NumPendingStartRequests:%d"), MaxActiveDownloads.Load(), NumCurrentlyActiveRequests, PendingStartRequests.Num());
 
 			//See how many more requests we can process and only do anything if we can still process more
-			const int NumRequestsWeCanProcess = (FPlatformBackgroundHttp::GetPlatformMaxActiveDownloads() - NumCurrentlyActiveRequests);
+			const int NumRequestsWeCanProcess = (MaxActiveDownloads - NumCurrentlyActiveRequests);
 			if (NumRequestsWeCanProcess > 0)
 			{
 				TArray<FBackgroundHttpRequestPtr> RemainingRequests;
diff --git a/Engine/Source/Runtime/Online/BackgroundHTTP/Private/GenericPlatform/GenericPlatformBackgroundHttp.cpp b/Engine/Source/Runtime/Online/BackgroundHTTP/Private/GenericPlatform/GenericPlatformBackgroundHttp.cpp
--- a/Engine/Source/Runtime/Online/BackgroundHTTP/Private/GenericPlatform/GenericPlatformBackgroundHttp.cpp
+++ b/Engine/Source/Runtime/Online/BackgroundHTTP/Private/GenericPlatform/GenericPlatformBackgroundHttp.cpp
@@ -41,10 +41,4 @@ const FString& FGenericPlatformBackgroundHttp::GetTemporaryRootPath()
 {
 	static FString BackgroundHttpDir = FPaths::Combine(FPlatformMisc::GamePersistentDownloadDir(), TEXT("BackgroundHttpTemp"));
 	return BackgroundHttpDir;
-}
-
-int FGenericPlatformBackgroundHttp::GetPlatformMaxActiveDownloads()
-{
-	//In general, lets just use 4 active downloads at once
-	return 4;
-}
+}
\ No newline at end of file
diff --git a/Engine/Source/Runtime/Online/BackgroundHTTP/Private/IOS/ApplePlatformBackgroundHttp.cpp b/Engine/Source/Runtime/Online/BackgroundHTTP/Private/IOS/ApplePlatformBackgroundHttp.cpp
--- a/Engine/Source/Runtime/Online/BackgroundHTTP/Private/IOS/ApplePlatformBackgroundHttp.cpp
+++ b/Engine/Source/Runtime/Online/BackgroundHTTP/Private/IOS/ApplePlatformBackgroundHttp.cpp
@@ -43,8 +43,3 @@ const FString& FApplePlatformBackgroundHttp::GetTemporaryRootPath()
 {
     return FBackgroundURLSessionHandler::GetBackgroundSessionWorkingDirectoryPath();
 }
-
-int FApplePlatformBackgroundHttp::GetPlatformMaxActiveDownloads()
-{
-	return 4;
-}
diff --git a/Engine/Source/Runtime/Online/BackgroundHTTP/Private/IOS/ApplePlatformBackgroundHttpManager.cpp b/Engine/Source/Runtime/Online/BackgroundHTTP/Private/IOS/ApplePlatformBackgroundHttpManager.cpp
--- a/Engine/Source/Runtime/Online/BackgroundHTTP/Private/IOS/ApplePlatformBackgroundHttpManager.cpp
+++ b/Engine/Source/Runtime/Online/BackgroundHTTP/Private/IOS/ApplePlatformBackgroundHttpManager.cpp
@@ -31,6 +31,7 @@ FApplePlatformBackgroundHttpManager::FApplePlatformBackgroundHttpManager()
     , bIsIteratingThroughSessionTasks(false)
     , RequestsPendingRemove()
     , NumCurrentlyActiveTasks(0)
+	, MaxNumActualTasks(MaxActiveDownloads)
 {
 }
 
@@ -233,6 +234,14 @@ void FApplePlatformBackgroundHttpManager::RemoveSessionTasksForRequest(FAppleBac
     Request->CancelActiveTask();
 }
 
+void FApplePlatformBackgroundHttpManager::SetMaxActiveDownloads(int InMaxActiveDownloads)
+{
+	FBackgroundHttpManagerImpl::SetMaxActiveDownloads(InMaxActiveDownloads);
+
+	// It's possible we have more than the new maximum active right now, so we gracefully reduce MaxNumActualTasks down to MaxActiveDownloads as the extra tasks finish.
+	MaxNumActualTasks = FMath::Max(MaxActiveDownloads.Load(), FPlatformAtomics::AtomicRead(&NumCurrentlyActiveTasks));
+}
+
 bool FApplePlatformBackgroundHttpManager::AssociateWithAnyExistingUnAssociatedTasks(const FBackgroundHttpRequestPtr Request)
 {
     bool bDidAssociateWithUnAssociatedTask = false;
@@ -491,8 +500,11 @@ void FApplePlatformBackgroundHttpManager::FinishRequest(FAppleBackgroundHttpRequ
                 {
                     int NumActualTasks = FPlatformAtomics::InterlockedDecrement(&NumCurrentlyActiveTasks);
 
+					// Handle the case that SetMaxActiveDownloads reduced the MaxActiveDownloads while we had more than the new maximum in progress.
+					MaxNumActualTasks = FMath::Max(MaxNumActualTasks - 1, MaxActiveDownloads.Load());
+
                     //Sanity check that our data is valid. Shouldn't ever trip if everything is working as intended.
-                    const bool bNumActualTasksIsValid = ((NumActualTasks >= 0) && (NumActualTasks <= FPlatformBackgroundHttp::GetPlatformMaxActiveDownloads()));
+                    const bool bNumActualTasksIsValid = ((NumActualTasks >= 0) && (NumActualTasks <= MaxNumActualTasks));
                     
                     UE_LOG(LogBackgroundHttpManager,Display, TEXT("Finishing Request lowering Task Count: %d"), NumActualTasks);
                     
@@ -769,7 +781,7 @@ void FApplePlatformBackgroundHttpManager::TickTasks(float DeltaTime)
              {
                  //Check to make sure we have room for more tasks to be active first
                  int CurrentCount = FPlatformAtomics::AtomicRead(&NumCurrentlyActiveTasks);
-                 if (CurrentCount < FPlatformBackgroundHttp::GetPlatformMaxActiveDownloads())
+                 if (CurrentCount < MaxActiveDownloads)
                  {
                      //Go through our tasks and try and activate as many as possible
                      for (NSURLSessionTask* task in tasks)
@@ -781,7 +793,7 @@ void FApplePlatformBackgroundHttpManager::TickTasks(float DeltaTime)
                              int NewRequestCount = FPlatformAtomics::InterlockedIncrement(&NumCurrentlyActiveTasks);
                              UE_LOG(LogBackgroundHttpManager,Verbose, TEXT("Incrementing Task Count: %d"), NewRequestCount);
                              
-                             if (NewRequestCount <= FPlatformBackgroundHttp::GetPlatformMaxActiveDownloads())
+                             if (NewRequestCount <= MaxActiveDownloads)
                              {
                                  FString TaskURL = [[[task currentRequest] URL] absoluteString];
                                  int TaskIdentifier = (int)[task taskIdentifier];
@@ -827,7 +839,7 @@ void FApplePlatformBackgroundHttpManager::TickTasks(float DeltaTime)
                              }
                              
                              //We now have enough requests queued up that we can stop looking for more
-                             if (NewRequestCount >= FPlatformBackgroundHttp::GetPlatformMaxActiveDownloads())
+                             if (NewRequestCount >= MaxActiveDownloads)
                              {
                                  break;
                              }
diff --git a/Engine/Source/Runtime/Online/BackgroundHTTP/Public/BackgroundHttpManagerImpl.h b/Engine/Source/Runtime/Online/BackgroundHTTP/Public/BackgroundHttpManagerImpl.h
--- a/Engine/Source/Runtime/Online/BackgroundHTTP/Public/BackgroundHttpManagerImpl.h
+++ b/Engine/Source/Runtime/Online/BackgroundHTTP/Public/BackgroundHttpManagerImpl.h
@@ -28,6 +28,9 @@ public:
 
 	virtual void CleanUpTemporaryFiles() override;
 
+	virtual int GetMaxActiveDownloads() const override;
+	virtual void SetMaxActiveDownloads(int MaxActiveDownloads) override;
+
 	//FTickerObjectBase implementation
 	virtual bool Tick(float DeltaTime) override;
 
@@ -50,4 +53,5 @@ protected:
 
 	/** Count of how many requests we have active **/
 	volatile int NumCurrentlyActiveRequests;
+	TAtomic<int> MaxActiveDownloads;
 };
diff --git a/Engine/Source/Runtime/Online/BackgroundHTTP/Public/GenericPlatform/GenericPlatformBackgroundHttp.h b/Engine/Source/Runtime/Online/BackgroundHTTP/Public/GenericPlatform/GenericPlatformBackgroundHttp.h
--- a/Engine/Source/Runtime/Online/BackgroundHTTP/Public/GenericPlatform/GenericPlatformBackgroundHttp.h
+++ b/Engine/Source/Runtime/Online/BackgroundHTTP/Public/GenericPlatform/GenericPlatformBackgroundHttp.h
@@ -64,11 +64,4 @@ public:
 	* @return FString File path that is the root path of our temp files for background http work on this platform.
 	*/
 	static const FString& GetTemporaryRootPath();
-
-	/**
-	* Function that returns how many active BackgroundHttpRequests we should have actively downloading at once. 
-	*
-	* @return int Number of downloads we should have active at once.
-	*/
-	static int GetPlatformMaxActiveDownloads();
 };
diff --git a/Engine/Source/Runtime/Online/BackgroundHTTP/Public/IOS/ApplePlatformBackgroundHttp.h b/Engine/Source/Runtime/Online/BackgroundHTTP/Public/IOS/ApplePlatformBackgroundHttp.h
--- a/Engine/Source/Runtime/Online/BackgroundHTTP/Public/IOS/ApplePlatformBackgroundHttp.h
+++ b/Engine/Source/Runtime/Online/BackgroundHTTP/Public/IOS/ApplePlatformBackgroundHttp.h
@@ -61,13 +61,6 @@ public:
      * @return FString File path that is the root path of our temp files for background http work on this platform.
      */
     static const FString& GetTemporaryRootPath();
-    
-	/**
-	* Function that returns how many active BackgroundHttpRequests we should have actively downloading at once.
-	*
-	* @return int Number of downloads we should have active at once.
-	*/
-	static int GetPlatformMaxActiveDownloads();
 };
 
 //Setup Platform Implementation calls to point here at the Apple Implementation
diff --git a/Engine/Source/Runtime/Online/BackgroundHTTP/Public/IOS/ApplePlatformBackgroundHttpManager.h b/Engine/Source/Runtime/Online/BackgroundHTTP/Public/IOS/ApplePlatformBackgroundHttpManager.h
--- a/Engine/Source/Runtime/Online/BackgroundHTTP/Public/IOS/ApplePlatformBackgroundHttpManager.h
+++ b/Engine/Source/Runtime/Online/BackgroundHTTP/Public/IOS/ApplePlatformBackgroundHttpManager.h
@@ -22,6 +22,7 @@ public:
 	virtual void AddRequest(const FBackgroundHttpRequestPtr Request) override;
 	virtual void RemoveRequest(const FBackgroundHttpRequestPtr Request) override;
 	virtual bool IsGenericImplementation() const override { return false;  }
+	virtual void SetMaxActiveDownloads(int MaxActiveDownloads) override;
 
 	/**
 	* Constructor
@@ -116,4 +117,5 @@ private:
     
     /** On iOS we need to track how many Tasks we have active. This is to replace the default implementations NumCurrentlyActiveRequests **/
     volatile int NumCurrentlyActiveTasks;
+	int MaxNumActualTasks;
 };
diff --git a/Engine/Source/Runtime/Online/BackgroundHTTP/Public/Interfaces/IBackgroundHttpManager.h b/Engine/Source/Runtime/Online/BackgroundHTTP/Public/Interfaces/IBackgroundHttpManager.h
--- a/Engine/Source/Runtime/Online/BackgroundHTTP/Public/Interfaces/IBackgroundHttpManager.h
+++ b/Engine/Source/Runtime/Online/BackgroundHTTP/Public/Interfaces/IBackgroundHttpManager.h
@@ -47,6 +47,20 @@ public:
 	/** Returns whether or not this is a platform specific implementation */
 	virtual bool IsGenericImplementation() const = 0;
 
+	/**
+	* Function that returns how many active BackgroundHttpRequests we should have actively downloading at once.
+	*
+	* @return int Number of downloads we should have active at once.
+	*/
+	virtual int GetMaxActiveDownloads() const = 0;
+
+	/**
+	* Function that sets how many active BackgroundHttpRequests we should have actively downloading at once.
+	*
+	* @param MaxActiveDownloads the maximum number of downloads that should be active at once
+	*/
+	virtual void SetMaxActiveDownloads(int MaxActiveDownloads) = 0;
+
 protected:
 	/**
 	* Designed to be called internally by AddRequest to associate our incoming request with any 
