Document BlockList lifespan on VoiceChat interface

diff --git a/Engine/Source/Runtime/Online/Voice/Public/Interfaces/VoiceChat.h b/Engine/Source/Runtime/Online/Voice/Public/Interfaces/VoiceChat.h
--- a/Engine/Source/Runtime/Online/Voice/Public/Interfaces/VoiceChat.h
+++ b/Engine/Source/Runtime/Online/Voice/Public/Interfaces/VoiceChat.h
@@ -357,6 +357,11 @@ public:
 	/**
 	 * Block communication with players
 	 *
+	 * Adds the PlayerNames to the block list, muting them in all current and subsequently joined channels.
+	 * A block for a player persists until either:
+	 *		- The local user logs out (e.g. by calling Logout/Disconnect/Uninitialize, or receiving equivalent events)
+	 *		- A subsequent call to UnblockPlayers is made with that PlayerName.
+	 *
 	 * @param PlayerNames List of players to block communication with
 	 */
 	virtual void BlockPlayers(const TArray<FString>& PlayerNames) = 0;
