Initialize SSL module when loading AppleHTTP to ensure correct unload sequence on shutdown

diff --git a/Engine/Source/Runtime/Online/HTTP/Private/Apple/ApplePlatformHttp.cpp b/Engine/Source/Runtime/Online/HTTP/Private/Apple/ApplePlatformHttp.cpp
--- a/Engine/Source/Runtime/Online/HTTP/Private/Apple/ApplePlatformHttp.cpp
+++ b/Engine/Source/Runtime/Online/HTTP/Private/Apple/ApplePlatformHttp.cpp
@@ -3,9 +3,17 @@
 #include "Apple/ApplePlatformHttp.h"
 #include "AppleHTTP.h"
 
+#if WITH_SSL
+#include "Ssl.h"
+#endif
 
 void FApplePlatformHttp::Init()
 {
+#if WITH_SSL
+	// Load SSL module during HTTP module's StatupModule() to make sure module manager figures out the dependencies correctly
+	// and doesn't unload SSL before unloading HTTP module at exit
+	FSslModule::Get();
+#endif
 }
 
 
