Fixed issue where UnAssociatedTasks could associate with an already completed task and create a timing issue where a BackgroundHttp request would do no work as it was pending a cancel that wasn't going to happen. No longer associates with tasks that are already completed/completing
#tests iOS Client

diff --git a/Engine/Source/Runtime/Online/BackgroundHTTP/Private/IOS/ApplePlatformBackgroundHttpManager.cpp b/Engine/Source/Runtime/Online/BackgroundHTTP/Private/IOS/ApplePlatformBackgroundHttpManager.cpp
--- a/Engine/Source/Runtime/Online/BackgroundHTTP/Private/IOS/ApplePlatformBackgroundHttpManager.cpp
+++ b/Engine/Source/Runtime/Online/BackgroundHTTP/Private/IOS/ApplePlatformBackgroundHttpManager.cpp
@@ -277,25 +277,32 @@ bool FApplePlatformBackgroundHttpManager::CheckForExistingUnAssociatedTask(const
 			NSURLSessionTask* FoundTask = [UnAssociatedTasks valueForKey:URL.GetNSString()];
 			if (nullptr != FoundTask)
 			{
-                UE_LOG(LogBackgroundHttpManager, Display, TEXT("Existing UnAssociateTask found for Request! Attempting to Associate! RequestDebugID:%s"), *(Request->GetRequestDebugID()));
-                    
-                //Associate with task so that our Request takes over ownership of this task so we can remove it from our UnAssociated Tasks list without it getting GC'd
-				if (Request->AssociateWithTask(FoundTask))
+				if ([FoundTask state] != NSURLSessionTaskStateCompleted && [FoundTask state] != NSURLSessionTaskStateCanceling)
 				{
-					//Always set our bWasTaskStartedInBG flag on our Request as true in the UnAssociated case as we don't know when it was really started
-					FPlatformAtomics::InterlockedExchange(&(Request->bWasTaskStartedInBG), true);
+					UE_LOG(LogBackgroundHttpManager, Display, TEXT("Existing UnAssociateTask found for Request! Attempting to Associate! RequestDebugID:%s"), *(Request->GetRequestDebugID()));
+					
+					//Associate with task so that our Request takes over ownership of this task so we can remove it from our UnAssociated Tasks list without it getting GC'd
+					if (Request->AssociateWithTask(FoundTask))
+					{
+						//Always set our bWasTaskStartedInBG flag on our Request as true in the UnAssociated case as we don't know when it was really started
+						FPlatformAtomics::InterlockedExchange(&(Request->bWasTaskStartedInBG), true);
 
-					//Suspend task in case it was running so that we can adhere to our desired platform max tasks
-					[FoundTask suspend];
+						//Suspend task in case it was running so that we can adhere to our desired platform max tasks
+						[FoundTask suspend];
 
-					bDidFindExistingTask = true;
-					break;
+						bDidFindExistingTask = true;
+						break;
+					}
+					else
+					{
+						UE_LOG(LogBackgroundHttpManager, Display, TEXT("UnAssociatedTask for request found, but failed to Associate with Task! -- RequestDebugID:%s | URL:%s"), *(Request->GetRequestDebugID()), *URL);
+					}
 				}
 				else
 				{
-					UE_LOG(LogBackgroundHttpManager, Display, TEXT("UnAssociatedTask for request found, but failed to Associate with Task! -- RequestDebugID:%s | URL:%s"), *(Request->GetRequestDebugID()), *URL);
+					UE_LOG(LogBackgroundHttpManager, Display, TEXT("UnAssociatedTask for request found, BUT NOT USING as it was cancelling or completed already! -- RequestDebugID:%s | URL:%s"), *(Request->GetRequestDebugID()), *URL);
 				}
-
+				
 				//Still want to remove UnAssociatedTask even though we didn't use it as something else can now be downloading this data and we do not want duplicates
 				[UnAssociatedTasks removeObjectForKey : URL.GetNSString()];
 			}
