Xmpp crash fix from Dev-Online

diff --git a/Engine/Source/Runtime/Online/XMPP/Private/XmppStrophe/XmppConnectionStrophe.cpp b/Engine/Source/Runtime/Online/XMPP/Private/XmppStrophe/XmppConnectionStrophe.cpp
--- a/Engine/Source/Runtime/Online/XMPP/Private/XmppStrophe/XmppConnectionStrophe.cpp
+++ b/Engine/Source/Runtime/Online/XMPP/Private/XmppStrophe/XmppConnectionStrophe.cpp
@@ -27,6 +27,20 @@ FXmppConnectionStrophe::FXmppConnectionStrophe()
 	PubSubStrophe = MakeShared<FXmppPubSubStrophe, ESPMode::ThreadSafe>(*this);
 }
 
+FXmppConnectionStrophe::~FXmppConnectionStrophe()
+{
+	// Shutdown our connections before we're fully destructed; strophe calls back into us on its disconnect event
+	if (StropheThread.IsValid())
+	{
+		StopXmppThread();
+	}
+
+	if (WebsocketConnection.IsValid())
+	{
+		WebsocketConnection.Reset();
+	}
+}
+
 void FXmppConnectionStrophe::SetServer(const FXmppServer& NewServerConfiguration)
 {
 	ServerConfiguration = NewServerConfiguration;
diff --git a/Engine/Source/Runtime/Online/XMPP/Private/XmppStrophe/XmppConnectionStrophe.h b/Engine/Source/Runtime/Online/XMPP/Private/XmppStrophe/XmppConnectionStrophe.h
--- a/Engine/Source/Runtime/Online/XMPP/Private/XmppStrophe/XmppConnectionStrophe.h
+++ b/Engine/Source/Runtime/Online/XMPP/Private/XmppStrophe/XmppConnectionStrophe.h
@@ -31,7 +31,7 @@ class FXmppConnectionStrophe
 public:
 	// FXmppConnectionStrophe
 	explicit FXmppConnectionStrophe();
-	virtual ~FXmppConnectionStrophe() = default;
+	virtual ~FXmppConnectionStrophe();
 	FXmppConnectionStrophe(const FXmppConnectionStrophe& Other) = delete;
 	FXmppConnectionStrophe(FXmppConnectionStrophe&& Other) = delete;
 	FXmppConnectionStrophe& operator=(const FXmppConnectionStrophe& Other) = delete;
