[HttpServer] HttpConnection reads utilize a 256kb buffer (was 2kb)

diff --git a/Engine/Source/Runtime/Online/HTTPServer/Private/HttpConnectionRequestReadContext.cpp b/Engine/Source/Runtime/Online/HTTPServer/Private/HttpConnectionRequestReadContext.cpp
--- a/Engine/Source/Runtime/Online/HTTPServer/Private/HttpConnectionRequestReadContext.cpp
+++ b/Engine/Source/Runtime/Online/HTTPServer/Private/HttpConnectionRequestReadContext.cpp
@@ -28,7 +28,8 @@ EHttpConnectionContextState FHttpConnectionRequestReadContext::ReadStream(float
 {
 	ElapsedIdleTime += DeltaTime;
 
-	uint8 ByteBuffer[2 * 1024] = { 0 };
+	const uint32 ByteBufferSize = 1024 * 256; // 256k read buffer (kernel default)
+	uint8 ByteBuffer[ByteBufferSize] = { 0 };
 	int32 BytesRead = 0;
 	if (!Socket->Recv(ByteBuffer, sizeof(ByteBuffer) - 1, BytesRead, ESocketReceiveFlags::None))
 	{
