[HttpServer] Fix SCA "/analyze: stacksize 81940" warning
+ Read buffer is now 64k

diff --git a/Engine/Source/Runtime/Online/HTTPServer/Private/HttpConnectionRequestReadContext.cpp b/Engine/Source/Runtime/Online/HTTPServer/Private/HttpConnectionRequestReadContext.cpp
--- a/Engine/Source/Runtime/Online/HTTPServer/Private/HttpConnectionRequestReadContext.cpp
+++ b/Engine/Source/Runtime/Online/HTTPServer/Private/HttpConnectionRequestReadContext.cpp
@@ -28,7 +28,7 @@ EHttpConnectionContextState FHttpConnectionRequestReadContext::ReadStream(float
 {
 	ElapsedIdleTime += DeltaTime;
 
-	const uint32 ByteBufferSize = 1024 * 256; // 256k read buffer (kernel default)
+	const uint32 ByteBufferSize = 1024 * 64; // 64k - safe value to remain under SCA limit (/analyze:stacksize 81940)
 	uint8 ByteBuffer[ByteBufferSize] = { 0 };
 	int32 BytesRead = 0;
 	if (!Socket->Recv(ByteBuffer, sizeof(ByteBuffer) - 1, BytesRead, ESocketReceiveFlags::None))
