Fix some Linux NonUnity build failures by including missing headers.

diff --git a/Engine/Source/Runtime/Online/XMPP/Private/XmppJingle/XmppMessagesJingle.cpp b/Engine/Source/Runtime/Online/XMPP/Private/XmppJingle/XmppMessagesJingle.cpp
--- a/Engine/Source/Runtime/Online/XMPP/Private/XmppJingle/XmppMessagesJingle.cpp
+++ b/Engine/Source/Runtime/Online/XMPP/Private/XmppJingle/XmppMessagesJingle.cpp
@@ -6,6 +6,7 @@
 #include "Serialization/JsonReader.h"
 #include "Serialization/JsonSerializer.h"
 #include "Policies/CondensedJsonPrintPolicy.h"
+#include "Containers/BackgroundableTicker.h"
 
 #if WITH_XMPP_JINGLE
 
diff --git a/Engine/Source/Runtime/Online/XMPP/Private/XmppJingle/XmppMultiUserChatJingle.cpp b/Engine/Source/Runtime/Online/XMPP/Private/XmppJingle/XmppMultiUserChatJingle.cpp
--- a/Engine/Source/Runtime/Online/XMPP/Private/XmppJingle/XmppMultiUserChatJingle.cpp
+++ b/Engine/Source/Runtime/Online/XMPP/Private/XmppJingle/XmppMultiUserChatJingle.cpp
@@ -6,7 +6,7 @@
 #include "Misc/ConfigCacheIni.h"
 #include "Logging/LogScopedVerbosityOverride.h"
 #include "XmppJingle/XmppConnectionJingle.h"
-
+#include "Containers/BackgroundableTicker.h"
 
 #if WITH_XMPP_JINGLE
 
