[HttpServer] Request query parameters are correctly parsed via '&' delimiter (BUGFIX)
+ Query parameters are now correctly split/parsed

diff --git a/Engine/Source/Runtime/Online/HTTPServer/Private/HttpConnectionRequestReadContext.cpp b/Engine/Source/Runtime/Online/HTTPServer/Private/HttpConnectionRequestReadContext.cpp
--- a/Engine/Source/Runtime/Online/HTTPServer/Private/HttpConnectionRequestReadContext.cpp
+++ b/Engine/Source/Runtime/Online/HTTPServer/Private/HttpConnectionRequestReadContext.cpp
@@ -265,8 +265,7 @@ TSharedPtr<FHttpServerRequest> FHttpConnectionRequestReadContext::BuildRequest(c
 
 		// Split query params
 		TArray<FString> QueryParamPairs;
-		const TCHAR QueryParamDelimiter[] = { TCHAR('&') };
-		QueryParamsStr.ParseIntoArray(QueryParamPairs, QueryParamDelimiter, true);
+		QueryParamsStr.ParseIntoArray(QueryParamPairs, TEXT("&"), true);
 		for (const FString& QueryParamPair : QueryParamPairs)
 		{
 			int32 Equalsindex = 0;
