Add option to reset XMPP ping timer on receiving a stanza

diff --git a/Engine/Source/Runtime/Online/XMPP/Private/XmppStrophe/XmppConnectionStrophe.cpp b/Engine/Source/Runtime/Online/XMPP/Private/XmppStrophe/XmppConnectionStrophe.cpp
--- a/Engine/Source/Runtime/Online/XMPP/Private/XmppStrophe/XmppConnectionStrophe.cpp
+++ b/Engine/Source/Runtime/Online/XMPP/Private/XmppStrophe/XmppConnectionStrophe.cpp
@@ -299,7 +299,7 @@ void FXmppConnectionStrophe::ReceiveStanza(const FStropheStanza& Stanza)
 	UE_LOG(LogXmpp, Verbose, TEXT("Received Strophe XMPP Stanza %s"), *Stanza.GetName());
 
 	// Reset our ping timer now that we've received traffic
-	if (PingStrophe.IsValid())
+	if (PingStrophe.IsValid() && ServerConfiguration.bResetPingTimeoutOnReceiveStanza)
 	{
 		PingStrophe->ResetPingTimer();
 	}
diff --git a/Engine/Source/Runtime/Online/XMPP/Public/XmppConnection.h b/Engine/Source/Runtime/Online/XMPP/Public/XmppConnection.h
--- a/Engine/Source/Runtime/Online/XMPP/Public/XmppConnection.h
+++ b/Engine/Source/Runtime/Online/XMPP/Public/XmppConnection.h
@@ -55,6 +55,7 @@ public:
 		, PingTimeout(30.0f)
 		, MaxPingRetries(1)
 		, bPrivateChatFriendsOnly(false)
+		, bResetPingTimeoutOnReceiveStanza(true)
 	{}
 
 	/** ip/host to connect to */
@@ -83,6 +84,8 @@ public:
 	int32 MaxPingRetries;
 	/** limit private chat to friends only */
 	bool bPrivateChatFriendsOnly;
+	/** reset ping timeout on received stanzas */
+	bool bResetPingTimeoutOnReceiveStanza;
 };
 
 /**
