Added missing copyright notices.

diff --git a/Engine/Source/Runtime/Online/HTTPServer/Private/HttpServerResponse.cpp b/Engine/Source/Runtime/Online/HTTPServer/Private/HttpServerResponse.cpp
--- a/Engine/Source/Runtime/Online/HTTPServer/Private/HttpServerResponse.cpp
+++ b/Engine/Source/Runtime/Online/HTTPServer/Private/HttpServerResponse.cpp
@@ -1,3 +1,4 @@
+// Copyright 1998-2019 Epic Games, Inc. All Rights Reserved.
 
 #include "HttpServerResponse.h"
 #include "HttpServerConstants.h"
@@ -60,4 +61,4 @@ TUniquePtr<FHttpServerResponse> FHttpServerResponse::Error(const FString& Messag
 	auto Response = Create(Message, MoveTemp(ContentType));
 	Response->Code = EHttpServerResponseCodes::ServerError;
 	return Response;
-}
\ No newline at end of file
+}
diff --git a/Engine/Source/Runtime/RHI/Public/GpuProfilerTrace.h b/Engine/Source/Runtime/RHI/Public/GpuProfilerTrace.h
--- a/Engine/Source/Runtime/RHI/Public/GpuProfilerTrace.h
+++ b/Engine/Source/Runtime/RHI/Public/GpuProfilerTrace.h
@@ -1,3 +1,5 @@
+// Copyright 1998-2019 Epic Games, Inc. All Rights Reserved.
+
 #pragma once
 
 #include "CoreTypes.h"
