Fix non-unity build (missing include)

diff --git a/Engine/Source/Runtime/Online/BackgroundHTTP/Private/GenericPlatform/GenericPlatformBackgroundHttpResponse.cpp b/Engine/Source/Runtime/Online/BackgroundHTTP/Private/GenericPlatform/GenericPlatformBackgroundHttpResponse.cpp
--- a/Engine/Source/Runtime/Online/BackgroundHTTP/Private/GenericPlatform/GenericPlatformBackgroundHttpResponse.cpp
+++ b/Engine/Source/Runtime/Online/BackgroundHTTP/Private/GenericPlatform/GenericPlatformBackgroundHttpResponse.cpp
@@ -4,6 +4,7 @@
 
 #include "Interfaces/IHttpResponse.h"
 #include "Misc/FileHelper.h"
+#include "HTML5/HTML5PlatformBackgroundHttp.h"
 
 FGenericPlatformBackgroundHttpResponse::FGenericPlatformBackgroundHttpResponse(int32 InResponseCode, const FString& InTempContentFilePath)
 {
