Back out revision 2 from //UE4/Dev-VR/Engine/Source/Runtime/Online/BackgroundHTTP/Private/GenericPlatform/GenericPlatformBackgroundHttpResponse.cpp

Reverting change, this was not the correct fix for the non-unity build

diff --git a/Engine/Source/Runtime/Online/BackgroundHTTP/Private/GenericPlatform/GenericPlatformBackgroundHttpResponse.cpp b/Engine/Source/Runtime/Online/BackgroundHTTP/Private/GenericPlatform/GenericPlatformBackgroundHttpResponse.cpp
--- a/Engine/Source/Runtime/Online/BackgroundHTTP/Private/GenericPlatform/GenericPlatformBackgroundHttpResponse.cpp
+++ b/Engine/Source/Runtime/Online/BackgroundHTTP/Private/GenericPlatform/GenericPlatformBackgroundHttpResponse.cpp
@@ -4,7 +4,6 @@
 
 #include "Interfaces/IHttpResponse.h"
 #include "Misc/FileHelper.h"
-#include "HTML5/HTML5PlatformBackgroundHttp.h"
 
 FGenericPlatformBackgroundHttpResponse::FGenericPlatformBackgroundHttpResponse(int32 InResponseCode, const FString& InTempContentFilePath)
 {
