Rename HoloLensHttp to HoloLensPlatformHttp

diff --git a/Engine/Source/Runtime/Online/HTTP/Private/HoloLens/HoloLensHttp.cpp b/Engine/Source/Runtime/Online/HTTP/Private/HoloLens/HoloLensHttp.cpp
deleted file mode 100644
diff --git a/Engine/Source/Runtime/Online/HTTP/Private/HoloLens/HoloLensPlatformHttp.cpp b/Engine/Source/Runtime/Online/HTTP/Private/HoloLens/HoloLensPlatformHttp.cpp
new file mode 100644
--- /dev/null
+++ b/Engine/Source/Runtime/Online/HTTP/Private/HoloLens/HoloLensPlatformHttp.cpp
@@ -0,0 +1,22 @@
+// Copyright 1998-2019 Epic Games, Inc. All Rights Reserved.
+
+#include "HoloLens/HoloLensPlatformHttp.h"
+#include "IXML/HttpIXML.h"
+
+void FHoloLensHttp::Init()
+{
+}
+
+void FHoloLensHttp::Shutdown()
+{
+}
+
+FHttpManager * FHoloLensHttp::CreatePlatformHttpManager()
+{
+	return nullptr;
+}
+
+IHttpRequest* FHoloLensHttp::ConstructRequest()
+{
+	return new FHttpRequestIXML();
+}
\ No newline at end of file
diff --git a/Engine/Source/Runtime/Online/HTTP/Public/HoloLens/HoloLensHttp.h b/Engine/Source/Runtime/Online/HTTP/Public/HoloLens/HoloLensHttp.h
deleted file mode 100644
diff --git a/Engine/Source/Runtime/Online/HTTP/Public/HoloLens/HoloLensPlatformHttp.h b/Engine/Source/Runtime/Online/HTTP/Public/HoloLens/HoloLensPlatformHttp.h
new file mode 100644
--- /dev/null
+++ b/Engine/Source/Runtime/Online/HTTP/Public/HoloLens/HoloLensPlatformHttp.h
@@ -0,0 +1,37 @@
+// Copyright 1998-2019 Epic Games, Inc. All Rights Reserved.
+
+#pragma once
+
+#include "GenericPlatform/GenericPlatformHttp.h"
+
+
+/**
+* Platform specific HTTP implementations.
+*/
+class FHoloLensHttp : public FGenericPlatformHttp
+{
+public:
+
+	/** Platform initialization step. */
+	static void Init();
+
+	/**
+	* Creates a platform-specific HTTP manager.
+	*
+	* @return nullptr if default implementation is to be used.
+	*/
+	static FHttpManager* CreatePlatformHttpManager();
+
+	/** Platform shutdown step. */
+	static void Shutdown();
+
+	/**
+	* Creates a new HTTP request instance for the current platform.
+	*
+	* @return The request object.
+	*/
+	static IHttpRequest* ConstructRequest();
+};
+
+
+typedef FHoloLensHttp FPlatformHttp;
