Fixed compile errors for HttpIXMLSupport

diff --git a/Engine/Source/Runtime/Online/HTTP/Private/IXML/HttpIXMLSupport.h b/Engine/Source/Runtime/Online/HTTP/Private/IXML/HttpIXMLSupport.h
--- a/Engine/Source/Runtime/Online/HTTP/Private/IXML/HttpIXMLSupport.h
+++ b/Engine/Source/Runtime/Online/HTTP/Private/IXML/HttpIXMLSupport.h
@@ -1,6 +1,8 @@
 // Copyright 1998-2019 Epic Games, Inc. All Rights Reserved.
 #pragma once
 
+#include "CoreTypes.h"
+
 #if PLATFORM_HOLOLENS
 
 
