Fix non-unity build

diff --git a/Engine/Source/Runtime/Online/HTTP/Private/IXML/HttpIXML.h b/Engine/Source/Runtime/Online/HTTP/Private/IXML/HttpIXML.h
--- a/Engine/Source/Runtime/Online/HTTP/Private/IXML/HttpIXML.h
+++ b/Engine/Source/Runtime/Online/HTTP/Private/IXML/HttpIXML.h
@@ -1,7 +1,8 @@
 // Copyright 1998-2019 Epic Games, Inc. All Rights Reserved.
-
 #pragma once
 
+#include "CoreTypes.h"
+
 #if PLATFORM_HOLOLENS
 
 #include "HttpIXMLSupport.h"
