Implement seek function in libcurl
Add a UpdateConfigs function to HTTP manager so implementation specific configs can be updated

diff --git a/Engine/Source/Runtime/Online/HTTP/Private/Curl/CurlHttp.cpp b/Engine/Source/Runtime/Online/HTTP/Private/Curl/CurlHttp.cpp
--- a/Engine/Source/Runtime/Online/HTTP/Private/Curl/CurlHttp.cpp
+++ b/Engine/Source/Runtime/Online/HTTP/Private/Curl/CurlHttp.cpp
@@ -14,8 +14,6 @@
 #include "HAL/FileManager.h"
 #include "Internationalization/Regex.h"
 
-int32 FCurlHttpRequest::NumberOfInfoMessagesToCache = 50;
-
 #if WITH_SSL
 #include "Ssl.h"
 
@@ -76,6 +74,7 @@ FCurlHttpRequest::FCurlHttpRequest()
 	,	TimeSinceLastResponse(0.0f)
 	,	bAnyHttpActivity(false)
 	,   BytesSent(0)
+	,	TotalBytesSent(0)
 	,	LastReportedBytesRead(0)
 	,	LastReportedBytesSent(0)
 	,   LeastRecentlyCachedInfoMessageIndex(0)
@@ -298,6 +297,7 @@ void FCurlHttpRequest::SetURL(const FString& InURL)
 void FCurlHttpRequest::SetContent(const TArray<uint8>& ContentPayload)
 {
 	RequestPayload = MakeUnique<FRequestPayloadInMemory>(ContentPayload);
+	bIsRequestPayloadSeekable = true;
 }
 
 void FCurlHttpRequest::SetContentAsString(const FString& ContentString)
@@ -306,7 +306,8 @@ void FCurlHttpRequest::SetContentAsString(const FString& ContentString)
 	TArray<uint8> Buffer;
 	Buffer.SetNum(Converter.Length());
 	FMemory::Memcpy(Buffer.GetData(), Converter.Get(), Buffer.Num());
-	RequestPayload = MakeUnique<FRequestPayloadInMemory>(Buffer);
+	bIsRequestPayloadSeekable = true;
+	RequestPayload = MakeUnique<FRequestPayloadInMemory>(MoveTemp(Buffer));
 }
 
 bool FCurlHttpRequest::SetContentAsStreamedFile(const FString& Filename)
@@ -323,14 +324,14 @@ bool FCurlHttpRequest::SetContentAsStreamedFile(const FString& Filename)
 	if (File)
 	{
 		RequestPayload = MakeUnique<FRequestPayloadInFileStream>(MakeShareable(File));
-		return true;
 	}
 	else
 	{
 		UE_LOG(LogHttp, Warning, TEXT("FCurlHttpRequest::SetContentAsStreamedFile Failed to open %s for reading"), *Filename);
 		RequestPayload.Reset();
-		return false;
 	}
+	bIsRequestPayloadSeekable = false;
+	return RequestPayload.IsValid();
 }
 
 bool FCurlHttpRequest::SetContentFromStream(TSharedRef<FArchive, ESPMode::ThreadSafe> Stream)
@@ -344,6 +345,7 @@ bool FCurlHttpRequest::SetContentFromStream(TSharedRef<FArchive, ESPMode::Thread
 	}
 
 	RequestPayload = MakeUnique<FRequestPayloadInFileStream>(Stream);
+	bIsRequestPayloadSeekable = false;
 	return true;
 }
 
@@ -384,6 +386,16 @@ size_t FCurlHttpRequest::StaticUploadCallback(void* Ptr, size_t SizeInBlocks, si
 	return Request->UploadCallback(Ptr, SizeInBlocks, BlockSizeInBytes);
 }
 
+int FCurlHttpRequest::StaticSeekCallback(void* UserData, curl_off_t Offset, int Origin)
+{
+	QUICK_SCOPE_CYCLE_COUNTER(STAT_FCurlHttpRequest_StaticSeekCallback);
+	check(UserData);
+
+	// dispatch
+	FCurlHttpRequest* Request = reinterpret_cast<FCurlHttpRequest*>(UserData);
+	return Request->SeekCallback(Offset, Origin);
+}
+
 size_t FCurlHttpRequest::StaticReceiveResponseHeaderCallback(void* Ptr, size_t SizeInBlocks, size_t BlockSizeInBytes, void* UserData)
 {
 	QUICK_SCOPE_CYCLE_COUNTER(STAT_FCurlHttpRequest_StaticReceiveResponseHeaderCallback);
@@ -522,11 +534,13 @@ size_t FCurlHttpRequest::UploadCallback(void* Ptr, size_t SizeInBlocks, size_t B
 	size_t SizeAlreadySent = static_cast<size_t>(BytesSent.GetValue());
 	size_t SizeSentThisTime = RequestPayload->FillOutputBuffer(Ptr, MaxBufferSize, SizeAlreadySent);
 	BytesSent.Add(SizeSentThisTime);
+	TotalBytesSent.Add(SizeSentThisTime);
 
-	UE_LOG(LogHttp, Verbose, TEXT("%p: UploadCallback: %d bytes out of %d sent. (SizeInBlocks=%d, BlockSizeInBytes=%d, SizeToSendThisTime=%d (<-this will get returned from the callback))"),
+	UE_LOG(LogHttp, Verbose, TEXT("%p: UploadCallback: %d bytes out of %d sent (%d bytes total sent). (SizeInBlocks=%d, BlockSizeInBytes=%d, SizeToSendThisTime=%d (<-this will get returned from the callback))"),
 		this,
 		static_cast< int32 >(BytesSent.GetValue()),
 		RequestPayload->GetContentLength(),
+		static_cast< int32 >(TotalBytesSent.GetValue()),
 		static_cast< int32 >(SizeInBlocks),
 		static_cast< int32 >(BlockSizeInBytes),
 		static_cast< int32 >(SizeSentThisTime)
@@ -535,6 +549,26 @@ size_t FCurlHttpRequest::UploadCallback(void* Ptr, size_t SizeInBlocks, size_t B
 	return SizeSentThisTime;
 }
 
+int FCurlHttpRequest::SeekCallback(curl_off_t Offset, int Origin)
+{
+	// Only support seeking to the very beginning
+	if (bIsRequestPayloadSeekable && Origin == SEEK_SET && Offset == 0)
+	{
+		UE_LOG(LogHttp, Log, TEXT("%p: SeekCallback: Resetting to the beginning. We had uploaded %d bytes"),
+			this,
+			static_cast<int32>(BytesSent.GetValue()));
+		BytesSent.Reset();
+		bIsRequestPayloadSeekable = false; // Do not attempt to re-seek
+		return CURL_SEEKFUNC_OK;
+	}
+	UE_LOG(LogHttp, Warning, TEXT("%p: SeekCallback: Failed to seek to Offset=%lld, Origin=%d %s"), 
+		this,
+		(int64)(Offset),
+		Origin, 
+		bIsRequestPayloadSeekable ? TEXT("not implemented") : TEXT("seek disabled"));
+	return CURL_SEEKFUNC_CANTSEEK;
+}
+
 size_t FCurlHttpRequest::DebugCallback(CURL * Handle, curl_infotype DebugInfoType, char * DebugInfo, size_t DebugInfoSize)
 {
 	check(Handle == EasyHandle);	// make sure it's us
@@ -667,6 +701,7 @@ bool FCurlHttpRequest::SetupRequest()
 	if (!RequestPayload.IsValid())
 	{
 		RequestPayload = MakeUnique<FRequestPayloadInMemory>(TArray<uint8>());
+		bIsRequestPayloadSeekable = true;
 	}
 
 	bCurlRequestCompleted = false;
@@ -722,7 +757,7 @@ bool FCurlHttpRequest::SetupRequest()
 	{
 		// If we don't pass any other Content-Type, RequestPayload is assumed to be URL-encoded by this time
 		// In the case of using a streamed file, you must explicitly set the Content-Type, because RequestPayload->IsURLEncoded returns false.
-		check(!GetHeader("Content-Type").IsEmpty() || RequestPayload->IsURLEncoded());
+		check(!GetHeader(TEXT("Content-Type")).IsEmpty() || RequestPayload->IsURLEncoded());
 		curl_easy_setopt(EasyHandle, CURLOPT_POST, 1L);
 		curl_easy_setopt(EasyHandle, CURLOPT_POSTFIELDS, NULL);
 		curl_easy_setopt(EasyHandle, CURLOPT_POSTFIELDSIZE, RequestPayload->GetContentLength());
@@ -752,7 +787,7 @@ bool FCurlHttpRequest::SetupRequest()
 	{
 		// If we don't pass any other Content-Type, RequestPayload is assumed to be URL-encoded by this time
 		// (if we pass, don't check here and trust the request)
-		check(!GetHeader("Content-Type").IsEmpty() || RequestPayload->IsURLEncoded());
+		check(!GetHeader(TEXT("Content-Type")).IsEmpty() || RequestPayload->IsURLEncoded());
 
 		curl_easy_setopt(EasyHandle, CURLOPT_POST, 1L);
 		curl_easy_setopt(EasyHandle, CURLOPT_CUSTOMREQUEST, "DELETE");
@@ -768,6 +803,7 @@ bool FCurlHttpRequest::SetupRequest()
 	if (bUseReadFunction)
 	{
 		BytesSent.Reset();
+		TotalBytesSent.Reset();
 		curl_easy_setopt(EasyHandle, CURLOPT_READDATA, this);
 		curl_easy_setopt(EasyHandle, CURLOPT_READFUNCTION, StaticUploadCallback);
 	}
@@ -788,20 +824,20 @@ bool FCurlHttpRequest::SetupRequest()
 		curl_easy_setopt(EasyHandle, CURLOPT_ACCEPT_ENCODING, "");
 	}
 
-	if (GetHeader("User-Agent").IsEmpty())
+	if (GetHeader(TEXT("User-Agent")).IsEmpty())
 	{
 		SetHeader(TEXT("User-Agent"), FPlatformHttp::GetDefaultUserAgent());
 	}
 
 	// content-length should be present http://www.w3.org/Protocols/rfc2616/rfc2616-sec4.html#sec4.4
-	if (GetHeader("Content-Length").IsEmpty())
+	if (GetHeader(TEXT("Content-Length")).IsEmpty())
 	{
 		SetHeader(TEXT("Content-Length"), FString::Printf(TEXT("%d"), RequestPayload->GetContentLength()));
 	}
 
 	// Remove "Expect: 100-continue" since this is supposed to cause problematic behavior on Amazon ELB (and WinInet doesn't send that either)
 	// (also see http://www.iandennismiller.com/posts/curl-http1-1-100-continue-and-multipartform-data-post.html , http://the-stickman.com/web-development/php-and-curl-disabling-100-continue-header/ )
-	if (GetHeader("Expect").IsEmpty())
+	if (GetHeader(TEXT("Expect")).IsEmpty())
 	{
 		SetHeader(TEXT("Expect"), TEXT(""));
 	}
@@ -846,6 +882,12 @@ bool FCurlHttpRequest::SetupRequest()
 		curl_easy_setopt(EasyHandle, CURLOPT_CONNECTTIMEOUT, HttpConnectionTimeout);
 	}
 
+	if (FCurlHttpManager::CurlRequestOptions.bAllowSeekFunction && bIsRequestPayloadSeekable)
+	{
+		curl_easy_setopt(EasyHandle, CURLOPT_SEEKDATA, this);
+		curl_easy_setopt(EasyHandle, CURLOPT_SEEKFUNCTION, StaticSeekCallback);
+	}
+
 	UE_LOG(LogHttp, Log, TEXT("%p: Starting %s request to URL='%s'"), this, *Verb, *URL);
 	return true;
 }
diff --git a/Engine/Source/Runtime/Online/HTTP/Private/Curl/CurlHttp.h b/Engine/Source/Runtime/Online/HTTP/Private/Curl/CurlHttp.h
--- a/Engine/Source/Runtime/Online/HTTP/Private/Curl/CurlHttp.h
+++ b/Engine/Source/Runtime/Online/HTTP/Private/Curl/CurlHttp.h
@@ -225,6 +225,25 @@ private:
 	 */
 	size_t UploadCallback(void* Ptr, size_t SizeInBlocks, size_t BlockSizeInBytes);
 
+	/**
+	 * Static callback to be used as seek function (CURLOPT_SEEKFUNCTION), will dispatch the call to proper instance
+	 *
+	 * @param UserData data we associated with request (will be a pointer to FCurlHttpRequest instance)
+	 * @param Offset offset from Origin to seek to
+	 * @param Origin where to seek to. Can be SEEK_SET, SEEK_CUR, or SEEK_END
+	 * @return CURL_SEEKFUNC_OK if the seek was successful, CURL_SEEKFUNC_FAIL if the request should be failed due to inability to seek, or CURL_SEEKFUNC_CANTSEEK to allow curl to try to workaround the inability to seek
+	 */
+	static int StaticSeekCallback(void* UserData, curl_off_t Offset, int Origin);
+
+	/**
+	 * Method called when libcurl wants us to seek to a position in the stream (see CURLOPT_SEEKFUNCTION)
+	 *
+	 * @param Offset offset from Origin to seek to
+	 * @param Origin where to seek to. Can be SEEK_SET, SEEK_CUR, or SEEK_END
+	 * @return CURL_SEEKFUNC_OK if the seek was successful, CURL_SEEKFUNC_FAIL if the request should be failed due to inability to seek, or CURL_SEEKFUNC_CANTSEEK to allow curl to try to workaround the inability to seek
+	 */
+	int SeekCallback(curl_off_t Offset, int Origin);
+
 	/**
 	 * Static callback to be used as header function (CURLOPT_HEADERFUNCTION), will dispatch the call to proper instance
 	 *
@@ -338,8 +357,10 @@ private:
 	CURLcode		CurlCompletionResult;
 	/** The response object which we will use to pair with this request */
 	TSharedPtr<class FCurlHttpResponse,ESPMode::ThreadSafe> Response;
-	/** BYTE array payload to use with the request. Typically for a POST */
+	/** Payload to use with the request. Typically for POST, PUT, or PATCH */
 	TUniquePtr<FRequestPayload> RequestPayload;
+	/** Is the request payload seekable? */
+	bool bIsRequestPayloadSeekable = false;
 	/** Current status of request being processed */
 	EHttpRequestStatus::Type CompletionStatus;
 	/** Mapping of header section to values. */
@@ -352,16 +373,18 @@ private:
 	bool bAnyHttpActivity;
 	/** Number of bytes sent already */
 	FThreadSafeCounter BytesSent;
+	/** Total number of bytes sent already (includes data re-sent by seek attempts) */
+	FThreadSafeCounter TotalBytesSent;
 	/** Last bytes read reported to progress delegate */
 	int32 LastReportedBytesRead;
 	/** Last bytes sent reported to progress delegate */
 	int32 LastReportedBytesSent;
 	/** Number of info channel messages to cache */
-	static int32 NumberOfInfoMessagesToCache;
+	static const constexpr int32 NumberOfInfoMessagesToCache = 50;
 	/** Index of least recently cached message */
 	int32 LeastRecentlyCachedInfoMessageIndex;
 	/** Cache of info messages from libcurl */
-	TArray<FString> InfoMessageCache;
+	TArray<FString, TFixedAllocator<NumberOfInfoMessagesToCache>> InfoMessageCache;
 };
 
 /**
diff --git a/Engine/Source/Runtime/Online/HTTP/Private/Curl/CurlHttpManager.cpp b/Engine/Source/Runtime/Online/HTTP/Private/Curl/CurlHttpManager.cpp
--- a/Engine/Source/Runtime/Online/HTTP/Private/Curl/CurlHttpManager.cpp
+++ b/Engine/Source/Runtime/Online/HTTP/Private/Curl/CurlHttpManager.cpp
@@ -26,6 +26,8 @@
 #include "SocketSubsystem.h"
 #include "IPAddress.h"
 
+#include "Http.h"
+
 #ifndef DISABLE_UNVERIFIED_CERTIFICATE_LOADING
 #define DISABLE_UNVERIFIED_CERTIFICATE_LOADING 0
 #endif
@@ -238,6 +240,8 @@ void FCurlHttpManager::InitCurl()
 		CurlRequestOptions.BufferSize = ConfigBufferSize;
 	}
 
+	GConfig->GetBool(TEXT("HTTP.Curl"), TEXT("bAllowSeekFunction"), CurlRequestOptions.bAllowSeekFunction, GEngineIni);
+
 	CurlRequestOptions.MaxHostConnections = FHttpModule::Get().GetHttpMaxConnectionsPerServer();
 	if (CurlRequestOptions.MaxHostConnections > 0)
 	{
@@ -355,6 +359,48 @@ void FCurlHttpManager::OnAfterFork()
 	FHttpManager::OnAfterFork();
 }
 
+void FCurlHttpManager::UpdateConfigs()
+{
+	// Update configs - update settings that are safe to update after initialize 
+	FHttpManager::UpdateConfigs();
+
+	{
+		bool bAcceptCompressedContent = true;
+		if (GConfig->GetBool(TEXT("HTTP"), TEXT("AcceptCompressedContent"), bAcceptCompressedContent, GEngineIni))
+		{
+			if (CurlRequestOptions.bAcceptCompressedContent != bAcceptCompressedContent)
+			{
+				UE_LOG(LogHttp, Log, TEXT("AcceptCompressedContent changed from %s to %s"), *LexToString(CurlRequestOptions.bAcceptCompressedContent), *LexToString(bAcceptCompressedContent));
+				CurlRequestOptions.bAcceptCompressedContent = bAcceptCompressedContent;
+			}
+		}
+	}
+
+	{
+		int32 ConfigBufferSize = 0;
+		if (GConfig->GetInt(TEXT("HTTP.Curl"), TEXT("BufferSize"), ConfigBufferSize, GEngineIni) && ConfigBufferSize > 0)
+		{
+			if (CurlRequestOptions.BufferSize != ConfigBufferSize)
+			{
+				UE_LOG(LogHttp, Log, TEXT("BufferSize changed from %d to %d"), CurlRequestOptions.BufferSize, ConfigBufferSize);
+				CurlRequestOptions.BufferSize = ConfigBufferSize;
+			}
+		}
+	}
+
+	{
+		bool bConfigAllowSeekFunction = false;
+		if (GConfig->GetBool(TEXT("HTTP.Curl"), TEXT("bAllowSeekFunction"), bConfigAllowSeekFunction, GEngineIni))
+		{
+			if (CurlRequestOptions.bAllowSeekFunction != bConfigAllowSeekFunction)
+			{
+				UE_LOG(LogHttp, Log, TEXT("bAllowSeekFunction changed from %s to %s"), *LexToString(CurlRequestOptions.bAllowSeekFunction), *LexToString(bConfigAllowSeekFunction));
+				CurlRequestOptions.bAllowSeekFunction = bConfigAllowSeekFunction;
+			}
+		}
+	}
+}
+
 FHttpThread* FCurlHttpManager::CreateHttpThread()
 {
 	return new FCurlHttpThread();
diff --git a/Engine/Source/Runtime/Online/HTTP/Private/Curl/CurlHttpManager.h b/Engine/Source/Runtime/Online/HTTP/Private/Curl/CurlHttpManager.h
--- a/Engine/Source/Runtime/Online/HTTP/Private/Curl/CurlHttpManager.h
+++ b/Engine/Source/Runtime/Online/HTTP/Private/Curl/CurlHttpManager.h
@@ -50,12 +50,16 @@ public:
 
 		/** Receive buffer size */
 		int32 BufferSize;
+
+		/** Do we allow seeking? */
+		bool bAllowSeekFunction = false;
 	}
 	CurlRequestOptions;
 
 	//~ Begin HttpManager Interface
 	virtual void OnBeforeFork() override;
 	virtual void OnAfterFork() override;
+	virtual void UpdateConfigs() override;
 
 public:
 	virtual bool SupportsDynamicProxy() const override;
diff --git a/Engine/Source/Runtime/Online/HTTP/Private/GenericPlatform/HttpRequestPayload.cpp b/Engine/Source/Runtime/Online/HTTP/Private/GenericPlatform/HttpRequestPayload.cpp
--- a/Engine/Source/Runtime/Online/HTTP/Private/GenericPlatform/HttpRequestPayload.cpp
+++ b/Engine/Source/Runtime/Online/HTTP/Private/GenericPlatform/HttpRequestPayload.cpp
@@ -82,6 +82,10 @@ FRequestPayloadInMemory::FRequestPayloadInMemory(const TArray<uint8>& Array) : B
 {
 }
 
+FRequestPayloadInMemory::FRequestPayloadInMemory(TArray<uint8>&& Array) : Buffer(MoveTemp(Array))
+{
+}
+
 FRequestPayloadInMemory::~FRequestPayloadInMemory()
 {
 }
diff --git a/Engine/Source/Runtime/Online/HTTP/Private/HttpManager.cpp b/Engine/Source/Runtime/Online/HTTP/Private/HttpManager.cpp
--- a/Engine/Source/Runtime/Online/HTTP/Private/HttpManager.cpp
+++ b/Engine/Source/Runtime/Online/HTTP/Private/HttpManager.cpp
@@ -108,6 +108,11 @@ void FHttpManager::OnAfterFork()
 
 }
 
+void FHttpManager::UpdateConfigs()
+{
+	// empty
+}
+
 FHttpThread* FHttpManager::CreateHttpThread()
 {
 	return new FHttpThread();
diff --git a/Engine/Source/Runtime/Online/HTTP/Private/HttpModule.cpp b/Engine/Source/Runtime/Online/HTTP/Private/HttpModule.cpp
--- a/Engine/Source/Runtime/Online/HTTP/Private/HttpModule.cpp
+++ b/Engine/Source/Runtime/Online/HTTP/Private/HttpModule.cpp
@@ -47,6 +47,11 @@ void FHttpModule::UpdateConfigs()
 
 	AllowedDomains.Empty();
 	GConfig->GetArray(TEXT("HTTP"), TEXT("AllowedDomains"), AllowedDomains, GEngineIni);
+
+	if (HttpManager != nullptr)
+	{
+		HttpManager->UpdateConfigs();
+	}
 }
 
 void FHttpModule::StartupModule()
@@ -85,7 +90,7 @@ void FHttpModule::StartupModule()
 	FPlatformHttp::Init();
 
 	HttpManager = FPlatformHttp::CreatePlatformHttpManager();
-	if (NULL == HttpManager)
+	if (nullptr == HttpManager)
 	{
 		// platform does not provide specific HTTP manager, use generic one
 		HttpManager = new FHttpManager();
diff --git a/Engine/Source/Runtime/Online/HTTP/Public/GenericPlatform/HttpRequestPayload.h b/Engine/Source/Runtime/Online/HTTP/Public/GenericPlatform/HttpRequestPayload.h
--- a/Engine/Source/Runtime/Online/HTTP/Public/GenericPlatform/HttpRequestPayload.h
+++ b/Engine/Source/Runtime/Online/HTTP/Public/GenericPlatform/HttpRequestPayload.h
@@ -44,6 +44,7 @@ class FRequestPayloadInMemory : public FRequestPayload
 {
 public:
 	FRequestPayloadInMemory(const TArray<uint8>& Array);
+	FRequestPayloadInMemory(TArray<uint8>&& Array);
 	virtual ~FRequestPayloadInMemory();
 	virtual int32 GetContentLength() const override;
 	virtual const TArray<uint8>& GetContent() const override;
diff --git a/Engine/Source/Runtime/Online/HTTP/Public/HttpManager.h b/Engine/Source/Runtime/Online/HTTP/Public/HttpManager.h
--- a/Engine/Source/Runtime/Online/HTTP/Public/HttpManager.h
+++ b/Engine/Source/Runtime/Online/HTTP/Public/HttpManager.h
@@ -146,6 +146,11 @@ public:
 	 */
 	virtual void OnAfterFork();
 
+	/**
+	 * Update configuration. Called when config has been updated and we need to apply any changes.
+	 */
+	virtual void UpdateConfigs();
+
 protected:
 	/** 
 	 * Create HTTP thread object
