#UE4 Include fix

diff --git a/Engine/Source/Runtime/Online/Voice/Private/Windows/VoiceCaptureWindows.h b/Engine/Source/Runtime/Online/Voice/Private/Windows/VoiceCaptureWindows.h
--- a/Engine/Source/Runtime/Online/Voice/Private/Windows/VoiceCaptureWindows.h
+++ b/Engine/Source/Runtime/Online/Voice/Private/Windows/VoiceCaptureWindows.h
@@ -9,7 +9,7 @@
 #include "DSP/EnvelopeFollower.h"
 #include "Containers/Ticker.h"
 #include "Windows/WindowsHWrapper.h"
-#include "SampleBuffer.h"
+#include "Sound/SampleBuffer.h"
 
 typedef struct IDirectSound8 *LPDIRECTSOUND8;
 
