Add support for stereo VOIP input/output.

diff --git a/Engine/Plugins/Online/OnlineSubsystemUtils/Source/OnlineSubsystemUtils/Private/Tests/TestVoice.cpp b/Engine/Plugins/Online/OnlineSubsystemUtils/Source/OnlineSubsystemUtils/Private/Tests/TestVoice.cpp
--- a/Engine/Plugins/Online/OnlineSubsystemUtils/Source/OnlineSubsystemUtils/Private/Tests/TestVoice.cpp
+++ b/Engine/Plugins/Online/OnlineSubsystemUtils/Source/OnlineSubsystemUtils/Private/Tests/TestVoice.cpp
@@ -27,8 +27,8 @@ FTestVoice::FTestVoice() :
 	EncodeHint(UVOIPStatics::GetAudioEncodingHint()),
 	InputSampleRate(UVOIPStatics::GetVoiceSampleRate()),
 	OutputSampleRate(UVOIPStatics::GetVoiceSampleRate()),
-	NumInChannels(DEFAULT_NUM_VOICE_CHANNELS),
-	NumOutChannels(DEFAULT_NUM_VOICE_CHANNELS),
+	NumInChannels(UVOIPStatics::GetVoiceNumChannels()),
+	NumOutChannels(UVOIPStatics::GetVoiceNumChannels()),
 	bLastWasPlaying(false),
 	StarvedDataCount(0),
 	MaxRawCaptureDataSize(0),
diff --git a/Engine/Plugins/Online/OnlineSubsystemUtils/Source/OnlineSubsystemUtils/Private/VoiceEngineImpl.cpp b/Engine/Plugins/Online/OnlineSubsystemUtils/Source/OnlineSubsystemUtils/Private/VoiceEngineImpl.cpp
--- a/Engine/Plugins/Online/OnlineSubsystemUtils/Source/OnlineSubsystemUtils/Private/VoiceEngineImpl.cpp
+++ b/Engine/Plugins/Online/OnlineSubsystemUtils/Source/OnlineSubsystemUtils/Private/VoiceEngineImpl.cpp
@@ -53,7 +53,7 @@ FRemoteTalkerDataImpl::FRemoteTalkerDataImpl() :
 	VoiceDecoder(nullptr)
 {
 	int32 SampleRate = UVOIPStatics::GetVoiceSampleRate();
-	int32 NumChannels = DEFAULT_NUM_VOICE_CHANNELS;
+	int32 NumChannels = UVOIPStatics::GetVoiceNumChannels();
 	VoiceDecoder = FVoiceModule::Get().CreateVoiceDecoder(SampleRate, NumChannels);
 	check(VoiceDecoder.IsValid());
 
@@ -672,7 +672,7 @@ void FVoiceEngineImpl::GenerateVoiceData(USoundWaveProcedural* InProceduralWave,
 	FRemoteTalkerDataImpl* QueuedData = RemoteTalkerBuffers.Find(FUniqueNetIdWrapper(TalkerId.AsShared()));
 	if (QueuedData)
 	{
-		const int32 SampleSize = sizeof(uint16) * DEFAULT_NUM_VOICE_CHANNELS;
+		const int32 SampleSize = sizeof(uint16) * UVOIPStatics::GetVoiceNumChannels();
 
 		{
 			FScopeLock ScopeLock(&QueuedData->QueueLock);
diff --git a/Engine/Plugins/Online/OnlineSubsystemUtils/Source/OnlineSubsystemUtils/Private/VoipListenerSynthComponent.cpp b/Engine/Plugins/Online/OnlineSubsystemUtils/Source/OnlineSubsystemUtils/Private/VoipListenerSynthComponent.cpp
--- a/Engine/Plugins/Online/OnlineSubsystemUtils/Source/OnlineSubsystemUtils/Private/VoipListenerSynthComponent.cpp
+++ b/Engine/Plugins/Online/OnlineSubsystemUtils/Source/OnlineSubsystemUtils/Private/VoipListenerSynthComponent.cpp
@@ -21,7 +21,7 @@ static float DefaultPatchGainCVar = 1.0f;
 FAutoConsoleVariableRef CVarDefaultPatchGain(
 	TEXT("voice.DefaultPatchGain"),
 	DefaultPatchGainCVar,
-	TEXT("Changes the default gain of audio patches, in samples.\n"),
+	TEXT("Changes the default gain of audio patches, in linear gain.\n"),
 	ECVF_Default);
 
 static int32 MuteAudioEngineOutputCVar = 0;
@@ -34,7 +34,7 @@ FAutoConsoleVariableRef CVarMuteAudioEngineOutput(
 
 bool UVoipListenerSynthComponent::Init(int32& SampleRate)
 {
-	NumChannels = 1;
+	NumChannels = UVOIPStatics::GetVoiceNumChannels();
 	SampleRate = UVOIPStatics::GetVoiceSampleRate();
 	MySampleRate = SampleRate;
 
diff --git a/Engine/Source/Runtime/Engine/Private/VoiceConfig.cpp b/Engine/Source/Runtime/Engine/Private/VoiceConfig.cpp
--- a/Engine/Source/Runtime/Engine/Private/VoiceConfig.cpp
+++ b/Engine/Source/Runtime/Engine/Private/VoiceConfig.cpp
@@ -21,6 +21,14 @@ static TAutoConsoleVariable<float> CVarVoiceSilenceDetectionThreshold(TEXT("voic
 	TEXT("Threshold to be set for the VOIP microphone's silence detection algorithm.\n"),	
 	ECVF_Default);
 
+static int32 NumVoiceChannelsCvar = 1;
+FAutoConsoleVariableRef CVarNumVoiceChannels(
+	TEXT("voice.NumChannels"),
+	NumVoiceChannelsCvar,
+	TEXT("Default number of channels to capture from mic input, encode to Opus, and output. Can be set to 1 or 2.\n")
+	TEXT("Value: Number of channels to use for VOIP input and output."),
+	ECVF_Default);
+
 int32 UVOIPStatics::GetVoiceSampleRate()
 {
 #if PLATFORM_UNIX
@@ -61,6 +69,11 @@ int32 UVOIPStatics::GetVoiceSampleRate()
 #endif
 }
 
+int32 UVOIPStatics::GetVoiceNumChannels()
+{
+	return FMath::Clamp<int32>(NumVoiceChannelsCvar, 1, 2);
+}
+
 uint32 UVOIPStatics::GetMaxVoiceDataSize()
 {
 	int32 SampleRate = GetVoiceSampleRate();
diff --git a/Engine/Source/Runtime/Engine/Public/Net/VoiceConfig.h b/Engine/Source/Runtime/Engine/Public/Net/VoiceConfig.h
--- a/Engine/Source/Runtime/Engine/Public/Net/VoiceConfig.h
+++ b/Engine/Source/Runtime/Engine/Public/Net/VoiceConfig.h
@@ -148,6 +148,7 @@ class ENGINE_API UVOIPStatics : public UBlueprintFunctionLibrary
 public:
 	// Returns the voice sample rate specified in Audio Settings.
 	static int32 GetVoiceSampleRate();
+	static int32 GetVoiceNumChannels();
 
 	// Returns the max voice data size per packet, uncompressed, and compressed based on the voice sample rate specified in Audio Settings.
 	static uint32 GetMaxVoiceDataSize();
diff --git a/Engine/Source/Runtime/Online/Voice/Public/VoiceModule.h b/Engine/Source/Runtime/Online/Voice/Public/VoiceModule.h
--- a/Engine/Source/Runtime/Online/Voice/Public/VoiceModule.h
+++ b/Engine/Source/Runtime/Online/Voice/Public/VoiceModule.h
@@ -17,8 +17,6 @@
 #define DEFAULT_VOICE_SAMPLE_RATE 16000
 /** Deprecated value, use DEFAULT_VOICE_SAMPLE_RATE */
 #define VOICE_SAMPLE_RATE DEFAULT_VOICE_SAMPLE_RATE
-/** Default voice chat number of channels (mono) */
-#define DEFAULT_NUM_VOICE_CHANNELS 1
 
 class IVoiceCapture;
 class IVoiceEncoder;
@@ -76,7 +74,7 @@ public:
 	 *
 	 * @return new voice capture object, possibly NULL
 	 */
-	virtual TSharedPtr<IVoiceCapture> CreateVoiceCapture(const FString& DeviceName = DEFAULT_DEVICE_NAME, int32 SampleRate = UVOIPStatics::GetVoiceSampleRate(), int32 NumChannels = DEFAULT_NUM_VOICE_CHANNELS);
+	virtual TSharedPtr<IVoiceCapture> CreateVoiceCapture(const FString& DeviceName = DEFAULT_DEVICE_NAME, int32 SampleRate = UVOIPStatics::GetVoiceSampleRate(), int32 NumChannels = UVOIPStatics::GetVoiceNumChannels());
 
 	/**
 	 * Instantiates a new voice encoder object
@@ -87,7 +85,7 @@ public:
 	 *
 	 * @return new voice encoder object, possibly NULL
 	 */
-	virtual TSharedPtr<IVoiceEncoder> CreateVoiceEncoder(int32 SampleRate = UVOIPStatics::GetVoiceSampleRate(), int32 NumChannels = DEFAULT_NUM_VOICE_CHANNELS, EAudioEncodeHint EncodeHint = UVOIPStatics::GetAudioEncodingHint());
+	virtual TSharedPtr<IVoiceEncoder> CreateVoiceEncoder(int32 SampleRate = UVOIPStatics::GetVoiceSampleRate(), int32 NumChannels = UVOIPStatics::GetVoiceNumChannels(), EAudioEncodeHint EncodeHint = UVOIPStatics::GetAudioEncodingHint());
 
 	/**
 	 * Instantiates a new voice decoder object
@@ -97,7 +95,7 @@ public:
 	 *
 	 * @return new voice decoder object, possibly NULL
 	 */
-	virtual TSharedPtr<IVoiceDecoder> CreateVoiceDecoder(int32 SampleRate = UVOIPStatics::GetVoiceSampleRate(), int32 NumChannels = DEFAULT_NUM_VOICE_CHANNELS);
+	virtual TSharedPtr<IVoiceDecoder> CreateVoiceDecoder(int32 SampleRate = UVOIPStatics::GetVoiceSampleRate(), int32 NumChannels = UVOIPStatics::GetVoiceNumChannels());
 
 	/**
 	* Checks to see if the current platform supports voice capture.
