BackgroundHttp Improvements:
- Added priority system for Background HTTP downloads.
- Fixed bug where dropping internet connection mid-download would trigger an ensure and try to re-activate all possible download tasks on iOS at once.
- Fixed issue where we were immediately calling the NSURLSession Background completion handler instead of waiting for threads to finish their work. Could lead to scenarios where we weren't finishing re-queuing failed tasks in the background if many failed.
#tests iOS Client, Android Client

diff --git a/Engine/Source/Runtime/ApplicationCore/Private/IOS/IOSBackgroundURLSessionHandler.cpp b/Engine/Source/Runtime/ApplicationCore/Private/IOS/IOSBackgroundURLSessionHandler.cpp
--- a/Engine/Source/Runtime/ApplicationCore/Private/IOS/IOSBackgroundURLSessionHandler.cpp
+++ b/Engine/Source/Runtime/ApplicationCore/Private/IOS/IOSBackgroundURLSessionHandler.cpp
@@ -13,12 +13,16 @@ FIOSBackgroundDownloadCoreDelegates::FIOSBackgroundDownload_DidFinishDownloading
 FIOSBackgroundDownloadCoreDelegates::FIOSBackgroundDownload_DidWriteData FIOSBackgroundDownloadCoreDelegates::OnIOSBackgroundDownload_DidWriteData;
 FIOSBackgroundDownloadCoreDelegates::FIOSBackgroundDownload_DidCompleteWithError FIOSBackgroundDownloadCoreDelegates::OnIOSBackgroundDownload_DidCompleteWithError;
 FIOSBackgroundDownloadCoreDelegates::FIOSBackgroundDownload_SessionDidFinishAllEvents FIOSBackgroundDownloadCoreDelegates::OnIOSBackgroundDownload_SessionDidFinishAllEvents;
+FIOSBackgroundDownloadCoreDelegates::FIOSBackgroundDownload_DelayedBackgroundURLSessionCompleteHandler FIOSBackgroundDownloadCoreDelegates::OnDelayedBackgroundURLSessionCompleteHandler;
 
 NSURLSession* FBackgroundURLSessionHandler::BackgroundSession = nil;
 FString FBackgroundURLSessionHandler::CachedIdentifierName = FString();
+volatile int32 FBackgroundURLSessionHandler::DelayedBackgroundURLSessionCompleteCount = 0;
 
 bool FBackgroundURLSessionHandler::InitBackgroundSession(const FString& SessionIdentifier)
 {
+	FIOSBackgroundDownloadCoreDelegates::OnDelayedBackgroundURLSessionCompleteHandler.BindStatic(&FBackgroundURLSessionHandler::OnDelayedBackgroundURLSessionCompleteHandlerCalled);
+	
 	bool bHasSessionWithIdentifier = false;
 
 	if (!SessionIdentifier.IsEmpty())
@@ -132,6 +136,45 @@ const FString FBackgroundURLSessionHandler::GetTemporaryFilePathFromURL(const FS
 	return FPaths::Combine(BackgroundHttpDir, FileName);
 }
 
+void FBackgroundURLSessionHandler::CallBackgroundURLSessionCompleteHandler()
+{
+	IOSAppDelegate* AppDelegate = [IOSAppDelegate GetDelegate];
+	
+	if ((nullptr != AppDelegate) && (nil != AppDelegate.BackgroundSessionEventCompleteDelegate))
+	{
+		//make a local copy of the completion handler to capture with the block
+		void(^completionHandler)() = AppDelegate.BackgroundSessionEventCompleteDelegate;
+		[completionHandler retain];
+		
+		//set the completion handler to nil to protect from anything else trying to use it
+		AppDelegate.BackgroundSessionEventCompleteDelegate = nil;
+		
+	   //Need to issue our stored completion callback on main thread as its a UI callback
+		[[NSOperationQueue mainQueue] addOperationWithBlock:^{
+			//call the copied completion handler
+			completionHandler();
+			
+			//Release local copy so we don't leak the completion handler
+			[completionHandler release];
+		}];
+	}
+}
+
+void FBackgroundURLSessionHandler::AddDelayedBackgroundURLSessionComplete()
+{
+	FPlatformAtomics::InterlockedIncrement(&DelayedBackgroundURLSessionCompleteCount);
+}
+
+void FBackgroundURLSessionHandler::OnDelayedBackgroundURLSessionCompleteHandlerCalled()
+{
+	int NewDecrementedCount = FPlatformAtomics::InterlockedDecrement(&DelayedBackgroundURLSessionCompleteCount);
+	
+	if (NewDecrementedCount == 0)
+	{
+		CallBackgroundURLSessionCompleteHandler();
+	}
+}
+
 @implementation BackgroundDownloadDelegate
 	
 -(void)URLSession:(NSURLSession *)session downloadTask : (NSURLSessionDownloadTask *)downloadTask didFinishDownloadingToURL : (NSURL *)location;
@@ -176,29 +219,40 @@ const FString FBackgroundURLSessionHandler::GetTemporaryFilePathFromURL(const FS
 
 -(void)URLSessionDidFinishEventsForBackgroundURLSession:(NSURLSession *)session;
 {
-	//First process other systems doing things for all events processing on URL Session
-	FIOSBackgroundDownloadCoreDelegates::OnIOSBackgroundDownload_SessionDidFinishAllEvents.Broadcast(session);
-
-    IOSAppDelegate* AppDelegate = [IOSAppDelegate GetDelegate];
-    
-    if ((nullptr != AppDelegate) && (nil != AppDelegate.BackgroundSessionEventCompleteDelegate))
-    {
-        //make a local copy of the completion handler to capture with the block
-        void(^completionHandler)() = AppDelegate.BackgroundSessionEventCompleteDelegate;
-        [completionHandler retain];
-        
-        //set the completion handler to nil to protect from anything else trying to use it
-        AppDelegate.BackgroundSessionEventCompleteDelegate = nil;
-        
-       //Need to issue our stored completion callback on main thread as its a UI callback
-        [[NSOperationQueue mainQueue] addOperationWithBlock:^{
-            //call the copied completion handler
-            completionHandler();
-            
-            //Release local copy so we don't leak the completion handler
-            [completionHandler release];
-        }];
-    }
+	//Check if we have this delegate setup. If not our app may have not been loaded or we were evicted.
+	const bool bHasURLSessionDelegateBeenSetup = FIOSBackgroundDownloadCoreDelegates::OnIOSBackgroundDownload_SessionDidFinishAllEvents.IsBound();
+	
+	if (bHasURLSessionDelegateBeenSetup)
+	{
+		//We call this here and then the completion handler ourselves below this. That way we delay the complete until all systems responding to this delegate have received the broadcast.
+		FBackgroundURLSessionHandler::AddDelayedBackgroundURLSessionComplete();
+		
+		//First process other systems doing things for all events processing on URL Session
+		FIOSBackgroundDownloadCoreDelegates::OnIOSBackgroundDownload_SessionDidFinishAllEvents.Broadcast(session, FIOSBackgroundDownloadCoreDelegates::OnDelayedBackgroundURLSessionCompleteHandler);
+		
+		//Now call the completion handler to erase the delay we had above
+		FBackgroundURLSessionHandler::OnDelayedBackgroundURLSessionCompleteHandlerCalled();
+	}
+	//Handle restarting any suspended tasks. These can exist because we didn't resume these before our app was evicted.
+	else
+	{
+		NSLog(@"Resuming all tasks regardless of request state as we have no registered handler of OnIOSBackgroundDownload_SessionDidFinishAllEvents");
+		
+		[session getTasksWithCompletionHandler : ^ (NSArray<__kindof NSURLSessionDataTask*>* DataTasks, NSArray<__kindof NSURLSessionUploadTask*>* UploadTasks, NSArray<__kindof NSURLSessionDownloadTask*>* DownloadTasks)
+		{
+			for (NSURLSessionDownloadTask* DownloadTask : DownloadTasks)
+			{
+				if ([DownloadTask state] == NSURLSessionTaskStateSuspended)
+				{
+					NSLog(@"Resuming Task for URL:%s", [[[[DownloadTask currentRequest] URL] absoluteString] UTF8String]);
+					[DownloadTask resume];
+				}
+			}
+			
+			//Go ahead and call the completion handler now that we have resumed everything.
+			FBackgroundURLSessionHandler::OnDelayedBackgroundURLSessionCompleteHandlerCalled();
+		}];
+	}
 }
 
 @end
diff --git a/Engine/Source/Runtime/ApplicationCore/Public/IOS/IOSBackgroundURLSessionHandler.h b/Engine/Source/Runtime/ApplicationCore/Public/IOS/IOSBackgroundURLSessionHandler.h
--- a/Engine/Source/Runtime/ApplicationCore/Public/IOS/IOSBackgroundURLSessionHandler.h
+++ b/Engine/Source/Runtime/ApplicationCore/Public/IOS/IOSBackgroundURLSessionHandler.h
@@ -12,12 +12,14 @@ public:
     DECLARE_MULTICAST_DELEGATE_ThreeParams(FIOSBackgroundDownload_DidFinishDownloadingToURL, NSURLSessionDownloadTask*, NSError*, const FString&);
     DECLARE_MULTICAST_DELEGATE_FourParams(FIOSBackgroundDownload_DidWriteData, NSURLSessionDownloadTask*, int64_t /*Bytes Written Since Last Call */, int64_t /*Total Bytes Written */, int64_t /*Total Bytes Expedted To Write */);
     DECLARE_MULTICAST_DELEGATE_TwoParams(FIOSBackgroundDownload_DidCompleteWithError, NSURLSessionTask*, NSError*);
-    DECLARE_MULTICAST_DELEGATE_OneParam(FIOSBackgroundDownload_SessionDidFinishAllEvents, NSURLSession*);
-    
+	DECLARE_DELEGATE(FIOSBackgroundDownload_DelayedBackgroundURLSessionCompleteHandler);
+	DECLARE_MULTICAST_DELEGATE_TwoParams(FIOSBackgroundDownload_SessionDidFinishAllEvents, NSURLSession*, FIOSBackgroundDownload_DelayedBackgroundURLSessionCompleteHandler);
+	
 	static FIOSBackgroundDownload_DidFinishDownloadingToURL OnIOSBackgroundDownload_DidFinishDownloadingToURL;
 	static FIOSBackgroundDownload_DidWriteData OnIOSBackgroundDownload_DidWriteData;
 	static FIOSBackgroundDownload_DidCompleteWithError OnIOSBackgroundDownload_DidCompleteWithError;
 	static FIOSBackgroundDownload_SessionDidFinishAllEvents OnIOSBackgroundDownload_SessionDidFinishAllEvents;
+	static FIOSBackgroundDownload_DelayedBackgroundURLSessionCompleteHandler OnDelayedBackgroundURLSessionCompleteHandler;
 };
 
 //Interface for wrapping a NSURLSession configured to support background downloading of NSURLSessionDownloadTasks.
@@ -46,12 +48,26 @@ public:
 	static const FString GetTemporaryFilePathFromURL(const FString& URL);
 
 	static void CreateBackgroundSessionWorkingDirectory();
+	
+	//Function to mark if you would like for the NSURLSession to wait to call the completion handler when
+	//OnIOSBackgroundDownload_SessionDidFinishAllEvents is called for you to call the passed completion handler
+	//NOTE: Call DURING OnIOSBackgroundDownload_SessionDidFinishAllEvents
+	static void AddDelayedBackgroundURLSessionComplete();
+	
+	//Function to handle calls to OnDelayedBackgroundURLSessionCompleteHandler
+	static void OnDelayedBackgroundURLSessionCompleteHandlerCalled();
+	
 private:
 	static NSURLSession* BackgroundSession;
 	static FString CachedIdentifierName;
+	
+	//Used to track calls to AddDelayedBackgroundURLSessionComplete vs calls to the completion handler.
+	static volatile int32 DelayedBackgroundURLSessionCompleteCount;
+	
+	//Used to call the stored background url session callback
+	static void CallBackgroundURLSessionCompleteHandler();
 };
 
-
 //Delegate object associated with our above NSURLSession
 @interface BackgroundDownloadDelegate : NSObject<NSURLSessionDownloadDelegate>
 @end
diff --git a/Engine/Source/Runtime/Online/BackgroundHTTP/Private/BackgroundHttpManagerImpl.cpp b/Engine/Source/Runtime/Online/BackgroundHTTP/Private/BackgroundHttpManagerImpl.cpp
--- a/Engine/Source/Runtime/Online/BackgroundHTTP/Private/BackgroundHttpManagerImpl.cpp
+++ b/Engine/Source/Runtime/Online/BackgroundHTTP/Private/BackgroundHttpManagerImpl.cpp
@@ -248,18 +248,48 @@ void FBackgroundHttpManagerImpl::ActivatePendingRequests()
 					RemainingRequests.Reserve(PendingStartRequests.Num() - NumRequestsWeCanProcess);
 				}
 
-				for (int RequestIndex = 0; RequestIndex < PendingStartRequests.Num(); ++RequestIndex)
+				//First preload the first N spots of the list with the start of the pending list.
+				int RequestIndex = 0;
+				for (;RequestIndex < PendingStartRequests.Num();++RequestIndex)
 				{
 					FBackgroundHttpRequestPtr& PendingRequest = PendingStartRequests[RequestIndex];
+					
 					if (RequestIndex <= NumRequestsWeCanProcess)
 					{
 						RequestsStartingThisTick.Add(PendingRequest);
 					}
 					else
 					{
-						RemainingRequests.Add(PendingRequest);
+						break;
 					}
 				}
+				
+				//Now go through the rest of the elements, replacing any items found with things of higher priority if needed
+				for (; RequestIndex < PendingStartRequests.Num(); ++RequestIndex)
+				{
+					FBackgroundHttpRequestPtr PendingRequestToCheck = PendingStartRequests[RequestIndex];
+					EBackgroundHTTPPriority PendingRequestPriority = PendingRequestToCheck->GetRequestPriority();
+					
+					for (int AlreadyGrabbedRequestIndex = 0; AlreadyGrabbedRequestIndex < RequestsStartingThisTick.Num();++AlreadyGrabbedRequestIndex)
+					{
+						FBackgroundHttpRequestPtr& AlreadyGrabbedRequest = RequestsStartingThisTick[AlreadyGrabbedRequestIndex];
+						EBackgroundHTTPPriority AlreadyGrabbedRequestPriority = AlreadyGrabbedRequest->GetRequestPriority();
+						
+						//Found a higher priority request, swap it and go back through entire list again
+						if (PendingRequestPriority < AlreadyGrabbedRequestPriority)
+						{
+							FBackgroundHttpRequestPtr RequestPtrCopy = AlreadyGrabbedRequest;
+							AlreadyGrabbedRequest = PendingRequestToCheck;
+							PendingRequestToCheck = RequestPtrCopy;
+							
+							//-1 so that it increments to 0
+							AlreadyGrabbedRequestIndex = -1;
+						}
+					}
+					
+					//Whatever is left in PendingRequestToCheck was too low priority to end up in the correct list.
+					RemainingRequests.Add(PendingRequestToCheck);
+				}
 
 				PendingStartRequests = MoveTemp(RemainingRequests);
 			}
diff --git a/Engine/Source/Runtime/Online/BackgroundHTTP/Private/BackgroundHttpRequestImpl.cpp b/Engine/Source/Runtime/Online/BackgroundHTTP/Private/BackgroundHttpRequestImpl.cpp
--- a/Engine/Source/Runtime/Online/BackgroundHTTP/Private/BackgroundHttpRequestImpl.cpp
+++ b/Engine/Source/Runtime/Online/BackgroundHTTP/Private/BackgroundHttpRequestImpl.cpp
@@ -13,6 +13,7 @@ FBackgroundHttpRequestImpl::FBackgroundHttpRequestImpl()
 	, URLList()
 	, RequestID()
 	, NumberOfTotalRetries(0)
+	, RequestPriority(EBackgroundHTTPPriority::Normal)
 	, HttpRequestCompleteDelegate()
 	, HttpProgressUpdateDelegate()
 {
@@ -138,3 +139,13 @@ bool FBackgroundHttpRequestImpl::HandleDelayedProcess()
 	//BackgroundHttpRequest if the platform expects to make use of this.
 	return ensureAlwaysMsgf(false, TEXT("Platform expects an implementation of HandleDelayedProcess on the BackgroundHttpRequest, but none found!"));
 }
+
+EBackgroundHTTPPriority FBackgroundHttpRequestImpl::GetRequestPriority() const
+{
+	return RequestPriority;
+}
+
+void FBackgroundHttpRequestImpl::SetRequestPriority(EBackgroundHTTPPriority NewPriority)
+{
+	RequestPriority = NewPriority;
+}
diff --git a/Engine/Source/Runtime/Online/BackgroundHTTP/Private/IOS/ApplePlatformBackgroundHttpManager.cpp b/Engine/Source/Runtime/Online/BackgroundHTTP/Private/IOS/ApplePlatformBackgroundHttpManager.cpp
--- a/Engine/Source/Runtime/Online/BackgroundHTTP/Private/IOS/ApplePlatformBackgroundHttpManager.cpp
+++ b/Engine/Source/Runtime/Online/BackgroundHTTP/Private/IOS/ApplePlatformBackgroundHttpManager.cpp
@@ -1,6 +1,7 @@
 // Copyright 1998-2019 Epic Games, Inc. All Rights Reserved.
 
 #include "IOS/ApplePlatformBackgroundHttpManager.h"
+#include "IOS/ApplePlatformBackgroundHttp.h"
 #include "IOS/ApplePlatformBackgroundHttpRequest.h"
 #include "IOS/ApplePlatformBackgroundHttpResponse.h"
 
@@ -200,7 +201,7 @@ void FApplePlatformBackgroundHttpManager::RemoveURLMapEntriesForRequest(FAppleBa
 void FApplePlatformBackgroundHttpManager::StartRequest(FAppleBackgroundHttpRequestPtr Request)
 {
     //Just count it as a retry that won't increment the retry counter before giving us the URL as our RetryCount 0 should start this up.
-    RetryRequest(Request,false, false, nullptr);
+    RetryRequest(Request,false, nullptr);
 }
 
 void FApplePlatformBackgroundHttpManager::RemoveRequest(const FBackgroundHttpRequestPtr Request)
@@ -341,7 +342,7 @@ void FApplePlatformBackgroundHttpManager::OnApp_EnteringForeground()
 void FApplePlatformBackgroundHttpManager::OnApp_EnteringBackground()
 {
 	FPlatformAtomics::InterlockedExchange(&bIsInBackground,true);
-	ResumeAllTasksWithoutPausedRequests();
+	ResumeTasksForBackgrounding();
 }
 
 void FApplePlatformBackgroundHttpManager::PauseAllActiveTasks()
@@ -372,9 +373,9 @@ void FApplePlatformBackgroundHttpManager::PauseAllActiveTasks()
 	}
 }
 
-void FApplePlatformBackgroundHttpManager::ResumeAllTasksWithoutPausedRequests()
+void FApplePlatformBackgroundHttpManager::ResumeTasksForBackgrounding(FIOSBackgroundHttpPostSessionWorkCallback Callback)
 {
-    UE_LOG(LogBackgroundHttpManager, Display, TEXT("Attempting to resume all active tasks that don't have a corresponding request that has paused them"));
+    UE_LOG(LogBackgroundHttpManager, Display, TEXT("Attempting to resume all active tasks that don't have a corresponding request that has paused them in priority order"));
     
 	NSURLSession* BackgroundDownloadSession = FBackgroundURLSessionHandler::GetBackgroundSession();
 	if (nullptr != BackgroundDownloadSession)
@@ -382,31 +383,72 @@ void FApplePlatformBackgroundHttpManager::ResumeAllTasksWithoutPausedRequests()
 		[BackgroundDownloadSession getTasksWithCompletionHandler : ^ (NSArray<__kindof NSURLSessionDataTask*>* DataTasks, NSArray<__kindof NSURLSessionUploadTask*>* UploadTasks, NSArray<__kindof NSURLSessionDownloadTask*>* DownloadTasks)
 		{
             FRWScopeLock ScopeLock(URLToRequestMapLock, SLT_ReadOnly);
-            
-			for (NSURLSessionDownloadTask* DownloadTask : DownloadTasks)
+			
+			//We only want to automatically re-queue the highest priority things, so go through all tasks by priority order and stop once we have requeued something
+			bool bDidResumeATask = false;
+			for (uint8 PriorityAsInt = (uint8)EBackgroundHTTPPriority::High; (!bDidResumeATask && (PriorityAsInt < (uint8)EBackgroundHTTPPriority::Num)); ++PriorityAsInt)
 			{
-				if ([DownloadTask state] == NSURLSessionTaskStateSuspended)
+				for (NSURLSessionDownloadTask* DownloadTask : DownloadTasks)
 				{
-                    FString TaskURL = [[[DownloadTask currentRequest] URL] absoluteString];
-                    int TaskIdentifier = (int)[DownloadTask taskIdentifier];
-                    FBackgroundHttpURLMappedRequestPtr* WeakRequestInMap = URLToRequestMap.Find(TaskURL);
-                    FAppleBackgroundHttpRequestPtr FoundRequest = (nullptr != WeakRequestInMap) ? WeakRequestInMap->Pin() : nullptr;
-                    
-                    bool bIsRequestPaused = FoundRequest.IsValid() ? FoundRequest->bIsTaskPaused : false;
-                    if (!bIsRequestPaused)
-                    {
-                        UE_LOG(LogBackgroundHttpManager, Display, TEXT("Resuming Task for URL:%s | TaskIdentifier:%d"), *TaskURL, TaskIdentifier);
-                        [DownloadTask resume];
-                    }
-                    else
-                    {
-                        UE_LOG(LogBackgroundHttpManager, Display, TEXT("NOT Resuming Task for URL as associated request was paused! URL:%s | TaskIdentifier:%d"), *TaskURL, TaskIdentifier);
-                    }
+					bDidResumeATask = bDidResumeATask || ResumeDownloadTaskForBackgroundingIfAppropriate(DownloadTask, EBackgroundHTTPPriority(PriorityAsInt));
+				}
+				
+				if (bDidResumeATask)
+				{
+					UE_LOG(LogBackgroundHttpManager, Display, TEXT("Resumed Tasks of Priority: %s . Not Resuming Any Lower Priority Tasks."), LexToString((EBackgroundHTTPPriority)PriorityAsInt));
 				}
 			}
+			
+			Callback.ExecuteIfBound();
 		}];
 	}
 }
+				 
+bool FApplePlatformBackgroundHttpManager::ResumeDownloadTaskForBackgroundingIfAppropriate(NSURLSessionDownloadTask* DownloadTask, EBackgroundHTTPPriority LowestPriorityToQueue)
+{
+	if (nullptr == DownloadTask)
+	{
+		return false;
+	}
+
+	bool bDidResumeTask = false;
+	
+	//We don't care about Cancelled, Completed, or Completing Task as those are already handled. Suspended tasks are waiting
+	//to be resumed to activate
+	if ([DownloadTask state] == NSURLSessionTaskStateSuspended)
+	{
+		FString TaskURL = [[[DownloadTask currentRequest] URL] absoluteString];
+		int TaskIdentifier = (int)[DownloadTask taskIdentifier];
+		FBackgroundHttpURLMappedRequestPtr* WeakRequestInMap = URLToRequestMap.Find(TaskURL);
+		FAppleBackgroundHttpRequestPtr FoundRequest = (nullptr != WeakRequestInMap) ? WeakRequestInMap->Pin() : nullptr;
+
+		const bool bIsRequestPaused = FoundRequest.IsValid() ? FoundRequest->bIsTaskPaused : false;
+		if (!bIsRequestPaused)
+		{
+			EBackgroundHTTPPriority FoundRequestPriority = FoundRequest->GetRequestPriority();
+			if (FoundRequestPriority <= LowestPriorityToQueue)
+			{
+				NSLog(@"TaskURL:%@",TaskURL.GetNSString());
+				NSLog(@"Priority:%d",(uint8)FoundRequestPriority);
+				UE_LOG(LogBackgroundHttpManager, Display, TEXT("TaskPriority:%s"), LexToString(FoundRequestPriority));
+				UE_LOG(LogBackgroundHttpManager, Display, TEXT("Resuming Task for URL:%s | TaskIdentifier:%d | TaskPriority:%s"), *TaskURL, TaskIdentifier, LexToString(FoundRequestPriority));
+				[DownloadTask resume];
+				
+				bDidResumeTask = true;
+			}
+			else
+			{
+				UE_LOG(LogBackgroundHttpManager, Verbose, TEXT("NOT RESUMING Task for URL because its a lower priority:%s | TaskIdentifier:%d | TaskPriority:%s | LowestPriorityToQueue:%s"), *TaskURL, TaskIdentifier, LexToString(FoundRequestPriority), LexToString(LowestPriorityToQueue));
+			}
+		}
+		else
+		{
+			UE_LOG(LogBackgroundHttpManager, Verbose, TEXT("NOT Resuming Task for URL as associated request was paused! URL:%s | TaskIdentifier:%d"), *TaskURL, TaskIdentifier);
+		}
+	}
+	
+	return bDidResumeTask;
+}
 
 void FApplePlatformBackgroundHttpManager::OnTask_DidFinishDownloadingToURL(NSURLSessionDownloadTask* Task, NSError* Error, const FString& TempFilePath)
 {
@@ -528,13 +570,12 @@ void FApplePlatformBackgroundHttpManager::FinishRequest(FAppleBackgroundHttpRequ
     }
 }
 
-void FApplePlatformBackgroundHttpManager::RetryRequest(FAppleBackgroundHttpRequestPtr Request, bool bShouldIncreaseRetryCount, bool bShouldStartImmediately, NSData* RetryData)
+void FApplePlatformBackgroundHttpManager::RetryRequest(FAppleBackgroundHttpRequestPtr Request, bool bShouldIncreaseRetryCount, NSData* RetryData)
 {
 	NSURLSessionDownloadTask* NewTask = nullptr;
 
 	if (ensureAlwaysMsgf(Request.IsValid(), TEXT("Call to RetryRequest with an invalid request!")))
 	{
-
 		NSURLSession* BackgroundDownloadSession = FBackgroundURLSessionHandler::GetBackgroundSession();
 		if (ensureAlwaysMsgf((nullptr != BackgroundDownloadSession), TEXT("Invalid Background Download NSURLSession during RetryRequest! Should have already Initialized the NSURLSession by this point!")))
 		{
@@ -565,15 +606,18 @@ void FApplePlatformBackgroundHttpManager::RetryRequest(FAppleBackgroundHttpReque
 			{
 				Request->AssociateWithTask(NewTask);
 
-				//If we are in BG or flagged this as an immediate start, resume right now without waiting for the FG tick.
+				//If we are in BG activate right now without waiting for the FG tick
 				volatile bool bCopyOfBGState = FPlatformAtomics::AtomicRead(&bIsInBackground);
-				if (!Request->bIsTaskPaused
-                    && (bCopyOfBGState || bShouldStartImmediately))
+				
+				//We also want to re-activate any already activated tasks
+				const bool bIsTaskActive = Request->IsUnderlyingTaskActive();
+				
+				if (!Request->bIsTaskPaused && (bCopyOfBGState || bIsTaskActive))
 				{
 					Request->ActivateUnderlyingTask();
 				}
 
-                UE_LOG(LogBackgroundHttpManager, Display, TEXT("Created Task for Request -- RequestDebugID:%s | bStartImmediately:%d | bIsAppInBG:%d | bIsPaused:%d"), *(Request->GetRequestDebugID()), (int)bShouldStartImmediately, (int)bCopyOfBGState, (int)Request->bIsTaskPaused);
+                UE_LOG(LogBackgroundHttpManager, Display, TEXT("Created Task for Request -- RequestDebugID:%s | bIsAppInBG:%d | bIsPaused:%d"), *(Request->GetRequestDebugID()), (int)bCopyOfBGState, (int)Request->bIsTaskPaused);
 
 				//Always set our bWasTaskStartedInBG flag on our Request so we will know if we need to restart this task next FG Tick.
 				FPlatformAtomics::InterlockedExchange(&(Request->bWasTaskStartedInBG), bCopyOfBGState);
@@ -683,7 +727,7 @@ void FApplePlatformBackgroundHttpManager::OnTask_DidCompleteWithError(NSURLSessi
             
 			if (bDidFindValidRequest)
 			{
-				RetryRequest(FoundRequest, bShouldRetryIncreaseRetryCount, true, ResumeData);
+				RetryRequest(FoundRequest, bShouldRetryIncreaseRetryCount, ResumeData);
 			}
 			else
 			{
@@ -694,10 +738,28 @@ void FApplePlatformBackgroundHttpManager::OnTask_DidCompleteWithError(NSURLSessi
 	}
 }
 
-void FApplePlatformBackgroundHttpManager::OnSession_SessionDidFinishAllEvents(NSURLSession* Session)
+void FApplePlatformBackgroundHttpManager::OnSession_SessionDidFinishAllEvents(NSURLSession* Session, FIOSBackgroundDownloadCoreDelegates::FIOSBackgroundDownload_DelayedBackgroundURLSessionCompleteHandler Callback)
 {
-	//@TODO: TRoss, For now not using this. Will be using this to track how long we take in BG Handling BG Downloads in analytics
-    UE_LOG(LogBackgroundHttpManager, Verbose, TEXT("NSURLSession done sending background events"));
+	UE_LOG(LogBackgroundHttpManager, Verbose, TEXT("NSURLSession done sending background events for all already queued tasks"));
+	
+	//Let BackgroundURLSessionHandler know that it should wait until we call the callback
+	FBackgroundURLSessionHandler::AddDelayedBackgroundURLSessionComplete();
+	
+	const bool bCopyIsBG = FPlatformAtomics::AtomicRead(&bIsInBackground);
+	
+	if (bCopyIsBG)
+	{
+		//Now that we have finished all queued background tasks, lets resume tasks for any lower priorities that weren't started
+		ResumeTasksForBackgrounding(FIOSBackgroundDownloadCoreDelegates::FIOSBackgroundDownload_DelayedBackgroundURLSessionCompleteHandler::CreateLambda(
+		[=]()
+		{
+			Callback.ExecuteIfBound();
+		}));
+	}
+	else
+	{
+		Callback.ExecuteIfBound();
+	}
 }
 
 bool FApplePlatformBackgroundHttpManager::Tick(float DeltaTime)
@@ -705,7 +767,6 @@ bool FApplePlatformBackgroundHttpManager::Tick(float DeltaTime)
     ensureAlwaysMsgf(IsInGameThread(), TEXT("Called from un-expected thread! Potential error in an implementation of background downloads!"));
     
     TickRequests(DeltaTime);
-    TickTasks(DeltaTime);
     TickUnassociatedTasks(DeltaTime);
     
     //Always keep ticking
@@ -718,7 +779,12 @@ void FApplePlatformBackgroundHttpManager::TickRequests(float DeltaTime)
     
     //First lets go through all our Requests to see if we need to complete or recreate any requests
     {
-        FRWScopeLock ScopeLock(ActiveRequestLock, SLT_Write);
+		//Check to make sure we have room for more tasks to be active first
+		int CurrentCount = FPlatformAtomics::AtomicRead(&NumCurrentlyActiveTasks);
+		bool bNeedsMoreTasks = (CurrentCount < MaxActiveDownloads);
+		FAppleBackgroundHttpRequestPtr FoundRequestToStart = nullptr;
+		
+        FRWScopeLock ScopeLock(ActiveRequestLock, SLT_ReadOnly);
         for (FBackgroundHttpRequestPtr& Request : ActiveRequests)
         {
             FAppleBackgroundHttpRequestPtr AppleRequest = StaticCastSharedPtr<FApplePlatformBackgroundHttpRequest>(Request);
@@ -761,110 +827,54 @@ void FApplePlatformBackgroundHttpManager::TickRequests(float DeltaTime)
                     
                     AppleRequest->SendDownloadProgressUpdate();
                 }
+				else if (bNeedsMoreTasks && !bIsTaskActive && !bIsTaskPaused && !bIsPendingCancel)
+				{
+					//This task is a possible to start, so lets see if its the highest priority task we have found yet.
+					
+					if (!FoundRequestToStart.IsValid())
+					{
+						FoundRequestToStart = AppleRequest;
+					}
+					else
+					{
+						if (AppleRequest->GetRequestPriority() < FoundRequestToStart->GetRequestPriority())
+						{
+							FoundRequestToStart = AppleRequest;
+						}
+					}
+					
+					AppleRequest->SendDownloadProgressUpdate();
+				}
                 else
                 {
                     AppleRequest->SendDownloadProgressUpdate();
                 }
             }
         }
+		
+		if (FoundRequestToStart.IsValid())
+		{
+			const int32 OldTotal = FPlatformAtomics::InterlockedIncrement(&NumCurrentlyActiveTasks);
+			
+			if (OldTotal <= MaxActiveDownloads)
+			{
+				UE_LOG(LogBackgroundHttpManager, Display, TEXT("Starting Task for Request -- RequestDebugID:%s | CurrentlyActiveRequests:%d"), *(FoundRequestToStart->GetRequestDebugID()), OldTotal);
+				FoundRequestToStart->ActivateUnderlyingTask();
+			}
+			else
+			{
+				UE_LOG(LogBackgroundHttpManager, Log, TEXT("Request failed to activate as we passed the platform max from another requeset before we could resume. RequestDebugID:%s | CurrentlyActiveRequests:%d"), *(FoundRequestToStart->GetRequestDebugID()), OldTotal);
+				
+				//Don't activate and remove our increment from above because something put us over the limit before we resumed
+				FPlatformAtomics::InterlockedDecrement(&NumCurrentlyActiveTasks);
+			}
+		}
     }
     
     //Now that we have gone through and finished all the requests, go ahead and delete any pending removes
     DeletePendingRemoveRequests();
 }
 
-void FApplePlatformBackgroundHttpManager::TickTasks(float DeltaTime)
-{
-    ensureAlwaysMsgf(IsInGameThread(), TEXT("Called from un-expected thread! Potential error in an implementation of background downloads!"));
-    
-    //Go through all our session's tasks and unpause as many tasks as our platform's limit allows
-    
-    NSURLSession* BackgroundDownloadSession = FBackgroundURLSessionHandler::GetBackgroundSession();
-    if (nullptr != BackgroundDownloadSession)
-    {
-        //Make sure we only queue up 1 tick's worth of task parsing at a time as this is an async callback
-        const bool IsIterateAlreadyQueued = FPlatformAtomics::InterlockedExchange(&bIsIteratingThroughSessionTasks, true);
-        if (!IsIterateAlreadyQueued)
-        {
-            [BackgroundDownloadSession getAllTasksWithCompletionHandler:^(NSArray<__kindof NSURLSessionTask*> *tasks)
-             {
-                 //Check to make sure we have room for more tasks to be active first
-                 int CurrentCount = FPlatformAtomics::AtomicRead(&NumCurrentlyActiveTasks);
-                 if (CurrentCount < MaxActiveDownloads)
-                 {
-                     //Go through our tasks and try and activate as many as possible
-                     for (NSURLSessionTask* task in tasks)
-                     {
-                         //Only resume state if the task is suspended, not intersted in already active, canceling, or completed tasks here
-                         if ([task state] == NSURLSessionTaskStateSuspended)
-                         {
-                             //Check to make sure we didn't have another task increment our count over the limit since we checked above
-                             int NewRequestCount = FPlatformAtomics::InterlockedIncrement(&NumCurrentlyActiveTasks);
-                             UE_LOG(LogBackgroundHttpManager,Verbose, TEXT("Incrementing Task Count: %d"), NewRequestCount);
-                             
-                             if (NewRequestCount <= MaxActiveDownloads)
-                             {
-                                 FString TaskURL = [[[task currentRequest] URL] absoluteString];
-                                 int TaskIdentifier = (int)[task taskIdentifier];
-                                 
-                                 //Try and find Request in map that matches this Task
-                                 {
-                                     FRWScopeLock ScopeLock(URLToRequestMapLock, SLT_ReadOnly);
-                                     FBackgroundHttpURLMappedRequestPtr* WeakRequestInMap = URLToRequestMap.Find(TaskURL);
-                                     FAppleBackgroundHttpRequestPtr FoundRequest = (nullptr != WeakRequestInMap) ? WeakRequestInMap->Pin() : nullptr;
-                                     
-                                     const bool bIsPaused = FoundRequest.IsValid() ? FoundRequest->IsUnderlyingTaskPaused() : false;
-                                     if (FoundRequest.IsValid() && !bIsPaused)
-                                     {
-                                         UE_LOG(LogBackgroundHttpManager, Display, TEXT("Manager Calling to Active Task For Request -- RequestDebugID:%s | TaskURL:%s | TaskIdentifier:%d | CurrentlyActiveRequests:%d"), *(FoundRequest->GetRequestDebugID()), *TaskURL, TaskIdentifier, NewRequestCount);
-                                         FoundRequest->ActivateUnderlyingTask();
-                                     }
-                                     else
-                                     {
-                                         if (FoundRequest.IsValid() && bIsPaused)
-                                         {
-                                             UE_LOG(LogBackgroundHttpManager, Verbose, TEXT("Skipping Activating Task as the associated request is paused. -- TaskURL:%s | TaskIdentifier:%d"), *TaskURL, TaskIdentifier);
-                                         }
-                                         else
-                                         {
-                                             UE_LOG(LogBackgroundHttpManager, Display, TEXT("Skipping Activating Task as there is no associated Request. Once a Request associates with this task, it can then be activated. -- TaskURL:%s | TaskIdentifier:%d"), *TaskURL, TaskIdentifier);
-                                         }
-                                         
-                                         //Don't activate and remove our increment from above because something put us over the limit before we resumed
-                                         FPlatformAtomics::InterlockedDecrement(&NumCurrentlyActiveTasks);
-                                     }
-                                 }
-                             }
-                             else
-                             {                                 
-                                 FString TaskURL = [[[task currentRequest] URL] absoluteString];
-                                 int TaskIdentifier = (int)[task taskIdentifier];
-                                 
-                                 UE_LOG(LogBackgroundHttpManager, Log, TEXT("Task failed to activate as we passed the platform max from another task before we could resume. Task -- TaskURL:%s | TaskIdentifier:%d | CurrentlyActiveRequests:%d"), *TaskURL, TaskIdentifier, NewRequestCount);
-                                 
-                                 //Don't activate and remove our increment from above because something put us over the limit before we resumed
-                                 FPlatformAtomics::InterlockedDecrement(&NumCurrentlyActiveTasks);
-                                 break;
-                             }
-                             
-                             //We now have enough requests queued up that we can stop looking for more
-                             if (NewRequestCount >= MaxActiveDownloads)
-                             {
-                                 break;
-                             }
-                         }
-                     }
-                 }
-                 
-                 //If this wasn't true when we exited this async callback, we have a logic error somewhere and have queued this async callback twice. Should always be true as this async callback ends
-                 //Also setting this to false so that future ticks will call this callback again
-                 const bool bWasLogicError = FPlatformAtomics::InterlockedExchange(&bIsIteratingThroughSessionTasks, false);
-                 ensureAlwaysMsgf(bWasLogicError, TEXT("Leaving Tick's getAllTasksWithCompletionHandler and bIsIteratingThroughSessionTasks was false before we finished. Something else has set bIsIteratingThroughSessionTasks during our callback!"));
-             }];
-        }
-    }
-}
-
 void FApplePlatformBackgroundHttpManager::TickUnassociatedTasks(float DeltaTime)
 {
     ensureAlwaysMsgf(IsInGameThread(), TEXT("Called from un-expected thread! Potential error in an implementation of background downloads!"));
diff --git a/Engine/Source/Runtime/Online/BackgroundHTTP/Public/BackgroundHttpRequestImpl.h b/Engine/Source/Runtime/Online/BackgroundHTTP/Public/BackgroundHttpRequestImpl.h
--- a/Engine/Source/Runtime/Online/BackgroundHTTP/Public/BackgroundHttpRequestImpl.h
+++ b/Engine/Source/Runtime/Online/BackgroundHTTP/Public/BackgroundHttpRequestImpl.h
@@ -35,7 +35,9 @@ public:
 	virtual const FString& GetRequestID() const override;
 	virtual void SetRequestID(const FString& NewRequestID) override;
 	virtual bool HandleDelayedProcess() override;
-
+	virtual EBackgroundHTTPPriority GetRequestPriority() const override;
+	virtual void SetRequestPriority(EBackgroundHTTPPriority NewPriority) override;
+	
 	virtual void NotifyNotificationObjectOfComplete(bool bWasSuccess);
 protected:
 	TSharedPtr<FBackgroundHttpNotificationObject, ESPMode::ThreadSafe> DownloadCompleteNotificationObject;
@@ -43,7 +45,8 @@ protected:
 	TArray<FString> URLList;
 	FString RequestID;
 	int NumberOfTotalRetries;
-
+	EBackgroundHTTPPriority RequestPriority;
+	
 	FBackgroundHttpRequestCompleteDelegate HttpRequestCompleteDelegate;
 	FBackgroundHttpProgressUpdateDelegate HttpProgressUpdateDelegate;
 };
diff --git a/Engine/Source/Runtime/Online/BackgroundHTTP/Public/IOS/ApplePlatformBackgroundHttpManager.h b/Engine/Source/Runtime/Online/BackgroundHTTP/Public/IOS/ApplePlatformBackgroundHttpManager.h
--- a/Engine/Source/Runtime/Online/BackgroundHTTP/Public/IOS/ApplePlatformBackgroundHttpManager.h
+++ b/Engine/Source/Runtime/Online/BackgroundHTTP/Public/IOS/ApplePlatformBackgroundHttpManager.h
@@ -8,8 +8,12 @@
 #include "Interfaces/IBackgroundHttpRequest.h"
 #include "IOS/ApplePlatformBackgroundHttpRequest.h"
 
+#include "IOS/IOSBackgroundURLSessionHandler.h"
+
 typedef TWeakPtr<FApplePlatformBackgroundHttpRequest,ESPMode::ThreadSafe> FBackgroundHttpURLMappedRequestPtr;
 
+DECLARE_DELEGATE(FIOSBackgroundHttpPostSessionWorkCallback);
+
 /**
  * Manages Background Http request that are currently being processed if we are on an Apple Platform
  */
@@ -74,18 +78,18 @@ private:
 
     //These are used internally to pause/resume things at the task level
 	void PauseAllActiveTasks();
-	void ResumeAllTasksWithoutPausedRequests();
+	void ResumeTasksForBackgrounding(FIOSBackgroundHttpPostSessionWorkCallback Callback = FIOSBackgroundHttpPostSessionWorkCallback());
+	bool ResumeDownloadTaskForBackgroundingIfAppropriate(NSURLSessionDownloadTask* DownloadTask, EBackgroundHTTPPriority LowestPriorityToQueue);
 
     void PauseAllUnassociatedTasks();
     void UnpauseAllUnassociatedTasks();
     
     void TickRequests(float DeltaTime);
-    void TickTasks(float DeltaTime);
     void TickUnassociatedTasks(float DeltaTime);
     
     void StartRequest(FAppleBackgroundHttpRequestPtr Request);
 	void FinishRequest(FAppleBackgroundHttpRequestPtr Request);
-	void RetryRequest(FAppleBackgroundHttpRequestPtr Request, bool bShouldIncreaseRetryCount, bool bShouldStartImmediately, NSData* RetryData);
+	void RetryRequest(FAppleBackgroundHttpRequestPtr Request, bool bShouldIncreaseRetryCount, NSData* RetryData);
     void RemoveSessionTasksForRequest(FAppleBackgroundHttpRequestPtr Request);
     bool GenerateURLMapEntriesForRequest(FAppleBackgroundHttpRequestPtr Request);
     void RemoveURLMapEntriesForRequest(FAppleBackgroundHttpRequestPtr Request);
@@ -99,7 +103,7 @@ private:
 	void OnTask_DidFinishDownloadingToURL(NSURLSessionDownloadTask* Task, NSError* Error, const FString& TempFilePath);
 	void OnTask_DidWriteData(NSURLSessionDownloadTask* Task, int64_t BytesWrittenSinceLastCall, int64_t TotalBytesWritten, int64_t TotalBytesExpectedToWrite);
 	void OnTask_DidCompleteWithError(NSURLSessionTask* Task, NSError* Error);
-	void OnSession_SessionDidFinishAllEvents(NSURLSession* Session);
+	void OnSession_SessionDidFinishAllEvents(NSURLSession* Session, FIOSBackgroundDownloadCoreDelegates::FIOSBackgroundDownload_DelayedBackgroundURLSessionCompleteHandler Callback);
 
 	FDelegateHandle OnApp_EnteringForegroundHandle;
 	FDelegateHandle OnApp_EnteringBackgroundHandle;
diff --git a/Engine/Source/Runtime/Online/BackgroundHTTP/Public/Interfaces/IBackgroundHttpRequest.h b/Engine/Source/Runtime/Online/BackgroundHTTP/Public/Interfaces/IBackgroundHttpRequest.h
--- a/Engine/Source/Runtime/Online/BackgroundHTTP/Public/Interfaces/IBackgroundHttpRequest.h
+++ b/Engine/Source/Runtime/Online/BackgroundHTTP/Public/Interfaces/IBackgroundHttpRequest.h
@@ -30,6 +30,34 @@ DECLARE_DELEGATE_TwoParams(FBackgroundHttpRequestCompleteDelegate, FBackgroundHt
 */
 DECLARE_DELEGATE_ThreeParams(FBackgroundHttpProgressUpdateDelegate, FBackgroundHttpRequestPtr /*Request*/, int32 /*TotalBytesWritten*/, int32 /*BytesWrittenSinceLastUpdate*/);
 
+/**
+ * Enum used to describe download priority. Higher priorities will be downloaded first.
+ * Note: Should always be kept in High -> Low priority order if adding more Priorities!
+ */
+enum class EBackgroundHTTPPriority : uint8
+{
+	  High
+	, Normal
+	, Low
+	, Num
+};
+
+inline const TCHAR* LexToString(EBackgroundHTTPPriority InType)
+{
+	switch (InType)
+	{
+	case EBackgroundHTTPPriority::High:
+		return TEXT("High");
+	case EBackgroundHTTPPriority::Normal:
+		return TEXT("Normal");
+	case EBackgroundHTTPPriority::Low:
+		return TEXT("Low");
+	default:
+		break;
+	}
+	return TEXT("<Unknown EBackgroundHTTPPriority>");
+}
+
 /**
  * Interface for Http requests (created using FHttpFactory)
  */
@@ -140,6 +168,20 @@ public:
 	*/
 	virtual void SetRequestID(const FString& NewRequestID) = 0;
 
+	/**
+	* Gets the associated Requeset's Priority for Background Downloadings. Where possible, we attempt to finish downloads in lowest-priority first order.
+	*
+	* @return The request Priority
+	*/
+	virtual EBackgroundHTTPPriority GetRequestPriority() const = 0;
+	
+	/**
+	* Sets the associated Requeset's Priority for Background Downloadings. Where possible, we attempt to finish downloads in lowest-priority first order.
+	*
+	* @param uint32 describing priority of this download. Lower happens first and thus 0 is the highest priority.
+	*/
+	virtual void SetRequestPriority(EBackgroundHTTPPriority NewPriority) = 0;
+	
 	/** 
 	* Destructor for overrides 
 	*/
