Adding LexTryParseString for EBackgroundHttpPriority enum.
#tests iOS Client

diff --git a/Engine/Source/Runtime/Online/BackgroundHTTP/Public/Interfaces/IBackgroundHttpRequest.h b/Engine/Source/Runtime/Online/BackgroundHTTP/Public/Interfaces/IBackgroundHttpRequest.h
--- a/Engine/Source/Runtime/Online/BackgroundHTTP/Public/Interfaces/IBackgroundHttpRequest.h
+++ b/Engine/Source/Runtime/Online/BackgroundHTTP/Public/Interfaces/IBackgroundHttpRequest.h
@@ -58,6 +58,26 @@ inline const TCHAR* LexToString(EBackgroundHTTPPriority InType)
 	return TEXT("<Unknown EBackgroundHTTPPriority>");
 }
 
+inline bool LexTryParseString(EBackgroundHTTPPriority& OutMode, const TCHAR* InBuffer)
+{
+	if (FCString::Stricmp(InBuffer, TEXT("High")) == 0)
+	{
+		OutMode = EBackgroundHTTPPriority::High;
+		return true;
+	}
+	if (FCString::Stricmp(InBuffer, TEXT("Normal")) == 0)
+	{
+		OutMode = EBackgroundHTTPPriority::Normal;
+		return true;
+	}
+	if (FCString::Stricmp(InBuffer, TEXT("Low")) == 0)
+	{
+		OutMode = EBackgroundHTTPPriority::Low;
+		return true;
+	}
+	return false;
+}
+
 /**
  * Interface for Http requests (created using FHttpFactory)
  */
