Add IsAvailable to HttpServerModule. Check if the module is loaded before stoping the WebServer listner.

diff --git a/Engine/Plugins/VirtualProduction/RemoteControl/Source/WebRemoteControl/Private/WebRemoteControl.cpp b/Engine/Plugins/VirtualProduction/RemoteControl/Source/WebRemoteControl/Private/WebRemoteControl.cpp
--- a/Engine/Plugins/VirtualProduction/RemoteControl/Source/WebRemoteControl/Private/WebRemoteControl.cpp
+++ b/Engine/Plugins/VirtualProduction/RemoteControl/Source/WebRemoteControl/Private/WebRemoteControl.cpp
@@ -599,9 +599,9 @@ public:
 		}
 		HttpRouter.Reset();
 
-		if (FHttpServerModule* HttpServerModule = FModuleManager::Get().GetModulePtr<FHttpServerModule>(TEXT("ContentBrowser")))
+		if (FHttpServerModule::IsAvailable())
 		{
-			HttpServerModule->StopAllListeners();
+			FHttpServerModule::Get().StopAllListeners();
 		}
 	}
 
diff --git a/Engine/Source/Runtime/Online/HTTPServer/Private/HttpServerModule.cpp b/Engine/Source/Runtime/Online/HTTPServer/Private/HttpServerModule.cpp
--- a/Engine/Source/Runtime/Online/HTTPServer/Private/HttpServerModule.cpp
+++ b/Engine/Source/Runtime/Online/HTTPServer/Private/HttpServerModule.cpp
@@ -4,7 +4,7 @@
 #include "HttpListener.h"
 #include "Modules/ModuleManager.h"
 
-DEFINE_LOG_CATEGORY(LogHttpServerModule); 
+DEFINE_LOG_CATEGORY(LogHttpServerModule);
 
 // FHttpServerModule 
 IMPLEMENT_MODULE(FHttpServerModule, HTTPServer);
@@ -14,18 +14,17 @@ FHttpServerModule* FHttpServerModule::Singleton = nullptr;
 void FHttpServerModule::StartupModule()
 {
 	Singleton = this;
-	bInitialized = true;
 }
 
 void FHttpServerModule::ShutdownModule()
 {
-	bInitialized = false;
-
 	// stop all listeners
 	StopAllListeners();
 
 	// destroy all listeners
 	Listeners.Empty();
+
+	Singleton = nullptr;
 }
 
 void FHttpServerModule::StartAllListeners()
@@ -48,7 +47,7 @@ void FHttpServerModule::StartAllListeners()
 
 void FHttpServerModule::StopAllListeners()
 {
-	UE_LOG(LogHttpServerModule, Log, 
+	UE_LOG(LogHttpServerModule, Log,
 		TEXT("Stopping all listeners..."));
 
 	for (const auto& Listener : Listeners)
@@ -64,7 +63,7 @@ void FHttpServerModule::StopAllListeners()
 
 }
 
-bool FHttpServerModule::HasPendingListeners() const 
+bool FHttpServerModule::HasPendingListeners() const
 {
 	for (const auto& Listener : Listeners)
 	{
@@ -76,6 +75,11 @@ bool FHttpServerModule::HasPendingListeners() const
 	return false;
 }
 
+bool FHttpServerModule::IsAvailable()
+{
+	return nullptr != Singleton;
+}
+
 FHttpServerModule& FHttpServerModule::Get()
 {
 	if (nullptr == Singleton)
@@ -89,7 +93,7 @@ FHttpServerModule& FHttpServerModule::Get()
 
 TSharedPtr<IHttpRouter> FHttpServerModule::GetHttpRouter(uint32 Port)
 {
-	check(bInitialized);
+	check(Singleton == this);
 
 	// We may already be listening on this port
 	TUniquePtr<FHttpListener>* ExistingListener = Listeners.Find(Port);
@@ -101,7 +105,7 @@ TSharedPtr<IHttpRouter> FHttpServerModule::GetHttpRouter(uint32 Port)
 	// Otherwise create a new one
 	TUniquePtr<FHttpListener> NewListener = MakeUnique<FHttpListener>(Port);
 
-    // Try to start this listener now
+	// Try to start this listener now
 	if (bHttpListenersEnabled)
 	{
 		NewListener->StartListening();
@@ -113,7 +117,6 @@ TSharedPtr<IHttpRouter> FHttpServerModule::GetHttpRouter(uint32 Port)
 bool FHttpServerModule::Tick(float DeltaTime)
 {
 	check(Singleton == this);
-	check(bInitialized);
 
 	if (bHttpListenersEnabled)
 	{
diff --git a/Engine/Source/Runtime/Online/HTTPServer/Public/HttpServerModule.h b/Engine/Source/Runtime/Online/HTTPServer/Public/HttpServerModule.h
--- a/Engine/Source/Runtime/Online/HTTPServer/Public/HttpServerModule.h
+++ b/Engine/Source/Runtime/Online/HTTPServer/Public/HttpServerModule.h
@@ -22,6 +22,13 @@ class FHttpServerModule :
 
 public:
 
+	/**
+	 * Checks to see if this module is loaded and ready.  It is only valid to call Get() if IsAvailable() returns true.
+	 *
+	 * @return True if the module is loaded and ready to use
+	 */
+	HTTPSERVER_API static bool IsAvailable();
+
 	/**
 	 * Singleton-like access to this module's interface.  This is just for convenience!
 	 * Beware of calling this during the shutdown phase, though.  Your module might have been unloaded already.
@@ -87,9 +94,6 @@ private:
 	/** The association of port bindings and respective HTTP listeners */
 	TMap<uint32, TUniquePtr<FHttpListener>> Listeners;
 
-	/** Whether this module has been initialized */
-	bool bInitialized = false;
-
 	/** Whether listeners can be started */
 	bool bHttpListenersEnabled = false;
 };
