Background HTTP Fixes:
- Fix ensure during backgrounding
- Fix potential crash when backgrounding
#tests iOS Client

diff --git a/Engine/Source/Runtime/Online/BackgroundHTTP/Private/IOS/ApplePlatformBackgroundHttpManager.cpp b/Engine/Source/Runtime/Online/BackgroundHTTP/Private/IOS/ApplePlatformBackgroundHttpManager.cpp
--- a/Engine/Source/Runtime/Online/BackgroundHTTP/Private/IOS/ApplePlatformBackgroundHttpManager.cpp
+++ b/Engine/Source/Runtime/Online/BackgroundHTTP/Private/IOS/ApplePlatformBackgroundHttpManager.cpp
@@ -18,6 +18,8 @@
 
 FApplePlatformBackgroundHttpManager::~FApplePlatformBackgroundHttpManager()
 {
+	bWasAppleBGHTTPInitialized = false;
+	
     [UnAssociatedTasks release];
     UnAssociatedTasks = nullptr;
     
@@ -27,6 +29,7 @@ FApplePlatformBackgroundHttpManager::~FApplePlatformBackgroundHttpManager()
 FString FApplePlatformBackgroundHttpManager::BackgroundSessionIdentifier = FString("");
 float FApplePlatformBackgroundHttpManager::ActiveTimeOutSetting = 30.0f;
 int FApplePlatformBackgroundHttpManager::RetryResumeDataLimitSetting = -1.f;
+volatile bool FApplePlatformBackgroundHttpManager::bWasAppleBGHTTPInitialized = false;
 
 FApplePlatformBackgroundHttpManager::FApplePlatformBackgroundHttpManager()
 	: bHasFinishedPopulatingUnassociatedTasks(false)
@@ -49,6 +52,8 @@ void FApplePlatformBackgroundHttpManager::Initialize()
 	SetupNSURLSessionResponseDelegates();
 
 	FBackgroundHttpManagerImpl::Initialize();
+	
+	bWasAppleBGHTTPInitialized = true;
 }
 
 void FApplePlatformBackgroundHttpManager::PopulateUnAssociatedTasks()
@@ -366,9 +371,6 @@ void FApplePlatformBackgroundHttpManager::PauseAllActiveTasks()
 					[DownloadTask suspend];
 				}
 			}
-            
-            //Reset our active requests to 0 now that we are pausing everything
-            FPlatformAtomics::InterlockedExchange(&NumCurrentlyActiveTasks, 0);
 		}];
 	}
 }
@@ -382,15 +384,28 @@ void FApplePlatformBackgroundHttpManager::ResumeTasksForBackgrounding(FIOSBackgr
 	{
 		[BackgroundDownloadSession getTasksWithCompletionHandler : ^ (NSArray<__kindof NSURLSessionDataTask*>* DataTasks, NSArray<__kindof NSURLSessionUploadTask*>* UploadTasks, NSArray<__kindof NSURLSessionDownloadTask*>* DownloadTasks)
 		{
-            FRWScopeLock ScopeLock(URLToRequestMapLock, SLT_ReadOnly);
-			
+			FRWScopeLock ScopeLock(URLToRequestMapLock, SLT_ReadOnly);
+		 
 			//We only want to automatically re-queue the highest priority things, so go through all tasks by priority order and stop once we have requeued something
 			bool bDidResumeATask = false;
 			for (uint8 PriorityAsInt = (uint8)EBackgroundHTTPPriority::High; (!bDidResumeATask && (PriorityAsInt < (uint8)EBackgroundHTTPPriority::Num)); ++PriorityAsInt)
 			{
 				for (NSURLSessionDownloadTask* DownloadTask : DownloadTasks)
 				{
-					bDidResumeATask = bDidResumeATask || ResumeDownloadTaskForBackgroundingIfAppropriate(DownloadTask, EBackgroundHTTPPriority(PriorityAsInt));
+					FString TaskURL = [[[DownloadTask currentRequest] URL] absoluteString];
+					int TaskIdentifier = (int)[DownloadTask taskIdentifier];
+		 
+					FBackgroundHttpURLMappedRequestPtr* WeakRequestInMap = URLToRequestMap.Find(TaskURL);
+					FAppleBackgroundHttpRequestPtr FoundRequest = ((nullptr != WeakRequestInMap) && (WeakRequestInMap->IsValid())) ? WeakRequestInMap->Pin() : nullptr;
+		 
+					if (FoundRequest.IsValid())
+					{
+						bDidResumeATask = bDidResumeATask || ResumeDownloadTaskForBackgroundingIfAppropriate(DownloadTask, FoundRequest, EBackgroundHTTPPriority(PriorityAsInt));
+					}
+					else
+					{
+						UE_LOG(LogBackgroundHttpManager, Verbose, TEXT("Skipped Resuming Task as there is no corresponding request! URL:%s | TaskIdentifier:%d"), *TaskURL, TaskIdentifier);
+					}
 				}
 				
 				if (bDidResumeATask)
@@ -403,10 +418,10 @@ void FApplePlatformBackgroundHttpManager::ResumeTasksForBackgrounding(FIOSBackgr
 		}];
 	}
 }
-				 
-bool FApplePlatformBackgroundHttpManager::ResumeDownloadTaskForBackgroundingIfAppropriate(NSURLSessionDownloadTask* DownloadTask, EBackgroundHTTPPriority LowestPriorityToQueue)
+
+bool FApplePlatformBackgroundHttpManager::ResumeDownloadTaskForBackgroundingIfAppropriate(NSURLSessionDownloadTask* DownloadTask, const FAppleBackgroundHttpRequestPtr Request, EBackgroundHTTPPriority LowestPriorityToQueue)
 {
-	if (nullptr == DownloadTask)
+	if ((nullptr == DownloadTask) || (nullptr == [DownloadTask currentRequest]) || (!Request.IsValid()))
 	{
 		return false;
 	}
@@ -419,13 +434,11 @@ bool FApplePlatformBackgroundHttpManager::ResumeDownloadTaskForBackgroundingIfAp
 	{
 		FString TaskURL = [[[DownloadTask currentRequest] URL] absoluteString];
 		int TaskIdentifier = (int)[DownloadTask taskIdentifier];
-		FBackgroundHttpURLMappedRequestPtr* WeakRequestInMap = URLToRequestMap.Find(TaskURL);
-		FAppleBackgroundHttpRequestPtr FoundRequest = (nullptr != WeakRequestInMap) ? WeakRequestInMap->Pin() : nullptr;
 
-		const bool bIsRequestPaused = FoundRequest.IsValid() ? FoundRequest->bIsTaskPaused : false;
+		const bool bIsRequestPaused = Request->bIsTaskPaused;
 		if (!bIsRequestPaused)
 		{
-			EBackgroundHTTPPriority FoundRequestPriority = FoundRequest->GetRequestPriority();
+			EBackgroundHTTPPriority FoundRequestPriority = Request->GetRequestPriority();
 			if (FoundRequestPriority <= LowestPriorityToQueue)
 			{
 				NSLog(@"TaskURL:%@",TaskURL.GetNSString());
@@ -469,7 +482,7 @@ void FApplePlatformBackgroundHttpManager::OnTask_DidFinishDownloadingToURL(NSURL
         {
             FRWScopeLock ScopeLock(URLToRequestMapLock, SLT_ReadOnly);
             FBackgroundHttpURLMappedRequestPtr* WeakRequestInMap = URLToRequestMap.Find(TaskURL);
-            FAppleBackgroundHttpRequestPtr FoundRequest = (nullptr != WeakRequestInMap) ? WeakRequestInMap->Pin() : nullptr;
+            FAppleBackgroundHttpRequestPtr FoundRequest = ((nullptr != WeakRequestInMap) && (WeakRequestInMap->IsValid())) ? WeakRequestInMap->Pin() : nullptr;
             
             if (FoundRequest.IsValid())
             {
@@ -545,9 +558,10 @@ void FApplePlatformBackgroundHttpManager::FinishRequest(FAppleBackgroundHttpRequ
             //If we are actually finishing this request, lets decrement our NumCurrentlyActiveTasks counter
             if (bIsRequestActuallyFinished)
             {
-                //See if we are in the BG or FG. In the BG, we don't want to decrement NumCurrentlyActiveTasks
-                const bool bWasFinishedInBG = FPlatformAtomics::AtomicRead(&bIsInBackground);
-                if (!bWasFinishedInBG)
+				//only decrement NumCurrentlyActiveTasks if this was an actual active task,
+				//can still be completed in the BG or from existing completed data without being active
+                const bool bIsTaskActive = Request->IsUnderlyingTaskActive();
+                if (bIsTaskActive)
                 {
                     int NumActualTasks = FPlatformAtomics::InterlockedDecrement(&NumCurrentlyActiveTasks);
 
@@ -666,7 +680,7 @@ void FApplePlatformBackgroundHttpManager::OnTask_DidWriteData(NSURLSessionDownlo
         {
             FRWScopeLock ScopeLock(URLToRequestMapLock, SLT_ReadOnly);
             FBackgroundHttpURLMappedRequestPtr* WeakRequestInMap = URLToRequestMap.Find(TaskURL);
-            FAppleBackgroundHttpRequestPtr FoundRequest = (nullptr != WeakRequestInMap) ? WeakRequestInMap->Pin() : nullptr;
+            FAppleBackgroundHttpRequestPtr FoundRequest = ((nullptr != WeakRequestInMap) && (WeakRequestInMap->IsValid())) ? WeakRequestInMap->Pin() : nullptr;
             
             if (FoundRequest.IsValid())
             {
@@ -710,7 +724,7 @@ void FApplePlatformBackgroundHttpManager::OnTask_DidCompleteWithError(NSURLSessi
 		{
             FRWScopeLock ScopeLock(URLToRequestMapLock, SLT_ReadOnly);
             FBackgroundHttpURLMappedRequestPtr* WeakRequestInMap = URLToRequestMap.Find(TaskURL);
-            FAppleBackgroundHttpRequestPtr FoundRequest = (nullptr != WeakRequestInMap) ? WeakRequestInMap->Pin() : nullptr;
+            FAppleBackgroundHttpRequestPtr FoundRequest = ((nullptr != WeakRequestInMap) && (WeakRequestInMap->IsValid())) ? WeakRequestInMap->Pin() : nullptr;
             const bool bDidFindValidRequest = FoundRequest.IsValid();
          
             //by default increase error count. Special cases below will overwrite this
@@ -747,7 +761,8 @@ void FApplePlatformBackgroundHttpManager::OnSession_SessionDidFinishAllEvents(NS
 	
 	const bool bCopyIsBG = FPlatformAtomics::AtomicRead(&bIsInBackground);
 	
-	if (bCopyIsBG)
+	//If we are in the BG, or if we have not yet been initialized, lets go ahead and just immediately send the callback
+	if (bWasAppleBGHTTPInitialized && bCopyIsBG)
 	{
 		//Now that we have finished all queued background tasks, lets resume tasks for any lower priorities that weren't started
 		ResumeTasksForBackgrounding(FIOSBackgroundDownloadCoreDelegates::FIOSBackgroundDownload_DelayedBackgroundURLSessionCompleteHandler::CreateLambda(
