Submitting missing header from CL 10454165
#tests none

diff --git a/Engine/Source/Runtime/Online/BackgroundHTTP/Public/IOS/ApplePlatformBackgroundHttpManager.h b/Engine/Source/Runtime/Online/BackgroundHTTP/Public/IOS/ApplePlatformBackgroundHttpManager.h
--- a/Engine/Source/Runtime/Online/BackgroundHTTP/Public/IOS/ApplePlatformBackgroundHttpManager.h
+++ b/Engine/Source/Runtime/Online/BackgroundHTTP/Public/IOS/ApplePlatformBackgroundHttpManager.h
@@ -66,6 +66,10 @@ protected:
     TMap<const FString, FBackgroundHttpURLMappedRequestPtr> URLToRequestMap;
 	FRWLock URLToRequestMapLock;
 
+	//Used to track if we came into certain code branches through a pathway that has initialized the manager yet or not.
+	//Some of our platform callbacks can end up hitting VERY early or in background versions of our app without the engine being fully spun-up.
+	static volatile bool bWasAppleBGHTTPInitialized;
+	
 private:
 	//Checks for tasks that already exist on the BackgroundSession. Should only be called during Initialize. These tasks should really only exist 
 	//due to previous application sessions starting background downloads that finished after the app closed. (IE: If the user started a patch and then our app was terminated by the OS while 
@@ -79,7 +83,7 @@ private:
     //These are used internally to pause/resume things at the task level
 	void PauseAllActiveTasks();
 	void ResumeTasksForBackgrounding(FIOSBackgroundHttpPostSessionWorkCallback Callback = FIOSBackgroundHttpPostSessionWorkCallback());
-	bool ResumeDownloadTaskForBackgroundingIfAppropriate(NSURLSessionDownloadTask* DownloadTask, EBackgroundHTTPPriority LowestPriorityToQueue);
+	bool ResumeDownloadTaskForBackgroundingIfAppropriate(NSURLSessionDownloadTask* DownloadTask, const FAppleBackgroundHttpRequestPtr Request, EBackgroundHTTPPriority LowestPriorityToQueue);
 
     void PauseAllUnassociatedTasks();
     void UnpauseAllUnassociatedTasks();
