Fixing Background HTTP continuous loop in general implementation.-- By [REDACTED_PERSON], submitting from Justin.Marcus' machine.
#tests Android Client

diff --git a/Engine/Source/Runtime/Online/BackgroundHTTP/Private/BackgroundHttpManagerImpl.cpp b/Engine/Source/Runtime/Online/BackgroundHTTP/Private/BackgroundHttpManagerImpl.cpp
--- a/Engine/Source/Runtime/Online/BackgroundHTTP/Private/BackgroundHttpManagerImpl.cpp
+++ b/Engine/Source/Runtime/Online/BackgroundHTTP/Private/BackgroundHttpManagerImpl.cpp
@@ -228,86 +228,53 @@ bool FBackgroundHttpManagerImpl::Tick(float DeltaTime)
 
 void FBackgroundHttpManagerImpl::ActivatePendingRequests()
 {
-	TArray<FBackgroundHttpRequestPtr> RequestsStartingThisTick;
+	FBackgroundHttpRequestPtr HighestPriorityRequestToStart = nullptr;
+	EBackgroundHTTPPriority HighestRequestPriority = EBackgroundHTTPPriority::Num;
 
-	//Handle Populating RequestsStartingThisTick from PendingStartRequests
+	//Go through and find the highest priority request
 	{
-		FRWScopeLock ScopeLock(PendingRequestLock, SLT_Write);
+		FRWScopeLock ScopeLock(PendingRequestLock, SLT_ReadOnly);
 
 		if (PendingStartRequests.Num() > 0)
 		{
 			UE_LOG(LogBackgroundHttpManager, Verbose, TEXT("Populating Requests to Start from PendingStartRequests - MaxActiveDownloads:%d | NumCurrentlyActiveRequests:%d | NumPendingStartRequests:%d"), MaxActiveDownloads.Load(), NumCurrentlyActiveRequests, PendingStartRequests.Num());
 
+			HighestPriorityRequestToStart = PendingStartRequests[0];
+			HighestRequestPriority = HighestPriorityRequestToStart.IsValid() ? HighestPriorityRequestToStart->GetRequestPriority() : EBackgroundHTTPPriority::Num;
+
 			//See how many more requests we can process and only do anything if we can still process more
 			const int NumRequestsWeCanProcess = (MaxActiveDownloads - NumCurrentlyActiveRequests);
 			if (NumRequestsWeCanProcess > 0)
 			{
-				TArray<FBackgroundHttpRequestPtr> RemainingRequests;
-				if (PendingStartRequests.Num() > NumRequestsWeCanProcess)
-				{
-					RemainingRequests.Reserve(PendingStartRequests.Num() - NumRequestsWeCanProcess);
-				}
-
-				//First preload the first N spots of the list with the start of the pending list.
-				int RequestIndex = 0;
-				for (;RequestIndex < PendingStartRequests.Num();++RequestIndex)
-				{
-					FBackgroundHttpRequestPtr& PendingRequest = PendingStartRequests[RequestIndex];
-					
-					if (RequestIndex <= NumRequestsWeCanProcess)
-					{
-						RequestsStartingThisTick.Add(PendingRequest);
-					}
-					else
-					{
-						break;
-					}
-				}
-				
-				//Now go through the rest of the elements, replacing any items found with things of higher priority if needed
-				for (; RequestIndex < PendingStartRequests.Num(); ++RequestIndex)
+				for (int RequestIndex = 1; RequestIndex < PendingStartRequests.Num(); ++RequestIndex)
 				{
 					FBackgroundHttpRequestPtr PendingRequestToCheck = PendingStartRequests[RequestIndex];
-					EBackgroundHTTPPriority PendingRequestPriority = PendingRequestToCheck->GetRequestPriority();
-					
-					for (int AlreadyGrabbedRequestIndex = 0; AlreadyGrabbedRequestIndex < RequestsStartingThisTick.Num();++AlreadyGrabbedRequestIndex)
+					EBackgroundHTTPPriority PendingRequestPriority = PendingRequestToCheck.IsValid() ? PendingRequestToCheck->GetRequestPriority() : EBackgroundHTTPPriority::Num;
+
+					//Found a higher priority request, so track that one
+					if (PendingRequestPriority < HighestRequestPriority)
 					{
-						FBackgroundHttpRequestPtr& AlreadyGrabbedRequest = RequestsStartingThisTick[AlreadyGrabbedRequestIndex];
-						EBackgroundHTTPPriority AlreadyGrabbedRequestPriority = AlreadyGrabbedRequest->GetRequestPriority();
-						
-						//Found a higher priority request, swap it and go back through entire list again
-						if (PendingRequestPriority < AlreadyGrabbedRequestPriority)
-						{
-							FBackgroundHttpRequestPtr RequestPtrCopy = AlreadyGrabbedRequest;
-							AlreadyGrabbedRequest = PendingRequestToCheck;
-							PendingRequestToCheck = RequestPtrCopy;
-							
-							//-1 so that it increments to 0
-							AlreadyGrabbedRequestIndex = -1;
-						}
+						HighestPriorityRequestToStart = PendingRequestToCheck;
+						HighestRequestPriority = PendingRequestPriority;
 					}
-					
-					//Whatever is left in PendingRequestToCheck was too low priority to end up in the correct list.
-					RemainingRequests.Add(PendingRequestToCheck);
 				}
-
-				PendingStartRequests = MoveTemp(RemainingRequests);
 			}
 		}
 	}
 
-	UE_LOG(LogBackgroundHttpManager, Verbose, TEXT("Starting %d Requests From PendingStartRequests Queue"), RequestsStartingThisTick.Num());
-
-	//Now actually add request to Active List and call Handle 
-	for (FBackgroundHttpRequestPtr& RequestToStart : RequestsStartingThisTick)
+	if (HighestPriorityRequestToStart.IsValid())
 	{
+		UE_LOG(LogBackgroundHttpManager, Verbose, TEXT("Starting Request: %s Priority:%s"), *HighestPriorityRequestToStart->GetRequestID(), LexToString(HighestRequestPriority));
+
 		//Actually move request to Active list now
-		FRWScopeLock ScopeLock(ActiveRequestLock, SLT_Write);
-		ActiveRequests.Add(RequestToStart);
+		FRWScopeLock ScopeLock0(ActiveRequestLock, SLT_Write);
+		FRWScopeLock ScopeLock1(PendingRequestLock, SLT_Write);
+		ActiveRequests.Add(HighestPriorityRequestToStart);
+		PendingStartRequests.RemoveSingle(HighestPriorityRequestToStart);
 
 		++NumCurrentlyActiveRequests;
 
 		//Call Handle for that task to now kick itself off
-		RequestToStart->HandleDelayedProcess();
+		HighestPriorityRequestToStart->HandleDelayedProcess();
 	}
 }
diff --git a/Engine/Source/Runtime/Online/BackgroundHTTP/Public/Interfaces/IBackgroundHttpRequest.h b/Engine/Source/Runtime/Online/BackgroundHTTP/Public/Interfaces/IBackgroundHttpRequest.h
--- a/Engine/Source/Runtime/Online/BackgroundHTTP/Public/Interfaces/IBackgroundHttpRequest.h
+++ b/Engine/Source/Runtime/Online/BackgroundHTTP/Public/Interfaces/IBackgroundHttpRequest.h
@@ -52,6 +52,9 @@ inline const TCHAR* LexToString(EBackgroundHTTPPriority InType)
 		return TEXT("Normal");
 	case EBackgroundHTTPPriority::Low:
 		return TEXT("Low");
+	case EBackgroundHTTPPriority::Num:
+		return TEXT("INVALID(Num)");
+
 	default:
 		break;
 	}
