Fixing Background Http Manager Generic Implementation not limiting queued requests to the system max active downloads.
#tests none

diff --git a/Engine/Source/Runtime/Online/BackgroundHTTP/Private/BackgroundHttpManagerImpl.cpp b/Engine/Source/Runtime/Online/BackgroundHTTP/Private/BackgroundHttpManagerImpl.cpp
--- a/Engine/Source/Runtime/Online/BackgroundHTTP/Private/BackgroundHttpManagerImpl.cpp
+++ b/Engine/Source/Runtime/Online/BackgroundHTTP/Private/BackgroundHttpManagerImpl.cpp
@@ -233,29 +233,31 @@ void FBackgroundHttpManagerImpl::ActivatePendingRequests()
 
 	//Go through and find the highest priority request
 	{
-		FRWScopeLock ScopeLock(PendingRequestLock, SLT_ReadOnly);
-
-		if (PendingStartRequests.Num() > 0)
+		FRWScopeLock ActiveScopeLock(ActiveRequestLock, SLT_ReadOnly);
+		FRWScopeLock PendingScopeLock(PendingRequestLock, SLT_ReadOnly);
+		const int NumRequestsWeCanProcess = (MaxActiveDownloads - NumCurrentlyActiveRequests);
+		if (NumRequestsWeCanProcess > 0)
 		{
-			UE_LOG(LogBackgroundHttpManager, Verbose, TEXT("Populating Requests to Start from PendingStartRequests - MaxActiveDownloads:%d | NumCurrentlyActiveRequests:%d | NumPendingStartRequests:%d"), MaxActiveDownloads.Load(), NumCurrentlyActiveRequests, PendingStartRequests.Num());
+			if (PendingStartRequests.Num() > 0)
+			{
+				UE_LOG(LogBackgroundHttpManager, Verbose, TEXT("Populating Requests to Start from PendingStartRequests - MaxActiveDownloads:%d | NumCurrentlyActiveRequests:%d | NumPendingStartRequests:%d"), MaxActiveDownloads.Load(), NumCurrentlyActiveRequests, PendingStartRequests.Num());
 
-			HighestPriorityRequestToStart = PendingStartRequests[0];
-			HighestRequestPriority = HighestPriorityRequestToStart.IsValid() ? HighestPriorityRequestToStart->GetRequestPriority() : EBackgroundHTTPPriority::Num;
+				HighestPriorityRequestToStart = PendingStartRequests[0];
+				HighestRequestPriority = HighestPriorityRequestToStart.IsValid() ? HighestPriorityRequestToStart->GetRequestPriority() : EBackgroundHTTPPriority::Num;
 
-			//See how many more requests we can process and only do anything if we can still process more
-			const int NumRequestsWeCanProcess = (MaxActiveDownloads - NumCurrentlyActiveRequests);
-			if (NumRequestsWeCanProcess > 0)
-			{
-				for (int RequestIndex = 1; RequestIndex < PendingStartRequests.Num(); ++RequestIndex)
+				//See how many more requests we can process and only do anything if we can still process more
 				{
-					FBackgroundHttpRequestPtr PendingRequestToCheck = PendingStartRequests[RequestIndex];
-					EBackgroundHTTPPriority PendingRequestPriority = PendingRequestToCheck.IsValid() ? PendingRequestToCheck->GetRequestPriority() : EBackgroundHTTPPriority::Num;
-
-					//Found a higher priority request, so track that one
-					if (PendingRequestPriority < HighestRequestPriority)
+					for (int RequestIndex = 1; RequestIndex < PendingStartRequests.Num(); ++RequestIndex)
 					{
-						HighestPriorityRequestToStart = PendingRequestToCheck;
-						HighestRequestPriority = PendingRequestPriority;
+						FBackgroundHttpRequestPtr PendingRequestToCheck = PendingStartRequests[RequestIndex];
+						EBackgroundHTTPPriority PendingRequestPriority = PendingRequestToCheck.IsValid() ? PendingRequestToCheck->GetRequestPriority() : EBackgroundHTTPPriority::Num;
+
+						//Found a higher priority request, so track that one
+						if (PendingRequestPriority < HighestRequestPriority)
+						{
+							HighestPriorityRequestToStart = PendingRequestToCheck;
+							HighestRequestPriority = PendingRequestPriority;
+						}
 					}
 				}
 			}
@@ -267,8 +269,8 @@ void FBackgroundHttpManagerImpl::ActivatePendingRequests()
 		UE_LOG(LogBackgroundHttpManager, Verbose, TEXT("Starting Request: %s Priority:%s"), *HighestPriorityRequestToStart->GetRequestID(), LexToString(HighestRequestPriority));
 
 		//Actually move request to Active list now
-		FRWScopeLock ScopeLock0(ActiveRequestLock, SLT_Write);
-		FRWScopeLock ScopeLock1(PendingRequestLock, SLT_Write);
+		FRWScopeLock ActiveScopeLock(ActiveRequestLock, SLT_Write);
+		FRWScopeLock PendingScopeLock(PendingRequestLock, SLT_Write);
 		ActiveRequests.Add(HighestPriorityRequestToStart);
 		PendingStartRequests.RemoveSingle(HighestPriorityRequestToStart);
 
