More VOIP tweakables: Add input gain CVar to voice capture windows.

diff --git a/Engine/Source/Runtime/Engine/Private/VoiceConfig.cpp b/Engine/Source/Runtime/Engine/Private/VoiceConfig.cpp
--- a/Engine/Source/Runtime/Engine/Private/VoiceConfig.cpp
+++ b/Engine/Source/Runtime/Engine/Private/VoiceConfig.cpp
@@ -21,7 +21,23 @@ static TAutoConsoleVariable<float> CVarVoiceSilenceDetectionThreshold(TEXT("voic
 	TEXT("Threshold to be set for the VOIP microphone's silence detection algorithm.\n"),	
 	ECVF_Default);
 
-static float JitterBufferDelayCvar = 0.3;
+static float MicInputGainCvar = 1.0f;
+FAutoConsoleVariableRef CVarMicInputGain(
+	TEXT("voice.MicInputGain"),
+	MicInputGainCvar,
+	TEXT("The default gain amount in linear amplitude.\n")
+	TEXT("Value: Gain multiplier."),
+	ECVF_Default);
+
+static float MicStereoBiasCvar = 0.0f;
+FAutoConsoleVariableRef CVarMicStereoBias(
+	TEXT("voice.MicStereoBias"),
+	MicStereoBiasCvar,
+	TEXT("This will attenuate the left or right channel.\n")
+	TEXT("0.0: Centered. 1.0: right channel only. -1.0: Left channel only."),
+	ECVF_Default);
+
+static float JitterBufferDelayCvar = 0.3f;
 FAutoConsoleVariableRef CVarJitterBufferDelay(
 	TEXT("voice.JitterBufferDelay"),
 	JitterBufferDelayCvar,
@@ -29,7 +45,7 @@ FAutoConsoleVariableRef CVarJitterBufferDelay(
 	TEXT("Value: Number of seconds of audio we buffer."),
 	ECVF_Default);
 
-static float MicNoiseGateThresholdCvar = 0.08;
+static float MicNoiseGateThresholdCvar = 0.08f;
 FAutoConsoleVariableRef CVarMicNoiseGateThreshold(
 	TEXT("voice.MicNoiseGateThreshold"),
 	JitterBufferDelayCvar,
@@ -37,7 +53,7 @@ FAutoConsoleVariableRef CVarMicNoiseGateThreshold(
 	TEXT("Value: Number of seconds of audio we buffer."),
 	ECVF_Default);
 
-static float MicNoiseGateAttackTimeCvar = 0.05;
+static float MicNoiseGateAttackTimeCvar = 0.05f;
 FAutoConsoleVariableRef CVarMicNoiseGateAttackTime(
 	TEXT("voice.MicNoiseAttackTime"),
 	MicNoiseGateAttackTimeCvar,
@@ -45,7 +61,7 @@ FAutoConsoleVariableRef CVarMicNoiseGateAttackTime(
 	TEXT("Value: Number of seconds we fade in over."),
 	ECVF_Default);
 
-static float MicNoiseGateReleaseTimeCvar = 0.30;
+static float MicNoiseGateReleaseTimeCvar = 0.30f;
 FAutoConsoleVariableRef CVarMicNoiseGateReleaseTime(
 	TEXT("voice.MicNoiseReleaseTime"),
 	MicNoiseGateReleaseTimeCvar,
diff --git a/Engine/Source/Runtime/Online/Voice/Private/Windows/VoiceCaptureWindows.cpp b/Engine/Source/Runtime/Online/Voice/Private/Windows/VoiceCaptureWindows.cpp
--- a/Engine/Source/Runtime/Online/Voice/Private/Windows/VoiceCaptureWindows.cpp
+++ b/Engine/Source/Runtime/Online/Voice/Private/Windows/VoiceCaptureWindows.cpp
@@ -418,7 +418,6 @@ void FVoiceCaptureWindows::ProcessData()
 			int32 SamplesPushedToUncompressedAudioBuffer = ReleaseBuffer.PopBufferedAudio(AudioBuffer, ReleaseBuffer.GetBufferCount());
 			AudioBuffer += SamplesPushedToUncompressedAudioBuffer;
 
-			// get threshold
 			static IConsoleVariable* SilenceDetectionThresholdCVar = IConsoleManager::Get().FindConsoleVariable(TEXT("voice.SilenceDetectionThreshold"));
 			check(SilenceDetectionThresholdCVar);			
 			const float MicSilenceThreshold = SilenceDetectionThresholdCVar->GetFloat();
@@ -435,6 +434,24 @@ void FVoiceCaptureWindows::ProcessData()
 			check(NoiseGateReleaseTimeCVar);
 			const float MicNoiseGateReleaseTime = NoiseGateReleaseTimeCVar->GetFloat();
 
+			static IConsoleVariable* MicInputGainCVar = IConsoleManager::Get().FindConsoleVariable(TEXT("voice.MicInputGain"));
+			check(MicInputGainCVar);
+			const float MicInputGain = MicInputGainCVar->GetFloat();
+
+			static IConsoleVariable* MicStereoBiasCVar = IConsoleManager::Get().FindConsoleVariable(TEXT("voice.MicStereoBias"));
+			check(MicStereoBiasCVar);
+			const float MicStereoBias = FMath::Clamp(MicStereoBiasCVar->GetFloat(), -1.0f, 1.0f);
+
+			float LeftGain = 1.0f;
+			float RightGain = 1.0f;
+
+			Audio::GetStereoPan(MicStereoBias, LeftGain, RightGain);
+
+			// Since we don't interpolate the pan here, we normalize stereo gains.
+			const float StereoGainMax = FMath::Max(LeftGain, RightGain);
+			LeftGain /= StereoGainMax;
+			RightGain /= StereoGainMax;
+
 			bool bMicReleased = false;
 
 			CurrentSampleStart = CachedSampleStart;
@@ -474,14 +491,15 @@ void FVoiceCaptureWindows::ProcessData()
 					if (bMicReleased)
 					{
 						// Apply noise gate attenuation.
+						const float TotalMicGain = MicInputGain * NoiseGateAttenuator.GetNextValue();
 						for (int32 ChannelIndex = 0; ChannelIndex < NumInputChannels; ChannelIndex++)
 						{
+							const float BiasedMicGain = (ChannelIndex % 2 == 0) ? (TotalMicGain * LeftGain) : (TotalMicGain * RightGain);
+
 							Audio::TSampleRef<int16> SampleRef(InputBuffer[FrameIndex + ChannelIndex]);
-							SampleRef = SampleRef * NoiseGateAttenuator.PeekCurrentValue();
+							SampleRef = SampleRef * BiasedMicGain;
 						}
 
-						NoiseGateAttenuator.GetNextValue();
-
 						ReleaseBuffer.PushFrame(&InputBuffer[FrameIndex], NumInputChannels);
 
 
@@ -494,14 +512,15 @@ void FVoiceCaptureWindows::ProcessData()
 					else
 					{
 						// Apply noise gate attenuation.
+						const float TotalMicGain = MicInputGain * NoiseGateAttenuator.GetNextValue();
 						for (int32 ChannelIndex = 0; ChannelIndex < NumInputChannels; ChannelIndex++)
 						{
+							const float BiasedMicGain = (ChannelIndex % 2 == 0) ? (TotalMicGain * LeftGain) : (TotalMicGain * RightGain);
+
 							Audio::TSampleRef<int16> SampleRef(InputBuffer[FrameIndex + ChannelIndex]);
-							SampleRef = SampleRef * NoiseGateAttenuator.PeekCurrentValue();
+							SampleRef = SampleRef * BiasedMicGain;
 						}
 
-						NoiseGateAttenuator.GetNextValue();
-
 						FMemory::Memcpy(&AudioBuffer[FrameIndex], &InputBuffer[FrameIndex], sizeof(int16) * NumInputChannels);
 						SamplesPushedToUncompressedAudioBuffer += NumInputChannels;
 					}
@@ -551,14 +570,15 @@ void FVoiceCaptureWindows::ProcessData()
 					if (bMicReleased)
 					{
 						// Apply noise gate attenuation.
+						const float TotalMicGain = MicInputGain * NoiseGateAttenuator.GetNextValue();
 						for (int32 ChannelIndex = 0; ChannelIndex < NumInputChannels; ChannelIndex++)
 						{
+							const float BiasedMicGain = (ChannelIndex % 2 == 0) ? (TotalMicGain * LeftGain) : (TotalMicGain * RightGain);
+
 							Audio::TSampleRef<int16> SampleRef(InputBuffer[FrameIndex + ChannelIndex]);
-							SampleRef = SampleRef * NoiseGateAttenuator.PeekCurrentValue();
+							SampleRef = SampleRef * BiasedMicGain;
 						}
 
-						NoiseGateAttenuator.GetNextValue();
-
 						ReleaseBuffer.PushFrame(&InputBuffer[FrameIndex], NumInputChannels);
 
 						if (!bSampleStartCached)
@@ -570,14 +590,15 @@ void FVoiceCaptureWindows::ProcessData()
 					else
 					{
 						// Apply noise gate attenuation.
+						const float TotalMicGain = MicInputGain * NoiseGateAttenuator.GetNextValue();
 						for (int32 ChannelIndex = 0; ChannelIndex < NumInputChannels; ChannelIndex++)
 						{
+							const float BiasedMicGain = (ChannelIndex % 2 == 0) ? (TotalMicGain * LeftGain) : (TotalMicGain * RightGain);
+
 							Audio::TSampleRef<int16> SampleRef(InputBuffer[FrameIndex + ChannelIndex]);
-							SampleRef = SampleRef * NoiseGateAttenuator.PeekCurrentValue();
+							SampleRef = SampleRef * BiasedMicGain;
 						}
 
-						NoiseGateAttenuator.GetNextValue();
-
 						FMemory::Memcpy(&AudioBuffer[FrameIndex], &InputBuffer[FrameIndex], sizeof(int16) * NumInputChannels);
 						SamplesPushedToUncompressedAudioBuffer += NumInputChannels;
 					}
diff --git a/Engine/Source/Runtime/SignalProcessing/Public/DSP/Dsp.h b/Engine/Source/Runtime/SignalProcessing/Public/DSP/Dsp.h
--- a/Engine/Source/Runtime/SignalProcessing/Public/DSP/Dsp.h
+++ b/Engine/Source/Runtime/SignalProcessing/Public/DSP/Dsp.h
@@ -6,6 +6,7 @@
 #include "Misc/ScopeLock.h"
 #include "Templates/IsFloatingPoint.h"
 #include "Templates/IsIntegral.h"
+#include "Templates/IsSigned.h"
 
 // Macros which can be enabled to cause DSP sample checking
 #if 0
@@ -691,6 +692,10 @@ namespace Audio
 		// This is the number used to convert from float to our fixed precision value.
 		static constexpr float QFactor = TGetPower<2, Q>::Value - 1;
 
+		// for fixed precision types, the max and min values that we can represent are calculated here:
+		static constexpr float MaxValue = TGetPower<2, (sizeof(SampleType) * 8 - Q)>;
+		static constexpr float MinValue = !!TIsSigned<SampleType> ? (-1.0f * MaxValue) : 0.0f;
+
 	public:
 
 		TSample(SampleType& InSample)
@@ -811,7 +816,7 @@ namespace Audio
 				{
 					// fixed * float.
 					OtherSampleType FloatLHS = ((OtherSampleType)LHS.Sample) / QFactor;
-					OtherSampleType Result = FloatLHS * RHS;
+					OtherSampleType Result = FMath::Clamp(FloatLHS * RHS, MinValue, MaxValue);
 					return static_cast<SampleType>(Result * QFactor);
 				}
 				else if (TIsIntegral<OtherSampleType>::Value)
@@ -819,7 +824,7 @@ namespace Audio
 					// Q * Q.
 					float FloatLHS = ((float)LHS.Sample) / QFactor;
 					float FloatRHS = ((float)RHS) / QFactor;
-					float Result = FloatLHS * FloatRHS;
+					float Result = FMath::Clamp(FloatLHS * FloatRHS, MinValue, MaxValue);
 					return static_cast<OtherSampleType>(Result * QFactor);
 				}
 				else
@@ -872,9 +877,13 @@ namespace Audio
 			static_assert(bIsTypeValid, "Invalid value for Q! TSampleRef only supports float or int types. For int types, Q must be smaller than the number of bits in the int type.");
 		}
 
-		// This is the number used to convert from float to our .
+		// This is the number used to convert from float to our fixed precision value.
 		static constexpr float QFactor = TGetPower<2, Q>::Value - 1;
 
+		// for fixed precision types, the max and min values that we can represent are calculated here:
+		static constexpr float MaxValue = TGetPower<2, (sizeof(SampleType) * 8 - Q)>::Value;
+		static constexpr float MinValue = !!TIsSigned<SampleType>::Value ? (-1.0f * MaxValue) : 0.0f;
+
 	public:
 
 		TSampleRef(SampleType& InSample)
@@ -996,15 +1005,16 @@ namespace Audio
 				{
 					// fixed * float.
 					OtherSampleType FloatLHS = ((OtherSampleType)LHS.Sample) / QFactor;
-					OtherSampleType Result = FloatLHS * RHS;
+					OtherSampleType Result = FMath::Clamp(FloatLHS * RHS, MinValue, MaxValue);
 					return static_cast<SampleType>(Result * QFactor);
 				}
 				else if (TIsIntegral<OtherSampleType>::Value)
 				{
-					// Q * Q.
+					// fixed * fixed.
 					float FloatLHS = ((float)LHS.Sample) / QFactor;
 					float FloatRHS = ((float)RHS) / QFactor;
-					float Result = FloatLHS * FloatRHS;
+					float Result = FMath::Clamp(FloatLHS * FloatRHS, MinValue, MaxValue);
+					
 					return static_cast<SampleType>(Result * QFactor);
 				}
 				else
