Fixing bug where iOS Background HTTP Requests are not resuming all work when going into background correctly.
#tests iOS Client

diff --git a/Engine/Source/Runtime/Online/BackgroundHTTP/Private/IOS/ApplePlatformBackgroundHttpManager.cpp b/Engine/Source/Runtime/Online/BackgroundHTTP/Private/IOS/ApplePlatformBackgroundHttpManager.cpp
--- a/Engine/Source/Runtime/Online/BackgroundHTTP/Private/IOS/ApplePlatformBackgroundHttpManager.cpp
+++ b/Engine/Source/Runtime/Online/BackgroundHTTP/Private/IOS/ApplePlatformBackgroundHttpManager.cpp
@@ -400,7 +400,7 @@ void FApplePlatformBackgroundHttpManager::ResumeTasksForBackgrounding(FIOSBackgr
 		 
 					if (FoundRequest.IsValid())
 					{
-						bDidResumeATask = bDidResumeATask || ResumeDownloadTaskForBackgroundingIfAppropriate(DownloadTask, FoundRequest, EBackgroundHTTPPriority(PriorityAsInt));
+						bDidResumeATask = ResumeDownloadTaskForBackgroundingIfAppropriate(DownloadTask, FoundRequest, EBackgroundHTTPPriority(PriorityAsInt)) || bDidResumeATask;
 					}
 					else
 					{
@@ -412,6 +412,10 @@ void FApplePlatformBackgroundHttpManager::ResumeTasksForBackgrounding(FIOSBackgr
 				{
 					UE_LOG(LogBackgroundHttpManager, Display, TEXT("Resumed Tasks of Priority: %s . Not Resuming Any Lower Priority Tasks."), LexToString((EBackgroundHTTPPriority)PriorityAsInt));
 				}
+				else
+				{
+					UE_LOG(LogBackgroundHttpManager, Display, TEXT("No tasks found of priority %s."), LexToString((EBackgroundHTTPPriority)PriorityAsInt));
+				}
 			}
 			
 			Callback.ExecuteIfBound();
