Fixing bug in iOS Background HTTP where we could end up queueing lower priority tasks when backgrounding earlier then intended when the task count was low.
#tests iOS Client

diff --git a/Engine/Source/Runtime/Online/BackgroundHTTP/Private/IOS/ApplePlatformBackgroundHttpManager.cpp b/Engine/Source/Runtime/Online/BackgroundHTTP/Private/IOS/ApplePlatformBackgroundHttpManager.cpp
--- a/Engine/Source/Runtime/Online/BackgroundHTTP/Private/IOS/ApplePlatformBackgroundHttpManager.cpp
+++ b/Engine/Source/Runtime/Online/BackgroundHTTP/Private/IOS/ApplePlatformBackgroundHttpManager.cpp
@@ -410,7 +410,7 @@ void FApplePlatformBackgroundHttpManager::ResumeTasksForBackgrounding(FIOSBackgr
 				
 				if (bDidResumeATask)
 				{
-					UE_LOG(LogBackgroundHttpManager, Display, TEXT("Resumed Tasks of Priority: %s . Not Resuming Any Lower Priority Tasks."), LexToString((EBackgroundHTTPPriority)PriorityAsInt));
+					UE_LOG(LogBackgroundHttpManager, Display, TEXT("Tasks of Priority: %s Found. Not Resuming Any Lower Priority Tasks."), LexToString((EBackgroundHTTPPriority)PriorityAsInt));
 				}
 				else
 				{
@@ -431,31 +431,30 @@ bool FApplePlatformBackgroundHttpManager::ResumeDownloadTaskForBackgroundingIfAp
 	}
 
 	bool bDidResumeTask = false;
+	bool bDidFindActiveTaskOfPriority = false;
 	
-	//We don't care about Cancelled, Completed, or Completing Task as those are already handled. Suspended tasks are waiting
-	//to be resumed to activate
-	if ([DownloadTask state] == NSURLSessionTaskStateSuspended)
-	{
-		FString TaskURL = [[[DownloadTask currentRequest] URL] absoluteString];
-		int TaskIdentifier = (int)[DownloadTask taskIdentifier];
+	FString TaskURL = [[[DownloadTask currentRequest] URL] absoluteString];
+	int TaskIdentifier = (int)[DownloadTask taskIdentifier];
 
+	EBackgroundHTTPPriority FoundRequestPriority = Request->GetRequestPriority();
+	if (FoundRequestPriority <= LowestPriorityToQueue)
+	{
 		const bool bIsRequestPaused = Request->bIsTaskPaused;
 		if (!bIsRequestPaused)
 		{
-			EBackgroundHTTPPriority FoundRequestPriority = Request->GetRequestPriority();
-			if (FoundRequestPriority <= LowestPriorityToQueue)
+			UE_LOG(LogBackgroundHttpManager, Display, TEXT("Resuming Found Task for URL:%s | TaskIdentifier:%d | TaskPriority:%s"), *TaskURL, TaskIdentifier, LexToString(FoundRequestPriority));
+			
+			const bool bIsTaskActive = Request->IsUnderlyingTaskActive();
+			
+			//We only want to resume Suspended tasks
+			if ([DownloadTask state] == NSURLSessionTaskStateSuspended)
 			{
-				NSLog(@"TaskURL:%@",TaskURL.GetNSString());
-				NSLog(@"Priority:%d",(uint8)FoundRequestPriority);
-				UE_LOG(LogBackgroundHttpManager, Display, TEXT("TaskPriority:%s"), LexToString(FoundRequestPriority));
-				UE_LOG(LogBackgroundHttpManager, Display, TEXT("Resuming Task for URL:%s | TaskIdentifier:%d | TaskPriority:%s"), *TaskURL, TaskIdentifier, LexToString(FoundRequestPriority));
 				[DownloadTask resume];
-				
 				bDidResumeTask = true;
 			}
-			else
+			else if (bIsTaskActive)
 			{
-				UE_LOG(LogBackgroundHttpManager, Verbose, TEXT("NOT RESUMING Task for URL because its a lower priority:%s | TaskIdentifier:%d | TaskPriority:%s | LowestPriorityToQueue:%s"), *TaskURL, TaskIdentifier, LexToString(FoundRequestPriority), LexToString(LowestPriorityToQueue));
+				bDidFindActiveTaskOfPriority = true;
 			}
 		}
 		else
@@ -463,8 +462,12 @@ bool FApplePlatformBackgroundHttpManager::ResumeDownloadTaskForBackgroundingIfAp
 			UE_LOG(LogBackgroundHttpManager, Verbose, TEXT("NOT Resuming Task for URL as associated request was paused! URL:%s | TaskIdentifier:%d"), *TaskURL, TaskIdentifier);
 		}
 	}
+	else
+	{
+		UE_LOG(LogBackgroundHttpManager, Verbose, TEXT("NOT RESUMING Task for URL because its a lower priority:%s | TaskIdentifier:%d | TaskPriority:%s | LowestPriorityToQueue:%s"), *TaskURL, TaskIdentifier, LexToString(FoundRequestPriority), LexToString(LowestPriorityToQueue));
+	}
 	
-	return bDidResumeTask;
+	return (bDidResumeTask || bDidFindActiveTaskOfPriority);
 }
 
 void FApplePlatformBackgroundHttpManager::OnTask_DidFinishDownloadingToURL(NSURLSessionDownloadTask* Task, NSError* Error, const FString& TempFilePath)
@@ -758,13 +761,13 @@ void FApplePlatformBackgroundHttpManager::OnTask_DidCompleteWithError(NSURLSessi
 
 void FApplePlatformBackgroundHttpManager::OnSession_SessionDidFinishAllEvents(NSURLSession* Session, FIOSBackgroundDownloadCoreDelegates::FIOSBackgroundDownload_DelayedBackgroundURLSessionCompleteHandler Callback)
 {
-	UE_LOG(LogBackgroundHttpManager, Verbose, TEXT("NSURLSession done sending background events for all already queued tasks"));
-	
 	//Let BackgroundURLSessionHandler know that it should wait until we call the callback
 	FBackgroundURLSessionHandler::AddDelayedBackgroundURLSessionComplete();
 	
 	const bool bCopyIsBG = FPlatformAtomics::AtomicRead(&bIsInBackground);
 	
+	UE_LOG(LogBackgroundHttpManager, Display, TEXT("NSURLSession done sending background events for all already queued tasks. bWasAppleBGHTTPInitialized:%d | bIsBG:%d"), bWasAppleBGHTTPInitialized, bCopyIsBG);
+	
 	//If we are in the BG, or if we have not yet been initialized, lets go ahead and just immediately send the callback
 	if (bWasAppleBGHTTPInitialized && bCopyIsBG)
 	{
