Adding voice.debug.PrintAmplitude so that we can check whether audio input is working.

diff --git a/Engine/Source/Runtime/Online/Voice/Private/Windows/VoiceCaptureWindows.cpp b/Engine/Source/Runtime/Online/Voice/Private/Windows/VoiceCaptureWindows.cpp
--- a/Engine/Source/Runtime/Online/Voice/Private/Windows/VoiceCaptureWindows.cpp
+++ b/Engine/Source/Runtime/Online/Voice/Private/Windows/VoiceCaptureWindows.cpp
@@ -9,6 +9,14 @@
 
 #include "Windows/AllowWindowsPlatformTypes.h"
 
+static int32 DisplayAmplitudeCvar = 0;
+FAutoConsoleVariableRef CVarDisplayAmplitude(
+	TEXT("voice.debug.PrintAmplitude"),
+	DisplayAmplitudeCvar,
+	TEXT("when set to 1, the current incoming amplitude of the VOIP engine will be displayed on screen.\n")
+	TEXT("0: disabled, 1: enabled."),
+	ECVF_Default);
+
 struct FVoiceCaptureWindowsVars
 {
 	/** GUID of current voice capture device */
@@ -694,6 +702,33 @@ EVoiceCaptureState::Type FVoiceCaptureWindows::GetVoiceData(uint8* OutVoiceBuffe
 		MicrophoneOutput.PushAudio(ConversionBuffer.GetData(), ConversionBuffer.GetNumSamples());
 	}
 
+	// print debug string with current amplitude:
+	if (DisplayAmplitudeCvar && GEngine)
+	{
+		static double TimeLastPrinted = FPlatformTime::Seconds();
+
+		static const double AmplitudeStringDisplayRate = 0.05;
+		static const int32 TotalNumTicks = 32;
+
+		if (FPlatformTime::Seconds() - TimeLastPrinted > AmplitudeStringDisplayRate)
+		{
+			const float MicLevel = MicLevelDetector.GetCurrentValue();
+			FString PrintString = FString::Printf(TEXT("Mic Amp: %.2f"), MicLevel);
+
+			int32 NumTicks = FMath::FloorToInt(MicLevelDetector.GetCurrentValue() * TotalNumTicks);
+
+			for (int32 Iteration = 0; Iteration < NumTicks; Iteration++)
+			{
+				PrintString.AppendChar(TCHAR('|'));
+			}
+
+			FColor TextColor = FLinearColor::LerpUsingHSV(FLinearColor::Green, FLinearColor::Red, MicLevel).ToFColor(true);
+
+			GEngine->AddOnScreenDebugMessage(30, AmplitudeStringDisplayRate, TextColor, PrintString, false);
+			TimeLastPrinted = FPlatformTime::Seconds();
+		}
+	}
+
 	return NewMicState;
 }
 
