iOS Background HTTP Fixes
- Fixing error where pausing and resuming tasks on iOS would activate too many tasks.
- Fixing potential multi-threaded crash accessing UnAssociatedTasks list
#tests iOS Client

diff --git a/Engine/Source/Runtime/Online/BackgroundHTTP/Private/IOS/ApplePlatformBackgroundHttpManager.cpp b/Engine/Source/Runtime/Online/BackgroundHTTP/Private/IOS/ApplePlatformBackgroundHttpManager.cpp
--- a/Engine/Source/Runtime/Online/BackgroundHTTP/Private/IOS/ApplePlatformBackgroundHttpManager.cpp
+++ b/Engine/Source/Runtime/Online/BackgroundHTTP/Private/IOS/ApplePlatformBackgroundHttpManager.cpp
@@ -20,9 +20,12 @@ FApplePlatformBackgroundHttpManager::~FApplePlatformBackgroundHttpManager()
 {
 	bWasAppleBGHTTPInitialized = false;
 	
-    [UnAssociatedTasks release];
-    UnAssociatedTasks = nullptr;
-    
+	{
+		FRWScopeLock ScopeLock(UnAssociatedTasksLock, SLT_Write);
+		[UnAssociatedTasks release];
+		UnAssociatedTasks = nullptr;
+	}
+	
     CleanUpNSURLSessionResponseDelegates();
 }
 
@@ -43,9 +46,14 @@ FApplePlatformBackgroundHttpManager::FApplePlatformBackgroundHttpManager()
 
 void FApplePlatformBackgroundHttpManager::Initialize()
 {
-	UnAssociatedTasks = [[NSMutableDictionary alloc] init];
+	//Initialize UnAssociatedTasks
+	{
+		FRWScopeLock ScopeLock(UnAssociatedTasksLock, SLT_Write);
+		UnAssociatedTasks = [[NSMutableDictionary alloc] init];
+	}
+	//This has its own lock when needed, so not included above
 	PopulateUnAssociatedTasks();
-
+	
     GConfig->GetFloat(TEXT("BackgroundHttp.iOSSettings"), TEXT("BackgroundHttp.ActiveReceiveTimeout"), ActiveTimeOutSetting, GEngineIni);
 	GConfig->GetInt(TEXT("BackgroundHttp.iOSSettings"), TEXT("BackgroundHttp.RetryResumeDataLimit"), RetryResumeDataLimitSetting, GEngineIni);
 
@@ -65,6 +73,8 @@ void FApplePlatformBackgroundHttpManager::PopulateUnAssociatedTasks()
 		{
 			[BackgroundDownloadSession getAllTasksWithCompletionHandler : ^ (NSArray<__kindof NSURLSessionTask*> *tasks)
 			{
+				FRWScopeLock ScopeLock(UnAssociatedTasksLock, SLT_Write);
+				
 				//Store all existing tasks by their URL
 				for (id task in tasks)
 				{
@@ -87,6 +97,8 @@ void FApplePlatformBackgroundHttpManager::PopulateUnAssociatedTasks()
 
 void FApplePlatformBackgroundHttpManager::PauseAllUnassociatedTasks()
 {
+	FRWScopeLock ScopeLock(UnAssociatedTasksLock, SLT_ReadOnly);
+	
     for (id Key in UnAssociatedTasks)
     {
         NSURLSessionDownloadTask* Task = (NSURLSessionDownloadTask*)([UnAssociatedTasks objectForKey:Key]);
@@ -102,6 +114,8 @@ void FApplePlatformBackgroundHttpManager::PauseAllUnassociatedTasks()
 
 void FApplePlatformBackgroundHttpManager::UnpauseAllUnassociatedTasks()
 {
+	FRWScopeLock ScopeLock(UnAssociatedTasksLock, SLT_ReadOnly);
+	
     for (id Key in UnAssociatedTasks)
     {
         NSURLSessionDownloadTask* Task = (NSURLSessionDownloadTask*)([UnAssociatedTasks objectForKey:Key]);
@@ -117,8 +131,11 @@ void FApplePlatformBackgroundHttpManager::UnpauseAllUnassociatedTasks()
 
 void FApplePlatformBackgroundHttpManager::Shutdown()
 {
-	[UnAssociatedTasks release];
-	UnAssociatedTasks = nullptr;
+	{
+		FRWScopeLock ScopeLock(UnAssociatedTasksLock, SLT_Write);
+		[UnAssociatedTasks release];
+		UnAssociatedTasks = nullptr;
+	}
 
 	CleanUpNSURLSessionResponseDelegates();
     FBackgroundURLSessionHandler::ShutdownBackgroundSession();
@@ -275,46 +292,63 @@ bool FApplePlatformBackgroundHttpManager::CheckForExistingUnAssociatedTask(const
 {
     bool bDidFindExistingTask = false;
     
-	if (ensureAlwaysMsgf(Request.IsValid(), TEXT("CheckForExistingUnAssociatedTask called with invalid Request!")))
+	TArray<FString> URLsToRemove;
+	
+	//Go through and read the UnAssociatedTasksLock and save off which URLs we need to remove as we associate with those tasks.
 	{
-		const TArray<FString>& URLList = Request->GetURLList();
-		for (const FString& URL : URLList)
+		FRWScopeLock ScopeLock(UnAssociatedTasksLock, SLT_ReadOnly);
+		
+		if (ensureAlwaysMsgf(Request.IsValid(), TEXT("CheckForExistingUnAssociatedTask called with invalid Request!")))
 		{
-			NSURLSessionTask* FoundTask = [UnAssociatedTasks valueForKey:URL.GetNSString()];
-			if (nullptr != FoundTask)
+			const TArray<FString>& URLList = Request->GetURLList();
+			for (const FString& URL : URLList)
 			{
-				if ([FoundTask state] != NSURLSessionTaskStateCompleted && [FoundTask state] != NSURLSessionTaskStateCanceling)
+				NSURLSessionTask* FoundTask = [UnAssociatedTasks valueForKey:URL.GetNSString()];
+				if (nullptr != FoundTask)
 				{
-					UE_LOG(LogBackgroundHttpManager, Display, TEXT("Existing UnAssociateTask found for Request! Attempting to Associate! RequestDebugID:%s"), *(Request->GetRequestDebugID()));
-					
-					//Associate with task so that our Request takes over ownership of this task so we can remove it from our UnAssociated Tasks list without it getting GC'd
-					if (Request->AssociateWithTask(FoundTask))
+					if ([FoundTask state] != NSURLSessionTaskStateCompleted && [FoundTask state] != NSURLSessionTaskStateCanceling)
 					{
-						//Always set our bWasTaskStartedInBG flag on our Request as true in the UnAssociated case as we don't know when it was really started
-						FPlatformAtomics::InterlockedExchange(&(Request->bWasTaskStartedInBG), true);
+						UE_LOG(LogBackgroundHttpManager, Display, TEXT("Existing UnAssociateTask found for Request! Attempting to Associate! RequestDebugID:%s"), *(Request->GetRequestDebugID()));
+						
+						//Associate with task so that our Request takes over ownership of this task so we can remove it from our UnAssociated Tasks list without it getting GC'd
+						if (Request->AssociateWithTask(FoundTask))
+						{
+							//Always set our bWasTaskStartedInBG flag on our Request as true in the UnAssociated case as we don't know when it was really started
+							FPlatformAtomics::InterlockedExchange(&(Request->bWasTaskStartedInBG), true);
 
-						//Suspend task in case it was running so that we can adhere to our desired platform max tasks
-						[FoundTask suspend];
+							//Suspend task in case it was running so that we can adhere to our desired platform max tasks
+							[FoundTask suspend];
 
-						bDidFindExistingTask = true;
-						break;
+							bDidFindExistingTask = true;
+							break;
+						}
+						else
+						{
+							UE_LOG(LogBackgroundHttpManager, Display, TEXT("UnAssociatedTask for request found, but failed to Associate with Task! -- RequestDebugID:%s | URL:%s"), *(Request->GetRequestDebugID()), *URL);
+						}
 					}
 					else
 					{
-						UE_LOG(LogBackgroundHttpManager, Display, TEXT("UnAssociatedTask for request found, but failed to Associate with Task! -- RequestDebugID:%s | URL:%s"), *(Request->GetRequestDebugID()), *URL);
+						UE_LOG(LogBackgroundHttpManager, Display, TEXT("UnAssociatedTask for request found, BUT NOT USING as it was cancelling or completed already! -- RequestDebugID:%s | URL:%s"), *(Request->GetRequestDebugID()), *URL);
 					}
+					
+					//Still want to remove UnAssociatedTask even though we didn't use it as something else can now be downloading this data and we do not want duplicates
+					URLsToRemove.Add(URL);
 				}
-				else
-				{
-					UE_LOG(LogBackgroundHttpManager, Display, TEXT("UnAssociatedTask for request found, BUT NOT USING as it was cancelling or completed already! -- RequestDebugID:%s | URL:%s"), *(Request->GetRequestDebugID()), *URL);
-				}
-				
-				//Still want to remove UnAssociatedTask even though we didn't use it as something else can now be downloading this data and we do not want duplicates
-				[UnAssociatedTasks removeObjectForKey : URL.GetNSString()];
 			}
 		}
 	}
-    
+	
+	//Remove all URLs from UnAssociatedTasks
+	{
+		FRWScopeLock ScopeLock(UnAssociatedTasksLock, SLT_Write);
+		
+		for (const FString& URL : URLsToRemove)
+		{
+			[UnAssociatedTasks removeObjectForKey : URL.GetNSString()];
+		}
+	}
+	    
     return bDidFindExistingTask;
 }
 
@@ -846,8 +880,6 @@ void FApplePlatformBackgroundHttpManager::TickRequests(float DeltaTime)
                         //Just cancel the task and let the OnTask_DidCompleteWithError callback handle retrying it if appropriate.
                         AppleRequest->CancelActiveTask();
                     }
-                    
-                    AppleRequest->SendDownloadProgressUpdate();
                 }
 				else if (bNeedsMoreTasks && !bIsTaskActive && !bIsTaskPaused && !bIsPendingCancel)
 				{
@@ -864,13 +896,9 @@ void FApplePlatformBackgroundHttpManager::TickRequests(float DeltaTime)
 							FoundRequestToStart = AppleRequest;
 						}
 					}
-					
-					AppleRequest->SendDownloadProgressUpdate();
 				}
-                else
-                {
-                    AppleRequest->SendDownloadProgressUpdate();
-                }
+				
+				AppleRequest->SendDownloadProgressUpdate();
             }
         }
 		
diff --git a/Engine/Source/Runtime/Online/BackgroundHTTP/Private/IOS/ApplePlatformBackgroundHttpRequest.cpp b/Engine/Source/Runtime/Online/BackgroundHTTP/Private/IOS/ApplePlatformBackgroundHttpRequest.cpp
--- a/Engine/Source/Runtime/Online/BackgroundHTTP/Private/IOS/ApplePlatformBackgroundHttpRequest.cpp
+++ b/Engine/Source/Runtime/Online/BackgroundHTTP/Private/IOS/ApplePlatformBackgroundHttpRequest.cpp
@@ -129,7 +129,6 @@ void FApplePlatformBackgroundHttpRequest::PauseUnderlyingTask()
             
             UE_LOG(LogBackgroundHttpRequest, Display, TEXT("Pausing Task for Request -- RequestDebugID:%s | TaskIdentifier:%d | TaskURL:%s"), *GetRequestDebugID(), TaskIdentifier, *TaskURL);
             
-            FPlatformAtomics::InterlockedExchange(&bIsTaskActive, false);
             FPlatformAtomics::InterlockedExchange(&bIsPendingCancel, false);
             FPlatformAtomics::InterlockedExchange(&bIsTaskPaused, true);
             
@@ -213,7 +212,18 @@ void FApplePlatformBackgroundHttpRequest::PauseRequest()
 
 void FApplePlatformBackgroundHttpRequest::ResumeRequest()
 {
-    ActivateUnderlyingTask();
+	FPlatformAtomics::InterlockedExchange(&bIsTaskPaused, false);
+	
+	//We only want to re-activate tasks that have already been flagged as active.
+	//Otherwise let our BackgroundHTTP Manager handle activating us on a tick now that we aren't paused.
+	if (IsUnderlyingTaskActive())
+	{
+		ActivateUnderlyingTask();
+	}
+	else
+	{
+		UE_LOG(LogBackgroundHttpRequest, Display, TEXT("ResumeRequest called on a task that wasn't active -- RequestDebugID:%s"), *GetRequestDebugID());
+	}
 }
 
 void FApplePlatformBackgroundHttpRequest::CancelActiveTask()
diff --git a/Engine/Source/Runtime/Online/BackgroundHTTP/Public/IOS/ApplePlatformBackgroundHttpManager.h b/Engine/Source/Runtime/Online/BackgroundHTTP/Public/IOS/ApplePlatformBackgroundHttpManager.h
--- a/Engine/Source/Runtime/Online/BackgroundHTTP/Public/IOS/ApplePlatformBackgroundHttpManager.h
+++ b/Engine/Source/Runtime/Online/BackgroundHTTP/Public/IOS/ApplePlatformBackgroundHttpManager.h
@@ -61,7 +61,8 @@ protected:
 protected:
 	//This dictionary is used to hold tasks that already existed on our Background Session when our BackgroundHttpManager was initialized. See PopulateUnAssociatedTasks and CheckForExistingUnAssociatedTask
 	NSMutableDictionary<NSString*, NSURLSessionDownloadTask*>* UnAssociatedTasks;
-
+	FRWLock UnAssociatedTasksLock;
+	
 	//Map to hold the associated BackgroundHttpRequest for any given URL. Multiple URLs will end up pointing to the same request in this list as it stores
     TMap<const FString, FBackgroundHttpURLMappedRequestPtr> URLToRequestMap;
 	FRWLock URLToRequestMapLock;
