exposing ssl certificate validation details for other modules
#rnx

diff --git a/Engine/Source/Runtime/Online/SSL/Private/SslCertificateManager.cpp b/Engine/Source/Runtime/Online/SSL/Private/SslCertificateManager.cpp
--- a/Engine/Source/Runtime/Online/SSL/Private/SslCertificateManager.cpp
+++ b/Engine/Source/Runtime/Online/SSL/Private/SslCertificateManager.cpp
@@ -11,6 +11,8 @@
 #include "Misc/CommandLine.h"
 #include <Algo/Count.h>
 
+FSslCertificateDelegates::FVerifySslCertificates FSslCertificateDelegates::VerifySslCertificates;
+
 #if WITH_SSL
 
 #if PLATFORM_WINDOWS
@@ -25,6 +27,27 @@
 #include "Windows/HideWindowsPlatformTypes.h"
 #endif
 
+namespace
+{
+	FString GetCertificateName(X509* const Certificate)
+	{
+		char StaticBuffer[2048];
+		// We do not have to free the return value of get_subject_name
+		X509_NAME_oneline(X509_get_subject_name(Certificate), StaticBuffer, sizeof(StaticBuffer));
+
+		return FString(ANSI_TO_TCHAR(StaticBuffer));
+	}
+
+	FString GetCertificateIssuer(X509* const Certificate)
+	{
+		char StaticBuffer[2048];
+		// We do not have to free the return value of get_subject_name
+		X509_NAME_oneline(X509_get_issuer_name(Certificate), StaticBuffer, sizeof(StaticBuffer));
+
+		return FString(ANSI_TO_TCHAR(StaticBuffer));
+	}
+}
+
 void FSslCertificateManager::AddCertificatesToSslContext(SSL_CTX* SslContextPtr) const
 {
 	X509_STORE* CertStore = SSL_CTX_get_cert_store(SslContextPtr);
@@ -188,8 +211,9 @@ bool FSslCertificateManager::VerifySslCertificates(X509_STORE_CTX* Context, cons
 	}
 
 	TArray<TArray<uint8, TFixedAllocator<PUBLIC_KEY_DIGEST_SIZE>>> CertDigests;
+	TArray<FSslCertificateDelegates::FCertInfo> CertInfoList;
+	const bool bCollectCertInfo = FSslCertificateDelegates::VerifySslCertificates.IsBound();
 
-	bool bFoundMatch = false;
 	for (int CertIndex = 0; CertIndex < NumCertsInChain; ++CertIndex)
 	{
 		X509* Certificate = sk_X509_value(Chain, CertIndex);
@@ -213,9 +237,34 @@ bool FSslCertificateManager::VerifySslCertificates(X509_STORE_CTX* Context, cons
 		SHA256_Final(Digest.GetData(), &ShaContext);
 
 		CertDigests.Add(Digest);
+
+		if (bCollectCertInfo)
+		{
+			FSslCertificateDelegates::FCertInfo CertInfo;
+			CertInfo.KeyDigest = Digest;
+			CertInfo.Issuer = GetCertificateIssuer(Certificate);
+			CertInfo.Subject = GetCertificateName(Certificate);
+
+			const EVP_MD* CertDigest = EVP_get_digestbyname("sha1");
+			if (CertDigest)
+			{
+				unsigned int DummySize = 0;
+				CertInfo.Thumbprint.AddZeroed(FSslCertificateDelegates::FCertInfo::CERT_DIGEST_SIZE);
+				X509_digest(Certificate, CertDigest, CertInfo.Thumbprint.GetData(), &DummySize);
+			}
+
+			CertInfoList.Add(CertInfo);
+		}
+	}
+
+	bool bFoundMatch = false;
+
+	bool bFoundMatchDelegate = FSslCertificateDelegates::VerifySslCertificates.IsBound() ? FSslCertificateDelegates::VerifySslCertificates.Execute(Domain, CertInfoList) : true;
+	if (bFoundMatchDelegate)
+	{
+		bFoundMatch = VerifySslCertificates(CertDigests, Domain);
 	}
 
-	bFoundMatch = VerifySslCertificates(CertDigests, Domain);
 	if (!bFoundMatch)
 	{
 		X509_STORE_CTX_set_error(Context, X509_V_ERR_CERT_UNTRUSTED);
@@ -376,18 +425,6 @@ void FSslCertificateManager::AddPEMFileToRootCertificateArray(const FString& Pat
 	}
 }
 
-namespace
-{
-	FString GetCertificateName(X509* const Certificate)
-	{
-		char StaticBuffer[2048];
-		// We do not have to free the return value of get_subject_name
-		X509_NAME_oneline(X509_get_subject_name(Certificate), StaticBuffer, sizeof(StaticBuffer));
-
-		return FString(ANSI_TO_TCHAR(StaticBuffer));
-	}
-}
-
 void FSslCertificateManager::AddCertificateToRootCertificateArray(X509* Certificate)
 {
 	bool bValidateRootCertificates = true;
diff --git a/Engine/Source/Runtime/Online/SSL/Public/Interfaces/ISslCertificateManager.h b/Engine/Source/Runtime/Online/SSL/Public/Interfaces/ISslCertificateManager.h
--- a/Engine/Source/Runtime/Online/SSL/Public/Interfaces/ISslCertificateManager.h
+++ b/Engine/Source/Runtime/Online/SSL/Public/Interfaces/ISslCertificateManager.h
@@ -5,6 +5,7 @@
 #include "Containers/ContainerAllocationPolicies.h"
 #include "CoreTypes.h"
 #include "CoreFwd.h"
+#include "Delegates/Delegate.h"
 
 struct ssl_ctx_st;
 typedef struct ssl_ctx_st SSL_CTX;
@@ -74,3 +75,20 @@ public:
 	 */
 	virtual bool VerifySslCertificates(TArray<TArray<uint8, TFixedAllocator<PUBLIC_KEY_DIGEST_SIZE>>>& Digests, const FString& Domain) const = 0;
 };
+
+class SSL_API FSslCertificateDelegates
+{
+public:
+	struct FCertInfo
+	{
+		static constexpr int CERT_DIGEST_SIZE = 20;
+
+		TArray<uint8, TInlineAllocator<ISslCertificateManager::PUBLIC_KEY_DIGEST_SIZE>> KeyDigest;
+		TArray<uint8, TInlineAllocator<CERT_DIGEST_SIZE>> Thumbprint;
+		FString Issuer;
+		FString Subject;
+	};
+
+	DECLARE_DELEGATE_RetVal_TwoParams(bool, FVerifySslCertificates, const FString&, const TArray<FCertInfo>&);
+	static FVerifySslCertificates VerifySslCertificates;
+};
