Include elapsed time in curlhttp response logs
#rnx

diff --git a/Engine/Source/Runtime/Online/HTTP/Private/Curl/CurlHttp.cpp b/Engine/Source/Runtime/Online/HTTP/Private/Curl/CurlHttp.cpp
--- a/Engine/Source/Runtime/Online/HTTP/Private/Curl/CurlHttp.cpp
+++ b/Engine/Source/Runtime/Online/HTTP/Private/Curl/CurlHttp.cpp
@@ -1135,13 +1135,13 @@ void FCurlHttpRequest::FinishedRequest()
 		{
 			if (bDebugServerResponse)
 			{
-				UE_LOG(LogHttp, Warning, TEXT("%p: request has been successfully processed. URL: %s, HTTP code: %d, content length: %d, actual payload size: %d"),
-					this, *GetURL(), Response->HttpCode, Response->ContentLength, Response->Payload.Num());
+				UE_LOG(LogHttp, Warning, TEXT("%p: request has been successfully processed. URL: %s, HTTP code: %d, content length: %d, actual payload size: %d, elapsed: %.2fs"),
+					this, *GetURL(), Response->HttpCode, Response->ContentLength, Response->Payload.Num(), ElapsedTime);
 			}
 			else
 			{
-				UE_LOG(LogHttp, Log, TEXT("%p: request has been successfully processed. URL: %s, HTTP code: %d, content length: %d, actual payload size: %d"),
-					this, *GetURL(), Response->HttpCode, Response->ContentLength, Response->Payload.Num());
+				UE_LOG(LogHttp, Log, TEXT("%p: request has been successfully processed. URL: %s, HTTP code: %d, content length: %d, actual payload size: %d, elapsed: %.2fs"),
+					this, *GetURL(), Response->HttpCode, Response->ContentLength, Response->Payload.Num(), ElapsedTime);
 			}
 
 			TArray<FString> AllHeaders = Response->GetAllHeaders();
