Added OnWebSocket created delegate
Added OnMessageSent delegate

diff --git a/Engine/Source/Runtime/Online/WebSockets/Private/HoloLens/HoloLensWebSocket.cpp b/Engine/Source/Runtime/Online/WebSockets/Private/HoloLens/HoloLensWebSocket.cpp
--- a/Engine/Source/Runtime/Online/WebSockets/Private/HoloLens/HoloLensWebSocket.cpp
+++ b/Engine/Source/Runtime/Online/WebSockets/Private/HoloLens/HoloLensWebSocket.cpp
@@ -129,6 +129,7 @@ void FHoloLensWebSocket::Send(const FString& Data)
 		{
 			Writer->WriteString(ref new Platform::String(*Data));
 			SendOperations.Add(Writer->StoreAsync());
+			OnMessageSent().Broadcast(Data);
 		}
 		catch (Platform::Exception^ Ex)
 		{
diff --git a/Engine/Source/Runtime/Online/WebSockets/Private/HoloLens/HoloLensWebSocket.h b/Engine/Source/Runtime/Online/WebSockets/Private/HoloLens/HoloLensWebSocket.h
--- a/Engine/Source/Runtime/Online/WebSockets/Private/HoloLens/HoloLensWebSocket.h
+++ b/Engine/Source/Runtime/Online/WebSockets/Private/HoloLens/HoloLensWebSocket.h
@@ -29,6 +29,7 @@ public:
 	virtual FWebSocketClosedEvent& OnClosed() override						{ return ClosedEvent; }
 	virtual FWebSocketMessageEvent& OnMessage() override					{ return MessageEvent; }
 	virtual FWebSocketRawMessageEvent& OnRawMessage() override				{ return RawMessageEvent; }
+	virtual FWebSocketMessageSentEvent& OnMessageSent() override { return OnMessageSentEvent; }
 
 public:
 	virtual bool Tick(float DeltaTime) override;
@@ -50,6 +51,8 @@ private:
 	FWebSocketClosedEvent ClosedEvent;
 	FWebSocketMessageEvent MessageEvent;
 	FWebSocketRawMessageEvent RawMessageEvent;
+	FWebSocketMessageSentEvent OnMessageSentEvent;
+
 	bool bUserClose;
 };
 
diff --git a/Engine/Source/Runtime/Online/WebSockets/Private/Lws/LwsWebSocket.cpp b/Engine/Source/Runtime/Online/WebSockets/Private/Lws/LwsWebSocket.cpp
--- a/Engine/Source/Runtime/Online/WebSockets/Private/Lws/LwsWebSocket.cpp
+++ b/Engine/Source/Runtime/Online/WebSockets/Private/Lws/LwsWebSocket.cpp
@@ -156,6 +156,7 @@ void FLwsWebSocket::Send(const FString& Data)
 {
 	FTCHARToUTF8 Converted(*Data);
 	Send((uint8*)Converted.Get(), Converted.Length(), false);
+	OnMessageSent().Broadcast(Data);
 }
 
 void FLwsWebSocket::SendFromQueue()
diff --git a/Engine/Source/Runtime/Online/WebSockets/Private/Lws/LwsWebSocket.h b/Engine/Source/Runtime/Online/WebSockets/Private/Lws/LwsWebSocket.h
--- a/Engine/Source/Runtime/Online/WebSockets/Private/Lws/LwsWebSocket.h
+++ b/Engine/Source/Runtime/Online/WebSockets/Private/Lws/LwsWebSocket.h
@@ -153,6 +153,12 @@ public:
 		return RawMessageEvent;
 	}
 
+	DECLARE_DERIVED_EVENT(FLwsWebSocket, IWebSocket::FWebSocketMessageSentEvent, FWebSocketMessageSentEvent);
+	virtual FWebSocketMessageSentEvent& OnMessageSent() override
+	{
+		return OnMessageSentEvent;
+	}
+
 	/** Callback on events for our libwebsockets connection */
 	int LwsCallback(lws* Instance, lws_callback_reasons Reason, void* Data, size_t Length);
 
@@ -272,6 +278,7 @@ private:
 	FWebSocketClosedEvent ClosedEvent;
 	FWebSocketMessageEvent MessageEvent;
 	FWebSocketRawMessageEvent RawMessageEvent;
+	FWebSocketMessageSentEvent OnMessageSentEvent;
 
 	/** libwebsockets connection */
 	struct lws *LwsConnection;
diff --git a/Engine/Source/Runtime/Online/WebSockets/Private/WebSocketsModule.cpp b/Engine/Source/Runtime/Online/WebSockets/Private/WebSocketsModule.cpp
--- a/Engine/Source/Runtime/Online/WebSockets/Private/WebSocketsModule.cpp
+++ b/Engine/Source/Runtime/Online/WebSockets/Private/WebSocketsModule.cpp
@@ -61,19 +61,21 @@ TSharedRef<IWebSocket> FWebSocketsModule::CreateWebSocket(const FString& Url, co
 	check(WebSocketsManager);
 
 	TArray<FString> ProtocolsCopy = Protocols;
-	ProtocolsCopy.RemoveAll([](const FString& Protocol){ return Protocol.IsEmpty(); });
-	return WebSocketsManager->CreateWebSocket(Url, ProtocolsCopy , UpgradeHeaders);
+	ProtocolsCopy.RemoveAll([](const FString& Protocol){ return Protocol.IsEmpty(); });	
+	TSharedRef<IWebSocket> WebSocket = WebSocketsManager->CreateWebSocket(Url, ProtocolsCopy, UpgradeHeaders);
+	OnWebSocketCreated.Broadcast(WebSocket, Protocols, Url);
+	
+	return WebSocket;
 }
 
 TSharedRef<IWebSocket> FWebSocketsModule::CreateWebSocket(const FString& Url, const FString& Protocol, const TMap<FString, FString>& UpgradeHeaders)
 {
-	check(WebSocketsManager);
-
 	TArray<FString> Protocols;
 	if (!Protocol.IsEmpty())
 	{
 		Protocols.Add(Protocol);
 	}
-	return WebSocketsManager->CreateWebSocket(Url, Protocols, UpgradeHeaders);
+	
+	return CreateWebSocket(Url, Protocols, UpgradeHeaders);
 }
 #endif // #if WITH_WEBSOCKETS
diff --git a/Engine/Source/Runtime/Online/WebSockets/Public/IWebSocket.h b/Engine/Source/Runtime/Online/WebSockets/Public/IWebSocket.h
--- a/Engine/Source/Runtime/Online/WebSockets/Public/IWebSocket.h
+++ b/Engine/Source/Runtime/Online/WebSockets/Public/IWebSocket.h
@@ -80,4 +80,10 @@ public:
 	DECLARE_EVENT_ThreeParams(IWebSocket, FWebSocketRawMessageEvent, const void* /* Data */, SIZE_T /* Size */, SIZE_T /* BytesRemaining */);
 	virtual FWebSocketRawMessageEvent& OnRawMessage() = 0;
 
+	/**
+	* Delegate called when a web socket text message has been sent.
+	* Assume UTF-8 encoding.
+	*/
+	DECLARE_EVENT_OneParam(IWebSocket, FWebSocketMessageSentEvent, const FString& /* MessageString */);
+	virtual FWebSocketMessageSentEvent& OnMessageSent() = 0;
 };
\ No newline at end of file
diff --git a/Engine/Source/Runtime/Online/WebSockets/Public/WebSocketsModule.h b/Engine/Source/Runtime/Online/WebSockets/Public/WebSocketsModule.h
--- a/Engine/Source/Runtime/Online/WebSockets/Public/WebSocketsModule.h
+++ b/Engine/Source/Runtime/Online/WebSockets/Public/WebSocketsModule.h
@@ -25,7 +25,18 @@ public:
 #endif // #if WITH_WEBSOCKETS
 	{
 	}
-
+#if WITH_WEBSOCKETS
+	/**
+	 * Delegate fired when a web socket is created.
+	 *
+	 * @param WebSocket Reference to newly created web socket
+	 * @param Protocols Array of protocols on web socket
+	 * @param Url Url the web socket is connected to
+	 */
+	DECLARE_MULTICAST_DELEGATE_ThreeParams(FOnWebSocketCreated, const TSharedPtr<class IWebSocket>& /* WebSocket */, const TArray<FString>& /* Protocols */, const FString& /* Url */);
+	FOnWebSocketCreated OnWebSocketCreated;
+#endif // #if WITH_WEBSOCKETS
+	
 	/**
 	 * Singleton-like access to this module's interface.  This is just for convenience!
 	 * Beware of calling this during the shutdown phase, though.  Your module might have been unloaded already.
