Added IXmppStanza abstract interface to expose stanza data getters independent of implementation; FStropheStanza now inherits and overrides members of it
Added OnXmppConnection created delegate
Added OnStanzaSent delegate, only implemented in Strophe
Added OnStanzaReceived delegate, only implemented in Strophe

diff --git a/Engine/Source/Runtime/Online/XMPP/Private/XmppJingle/XmppConnectionJingle.h b/Engine/Source/Runtime/Online/XMPP/Private/XmppJingle/XmppConnectionJingle.h
--- a/Engine/Source/Runtime/Online/XMPP/Private/XmppJingle/XmppConnectionJingle.h
+++ b/Engine/Source/Runtime/Online/XMPP/Private/XmppJingle/XmppConnectionJingle.h
@@ -37,6 +37,8 @@ public:
 	virtual FOnXmppLoginComplete& OnLoginComplete() override { return OnXmppLoginCompleteDelegate; }
 	virtual FOnXmppLoginChanged& OnLoginChanged() override { return OnXmppLoginChangedDelegate; }
 	virtual FOnXmppLogoutComplete& OnLogoutComplete() override { return OnXmppLogoutCompleteDelegate; }
+	virtual FOnXmppStanzaSent& OnStanzaSent() override { return OnXmppStanzaSentDelegate; }
+	virtual FOnXmppStanzaReceived& OnStanzaReceived() override { return OnXmppStanzaReceivedDelegate; }
 
 	virtual IXmppPresencePtr Presence() override;
 	virtual IXmppPubSubPtr PubSub() override;
@@ -100,6 +102,8 @@ private:
 	FOnXmppLoginComplete OnXmppLoginCompleteDelegate;
 	FOnXmppLoginChanged OnXmppLoginChangedDelegate;
 	FOnXmppLogoutComplete OnXmppLogoutCompleteDelegate;
+	FOnXmppStanzaReceived OnXmppStanzaReceivedDelegate;
+	FOnXmppStanzaSent OnXmppStanzaSentDelegate;
 
 	/** access to presence implementation */
 	TSharedPtr<class FXmppPresenceJingle, ESPMode::ThreadSafe> PresenceJingle;
diff --git a/Engine/Source/Runtime/Online/XMPP/Private/XmppModule.cpp b/Engine/Source/Runtime/Online/XMPP/Private/XmppModule.cpp
--- a/Engine/Source/Runtime/Online/XMPP/Private/XmppModule.cpp
+++ b/Engine/Source/Runtime/Online/XMPP/Private/XmppModule.cpp
@@ -734,7 +734,11 @@ TSharedRef<IXmppConnection> FXmppModule::CreateConnection(const FString& UserId)
 			Connection = FXmppNull::CreateConnection();
 		}
 
-		return ActiveConnections.Add(UserId, Connection.ToSharedRef());
+		TSharedRef<IXmppConnection> ConnectionRef = Connection.ToSharedRef();
+		ActiveConnections.Add(UserId, ConnectionRef);
+		OnXmppConnectionCreated.Broadcast(ConnectionRef);
+		
+		return ConnectionRef;
 	}
 }
 
diff --git a/Engine/Source/Runtime/Online/XMPP/Private/XmppNull.cpp b/Engine/Source/Runtime/Online/XMPP/Private/XmppNull.cpp
--- a/Engine/Source/Runtime/Online/XMPP/Private/XmppNull.cpp
+++ b/Engine/Source/Runtime/Online/XMPP/Private/XmppNull.cpp
@@ -37,6 +37,8 @@ public:
 	virtual FOnXmppLoginComplete& OnLoginComplete() override { return OnXmppLoginCompleteDelegate; }
 	virtual FOnXmppLoginChanged& OnLoginChanged() override { return OnXmppLoginChangedDelegate; }
 	virtual FOnXmppLogoutComplete& OnLogoutComplete() override { return OnXmppLogoutCompleteDelegate; }
+	virtual FOnXmppStanzaSent& OnStanzaSent() override { return OnXmppStanzaSentDelegate; }
+	virtual FOnXmppStanzaReceived& OnStanzaReceived() override { return OnXmppStanzaReceivedDelegate; }
 
 	virtual IXmppPresencePtr Presence() override { return NULL; }
 	virtual IXmppPubSubPtr PubSub() override { return NULL; }
@@ -59,7 +61,9 @@ private:
 	// completion delegates
 	FOnXmppLoginComplete OnXmppLoginCompleteDelegate;
 	FOnXmppLoginChanged OnXmppLoginChangedDelegate;
-	FOnXmppLogoutComplete OnXmppLogoutCompleteDelegate;	
+	FOnXmppLogoutComplete OnXmppLogoutCompleteDelegate;
+	FOnXmppStanzaReceived OnXmppStanzaReceivedDelegate;
+	FOnXmppStanzaSent OnXmppStanzaSentDelegate;
 };
 
 
diff --git a/Engine/Source/Runtime/Online/XMPP/Private/XmppStrophe/StropheStanza.cpp b/Engine/Source/Runtime/Online/XMPP/Private/XmppStrophe/StropheStanza.cpp
--- a/Engine/Source/Runtime/Online/XMPP/Private/XmppStrophe/StropheStanza.cpp
+++ b/Engine/Source/Runtime/Online/XMPP/Private/XmppStrophe/StropheStanza.cpp
@@ -117,7 +117,7 @@ void FStropheStanza::AddChild(const FStropheStanza& Child)
 	}
 }
 
-TOptional<FStropheStanza> FStropheStanza::GetChild(const FString& ChildName)
+TOptional<FStropheStanza> FStropheStanza::GetChildStropheStanza(const FString& ChildName)
 {
 	TOptional<FStropheStanza> OptionalStanza;
 
@@ -130,7 +130,17 @@ TOptional<FStropheStanza> FStropheStanza::GetChild(const FString& ChildName)
 	return OptionalStanza;
 }
 
-TOptional<const FStropheStanza> FStropheStanza::GetChild(const FString& ChildName) const
+TUniquePtr<IXmppStanza> FStropheStanza::GetChild(const FString& ChildName) const
+{
+	TOptional<const FStropheStanza> Stanza = GetChildStropheStanza(ChildName);
+	if (Stanza.IsSet())
+	{
+		return MakeUnique<FStropheStanza>(Stanza.GetValue());
+	}
+	return nullptr;
+}
+
+TOptional<const FStropheStanza> FStropheStanza::GetChildStropheStanza(const FString& ChildName) const
 {
 	TOptional<const FStropheStanza> OptionalStanza;
 
diff --git a/Engine/Source/Runtime/Online/XMPP/Private/XmppStrophe/StropheStanza.h b/Engine/Source/Runtime/Online/XMPP/Private/XmppStrophe/StropheStanza.h
--- a/Engine/Source/Runtime/Online/XMPP/Private/XmppStrophe/StropheStanza.h
+++ b/Engine/Source/Runtime/Online/XMPP/Private/XmppStrophe/StropheStanza.h
@@ -2,7 +2,7 @@
 
 #pragma once
 
-#include "CoreMinimal.h"
+#include "XmppStanza.h"
 
 #if WITH_XMPP_STROPHE
 
@@ -15,6 +15,7 @@ typedef struct _xmpp_conn_t xmpp_conn_t;
 typedef struct _xmpp_stanza_t xmpp_stanza_t;
 
 class FStropheStanza
+	: public IXmppStanza
 {
 	// For TArray<FStropheStanza> Emplace clone from xmp_stanza_t*
 	friend class TArray<FStropheStanza>;
@@ -32,18 +33,19 @@ public:
 	explicit FStropheStanza(const FXmppConnectionStrophe& Context, const FString& StanzaName = FString());
 	FStropheStanza(const FStropheStanza& Other);
 	FStropheStanza(FStropheStanza&& Other);
-	~FStropheStanza();
+	virtual ~FStropheStanza();
 	FStropheStanza& operator=(const FStropheStanza& Other);
 	FStropheStanza& operator=(FStropheStanza&& Other);
 
 	FStropheStanza Clone();
 
 	void AddChild(const FStropheStanza& Child);
-	TOptional<FStropheStanza> GetChild(const FString& ChildName);
-	TOptional<const FStropheStanza> GetChild(const FString& ChildName) const;
+	virtual TUniquePtr<IXmppStanza> GetChild(const FString& ChildName) const override;
+	TOptional<FStropheStanza> GetChildStropheStanza(const FString& ChildName);
+	TOptional<const FStropheStanza> GetChildStropheStanza(const FString& ChildName) const;
 	TOptional<FStropheStanza> GetChildByNameAndNamespace(const FString& ChildName, const FString& Namespace);
 	TOptional<const FStropheStanza> GetChildByNameAndNamespace(const FString& ChildName, const FString& Namespace) const;
-	bool HasChild(const FString& ChildName) const;
+	virtual bool HasChild(const FString& ChildName) const override;
 	bool HasChildByNameAndNamespace(const FString& ChildName, const FString& Namespace) const;
 
 	TArray<FStropheStanza> GetChildren();
@@ -53,35 +55,35 @@ public:
 	FString GetNamespace() const;
 
 	void SetAttribute(const FString& Key, const FString& Value);
-	FString GetAttribute(const FString& Key) const;
-	bool HasAttribute(const FString& Key) const;
+	virtual FString GetAttribute(const FString& Key) const override;
+	virtual bool HasAttribute(const FString& Key) const override;
 
 	void SetName(const FString& Name);
-	FString GetName() const;
+	virtual FString GetName() const override;
 
 	void SetText(const FString& Text);
-	FString GetText() const;
+	virtual FString GetText() const override;
 
 	void SetType(const FString& Type);
-	FString GetType() const;
+	virtual FString GetType() const override;
 
 	void SetId(const FString& Id);
-	FString GetId() const;
+	virtual FString GetId() const override;
 
 	void SetTo(const FXmppUserJid& To);
 	void SetTo(const FString& To);
-	FXmppUserJid GetTo() const;
+	virtual FXmppUserJid GetTo() const override;
 
 	void SetFrom(const FXmppUserJid& From);
 	void SetFrom(const FString& From);
-	FXmppUserJid GetFrom() const;
+	virtual FXmppUserJid GetFrom() const override;
 
 	// Helpers for Message stanzas */
 
 	/** Add a child stanza of name Body with the requested text.  Fails if we already have a body stanza, or if we are a text stanza. */
 	bool AddBodyWithText(const FString& Text);
 	/** Get the text from a child Body stanza, if one exists */
-	TOptional<FString> GetBodyText() const;
+	virtual TOptional<FString> GetBodyText() const override;
 
 protected:
 	/** Get the current stanza */
diff --git a/Engine/Source/Runtime/Online/XMPP/Private/XmppStrophe/XmppConnectionStrophe.cpp b/Engine/Source/Runtime/Online/XMPP/Private/XmppStrophe/XmppConnectionStrophe.cpp
--- a/Engine/Source/Runtime/Online/XMPP/Private/XmppStrophe/XmppConnectionStrophe.cpp
+++ b/Engine/Source/Runtime/Online/XMPP/Private/XmppStrophe/XmppConnectionStrophe.cpp
@@ -221,6 +221,7 @@ bool FXmppConnectionStrophe::SendStanza(FStropheStanza&& Stanza)
 		return false;
 	}
 
+	OnStanzaSent().Broadcast(Stanza);
 	bool bQueuedStanzaToBeSent = true;
 	if (StropheThread.IsValid())
 	{
@@ -243,7 +244,7 @@ bool FXmppConnectionStrophe::SendStanza(FStropheStanza&& Stanza)
 			PingStrophe->ResetPingTimer();
 		}
 	}
-
+	
 	return bQueuedStanzaToBeSent;
 }
 
@@ -297,7 +298,7 @@ void FXmppConnectionStrophe::ReceiveConnectionError(const FStropheError& Error,
 void FXmppConnectionStrophe::ReceiveStanza(const FStropheStanza& Stanza)
 {
 	UE_LOG(LogXmpp, Verbose, TEXT("Received Strophe XMPP Stanza %s"), *Stanza.GetName());
-
+	OnStanzaReceived().Broadcast(Stanza);
 	// Reset our ping timer now that we've received traffic
 	if (PingStrophe.IsValid() && ServerConfiguration.bResetPingTimeoutOnReceiveStanza)
 	{
diff --git a/Engine/Source/Runtime/Online/XMPP/Private/XmppStrophe/XmppConnectionStrophe.h b/Engine/Source/Runtime/Online/XMPP/Private/XmppStrophe/XmppConnectionStrophe.h
--- a/Engine/Source/Runtime/Online/XMPP/Private/XmppStrophe/XmppConnectionStrophe.h
+++ b/Engine/Source/Runtime/Online/XMPP/Private/XmppStrophe/XmppConnectionStrophe.h
@@ -53,6 +53,8 @@ public:
 	virtual FOnXmppLoginComplete& OnLoginComplete() override { return OnXmppLoginCompleteDelegate; }
 	virtual FOnXmppLoginChanged& OnLoginChanged() override { return OnXmppLoginChangedDelegate; }
 	virtual FOnXmppLogoutComplete& OnLogoutComplete() override { return OnXmppLogoutCompleteDelegate; }
+	virtual FOnXmppStanzaSent& OnStanzaSent() override { return OnXmppStanzaSentDelegate; }
+	virtual FOnXmppStanzaReceived& OnStanzaReceived() override { return OnXmppStanzaReceivedDelegate; }
 
 	virtual IXmppMessagesPtr Messages() override;
 	virtual IXmppMultiUserChatPtr MultiUserChat() override;
@@ -121,6 +123,10 @@ protected:
 	FOnXmppLoginChanged OnXmppLoginChangedDelegate;
 	FOnXmppLogoutComplete OnXmppLogoutCompleteDelegate;
 
+	/** Message Status delegates */
+	FOnXmppStanzaReceived OnXmppStanzaReceivedDelegate;
+	FOnXmppStanzaSent OnXmppStanzaSentDelegate;
+
 	/** XMPP Implementation Shared Pointers */
 	FXmppMessagesStrophePtr MessagesStrophe;
 	FXmppMultiUserChatStrophePtr MultiUserChatStrophe;
diff --git a/Engine/Source/Runtime/Online/XMPP/Private/XmppStrophe/XmppMessagesStrophe.cpp b/Engine/Source/Runtime/Online/XMPP/Private/XmppStrophe/XmppMessagesStrophe.cpp
--- a/Engine/Source/Runtime/Online/XMPP/Private/XmppStrophe/XmppMessagesStrophe.cpp
+++ b/Engine/Source/Runtime/Online/XMPP/Private/XmppStrophe/XmppMessagesStrophe.cpp
@@ -51,7 +51,7 @@ bool FXmppMessagesStrophe::ReceiveStanza(const FStropheStanza& IncomingStanza)
 		return false;
 	}
 
-	const TOptional<const FStropheStanza> ErrorStanza = IncomingStanza.GetChild(Strophe::SN_ERROR);
+	const TOptional<const FStropheStanza> ErrorStanza = IncomingStanza.GetChildStropheStanza(Strophe::SN_ERROR);
 	if (ErrorStanza.IsSet())
 	{
 		return HandleMessageErrorStanza(ErrorStanza.GetValue());
diff --git a/Engine/Source/Runtime/Online/XMPP/Private/XmppStrophe/XmppMultiUserChatStrophe.cpp b/Engine/Source/Runtime/Online/XMPP/Private/XmppStrophe/XmppMultiUserChatStrophe.cpp
--- a/Engine/Source/Runtime/Online/XMPP/Private/XmppStrophe/XmppMultiUserChatStrophe.cpp
+++ b/Engine/Source/Runtime/Online/XMPP/Private/XmppStrophe/XmppMultiUserChatStrophe.cpp
@@ -120,7 +120,7 @@ bool FXmppMultiUserChatStrophe::HandlePresenceStanza(const FStropheStanza& Incom
 	TOptional<const FStropheStanza> UserStanza = IncomingStanza.GetChildByNameAndNamespace(Strophe::SN_X, Strophe::SNS_MUC_USER);
 	if (UserStanza.IsSet())
 	{
-		TOptional<const FStropheStanza> UserItemStanza = UserStanza->GetChild(Strophe::SN_ITEM);
+		TOptional<const FStropheStanza> UserItemStanza = UserStanza->GetChildStropheStanza(Strophe::SN_ITEM);
 		if (UserItemStanza.IsSet())
 		{
 			Presence.MemberJid = FXmppStrophe::JidFromString(UserItemStanza->GetAttribute(Strophe::SA_JID));
@@ -139,7 +139,7 @@ bool FXmppMultiUserChatStrophe::HandlePresenceErrorStanza(const FStropheStanza&
 	FXmppStropheErrorPair OutError;
 	OutError.RoomId = IncomingStanza.GetFrom().Id;
 
-	TOptional<const FStropheStanza> Error = IncomingStanza.GetChild(Strophe::SN_ERROR);
+	TOptional<const FStropheStanza> Error = IncomingStanza.GetChildStropheStanza(Strophe::SN_ERROR);
 	if (Error.IsSet())
 	{
 		const TArray<FStropheStanza> ErrorList = Error->GetChildren();
@@ -205,7 +205,7 @@ bool FXmppMultiUserChatStrophe::HandlePresenceErrorStanza(const FStropheStanza&
 
 bool FXmppMultiUserChatStrophe::HandleGroupChatStanza(const FStropheStanza& IncomingStanza)
 {
-	TOptional<const FStropheStanza> SubjectStanza = IncomingStanza.GetChild(Strophe::SN_SUBJECT);
+	TOptional<const FStropheStanza> SubjectStanza = IncomingStanza.GetChildStropheStanza(Strophe::SN_SUBJECT);
 	if (SubjectStanza.IsSet())
 	{
 		FXmppStropheSubjectUpdate SubjectUpdate;
@@ -253,7 +253,7 @@ bool FXmppMultiUserChatStrophe::HandleGroupChatStanza(const FStropheStanza& Inco
 	ChatMessage.Body = MoveTemp(BodyText.GetValue());
 
 	// Parse Timezone
-	TOptional<const FStropheStanza> StanzaDelay = IncomingStanza.GetChild(Strophe::SN_DELAY);
+	TOptional<const FStropheStanza> StanzaDelay = IncomingStanza.GetChildStropheStanza(Strophe::SN_DELAY);
 	if (StanzaDelay.IsSet())
 	{
 		if (StanzaDelay->HasAttribute(Strophe::SA_STAMP))
@@ -277,7 +277,7 @@ bool FXmppMultiUserChatStrophe::HandleGroupChatErrorStanza(const FStropheStanza&
 {
 	FString ErrorMessage;
 
-	TOptional<const FStropheStanza> Error = IncomingStanza.GetChild(Strophe::SN_ERROR);
+	TOptional<const FStropheStanza> Error = IncomingStanza.GetChildStropheStanza(Strophe::SN_ERROR);
 	if (Error.IsSet())
 	{
 		const TArray<FStropheStanza> ErrorList = Error->GetChildren();
@@ -328,7 +328,7 @@ bool FXmppMultiUserChatStrophe::HandleRoomConfigStanza(const FStropheStanza& Inc
 	//    means the channel already exists and we're querying the options for it
 
 	// Check for config write case (No Query child)
-	TOptional<const FStropheStanza> QueryStanza = IncomingStanza.GetChild(Strophe::SN_QUERY);
+	TOptional<const FStropheStanza> QueryStanza = IncomingStanza.GetChildStropheStanza(Strophe::SN_QUERY);
 	if (!QueryStanza.IsSet())
 	{
 		FEmbeddedCommunication::KeepAwake(TickRequesterId, false);
@@ -348,7 +348,7 @@ bool FXmppMultiUserChatStrophe::HandleRoomConfigErrorStanza(const FStropheStanza
 	FXmppStropheErrorPair OutError;
 	OutError.RoomId = IncomingStanza.GetFrom().Id;
 
-	TOptional<const FStropheStanza> ErrorStanza = IncomingStanza.GetChild(Strophe::SN_ERROR);
+	TOptional<const FStropheStanza> ErrorStanza = IncomingStanza.GetChildStropheStanza(Strophe::SN_ERROR);
 	if (ErrorStanza.IsSet())
 	{
 		const FString ErrorType = ErrorStanza->GetType();
diff --git a/Engine/Source/Runtime/Online/XMPP/Private/XmppStrophe/XmppPresenceStrophe.cpp b/Engine/Source/Runtime/Online/XMPP/Private/XmppStrophe/XmppPresenceStrophe.cpp
--- a/Engine/Source/Runtime/Online/XMPP/Private/XmppStrophe/XmppPresenceStrophe.cpp
+++ b/Engine/Source/Runtime/Online/XMPP/Private/XmppStrophe/XmppPresenceStrophe.cpp
@@ -78,7 +78,7 @@ bool FXmppPresenceStrophe::ReceiveStanza(const FStropheStanza& IncomingStanza)
 	{
 		Presence.bIsAvailable = true;
 
-		TOptional<const FStropheStanza> StatusTextStanza = IncomingStanza.GetChild(Strophe::SN_STATUS);
+		TOptional<const FStropheStanza> StatusTextStanza = IncomingStanza.GetChildStropheStanza(Strophe::SN_STATUS);
 		if (StatusTextStanza.IsSet())
 		{
 			Presence.StatusStr = StatusTextStanza->GetText();
@@ -86,7 +86,7 @@ bool FXmppPresenceStrophe::ReceiveStanza(const FStropheStanza& IncomingStanza)
 
 		Presence.Status = EXmppPresenceStatus::Online;
 
-		TOptional<const FStropheStanza> StatusEnumStanza = IncomingStanza.GetChild(Strophe::SN_SHOW);
+		TOptional<const FStropheStanza> StatusEnumStanza = IncomingStanza.GetChildStropheStanza(Strophe::SN_SHOW);
 		if (StatusEnumStanza.IsSet())
 		{
 			FString StatusEnum(StatusEnumStanza->GetText());
@@ -108,7 +108,7 @@ bool FXmppPresenceStrophe::ReceiveStanza(const FStropheStanza& IncomingStanza)
 			}
 		}
 
-		TOptional<const FStropheStanza> TimestampStanza = IncomingStanza.GetChild(Strophe::SN_DELAY);
+		TOptional<const FStropheStanza> TimestampStanza = IncomingStanza.GetChildStropheStanza(Strophe::SN_DELAY);
 		if (TimestampStanza.IsSet())
 		{
 			FDateTime::ParseIso8601(*TimestampStanza->GetAttribute(Strophe::SA_STAMP), Presence.SentTime);
diff --git a/Engine/Source/Runtime/Online/XMPP/Private/XmppStrophe/XmppPrivateChatStrophe.cpp b/Engine/Source/Runtime/Online/XMPP/Private/XmppStrophe/XmppPrivateChatStrophe.cpp
--- a/Engine/Source/Runtime/Online/XMPP/Private/XmppStrophe/XmppPrivateChatStrophe.cpp
+++ b/Engine/Source/Runtime/Online/XMPP/Private/XmppStrophe/XmppPrivateChatStrophe.cpp
@@ -72,7 +72,7 @@ bool FXmppPrivateChatStrophe::ReceiveStanza(const FStropheStanza& IncomingStanza
 	ChatMessage.Body = MoveTemp(BodyText.GetValue());
 
 	// Parse Timezone
-	TOptional<const FStropheStanza> StanzaDelay = IncomingStanza.GetChild(Strophe::SN_DELAY);
+	TOptional<const FStropheStanza> StanzaDelay = IncomingStanza.GetChildStropheStanza(Strophe::SN_DELAY);
 	if (StanzaDelay.IsSet())
 	{
 		if (StanzaDelay->HasAttribute(Strophe::SA_STAMP))
diff --git a/Engine/Source/Runtime/Online/XMPP/Public/XmppConnection.h b/Engine/Source/Runtime/Online/XMPP/Public/XmppConnection.h
--- a/Engine/Source/Runtime/Online/XMPP/Public/XmppConnection.h
+++ b/Engine/Source/Runtime/Online/XMPP/Public/XmppConnection.h
@@ -9,6 +9,7 @@ class IXmppMessages;
 class IXmppMultiUserChat;
 class IXmppPresence;
 class IXmppPubSub;
+class IXmppStanza;
 
 /** Possible XMPP login states */
 namespace EXmppLoginStatus
@@ -291,6 +292,19 @@ public:
 
 	/** Alias for old typo, remove in 4.24 */
 	using FOnXmppLogingChanged UE_DEPRECATED(4.23, "Please update usages of 'FOnXmppLogingChanged' to 'FOnXmppLoginChanged'") = FOnXmppLoginChanged;
+	/**
+	* Delegate called when a message is received (ONLY STROPHE)
+	*
+	* @param reference to stanza
+	*/
+	DECLARE_MULTICAST_DELEGATE_OneParam(FOnXmppStanzaReceived, const IXmppStanza& /* Stanza */);
+
+	/**
+	* Delegate called when a message is sent (ONLY STROPHE)
+	*
+	* @param reference to stanza
+	*/
+	DECLARE_MULTICAST_DELEGATE_OneParam(FOnXmppStanzaSent, const IXmppStanza& /* Stanza */);
 
 	/** @return login complete delegate */
 	virtual FOnXmppLoginComplete& OnLoginComplete() = 0;
@@ -298,6 +312,10 @@ public:
 	virtual FOnXmppLoginChanged& OnLoginChanged() = 0;
 	/** @return logout complete delegate */
 	virtual FOnXmppLogoutComplete& OnLogoutComplete() = 0;
+	/** @return Stanza sent delegate */
+	virtual FOnXmppStanzaSent& OnStanzaSent() = 0;
+	/** @return Stanza received delegate */
+	virtual FOnXmppStanzaReceived& OnStanzaReceived() = 0;
 
 	/** @return Presence interface if available. NULL otherwise */
 	virtual IXmppPresencePtr Presence() = 0;
diff --git a/Engine/Source/Runtime/Online/XMPP/Public/XmppModule.h b/Engine/Source/Runtime/Online/XMPP/Public/XmppModule.h
--- a/Engine/Source/Runtime/Online/XMPP/Public/XmppModule.h
+++ b/Engine/Source/Runtime/Online/XMPP/Public/XmppModule.h
@@ -114,6 +114,16 @@ public:
 	DECLARE_EVENT_TwoParams(FXmppModule, FOnXmppConnectionRelinquished, const TSharedRef<IXmppConnection>& /*XmppConnection*/, const FString& /*SystemName*/);
 	FOnXmppConnectionRelinquished OnXmppConnectionRelinquished;
 
+	/**
+	* Delegate fired when an Xmpp connection is created.
+	*
+	* @param NewConnection Reference to newly created Xmpp connection
+	*
+	*/
+	DECLARE_MULTICAST_DELEGATE_OneParam(FOnXmppConnectionCreated, const TSharedRef<IXmppConnection>& /*NewConnection*/);
+	FOnXmppConnectionCreated OnXmppConnectionCreated;
+
+
 private:
 
 	// IModuleInterface
diff --git a/Engine/Source/Runtime/Online/XMPP/Public/XmppStanza.h b/Engine/Source/Runtime/Online/XMPP/Public/XmppStanza.h
new file mode 100644
--- /dev/null
+++ b/Engine/Source/Runtime/Online/XMPP/Public/XmppStanza.h
@@ -0,0 +1,88 @@
+// Copyright 1998-2019 Epic Games, Inc. All Rights Reserved.
+
+#pragma once
+
+#include "CoreMinimal.h"
+
+class FXmppUserJid;
+
+/**
+ * Xmpp stanza abstract interface
+ * Exposes getters for various attributes and fields of an Xmpp stanza
+ */
+class IXmppStanza
+{
+public:
+
+	virtual ~IXmppStanza() {};
+
+	/**
+	 * Get Stanza type (presence, query, message)
+	 * @return stanza type
+	 */
+	virtual FString GetName() const = 0;
+	
+	/**
+	 * Get Stanza text (concatenates all child text nodes)
+	 * @return stanza text
+	 */
+	virtual FString GetText() const = 0;
+
+	/**
+	 * Get Stanza type
+	 * @return stanza type
+	 */
+	virtual FString GetType() const = 0;
+
+	/**
+	 * Get Stanza Id
+	 * @return Stanza Id
+	 */
+	virtual FString GetId() const = 0;
+
+	/**
+	 * Get Stanza to (destination) user Jid
+	 * @return Stanza to Jid
+	 */
+	virtual FXmppUserJid GetTo() const = 0;
+
+	/**
+	 * Get Stanza from (source) user Jid
+	 * @return Stanza from Jid
+	 */
+	virtual FXmppUserJid GetFrom() const = 0;
+
+	/**
+	 * Get Stanza attribute
+	 * @param Key attribute name
+	 * @return attribute text
+	 */
+	virtual FString GetAttribute(const FString& Key) const = 0;
+
+	/**
+	 * Returns whether or not stanza has a given attribute
+	 * @param Key attribute name
+	 * @return whether or not stanza has attribute named Key
+	 */
+	virtual bool HasAttribute(const FString& Key) const = 0;
+
+	/**
+	 * Get stanza body text (if any)
+	 * @return Stanza body text
+	 */
+	virtual TOptional<FString> GetBodyText() const = 0;
+
+	/**
+	 * Get given Stanza child, if existing
+	 * @param ChildName name of child to return
+	 * @return Stanza ChildName, if it exists
+	 */
+	virtual TUniquePtr<IXmppStanza> GetChild(const FString& ChildName) const = 0;
+
+	/**
+	 * Whether or not Stanza has a given child
+	 * @param ChildName name of child to check for 
+	 * @return Whether or not Stanza has child ChildName
+	 */
+	virtual bool HasChild(const FString& ChildName) const = 0;
+};
\ No newline at end of file
