Moved OnRequestWillRetryDelegate from HttpRetrySystem to IHttpRequest

diff --git a/Engine/Source/Runtime/Online/BuildPatchServices/Private/Tests/Mock/HttpRequest.mock.h b/Engine/Source/Runtime/Online/BuildPatchServices/Private/Tests/Mock/HttpRequest.mock.h
--- a/Engine/Source/Runtime/Online/BuildPatchServices/Private/Tests/Mock/HttpRequest.mock.h
+++ b/Engine/Source/Runtime/Online/BuildPatchServices/Private/Tests/Mock/HttpRequest.mock.h
@@ -22,6 +22,12 @@ namespace BuildPatchServices
 			return FString();
 		}
 
+		virtual FHttpRequestWillRetryDelegate& OnRequestWillRetry() override
+		{
+			MOCK_FUNC_NOT_IMPLEMENTED("FMockHttpRequest::OnRequestWillRetry");
+			return HttpRequestWillRetryDelegate;
+		}
+
 		virtual FString GetURLParameter(const FString& ParameterName) const override
 		{
 			MOCK_FUNC_NOT_IMPLEMENTED("FMockHttpRequest::GetURLParameter");
@@ -160,6 +166,7 @@ namespace BuildPatchServices
 		FHttpRequestProgressDelegate HttpRequestProgressDelegate;
 		FHttpRequestCompleteDelegate HttpRequestCompleteDelegate;
 		FHttpRequestHeaderReceivedDelegate HttpHeaderReceivedDelegate;
+		FHttpRequestWillRetryDelegate HttpRequestWillRetryDelegate;
 
 		TArray<FRxSetVerb> RxSetVerb;
 		TArray<FRxSetURL> RxSetURL;
diff --git a/Engine/Source/Runtime/Online/HTTP/Private/GenericPlatform/HttpRequestImpl.cpp b/Engine/Source/Runtime/Online/HTTP/Private/GenericPlatform/HttpRequestImpl.cpp
--- a/Engine/Source/Runtime/Online/HTTP/Private/GenericPlatform/HttpRequestImpl.cpp
+++ b/Engine/Source/Runtime/Online/HTTP/Private/GenericPlatform/HttpRequestImpl.cpp
@@ -21,6 +21,12 @@ FHttpRequestHeaderReceivedDelegate& FHttpRequestImpl::OnHeaderReceived()
 	return HeaderReceivedDelegate;
 }
 
+FHttpRequestWillRetryDelegate& FHttpRequestImpl::OnRequestWillRetry()
+{
+	UE_LOG(LogHttp, VeryVerbose, TEXT("FHttpRequestImpl::OnRequestWillRetry()"));
+	return OnRequestWillRetryDelegate;
+}
+
 void FHttpRequestImpl::BroadcastResponseHeadersReceived()
 {
 	if (OnHeaderReceived().IsBound())
diff --git a/Engine/Source/Runtime/Online/HTTP/Public/GenericPlatform/HttpRequestImpl.h b/Engine/Source/Runtime/Online/HTTP/Public/GenericPlatform/HttpRequestImpl.h
--- a/Engine/Source/Runtime/Online/HTTP/Public/GenericPlatform/HttpRequestImpl.h
+++ b/Engine/Source/Runtime/Online/HTTP/Public/GenericPlatform/HttpRequestImpl.h
@@ -14,6 +14,7 @@ public:
 	virtual FHttpRequestCompleteDelegate& OnProcessRequestComplete() override;
 	virtual FHttpRequestProgressDelegate& OnRequestProgress() override;
 	virtual FHttpRequestHeaderReceivedDelegate& OnHeaderReceived() override;
+	virtual FHttpRequestWillRetryDelegate& OnRequestWillRetry() override;
 
 protected:
 	/** 
@@ -31,4 +32,7 @@ protected:
 
 	/** Delegate that will get called for each new header received */
 	FHttpRequestHeaderReceivedDelegate HeaderReceivedDelegate;
+	
+	/** Delegate that will get called when request will be retried */
+	FHttpRequestWillRetryDelegate OnRequestWillRetryDelegate;
 };
diff --git a/Engine/Source/Runtime/Online/HTTP/Public/HttpRetrySystem.h b/Engine/Source/Runtime/Online/HTTP/Public/HttpRetrySystem.h
--- a/Engine/Source/Runtime/Online/HTTP/Public/HttpRetrySystem.h
+++ b/Engine/Source/Runtime/Online/HTTP/Public/HttpRetrySystem.h
@@ -56,14 +56,6 @@ namespace FHttpRetrySystem
 	TOptional<double> HTTP_API ReadThrottledTimeFromResponseInSeconds(FHttpResponsePtr Response);
 };
 
-/**
-* Delegate called when an Http request will be retried in the future
-*
-* @param first parameter - original Http request that started things
-* @param second parameter - response received from the server if a successful connection was established
-* @param third parameter - seconds in the future when the response will be retried
-*/
-DECLARE_DELEGATE_ThreeParams(FHttpRequestWillRetryDelegate, FHttpRequestPtr, FHttpResponsePtr, float);
 
 namespace FHttpRetrySystem
 {
@@ -92,7 +84,6 @@ namespace FHttpRetrySystem
 		// IHttpRequest interface
 		HTTP_API virtual bool ProcessRequest() override;
 		HTTP_API virtual void CancelRequest() override;
-		virtual FHttpRequestWillRetryDelegate& OnRequestWillRetry() { return OnRequestWillRetryDelegate; }
 		
 		// FRequest
 		EStatus::Type GetRetryStatus() const { return Status; }
@@ -129,8 +120,6 @@ namespace FHttpRetrySystem
 		/** The original URL before replacing anything from RetryDomains */
 		FString								 OriginalUrl;
 
-		FHttpRequestWillRetryDelegate OnRequestWillRetryDelegate;
-
 		FManager& RetryManager;
     };
 }
diff --git a/Engine/Source/Runtime/Online/HTTP/Public/Interfaces/IHttpRequest.h b/Engine/Source/Runtime/Online/HTTP/Public/Interfaces/IHttpRequest.h
--- a/Engine/Source/Runtime/Online/HTTP/Public/Interfaces/IHttpRequest.h
+++ b/Engine/Source/Runtime/Online/HTTP/Public/Interfaces/IHttpRequest.h
@@ -98,6 +98,15 @@ DECLARE_DELEGATE_ThreeParams(FHttpRequestHeaderReceivedDelegate, FHttpRequestPtr
  */
 DECLARE_DELEGATE_ThreeParams(FHttpRequestProgressDelegate, FHttpRequestPtr /*Request*/, int32 /*BytesSent*/, int32 /*BytesReceived*/);
 
+/**
+ * Delegate called when an Http request will be retried in the future
+ *
+ * @param Request - original Http request that started things
+ * @param Response - response received from the server if a successful connection was established
+ * @param SecondsToRetry - seconds in the future when the response will be retried
+ */
+DECLARE_DELEGATE_ThreeParams(FHttpRequestWillRetryDelegate, FHttpRequestPtr /*Request*/, FHttpResponsePtr /*Response*/, float /*SecondsToRetry*/);
+
 /**
  * Interface for Http requests (created using FHttpFactory)
  */
@@ -207,6 +216,11 @@ public:
 	 * Delegate called to update the request/response progress. See FHttpRequestProgressDelegate
 	 */
 	virtual FHttpRequestProgressDelegate& OnRequestProgress() = 0;
+	
+	/**
+	* Delegate called when the request will be retried
+	*/
+	virtual FHttpRequestWillRetryDelegate& OnRequestWillRetry() = 0;
 
 	/** 
 	 * Delegate called to signal the receipt of a header.  See FHttpRequestHeaderReceivedDelegate
