CIS fix
FHttpRequestIXML missing OnRequestWillRetry

diff --git a/Engine/Source/Runtime/Online/HTTP/Private/IXML/HttpIXML.h b/Engine/Source/Runtime/Online/HTTP/Private/IXML/HttpIXML.h
--- a/Engine/Source/Runtime/Online/HTTP/Private/IXML/HttpIXML.h
+++ b/Engine/Source/Runtime/Online/HTTP/Private/IXML/HttpIXML.h
@@ -46,6 +46,7 @@ public:
 	virtual FHttpRequestCompleteDelegate& OnProcessRequestComplete() override;
 	virtual FHttpRequestProgressDelegate& OnRequestProgress() override;
 	virtual FHttpRequestHeaderReceivedDelegate& OnHeaderReceived() override;
+	virtual FHttpRequestWillRetryDelegate& OnRequestWillRetry() override;
 	virtual void CancelRequest() override;
 	virtual EHttpRequestStatus::Type GetStatus() const override;
 	virtual const FHttpResponsePtr GetResponse() const override;
@@ -78,6 +79,7 @@ private:
 	FHttpRequestCompleteDelegate		CompleteDelegate;
 	FHttpRequestProgressDelegate		RequestProgressDelegate;
 	FHttpRequestHeaderReceivedDelegate	RequestHeaderReceivedDelegate;
+	FHttpRequestWillRetryDelegate		OnRequestWillRetryDelegate;
 
 	EHttpRequestStatus::Type			RequestStatus;
 	float								ElapsedTime;
diff --git a/Engine/Source/Runtime/Online/HTTP/Private/IXML/HttpIXml.cpp b/Engine/Source/Runtime/Online/HTTP/Private/IXML/HttpIXml.cpp
--- a/Engine/Source/Runtime/Online/HTTP/Private/IXML/HttpIXml.cpp
+++ b/Engine/Source/Runtime/Online/HTTP/Private/IXML/HttpIXml.cpp
@@ -362,6 +362,11 @@ FHttpRequestHeaderReceivedDelegate& FHttpRequestIXML::OnHeaderReceived()
 	return RequestHeaderReceivedDelegate;
 }
 
+FHttpRequestWillRetryDelegate& FHttpRequestIXML::OnRequestWillRetry()
+{
+	return OnRequestWillRetryDelegate;
+}
+
 //-----------------------------------------------------------------------------
 //	
 //-----------------------------------------------------------------------------
