Fix for lambda crash when running from main thread after BuildPatchInstaller has already finished and destructed.
This bug would be exposed by scenarios where the installer completed in a short time while game thread is ticking at a lower rate.

diff --git a/Engine/Source/Runtime/Online/BuildPatchServices/Private/BuildPatchInstaller.cpp b/Engine/Source/Runtime/Online/BuildPatchServices/Private/BuildPatchInstaller.cpp
--- a/Engine/Source/Runtime/Online/BuildPatchServices/Private/BuildPatchInstaller.cpp
+++ b/Engine/Source/Runtime/Online/BuildPatchServices/Private/BuildPatchInstaller.cpp
@@ -700,7 +700,7 @@ namespace BuildPatchServices
 			{
 				ChunkDataSizeProvider->AddManifestData(InstallerAction.TryGetInstallManifest());
 			}
-		});
+		}).Wait();
 
 		// Init build statistics that are known.
 		{
