#UE4 Memory read fix in HTTPServer. Found by asan

diff --git a/Engine/Source/Runtime/Online/HTTPServer/Private/HttpConnectionRequestReadContext.cpp b/Engine/Source/Runtime/Online/HTTPServer/Private/HttpConnectionRequestReadContext.cpp
--- a/Engine/Source/Runtime/Online/HTTPServer/Private/HttpConnectionRequestReadContext.cpp
+++ b/Engine/Source/Runtime/Online/HTTPServer/Private/HttpConnectionRequestReadContext.cpp
@@ -298,8 +298,7 @@ TSharedPtr<FHttpServerRequest> FHttpConnectionRequestReadContext::BuildRequest(c
 			const auto& HeaderValuesStr = HeaderLine.Mid(SplitIndex + 1).TrimStartAndEnd();
 
 			TArray<FString> HeaderValues;
-			const TCHAR HeaderValueDelimiters[] = { TCHAR(',') };
-			HeaderValuesStr.ParseIntoArray(HeaderValues, HeaderValueDelimiters, true);
+			HeaderValuesStr.ParseIntoArray(HeaderValues, TEXT(","), true);
 
 			if (HeaderValues.Num() > 0)
 			{
