Merging from //UE4/Portal-Staging/Engine/Source/Runtime/Online/BuildPatchServices/...

CL# 8925802
----------
Fix for calling GLog->FlushThreadedLogs() when we are not on the main thread during directory chunking

CL# 9630632
----------

CL# 9828174
----------
Fixing an issue where a generated optimised delta between two binaries will break an installation if the destination binary is re-processed as different data, replacing the original manifest at same version string.
Rename EFeatureLevel::StoresUniqueBuildId to EFeatureLevel::UsesRuntimeGeneratedBuildId.
Add new manifest EFeatureLevel, storing a build-time generate id - EFeatureLevel::UsesBuildTimeGeneratedBuildId.
When saving a manifest of EFeatureLevel::UsesBuildTimeGeneratedBuildId or higher, a buildid field is serialised as part of the manifest meta data block.
This buildid is generated as part of the FManifestMeta constructor, and thus uniquely saved upon creation of a new manifest object.
When loading a manifest object, the buildid will be serialised or runtime generated based on manifest object version.
When downloading an optimised delta, we verify for agreeable SHA1 and file list, rejecting with error if a problem is detected.
ChunkDeltaOptimise should chose appropriate output feature level based on source and destination manifests.
MergeManifests - Ensure to preserve the original newly generated buildid rather than clobber it with one of the input manifets.

CL# 9894645
----------
BuildPatchServices: Extending the life of the IOptimisedDelta instances to match the FBuildPatchInstaller so that we can avoid race conditions for destruction that currently exists when FBuildPatchInstaller::Initialize returns.

CL# 10075821
----------
BuildPatchServices: Exposing the file verification error counter from verify statistics.

CL# 10114262
----------
Remove unnecessary explicit delete of resume data when completing a EInstallMode::StageFiles installer.
Add meta folder to staging so that the Install staging folder will only ever contain install files.
Move $resumeData file to new location from legacy location if found.

CL# 10892753
----------
BuildPatchServices: Refactor of FBuildPatchFileConstructor error handling to catch serialisation errors as soon as they occur.

CL# 10950441
----------
Fix for possible shutdown crash when packaging chunks.
When processor classes complete quickly, the IBuildManifestSet instance could be dealocated before other systems that are using it.

diff --git a/Engine/Source/Runtime/Online/BuildPatchServices/Private/BuildPatchFeatureLevel.cpp b/Engine/Source/Runtime/Online/BuildPatchServices/Private/BuildPatchFeatureLevel.cpp
--- a/Engine/Source/Runtime/Online/BuildPatchServices/Private/BuildPatchFeatureLevel.cpp
+++ b/Engine/Source/Runtime/Online/BuildPatchServices/Private/BuildPatchFeatureLevel.cpp
@@ -2,7 +2,7 @@
 
 #include "BuildPatchFeatureLevel.h"
 
-static_assert((int32)BuildPatchServices::EFeatureLevel::Latest == 17, "Please add support for the extra values to the two functions below.");
+static_assert((int32)BuildPatchServices::EFeatureLevel::Latest == 18, "Please add support for the extra values to the two functions below.");
 
 BUILDPATCHSERVICES_API const TCHAR* BuildPatchServices::FeatureLevelToString(const BuildPatchServices::EFeatureLevel& FeatureLevel)
 {
@@ -23,7 +23,8 @@ BUILDPATCHSERVICES_API const TCHAR* BuildPatchServices::FeatureLevelToString(con
 		CASE_ENUM_TO_STR(StoresPrerequisiteIds);
 		CASE_ENUM_TO_STR(StoredAsBinaryData);
 		CASE_ENUM_TO_STR(VariableSizeChunks);
-		CASE_ENUM_TO_STR(StoresUniqueBuildId);
+		CASE_ENUM_TO_STR(UsesRuntimeGeneratedBuildId);
+		CASE_ENUM_TO_STR(UsesBuildTimeGeneratedBuildId);
 		default: return TEXT("Invalid");
 	}
 #undef CASE_ENUM_TO_STR
@@ -50,11 +51,13 @@ BUILDPATCHSERVICES_API bool BuildPatchServices::FeatureLevelFromString(const TCH
 	RETURN_IF_EQUAL(StoresPrerequisiteIds);
 	RETURN_IF_EQUAL(StoredAsBinaryData);
 	RETURN_IF_EQUAL(VariableSizeChunks);
-	RETURN_IF_EQUAL(StoresUniqueBuildId);
+	RETURN_IF_EQUAL(UsesRuntimeGeneratedBuildId);
+	RETURN_IF_EQUAL(UsesBuildTimeGeneratedBuildId);
 	RETURN_IF_EQUAL(Latest);
 	RETURN_IF_EQUAL(LatestNoChunks);
 	RETURN_IF_EQUAL(LatestJson);
 	RETURN_IF_EQUAL(FirstOptimisedDelta);
+	RETURN_IF_EQUAL(StoresUniqueBuildId);
 	return false;
 #undef RETURN_IF_EQUAL
 }
diff --git a/Engine/Source/Runtime/Online/BuildPatchServices/Private/BuildPatchFileConstructor.cpp b/Engine/Source/Runtime/Online/BuildPatchServices/Private/BuildPatchFileConstructor.cpp
--- a/Engine/Source/Runtime/Online/BuildPatchServices/Private/BuildPatchFileConstructor.cpp
+++ b/Engine/Source/Runtime/Online/BuildPatchServices/Private/BuildPatchFileConstructor.cpp
@@ -38,28 +38,18 @@ namespace FileConstructorHelpers
 		}
 	}
 
-	bool CheckAndReportRemainingDiskSpaceError(IInstallerError* InstallerError, const FString& InstallDirectory, uint64 RemainingBytesRequired, const TCHAR* SpaceErrorCode, uint64* OutAvailableDiskSpace = nullptr)
+	bool CheckRemainingDiskSpace(const FString& InstallDirectory, uint64 RemainingBytesRequired, uint64& OutAvailableDiskSpace)
 	{
 		bool bContinueConstruction = true;
 		uint64 TotalSize = 0;
-		uint64 AvailableSpace = 0;
-		if (FPlatformMisc::GetDiskTotalAndFreeSpace(InstallDirectory, TotalSize, AvailableSpace))
+		OutAvailableDiskSpace = 0;
+		if (FPlatformMisc::GetDiskTotalAndFreeSpace(InstallDirectory, TotalSize, OutAvailableDiskSpace))
 		{
-			if (AvailableSpace < RemainingBytesRequired)
+			if (OutAvailableDiskSpace < RemainingBytesRequired)
 			{
-				UE_LOG(LogBuildPatchServices, Error, TEXT("Out of HDD space. Needs %llu bytes, Free %llu bytes"), RemainingBytesRequired, AvailableSpace);
-				InstallerError->SetError(
-					EBuildPatchInstallError::OutOfDiskSpace,
-					SpaceErrorCode,
-					0,
-					BuildPatchServices::GetDiskSpaceMessage(InstallDirectory, RemainingBytesRequired, AvailableSpace));
 				bContinueConstruction = false;
 			}
 		}
-		if (OutAvailableDiskSpace != nullptr)
-		{
-			*OutAvailableDiskSpace = AvailableSpace;
-		}
 		return bContinueConstruction;
 	}
 
@@ -110,6 +100,17 @@ namespace FileConstructorHelpers
 	}
 }
 
+enum class EConstructionError : uint8
+{
+	None = 0,
+	CannotCreateFile,
+	OutOfDiskSpace,
+	MissingChunk,
+	SerializeError,
+	TrackingError,
+	OutboundDataError
+};
+
 /**
  * This struct handles loading and saving of simple resume information, that will allow us to decide which
  * files should be resumed from. It will also check that we are creating the same version and app as we expect to be.
@@ -127,7 +128,7 @@ public:
 	const FString StagingDir;
 
 	// The filename to the resume data information
-	const FString ResumeDataFile;
+	const FString ResumeDataFilename;
 
 	// The resume ids that we loaded from disk
 	TSet<FString> LoadedResumeIds;
@@ -146,15 +147,15 @@ public:
 
 public:
 
-	FResumeData(IFileSystem* InFileSystem, IBuildManifestSet* InManifestSet, const FString& InStagingDir)
+	FResumeData(IFileSystem* InFileSystem, IBuildManifestSet* InManifestSet, const FString& InStagingDir, const FString& InResumeDataFilename)
 		: FileSystem(InFileSystem)
 		, ManifestSet(InManifestSet)
 		, StagingDir(InStagingDir)
-		, ResumeDataFile(InStagingDir / TEXT("$resumeData"))
+		, ResumeDataFilename(InResumeDataFilename)
 		, bHasResumeData(false)
 	{
 		// Load data from previous resume file
-		bHasResumeData = FileSystem->FileExists(*ResumeDataFile);
+		bHasResumeData = FileSystem->FileExists(*ResumeDataFilename);
 		GLog->Logf(TEXT("BuildPatchResumeData file found: %s"), bHasResumeData ? TEXT("true") : TEXT("false"));
 		if (bHasResumeData)
 		{
@@ -162,7 +163,7 @@ public:
 			const bool bCullEmptyLines = true;
 			FString PrevResumeData;
 			TArray<FString> PrevResumeDataLines;
-			FileSystem->LoadFileToString(*ResumeDataFile, PrevResumeData);
+			FileSystem->LoadFileToString(*ResumeDataFilename, PrevResumeData);
 			PrevResumeData.ParseIntoArrayLines(PrevResumeDataLines, bCullEmptyLines);
 			// Grab current resume ids
 			const bool bCheckLegacyIds = true;
@@ -189,7 +190,7 @@ public:
 	void SaveOut(const TSet<FString>& ResumeIds)
 	{
 		// Save out the patch versions
-		FileSystem->SaveStringToFile(*ResumeDataFile, FString::Join(ResumeIds, TEXT("\n")));
+		FileSystem->SaveStringToFile(*ResumeDataFilename, FString::Join(ResumeIds, TEXT("\n")));
 	}
 
 	/**
@@ -290,34 +291,22 @@ FBuildPatchFileConstructor::~FBuildPatchFileConstructor()
 	}
 }
 
-bool FBuildPatchFileConstructor::Init()
-{
-	// We are ready to go if our delegates are bound and directories successfully created
-	bool bStageDirExists = IFileManager::Get().DirectoryExists(*Configuration.StagingDirectory);
-	if (!bStageDirExists)
-	{
-		UE_LOG(LogBuildPatchServices, Error, TEXT("FBuildPatchFileConstructor: Stage directory missing %s"), *Configuration.StagingDirectory);
-		InstallerError->SetError(EBuildPatchInstallError::InitializationError, InitializationErrorCodes::MissingStageDirectory);
-	}
-	SetInitFailed(!bStageDirExists);
-	return bStageDirExists;
-}
-
 uint32 FBuildPatchFileConstructor::Run()
 {
-	// TODO: Remove this check, or assert on construction
-	if (Configuration.ManifestSet == nullptr)
-	{
-		UE_LOG(LogBuildPatchServices, Error, TEXT("FBuildPatchFileConstructor: ManifestSet is null."));
-		return 0;
-	}
-
 	SetRunning(true);
 	SetInited(true);
 	FileConstructorStat->OnTotalRequiredUpdated(TotalJobSize);
 
-	// Check for resume data
-	FResumeData ResumeData(FileSystem, Configuration.ManifestSet, Configuration.StagingDirectory);
+	// Check for resume data, we need to also look for a legacy resume file to use instead in case we are resuming from an install of previous code version.
+	const FString LegacyResumeDataFilename = Configuration.StagingDirectory / TEXT("$resumeData");
+	const FString ResumeDataFilename = Configuration.MetaDirectory / TEXT("$resumeData");
+	const bool bHasLegacyResumeData = FileSystem->FileExists(*LegacyResumeDataFilename);
+	// If we find a legacy resume data file, lets move it first.
+	if (bHasLegacyResumeData)
+	{
+		FileSystem->MoveFile(*ResumeDataFilename, *LegacyResumeDataFilename);
+	}
+	FResumeData ResumeData(FileSystem, Configuration.ManifestSet, Configuration.StagingDirectory, ResumeDataFilename);
 
 	// Remove incompatible files
 	if (ResumeData.bHasResumeData)
@@ -374,7 +363,7 @@ uint32 FBuildPatchFileConstructor::Run()
 			}
 			else
 			{
-				bFileSuccess = ConstructFileFromChunks(FileToConstruct, bFilePreviouslyStarted);
+				bFileSuccess = ConstructFileFromChunks(*FileManifest, bFilePreviouslyStarted);
 			}
 		}
 		else
@@ -541,214 +530,252 @@ uint64 FBuildPatchFileConstructor::CalculateRequiredDiskSpace(const FFileManifes
 	return FMath::Max<int64>(DiskSpaceDeltaPeak, 0);
 }
 
-bool FBuildPatchFileConstructor::ConstructFileFromChunks( const FString& Filename, bool bResumeExisting )
+bool FBuildPatchFileConstructor::ConstructFileFromChunks(const FFileManifest& FileManifest, bool bResumeExisting)
 {
 	bool bSuccess = true;
-	FString NewFilename = Configuration.StagingDirectory / Filename;
+	EConstructionError ConstructionError = EConstructionError::None;
+	uint32 LastError = 0;
+	FString NewFilename = Configuration.StagingDirectory / FileManifest.Filename;
 
 	// Calculate the hash as we write the data
 	FSHA1 HashState;
 	FSHAHash HashValue;
 
-	// First make sure we can get the file manifest
-	const FFileManifest* FileManifest = Configuration.ManifestSet->GetNewFileManifest(Filename);
-	bSuccess = FileManifest != nullptr;
-	if( bSuccess )
+	if (!FileManifest.SymlinkTarget.IsEmpty())
 	{
-		if( !FileManifest->SymlinkTarget.IsEmpty() )
-		{
 #if PLATFORM_MAC
-			bSuccess = symlink(TCHAR_TO_UTF8(*FileManifest->SymlinkTarget), TCHAR_TO_UTF8(*NewFilename)) == 0;
+		bSuccess = symlink(TCHAR_TO_UTF8(*FileManifest.SymlinkTarget), TCHAR_TO_UTF8(*NewFilename)) == 0;
 #else
-			const bool bSymlinkNotImplemented = false;
-			check(bSymlinkNotImplemented);
-			bSuccess = false;
+		const bool bSymlinkNotImplemented = false;
+		check(bSymlinkNotImplemented);
+		bSuccess = false;
 #endif
-			return bSuccess;
-		}
+		return bSuccess;
+	}
 
-		// Check for resuming of existing file
-		int64 StartPosition = 0;
-		int32 StartChunkPart = 0;
-		if (bResumeExisting)
+	// Check for resuming of existing file
+	int64 StartPosition = 0;
+	int32 StartChunkPart = 0;
+	if (bResumeExisting)
+	{
+		// We have to read in the existing file so that the hash check can still be done.
+		TUniquePtr<FArchive> NewFileReader(IFileManager::Get().CreateFileReader(*NewFilename));
+		if (NewFileReader.IsValid())
 		{
-			// We have to read in the existing file so that the hash check can still be done.
-			TUniquePtr<FArchive> NewFileReader(IFileManager::Get().CreateFileReader(*NewFilename));
-			if (NewFileReader.IsValid())
+			// Start with a sensible buffer size for reading. 4 MiB.
+			const int32 ReadBufferSize = 4 * 1024 * 1024;
+			// Read buffer
+			TArray<uint8> ReadBuffer;
+			ReadBuffer.Empty(ReadBufferSize);
+			ReadBuffer.SetNumUninitialized(ReadBufferSize);
+			// Reuse a certain amount of the file
+			StartPosition = FMath::Max<int64>(0, NewFileReader->TotalSize() - NUM_BYTES_RESUME_IGNORE);
+			// We'll also find the correct chunkpart to start writing from
+			int64 ByteCounter = 0;
+			for (int32 ChunkPartIdx = StartChunkPart; ChunkPartIdx < FileManifest.ChunkParts.Num() && !bShouldAbort; ++ChunkPartIdx)
 			{
-				// Start with a sensible buffer size for reading. 4 MiB.
-				const int32 ReadBufferSize = 4 * 1024 * 1024;
-				// Read buffer
-				TArray<uint8> ReadBuffer;
-				ReadBuffer.Empty(ReadBufferSize);
-				ReadBuffer.SetNumUninitialized(ReadBufferSize);
-				// Reuse a certain amount of the file
-				StartPosition = FMath::Max<int64>(0, NewFileReader->TotalSize() - NUM_BYTES_RESUME_IGNORE);
-				// We'll also find the correct chunkpart to start writing from
-				int64 ByteCounter = 0;
-				for (int32 ChunkPartIdx = StartChunkPart; ChunkPartIdx < FileManifest->ChunkParts.Num() && !bShouldAbort; ++ChunkPartIdx)
+				const FChunkPart& ChunkPart = FileManifest.ChunkParts[ChunkPartIdx];
+				const int64 NextBytePosition = ByteCounter + ChunkPart.Size;
+				if (NextBytePosition <= StartPosition)
 				{
-					const FChunkPart& ChunkPart = FileManifest->ChunkParts[ChunkPartIdx];
-					const int64 NextBytePosition = ByteCounter + ChunkPart.Size;
-					if (NextBytePosition <= StartPosition)
-					{
-						// Ensure buffer is large enough
-						ReadBuffer.SetNumUninitialized(ChunkPart.Size, false);
-						ISpeedRecorder::FRecord ActivityRecord;
-						// Read data for hash check
-						FileConstructorStat->OnBeforeRead();
-						ActivityRecord.CyclesStart = FStatsCollector::GetCycles();
-						NewFileReader->Serialize(ReadBuffer.GetData(), ChunkPart.Size);
-						ActivityRecord.CyclesEnd = FStatsCollector::GetCycles();
-						ActivityRecord.Size = ChunkPart.Size;
-						HashState.Update(ReadBuffer.GetData(), ChunkPart.Size);
-						FileConstructorStat->OnAfterRead(ActivityRecord);
-						// Count bytes read from file
-						ByteCounter = NextBytePosition;
-						// Set to resume from next chunk part
-						StartChunkPart = ChunkPartIdx + 1;
-						// Inform the reference tracker of the chunk part skip
-						bSuccess = ChunkReferenceTracker->PopReference(ChunkPart.Guid) && bSuccess;
-						CountBytesProcessed(ChunkPart.Size);
-						FileConstructorStat->OnFileProgress(Filename, NewFileReader->Tell());
-						// Wait if paused
-						FileConstructorHelpers::WaitWhilePaused(bIsPaused, bShouldAbort);
-					}
-					else
-					{
-						// No more parts on disk
-						break;
-					}
+					// Ensure buffer is large enough
+					ReadBuffer.SetNumUninitialized(ChunkPart.Size, false);
+					ISpeedRecorder::FRecord ActivityRecord;
+					// Read data for hash check
+					FileConstructorStat->OnBeforeRead();
+					ActivityRecord.CyclesStart = FStatsCollector::GetCycles();
+					NewFileReader->Serialize(ReadBuffer.GetData(), ChunkPart.Size);
+					ActivityRecord.CyclesEnd = FStatsCollector::GetCycles();
+					ActivityRecord.Size = ChunkPart.Size;
+					HashState.Update(ReadBuffer.GetData(), ChunkPart.Size);
+					FileConstructorStat->OnAfterRead(ActivityRecord);
+					// Count bytes read from file
+					ByteCounter = NextBytePosition;
+					// Set to resume from next chunk part
+					StartChunkPart = ChunkPartIdx + 1;
+					// Inform the reference tracker of the chunk part skip
+					bSuccess = ChunkReferenceTracker->PopReference(ChunkPart.Guid) && bSuccess;
+					CountBytesProcessed(ChunkPart.Size);
+					FileConstructorStat->OnFileProgress(FileManifest.Filename, NewFileReader->Tell());
+					// Wait if paused
+					FileConstructorHelpers::WaitWhilePaused(bIsPaused, bShouldAbort);
+				}
+				else
+				{
+					// No more parts on disk
+					break;
 				}
-				// Set start position to the byte we got up to
-				StartPosition = ByteCounter;
-				// Close file
-				NewFileReader->Close();
 			}
+			// Set start position to the byte we got up to
+			StartPosition = ByteCounter;
+			// Close file
+			NewFileReader->Close();
 		}
+	}
 
-		// If we haven't done so yet, make the initial disk space check
-		if (!bInitialDiskSizeCheck)
+	// If we haven't done so yet, make the initial disk space check
+	if (!bInitialDiskSizeCheck)
+	{
+		bInitialDiskSizeCheck = true;
+		const uint64 RequiredSpace = CalculateRequiredDiskSpace(FileManifest, FileManifest.FileSize - StartPosition);
+		// ThreadLock protects access to members RequiredDiskSpace and AvailableDiskSpace;
+		FScopeLock Lock(&ThreadLock);
+		RequiredDiskSpace = RequiredSpace;
+		if (!FileConstructorHelpers::CheckRemainingDiskSpace(Configuration.InstallDirectory, RequiredSpace, AvailableDiskSpace))
 		{
-			bInitialDiskSizeCheck = true;
-			const uint64 RequiredSpace = CalculateRequiredDiskSpace(*FileManifest, FileManifest->FileSize - StartPosition);
-			// ThreadLock protects access to member RequiredDiskSpace and AvailableDiskSpace;
-			FScopeLock Lock(&ThreadLock);
-			RequiredDiskSpace = RequiredSpace;
-			if (!FileConstructorHelpers::CheckAndReportRemainingDiskSpaceError(InstallerError, Configuration.InstallDirectory, RequiredSpace, DiskSpaceErrorCodes::InitialSpaceCheck, &AvailableDiskSpace))
-			{
-				return false;
-			}
+			UE_LOG(LogBuildPatchServices, Error, TEXT("Out of HDD space. Needs %llu bytes, Free %llu bytes"), RequiredSpace, AvailableDiskSpace);
+			InstallerError->SetError(
+				EBuildPatchInstallError::OutOfDiskSpace,
+				DiskSpaceErrorCodes::InitialSpaceCheck,
+				0,
+				BuildPatchServices::GetDiskSpaceMessage(Configuration.InstallDirectory, RequiredSpace, AvailableDiskSpace));
+			return false;
 		}
+	}
 
-		// Now we can make sure the chunk cache knows to start downloading chunks
-		if (!bIsDownloadStarted)
+	// Now we can make sure the chunk cache knows to start downloading chunks
+	if (!bIsDownloadStarted)
+	{
+		bIsDownloadStarted = true;
+		FileConstructorStat->OnResumeCompleted();
+	}
+
+	// Attempt to create the file
+	ISpeedRecorder::FRecord ActivityRecord;
+	FileConstructorStat->OnBeforeAdminister();
+	ActivityRecord.CyclesStart = FStatsCollector::GetCycles();
+	TUniquePtr<FArchive> NewFile = FileSystem->CreateFileWriter(*NewFilename, bResumeExisting ? EWriteFlags::Append : EWriteFlags::None);
+	LastError = FPlatformMisc::GetLastError();
+	ActivityRecord.CyclesEnd = FStatsCollector::GetCycles();
+	ActivityRecord.Size = 0;
+	FileConstructorStat->OnAfterAdminister(ActivityRecord);
+	bSuccess = NewFile != nullptr;
+	if (bSuccess)
+	{
+		// Seek to file write position
+		if (NewFile->Tell() != StartPosition)
 		{
-			bIsDownloadStarted = true;
-			FileConstructorStat->OnResumeCompleted();
+			FileConstructorStat->OnBeforeAdminister();
+			ActivityRecord.CyclesStart = FStatsCollector::GetCycles();
+			NewFile->Seek(StartPosition);
+			ActivityRecord.CyclesEnd = FStatsCollector::GetCycles();
+			ActivityRecord.Size = 0;
+			FileConstructorStat->OnAfterAdminister(ActivityRecord);
 		}
 
-		// Attempt to create the file
-		ISpeedRecorder::FRecord ActivityRecord;
-		FileConstructorStat->OnBeforeAdminister();
-		ActivityRecord.CyclesStart = FStatsCollector::GetCycles();
-		TUniquePtr<FArchive> NewFile = FileSystem->CreateFileWriter( *NewFilename, bResumeExisting ? EWriteFlags::Append : EWriteFlags::None );
-		uint32 LastError = FPlatformMisc::GetLastError();
-		ActivityRecord.CyclesEnd = FStatsCollector::GetCycles();
-		ActivityRecord.Size = 0;
-		FileConstructorStat->OnAfterAdminister(ActivityRecord);
-		bSuccess = NewFile != nullptr;
-		if (bSuccess)
+		// For each chunk, load it, and place it's data into the file
+		for (int32 ChunkPartIdx = StartChunkPart; ChunkPartIdx < FileManifest.ChunkParts.Num() && bSuccess && !bShouldAbort; ++ChunkPartIdx)
 		{
-			// Seek to file write position
-			if (NewFile->Tell() != StartPosition)
+			const FChunkPart& ChunkPart = FileManifest.ChunkParts[ChunkPartIdx];
+			bSuccess = InsertChunkData(ChunkPart, *NewFile, HashState, ConstructionError);
+			FileConstructorStat->OnFileProgress(FileManifest.Filename, NewFile->Tell());
+			if (bSuccess)
 			{
-				FileConstructorStat->OnBeforeAdminister();
-				ActivityRecord.CyclesStart = FStatsCollector::GetCycles();
-				NewFile->Seek(StartPosition);
-				ActivityRecord.CyclesEnd = FStatsCollector::GetCycles();
-				ActivityRecord.Size = 0;
-				FileConstructorStat->OnAfterAdminister(ActivityRecord);
+				CountBytesProcessed(ChunkPart.Size);
+				// Wait while paused
+				FileConstructorHelpers::WaitWhilePaused(bIsPaused, bShouldAbort);
 			}
-
-			// For each chunk, load it, and place it's data into the file
-			for (int32 ChunkPartIdx = StartChunkPart; ChunkPartIdx < FileManifest->ChunkParts.Num() && bSuccess && !bShouldAbort; ++ChunkPartIdx)
+			// Only report or log if this is the first error
+			else if (InstallerError->HasError() == false)
 			{
-				const FChunkPart& ChunkPart = FileManifest->ChunkParts[ChunkPartIdx];
-				bSuccess = InsertChunkData(ChunkPart, *NewFile, HashState);
-				FileConstructorStat->OnFileProgress(Filename, NewFile->Tell());
-				if (bSuccess)
+				if (ConstructionError == EConstructionError::MissingChunk)
 				{
-					CountBytesProcessed(ChunkPart.Size);
-					// Wait while paused
-					FileConstructorHelpers::WaitWhilePaused(bIsPaused, bShouldAbort);
+					InstallerAnalytics->RecordConstructionError(FileManifest.Filename, INDEX_NONE, TEXT("Missing Chunk"));
+					UE_LOG(LogBuildPatchServices, Error, TEXT("FBuildPatchFileConstructor: Failed %s due to missing chunk %s"), *FileManifest.Filename, *ChunkPart.Guid.ToString());
 				}
-				else
+				else if (ConstructionError == EConstructionError::TrackingError)
 				{
-					// Only report or log if the first error
-					if (InstallerError->HasError() == false)
-					{
-						InstallerAnalytics->RecordConstructionError(Filename, INDEX_NONE, TEXT("Missing Chunk"));
-						UE_LOG(LogBuildPatchServices, Error, TEXT("FBuildPatchFileConstructor: Failed %s due to chunk %s"), *Filename, *ChunkPart.Guid.ToString());
-					}
-					InstallerError->SetError(EBuildPatchInstallError::FileConstructionFail, ConstructionErrorCodes::MissingChunkData);
+					InstallerAnalytics->RecordConstructionError(FileManifest.Filename, INDEX_NONE, TEXT("Tracking Error"));
+					UE_LOG(LogBuildPatchServices, Error, TEXT("FBuildPatchFileConstructor: Failed %s due to untracked chunk %s"), *FileManifest.Filename, *ChunkPart.Guid.ToString());
 				}
 			}
+		}
 
-			// Close the file writer
-			FileConstructorStat->OnBeforeAdminister();
-			ActivityRecord.CyclesStart = FStatsCollector::GetCycles();
-			NewFile->Close();
-			NewFile.Reset();
-			ActivityRecord.CyclesEnd = FStatsCollector::GetCycles();
-			ActivityRecord.Size = 0;
-			FileConstructorStat->OnAfterAdminister(ActivityRecord);
+		// Close the file writer
+		FileConstructorStat->OnBeforeAdminister();
+		ActivityRecord.CyclesStart = FStatsCollector::GetCycles();
+		const bool bArchiveSuccess = NewFile->Close();
+		NewFile.Reset();
+		ActivityRecord.CyclesEnd = FStatsCollector::GetCycles();
+		ActivityRecord.Size = 0;
+		FileConstructorStat->OnAfterAdminister(ActivityRecord);
+
+		// Check for final success
+		if (ConstructionError == EConstructionError::None && !bArchiveSuccess)
+		{
+			ConstructionError = EConstructionError::SerializeError;
+			bSuccess = false;
+		}
+	}
+	else
+	{
+		ConstructionError = EConstructionError::CannotCreateFile;
+	}
+
+	// Check for error state
+	if (!bSuccess)
+	{
+		// Recalculate disk space first
+		int64 InProgressFileSize = FileManifest.FileSize;
+		FileSystem->GetFileSize(*NewFilename, InProgressFileSize);
+		const uint64 RemainingRequiredSpace = CalculateRequiredDiskSpace(FileManifest, FileManifest.FileSize - InProgressFileSize);
+		uint64 RemainingAvailableDiskSpace = 0;
+		if (!FileConstructorHelpers::CheckRemainingDiskSpace(Configuration.InstallDirectory, RemainingRequiredSpace, RemainingAvailableDiskSpace))
+		{
+			// ThreadLock protects access to members RequiredDiskSpace and AvailableDiskSpace
+			ThreadLock.Lock();
+			RequiredDiskSpace = RemainingRequiredSpace;
+			AvailableDiskSpace = RemainingAvailableDiskSpace;
+			ThreadLock.Unlock();
+			// Convert error to disk space rather than reported.
+			UE_LOG(LogBuildPatchServices, Error, TEXT("Out of HDD space. Needs %llu bytes, Free %llu bytes"), RemainingRequiredSpace, RemainingAvailableDiskSpace);
+			InstallerError->SetError(
+				EBuildPatchInstallError::OutOfDiskSpace,
+				DiskSpaceErrorCodes::DuringInstallation,
+				0,
+				BuildPatchServices::GetDiskSpaceMessage(Configuration.InstallDirectory, RemainingRequiredSpace, RemainingAvailableDiskSpace));
+			ConstructionError = EConstructionError::OutOfDiskSpace;
 		}
 		else
 		{
-			// Check if drive space was the issue here
-			const uint64 RequiredSpace = CalculateRequiredDiskSpace(*FileManifest, FileManifest->FileSize);
-			bool bError = !FileConstructorHelpers::CheckAndReportRemainingDiskSpaceError(InstallerError, Configuration.InstallDirectory, RequiredSpace, DiskSpaceErrorCodes::DuringInstallation);
-
-			// Otherwise we just couldn't make the file
-			if (!bError)
+			const bool bReportAnalytic = InstallerError->HasError() == false;
+			switch (ConstructionError)
 			{
-				// Only report or log if the first error
-				if (InstallerError->HasError() == false)
+			case EConstructionError::CannotCreateFile:
+				if (bReportAnalytic)
 				{
-					InstallerAnalytics->RecordConstructionError(Filename, LastError, TEXT("Could Not Create File"));
-					UE_LOG(LogBuildPatchServices, Error, TEXT("FBuildPatchFileConstructor: Could not create %s"), *Filename);
+					InstallerAnalytics->RecordConstructionError(FileManifest.Filename, LastError, TEXT("Could Not Create File"));
+					UE_LOG(LogBuildPatchServices, Error, TEXT("FBuildPatchFileConstructor: Could not create %s"), *FileManifest.Filename);
 				}
-				// Always set
 				InstallerError->SetError(EBuildPatchInstallError::FileConstructionFail, ConstructionErrorCodes::FileCreateFail, LastError);
+				break;
+			case EConstructionError::MissingChunk:
+				InstallerError->SetError(EBuildPatchInstallError::FileConstructionFail, ConstructionErrorCodes::MissingChunkData);
+				break;
+			case EConstructionError::SerializeError:
+				InstallerError->SetError(EBuildPatchInstallError::FileConstructionFail, ConstructionErrorCodes::SerializationError);
+				break;
+			case EConstructionError::TrackingError:
+				InstallerError->SetError(EBuildPatchInstallError::FileConstructionFail, ConstructionErrorCodes::TrackingError);
+				break;
 			}
 		}
 	}
-	else
-	{
-		// Only report or log if the first error
-		if (InstallerError->HasError() == false)
-		{
-			InstallerAnalytics->RecordConstructionError(Filename, INDEX_NONE, TEXT("Missing File Manifest"));
-			UE_LOG(LogBuildPatchServices, Error, TEXT("FBuildPatchFileConstructor: Missing file manifest for %s"), *Filename);
-		}
-		// Always set
-		InstallerError->SetError(EBuildPatchInstallError::FileConstructionFail, ConstructionErrorCodes::MissingFileInfo);
-	}
 
 	// Verify the hash for the file that we created
-	if( bSuccess )
+	if (bSuccess)
 	{
 		HashState.Final();
-		HashState.GetHash( HashValue.Hash );
-		bSuccess = HashValue == FileManifest->FileHash;
-		if( !bSuccess )
+		HashState.GetHash(HashValue.Hash);
+		bSuccess = HashValue == FileManifest.FileHash;
+		if (!bSuccess)
 		{
+			ConstructionError = EConstructionError::OutboundDataError;
 			// Only report or log if the first error
 			if (InstallerError->HasError() == false)
 			{
-				InstallerAnalytics->RecordConstructionError(Filename, INDEX_NONE, TEXT("Serialised Verify Fail"));
-				UE_LOG(LogBuildPatchServices, Error, TEXT("FBuildPatchFileConstructor: Verify failed after constructing %s"), *Filename);
+				InstallerAnalytics->RecordConstructionError(FileManifest.Filename, INDEX_NONE, TEXT("Serialised Verify Fail"));
+				UE_LOG(LogBuildPatchServices, Error, TEXT("FBuildPatchFileConstructor: Verify failed after constructing %s"), *FileManifest.Filename);
 			}
 			// Always set
 			InstallerError->SetError(EBuildPatchInstallError::FileConstructionFail, ConstructionErrorCodes::OutboundCorrupt);
@@ -756,7 +783,7 @@ bool FBuildPatchFileConstructor::ConstructFileFromChunks( const FString& Filenam
 	}
 
 #if PLATFORM_MAC
-	if (bSuccess && EnumHasAllFlags(FileManifest->FileMetaFlags, EFileMetaFlags::UnixExecutable))
+	if (bSuccess && EnumHasAllFlags(FileManifest.FileMetaFlags, EFileMetaFlags::UnixExecutable))
 	{
 		// Enable executable permission bit
 		struct stat FileInfo;
@@ -766,7 +793,7 @@ bool FBuildPatchFileConstructor::ConstructFileFromChunks( const FString& Filenam
 		}
 	}
 #endif
-	
+
 #if PLATFORM_ANDROID || PLATFORM_ANDROIDESDEFERRED
 	if (bSuccess)
 	{
@@ -774,27 +801,35 @@ bool FBuildPatchFileConstructor::ConstructFileFromChunks( const FString& Filenam
 	}
 #endif
 
-	// Delete the staging file if unsuccessful by means of construction fail (i.e. keep if canceled or download issue)
-	if( !bSuccess && InstallerError->GetErrorType() == EBuildPatchInstallError::FileConstructionFail )
+	// Delete the staging file if unsuccessful by means of any failure that could leave the file in unknown state.
+	if (!bSuccess)
 	{
-		if (!FileSystem->DeleteFile(*NewFilename))
+		switch (ConstructionError)
 		{
-			UE_LOG(LogBuildPatchServices, Warning, TEXT("FBuildPatchFileConstructor: Error deleting file: %s (Error Code %i)"), *NewFilename, FPlatformMisc::GetLastError());
+		case EConstructionError::CannotCreateFile:
+		case EConstructionError::SerializeError:
+		case EConstructionError::TrackingError:
+		case EConstructionError::OutboundDataError:
+			if (!FileSystem->DeleteFile(*NewFilename))
+			{
+				UE_LOG(LogBuildPatchServices, Warning, TEXT("FBuildPatchFileConstructor: Error deleting file: %s (Error Code %i)"), *NewFilename, FPlatformMisc::GetLastError());
+			}
+			break;
 		}
 	}
 
 	return bSuccess;
 }
 
-bool FBuildPatchFileConstructor::InsertChunkData(const FChunkPart& ChunkPart, FArchive& DestinationFile, FSHA1& HashState)
+bool FBuildPatchFileConstructor::InsertChunkData(const FChunkPart& ChunkPart, FArchive& DestinationFile, FSHA1& HashState, EConstructionError& ConstructionError)
 {
 	uint8* Data;
 	uint8* DataStart;
+	ConstructionError = EConstructionError::None;
 	ISpeedRecorder::FRecord ActivityRecord;
 	FileConstructorStat->OnChunkGet(ChunkPart.Guid);
 	IChunkDataAccess* ChunkDataAccess = ChunkSource->Get(ChunkPart.Guid);
-	bool bSuccess = ChunkDataAccess != nullptr && !bShouldAbort;
-	if (bSuccess)
+	if (ChunkDataAccess != nullptr)
 	{
 		ChunkDataAccess->GetDataLock(&Data, nullptr);
 		FileConstructorStat->OnBeforeWrite();
@@ -802,13 +837,26 @@ bool FBuildPatchFileConstructor::InsertChunkData(const FChunkPart& ChunkPart, FA
 		DataStart = &Data[ChunkPart.Offset];
 		HashState.Update(DataStart, ChunkPart.Size);
 		DestinationFile.Serialize(DataStart, ChunkPart.Size);
+		const bool bSerializeOk = !DestinationFile.IsError();
 		ActivityRecord.Size = ChunkPart.Size;
 		ActivityRecord.CyclesEnd = FStatsCollector::GetCycles();
 		FileConstructorStat->OnAfterWrite(ActivityRecord);
 		ChunkDataAccess->ReleaseDataLock();
-		bSuccess = ChunkReferenceTracker->PopReference(ChunkPart.Guid);
+		const bool bPopReferenceOk = ChunkReferenceTracker->PopReference(ChunkPart.Guid);
+		if (!bSerializeOk)
+		{
+			ConstructionError = EConstructionError::SerializeError;
+		}
+		else if (!bPopReferenceOk)
+		{
+			ConstructionError = EConstructionError::TrackingError;
+		}
 	}
-	return bSuccess;
+	else
+	{
+		ConstructionError = EConstructionError::MissingChunk;
+	}
+	return ConstructionError == EConstructionError::None;
 }
 
 void FBuildPatchFileConstructor::DeleteDirectoryContents(const FString& RootDirectory)
diff --git a/Engine/Source/Runtime/Online/BuildPatchServices/Private/BuildPatchFileConstructor.h b/Engine/Source/Runtime/Online/BuildPatchServices/Private/BuildPatchFileConstructor.h
--- a/Engine/Source/Runtime/Online/BuildPatchServices/Private/BuildPatchFileConstructor.h
+++ b/Engine/Source/Runtime/Online/BuildPatchServices/Private/BuildPatchFileConstructor.h
@@ -13,6 +13,7 @@
 
 // Forward declarations
 class FBuildPatchAppManifest;
+enum class EConstructionError : uint8;
 
 namespace BuildPatchServices
 {
@@ -30,14 +31,18 @@ namespace BuildPatchServices
 	 */
 	struct FFileConstructorConfig
 	{
+		// The manifest set class for details on the installation files.
 		IBuildManifestSet* ManifestSet;
 
-		// The root location where the installation is going.
+		// The location for the installation.
 		FString InstallDirectory;
 
-		// The location where we will store temporary files.
+		// The location where new installation files will be constructed.
 		FString StagingDirectory;
 
+		// The location where temporary files for tracking can be stored.
+		FString MetaDirectory;
+
 		// The list of files to be constructed, filename paths should match those contained in manifest.
 		TArray<FString> ConstructList;
 
@@ -73,7 +78,6 @@ namespace BuildPatchServices
 		~FBuildPatchFileConstructor();
 
 		// FRunnable interface begin.
-		virtual bool Init() override;
 		virtual uint32 Run() override;
 		// FRunnable interface end.
 
@@ -168,20 +172,21 @@ namespace BuildPatchServices
 		/**
 		 * Constructs a particular file referenced by the given BuildManifest. The function takes an interface to a class that can provide availability information of chunks so that this
 		 * file construction process can be ran alongside chunk acquisition threads. It will Sleep while waiting for chunks that it needs.
-		 * @param Filename			The Filename for the file to construct, that matches an entry in the BuildManifest.
+		 * @param FileManifest		The FFileManifest for the file to construct.
 		 * @param bResumeExisting	Whether we should resume from an existing file
 		 * @return	true if no file errors occurred
 		 */
-		bool ConstructFileFromChunks(const FString& Filename, bool bResumeExisting);
+		bool ConstructFileFromChunks(const FFileManifest& FileManifest, bool bResumeExisting);
 
 		/**
 		 * Inserts the data data from a chunk into the destination file according to the chunk part info
-		 * @param ChunkPart			The chunk part details.
-		 * @param DestinationFile	The Filename for the file being constructed.
-		 * @param HashState			An FSHA1 hash state to update with the data going into the destination file.
+		 * @param ChunkPart          The chunk part details.
+		 * @param DestinationFile    The Filename for the file being constructed.
+		 * @param HashState          An FSHA1 hash state to update with the data going into the destination file.
+		 * @param ConstructionError  Will be set to the error type that ocurred or EConstructionError::None.
 		 * @return true if no errors were detected
 		 */
-		bool InsertChunkData(const FChunkPart& ChunkPart, FArchive& DestinationFile, FSHA1& HashState);
+		bool InsertChunkData(const FChunkPart& ChunkPart, FArchive& DestinationFile, FSHA1& HashState, EConstructionError& ConstructionError);
 
 		/**
 		 * Delete all contents of a directory
diff --git a/Engine/Source/Runtime/Online/BuildPatchServices/Private/BuildPatchGeneration.cpp b/Engine/Source/Runtime/Online/BuildPatchServices/Private/BuildPatchGeneration.cpp
--- a/Engine/Source/Runtime/Online/BuildPatchServices/Private/BuildPatchGeneration.cpp
+++ b/Engine/Source/Runtime/Online/BuildPatchServices/Private/BuildPatchGeneration.cpp
@@ -190,7 +190,10 @@ bool FBuildDataGenerator::ChunkBuildDirectory(const BuildPatchServices::FChunkBu
 	while (CloudEnumeration->IsComplete() == false)
 	{
 		// Log collected stats.
-		GLog->FlushThreadedLogs();
+		if (bIsGameThread)
+		{
+			GLog->FlushThreadedLogs();
+		}
 		FStatsCollector::Set(StatTotalTime, FStatsCollector::GetCycles() - StartTime);
 		StatsCollector->LogStats(StatsLoggerTimeSeconds);
 
@@ -251,7 +254,7 @@ bool FBuildDataGenerator::ChunkBuildDirectory(const BuildPatchServices::FChunkBu
 	uint32 ReadLen = 0;
 	TArray<TUniquePtr<FScannerDetails>> Scanners;
 	bool bHasUnknownData = true;
-	while (!BuildStream->IsEndOfData() || Scanners.Num() > 0 || bHasUnknownData)
+	while (!BuildStream->HasAborted() && (!BuildStream->IsEndOfData() || Scanners.Num() > 0 || bHasUnknownData))
 	{
 		// Grab a scanner result.
 		if (Scanners.Num() > 0 && Scanners[0]->Scanner->IsComplete())
@@ -694,6 +697,12 @@ bool FBuildDataGenerator::ChunkBuildDirectory(const BuildPatchServices::FChunkBu
 		FPlatformProcess::Sleep(0.01f);
 	}
 
+	if (BuildStream->HasAborted())
+	{
+		UE_LOG(LogPatchGeneration, Error, TEXT("Directory Build Streamer aborted"));
+		return false;
+	}
+
 	// Complete chunk writer.
 	FParallelChunkWriterSummaries ChunkWriterSummaries = ChunkWriter->OnProcessComplete();
 
diff --git a/Engine/Source/Runtime/Online/BuildPatchServices/Private/BuildPatchInstaller.cpp b/Engine/Source/Runtime/Online/BuildPatchServices/Private/BuildPatchInstaller.cpp
--- a/Engine/Source/Runtime/Online/BuildPatchServices/Private/BuildPatchInstaller.cpp
+++ b/Engine/Source/Runtime/Online/BuildPatchServices/Private/BuildPatchInstaller.cpp
@@ -406,6 +406,7 @@ namespace BuildPatchServices
 		, Configuration(MoveTemp(InConfiguration))
 		, DataStagingDir(Configuration.StagingDirectory / TEXT("PatchData"))
 		, InstallStagingDir(Configuration.StagingDirectory / TEXT("Install"))
+		, MetaStagingDir(Configuration.StagingDirectory / TEXT("Meta"))
 		, PreviousMoveMarker(Configuration.InstallDirectory / TEXT("$movedMarker"))
 		, ThreadLock()
 		, bSuccess(false)
@@ -653,33 +654,34 @@ namespace BuildPatchServices
 
 		// do the delta optimization
 		FOptimisedDeltaDependencies OptimisedDeltaDependencies = InstallerHelpers::BuildOptimisedDeltaDependencies(DownloadService);
-		typedef TTuple<FBuildPatchInstallerAction&, TUniquePtr<IOptimisedDelta>> FManifestInfoDeltaPair;
+		typedef TTuple<FBuildPatchInstallerAction&, IOptimisedDelta*> FManifestInfoDeltaPair;
 		TArray<FManifestInfoDeltaPair> RunningOptimisedDeltas;
 		for (FBuildPatchInstallerAction& InstallerAction : InstallerActions)
 		{
 			if (InstallerAction.IsUpdate())
 			{
-				RunningOptimisedDeltas.Add(FManifestInfoDeltaPair{ InstallerAction, FOptimisedDeltaFactory::Create(ConfigHelpers::BuildOptimisedDeltaConfig(Configuration, InstallerAction), OptimisedDeltaDependencies) });
+				IOptimisedDelta* OptimisedDelta = FOptimisedDeltaFactory::Create(ConfigHelpers::BuildOptimisedDeltaConfig(Configuration, InstallerAction), OptimisedDeltaDependencies);
+				OptimisedDeltas.Add(TUniquePtr<IOptimisedDelta>(OptimisedDelta));
+				RunningOptimisedDeltas.Add(FManifestInfoDeltaPair{ InstallerAction, OptimisedDelta });
 			}
 		}
 		for (FManifestInfoDeltaPair& RunningOptimisedDelta : RunningOptimisedDeltas)
 		{
 			FBuildPatchInstallerAction& InstallerAction = RunningOptimisedDelta.Get<0>();
-			const TUniquePtr<IOptimisedDelta>& OptimisedDelta = RunningOptimisedDelta.Get<1>();
-
-			FBuildPatchAppManifestPtr DestinationManifest(OptimisedDelta->GetDestinationManifest());
+			IOptimisedDelta* const OptimisedDelta = RunningOptimisedDelta.Get<1>();
+			const IOptimisedDelta::FResultValueOrError& OptimisedDeltaResult = OptimisedDelta->GetResult();
 			PreviousTotalDownloadRequired.Add(OptimisedDelta->GetMetaDownloadSize());
 			// The OptimiseDelta class handles policy, so if we get a nullptr back, that is a hard error.
-			if (!DestinationManifest.IsValid())
+			if (!OptimisedDeltaResult.IsValid())
 			{
 				UE_LOG(LogBuildPatchServices, Error, TEXT("Installer setup: Destination manifest could not be obtained."));
-				InstallerError->SetError(EBuildPatchInstallError::DownloadError, DownloadErrorCodes::MissingDestinationManifest);
+				InstallerError->SetError(EBuildPatchInstallError::DownloadError, *OptimisedDeltaResult.GetError());
 				bInstallerInitSuccess = false;
 				break;
 			}
 			else
 			{
-				InstallerAction.SetDeltaManifest(DestinationManifest.ToSharedRef());
+				InstallerAction.SetDeltaManifest(OptimisedDeltaResult.GetValue().ToSharedRef());
 			}
 		}
 
@@ -786,7 +788,6 @@ namespace BuildPatchServices
 					{
 						UE_LOG(LogBuildPatchServices, Log, TEXT("Deleting litter from staging area."));
 						IFileManager::Get().DeleteDirectory(*DataStagingDir, false, true);
-						IFileManager::Get().Delete(*(InstallStagingDir / TEXT("$resumeData")), false, true);
 					}
 					else
 					{
@@ -1081,10 +1082,12 @@ namespace BuildPatchServices
 		// Save the staging directories
 		FPaths::NormalizeDirectoryName(DataStagingDir);
 		FPaths::NormalizeDirectoryName(InstallStagingDir);
+		FPaths::NormalizeDirectoryName(MetaStagingDir);
 
 		// Make sure staging directories exist
 		IFileManager::Get().MakeDirectory(*DataStagingDir, true);
 		IFileManager::Get().MakeDirectory(*InstallStagingDir, true);
+		IFileManager::Get().MakeDirectory(*MetaStagingDir, true);
 
 		// Reset our error and build progress
 		InstallerError.Reset(FInstallerErrorFactory::Create());
@@ -1274,6 +1277,7 @@ namespace BuildPatchServices
 					ManifestSet.Get(),
 					Configuration.InstallDirectory,
 					InstallStagingDir,
+					MetaStagingDir,
 					FilesToConstruct.Array(),
 					Configuration.InstallMode}),
 				FileSystem.Get(),
diff --git a/Engine/Source/Runtime/Online/BuildPatchServices/Private/BuildPatchInstaller.h b/Engine/Source/Runtime/Online/BuildPatchServices/Private/BuildPatchInstaller.h
--- a/Engine/Source/Runtime/Online/BuildPatchServices/Private/BuildPatchInstaller.h
+++ b/Engine/Source/Runtime/Online/BuildPatchServices/Private/BuildPatchInstaller.h
@@ -85,6 +85,9 @@ namespace BuildPatchServices
 		// The directory created in staging to construct install files to.
 		FString InstallStagingDir;
 
+		// The directory created in staging to store any required meta, such as resume info.
+		FString MetaStagingDir;
+
 		// The filename used to mark a previous install that did not complete but moved staged files into the install directory.
 		FString PreviousMoveMarker;
 
@@ -179,6 +182,9 @@ namespace BuildPatchServices
 		// The message pump controller.
 		TUniquePtr<IMessagePump> MessagePump;
 
+		// Holds the optimised delta interfaces created during Initialize.
+		TArray<TUniquePtr<IOptimisedDelta>> OptimisedDeltas;
+
 		// The interface for manifest data aggregation.
 		TUniquePtr<IBuildManifestSet> ManifestSet;
 
diff --git a/Engine/Source/Runtime/Online/BuildPatchServices/Private/BuildPatchMergeManifests.cpp b/Engine/Source/Runtime/Online/BuildPatchServices/Private/BuildPatchMergeManifests.cpp
--- a/Engine/Source/Runtime/Online/BuildPatchServices/Private/BuildPatchMergeManifests.cpp
+++ b/Engine/Source/Runtime/Online/BuildPatchServices/Private/BuildPatchMergeManifests.cpp
@@ -190,9 +190,13 @@ bool FBuildMergeManifests::MergeManifests(const FString& ManifestFilePathA, cons
 	// Create the new manifest
 	FBuildPatchAppManifest MergedManifest;
 
-	// Copy basic info from B
-	MergedManifest.ManifestMeta = ManifestB->ManifestMeta;
-	MergedManifest.CustomFields = ManifestB->CustomFields;
+	// Copy basic info from B, preserving the generated build ID
+	{
+		FString NewBuildId = MoveTemp(MergedManifest.ManifestMeta.BuildId);
+		MergedManifest.ManifestMeta = ManifestB->ManifestMeta;
+		MergedManifest.CustomFields = ManifestB->CustomFields;
+		MergedManifest.ManifestMeta.BuildId = MoveTemp(NewBuildId);
+	}
 
 	// Set the new version string
 	MergedManifest.ManifestMeta.BuildVersion = NewVersionString;
diff --git a/Engine/Source/Runtime/Online/BuildPatchServices/Private/BuildPatchUtil.h b/Engine/Source/Runtime/Online/BuildPatchServices/Private/BuildPatchUtil.h
--- a/Engine/Source/Runtime/Online/BuildPatchServices/Private/BuildPatchUtil.h
+++ b/Engine/Source/Runtime/Online/BuildPatchServices/Private/BuildPatchUtil.h
@@ -119,13 +119,13 @@ struct FBuildPatchUtils
 	static bool GetGUIDFromFilename(const FString& DataFilename, FGuid& DataGUID);
 
 	/**
-	 * Generates a new BuildId for a manifest. This should be used only when created new builds.
+	 * Generates a new BuildId for a manifest. This should be used only when creating new builds, and thus saving out brand new manifests rather than copies of manifests.
 	 * @return the generated id.
 	 */
 	static FString GenerateNewBuildId();
 
 	/**
-	 * Creates a deterministic BuildId for a manifest that is older than EFeatureLevel::StoresUniqueBuildId.
+	 * Creates a deterministic BuildId for use with a manifest that is at EFeatureLevel::UsesRuntimeGeneratedBuildId or older.
 	 * The id is created based on the meta data, which itself should be unique per build created.
 	 * @param ManifestMeta     The meta for the old manifest.
 	 * @return the id for this manifest.
diff --git a/Engine/Source/Runtime/Online/BuildPatchServices/Private/Data/ChunkData.cpp b/Engine/Source/Runtime/Online/BuildPatchServices/Private/Data/ChunkData.cpp
--- a/Engine/Source/Runtime/Online/BuildPatchServices/Private/Data/ChunkData.cpp
+++ b/Engine/Source/Runtime/Online/BuildPatchServices/Private/Data/ChunkData.cpp
@@ -133,14 +133,15 @@ namespace BuildPatchServices
 				case BuildPatchServices::EFeatureLevel::StoredAsBinaryData:
 				case BuildPatchServices::EFeatureLevel::VariableSizeChunksWithoutWindowSizeChunkInfo:
 				case BuildPatchServices::EFeatureLevel::VariableSizeChunks:
-				case BuildPatchServices::EFeatureLevel::StoresUniqueBuildId:
+				case BuildPatchServices::EFeatureLevel::UsesRuntimeGeneratedBuildId:
+				case BuildPatchServices::EFeatureLevel::UsesBuildTimeGeneratedBuildId:
 					return EChunkVersion::StoresDataSizeUncompressed;
 			}
 			checkf(false, TEXT("Unhandled FeatureLevel %s"), FeatureLevelToString(FeatureLevel));
 			return EChunkVersion::Invalid;
 		}
 	}
-	static_assert((uint32)EFeatureLevel::Latest == 17, "Please adjust HeaderHelpers::FeatureLevelToChunkVersion for new feature levels.");
+	static_assert((uint32)EFeatureLevel::Latest == 18, "Please adjust HeaderHelpers::FeatureLevelToChunkVersion for new feature levels.");
 
 	FChunkHeader::FChunkHeader()
 		: Version((uint32)EChunkVersion::Latest)
diff --git a/Engine/Source/Runtime/Online/BuildPatchServices/Private/Data/ManifestData.cpp b/Engine/Source/Runtime/Online/BuildPatchServices/Private/Data/ManifestData.cpp
--- a/Engine/Source/Runtime/Online/BuildPatchServices/Private/Data/ManifestData.cpp
+++ b/Engine/Source/Runtime/Online/BuildPatchServices/Private/Data/ManifestData.cpp
@@ -20,6 +20,58 @@
 
 namespace BuildPatchServices
 {
+
+	// The constant minimum sizes for each version of a header struct. Must be updated.
+	// If new member variables are added the version MUST be bumped and handled properly here,
+	// and these values must never change.
+	static const uint32 ManifestHeaderVersionSizes[(int32)EFeatureLevel::LatestPlusOne] =
+	{
+		// EFeatureLevel::Original is 37B (32b Magic, 32b HeaderSize, 32b DataSizeUncompressed, 32b DataSizeCompressed, 160b SHA1, 8b StoredAs)
+		// This remained the same all up to including EFeatureLevel::StoresPrerequisiteIds.
+		37, 37, 37, 37, 37, 37, 37, 37, 37, 37, 37, 37, 37, 37,
+		// EFeatureLevel::StoredAsBinaryData is 41B, (296b Original, 32b Version).
+		// This remained the same all up to including EFeatureLevel::UsesBuildTimeGeneratedBuildId.
+		41, 41, 41, 41, 41
+	};
+	static_assert((int32)EFeatureLevel::Latest == 18, "Please adjust ManifestHeaderVersionSizes values accordingly.");
+
+	/**
+	 * Enum which describes the FManifestMeta data version.
+	 */
+	enum class EManifestMetaVersion : uint8
+	{
+		Original = 0,
+		SerialisesBuildId,
+
+		// Always after the latest version, signifies the latest version plus 1 to allow initialization simplicity.
+		LatestPlusOne,
+		Latest = (LatestPlusOne - 1)
+	};
+
+	/**
+	 * Enum which describes the FChunkDataList data version.
+	 */
+	enum class EChunkDataListVersion : uint8
+	{
+		Original = 0,
+
+		// Always after the latest version, signifies the latest version plus 1 to allow initialization simplicity.
+		LatestPlusOne,
+		Latest = (LatestPlusOne - 1)
+	};
+
+	/**
+	 * Enum which describes the FFileManifestList data version.
+	 */
+	enum class EFileManifestListVersion : uint8
+	{
+		Original = 0,
+
+		// Always after the latest version, signifies the latest version plus 1 to allow initialization simplicity.
+		LatestPlusOne,
+		Latest = (LatestPlusOne - 1)
+	};
+
 	namespace ManifestVersionHelpers
 	{
 		const TCHAR* GetChunkSubdir(EFeatureLevel FeatureLevel)
@@ -36,7 +88,38 @@
 				: FeatureLevel < EFeatureLevel::StoresChunkDataShaHashes ? TEXT("FilesV2")
 				: TEXT("FilesV3");
 		}
+
+		EManifestMetaVersion FeatureLevelToManifestMetaVersion(EFeatureLevel FeatureLevel)
+		{
+			switch (FeatureLevel)
+			{
+			case BuildPatchServices::EFeatureLevel::Original:
+			case BuildPatchServices::EFeatureLevel::CustomFields:
+			case BuildPatchServices::EFeatureLevel::StartStoringVersion:
+			case BuildPatchServices::EFeatureLevel::DataFileRenames:
+			case BuildPatchServices::EFeatureLevel::StoresIfChunkOrFileData:
+			case BuildPatchServices::EFeatureLevel::StoresDataGroupNumbers:
+			case BuildPatchServices::EFeatureLevel::ChunkCompressionSupport:
+			case BuildPatchServices::EFeatureLevel::StoresPrerequisitesInfo:
+			case BuildPatchServices::EFeatureLevel::StoresChunkFileSizes:
+			case BuildPatchServices::EFeatureLevel::StoredAsCompressedUClass:
+			case BuildPatchServices::EFeatureLevel::UNUSED_0:
+			case BuildPatchServices::EFeatureLevel::UNUSED_1:
+			case BuildPatchServices::EFeatureLevel::StoresChunkDataShaHashes:
+			case BuildPatchServices::EFeatureLevel::StoresPrerequisiteIds:
+			case BuildPatchServices::EFeatureLevel::StoredAsBinaryData:
+			case BuildPatchServices::EFeatureLevel::VariableSizeChunksWithoutWindowSizeChunkInfo:
+			case BuildPatchServices::EFeatureLevel::VariableSizeChunks:
+			case BuildPatchServices::EFeatureLevel::UsesRuntimeGeneratedBuildId:
+				return EManifestMetaVersion::Original;
+			case BuildPatchServices::EFeatureLevel::UsesBuildTimeGeneratedBuildId:
+				return EManifestMetaVersion::SerialisesBuildId;
+			}
+			checkf(false, TEXT("Unhandled FeatureLevel %s"), FeatureLevelToString(FeatureLevel));
+			return EManifestMetaVersion::Latest;
+		}
 	}
+	static_assert((uint32)EFeatureLevel::Latest == 18, "Please adjust ManifestVersionHelpers::FeatureLevelToManifestMetaVersion for new feature levels.");
 
 	namespace ManifestDataHelpers
 	{
@@ -62,20 +145,6 @@
 	/* FManifestHeader - The header for a compressed/encoded manifest file.
 	*****************************************************************************/
 
-	// The constant minimum sizes for each version of a header struct. Must be updated.
-	// If new member variables are added the version MUST be bumped and handled properly here,
-	// and these values must never change.
-	static const uint32 ManifestHeaderVersionSizes[(int32)EFeatureLevel::LatestPlusOne] =
-	{
-		// EFeatureLevel::Original is 37B (32b Magic, 32b HeaderSize, 32b DataSizeUncompressed, 32b DataSizeCompressed, 160b SHA1, 8b StoredAs)
-		// This remained the same all up to including EFeatureLevel::StoresPrerequisiteIds.
-		37, 37, 37, 37, 37, 37, 37, 37, 37, 37, 37, 37, 37, 37,
-		// EFeatureLevel::StoredAsBinaryData is 41B, (296b Original, 32b Version).
-		// This remained the same all up to including EFeatureLevel::StoresUniqueBuildId.
-		41, 41, 41, 41
-	};
-	static_assert((int32)EFeatureLevel::Latest == 17, "Please adjust ManifestHeaderVersionSizes values accordingly.");
-
 	FManifestHeader::FManifestHeader()
 		: Version(EFeatureLevel::Latest)
 		, HeaderSize(0)
@@ -155,21 +224,6 @@
 	/* FManifestMeta - The data implementation for a build meta data.
 	*****************************************************************************/
 
-	/**
-	 * Enum which describes the FManifestMeta data version.
-	 */
-	enum class ManifestMetaVersion : uint8
-	{
-		Original = 0,
-		/* Due to some specific EpicGamesLauncher functionality, we're going to not start saving the build ID until a client is released that can save it properly.
-		   It does not cause a serialisation issue, it just means we can't use optimised deltas immediately unless we forgo this field until later.
-		StoresBuildId,*/
-
-		// Always after the latest version, signifies the latest version plus 1 to allow initialization simplicity.
-		LatestPlusOne,
-		Latest = (LatestPlusOne - 1)
-	};
-
 	FManifestMeta::FManifestMeta()
 		: FeatureLevel(BuildPatchServices::EFeatureLevel::Invalid)
 		, bIsFileData(false)
@@ -188,16 +242,16 @@
 		// Serialise the data header type values.
 		const int64 StartPos = Ar.Tell();
 		uint32 DataSize = 0;
-		ManifestMetaVersion DataVersion = ManifestMetaVersion::Latest;
+		EManifestMetaVersion DataVersion = Ar.IsSaving() ? ManifestVersionHelpers::FeatureLevelToManifestMetaVersion(Meta.FeatureLevel) : EManifestMetaVersion::Latest;
 		{
 			uint8 DataVersionInt = (uint8)DataVersion;
 			Ar << DataSize;
 			Ar << DataVersionInt;
-			DataVersion = (ManifestMetaVersion)DataVersionInt;
+			DataVersion = (EManifestMetaVersion)DataVersionInt;
 		}
 
 		// Serialise the ManifestMetaVersion::Original version variables.
-		if (!Ar.IsError() && DataVersion >= ManifestMetaVersion::Original)
+		if (!Ar.IsError() && DataVersion >= EManifestMetaVersion::Original)
 		{
 			int32 FeatureLevelInt = (int32)Meta.FeatureLevel;
 			uint8 IsFileDataInt = Meta.bIsFileData ? 1 : 0;
@@ -216,22 +270,13 @@
 			Meta.bIsFileData = IsFileDataInt == 1;
 		}
 
-		/* Due to some specific EpicGamesLauncher functionality, we're going to not start saving the build ID until a client is released that can save it properly.
-		 * It does not cause a serialisation issue, it just means we can't use optimised deltas immediately unless we forgo this field until later.
 		// Serialise the BuildId.
-		if (!Ar.IsError())
+		if (!Ar.IsError() && DataVersion >= EManifestMetaVersion::SerialisesBuildId)
 		{
-			if (DataVersion >= ManifestMetaVersion::StoresBuildId)
-			{
-				Ar << Meta.BuildId;
-			}
-			// Otherwise, initialise with backwards compat default when loading
-			else if (Ar.IsLoading())
-			{
-				Meta.BuildId = FBuildPatchUtils::GetBackwardsCompatibleBuildId(Meta);
-			}
-		}*/
-		if (!Ar.IsError() && Ar.IsLoading())
+			Ar << Meta.BuildId;
+		}
+		// Otherwise, initialise with backwards compatible default when loading.
+		else if (!Ar.IsError() && Ar.IsLoading())
 		{
 			Meta.BuildId = FBuildPatchUtils::GetBackwardsCompatibleBuildId(Meta);
 		}
@@ -260,18 +305,6 @@
 	/* FChunkDataList - The data implementation for a list of referenced chunk data.
 	*****************************************************************************/
 
-	/**
-	 * Enum which describes the FChunkDataList data version.
-	 */
-	enum class EChunkDataListVersion : uint8
-	{
-		Original = 0,
-
-		// Always after the latest version, signifies the latest version plus 1 to allow initialization simplicity.
-		LatestPlusOne,
-		Latest = (LatestPlusOne - 1)
-	};
-
 	FChunkDataList::FChunkDataList()
 	{
 	}
@@ -347,18 +380,6 @@
 	/* FFileManifestList - The data implementation for a list of referenced files.
 	*****************************************************************************/
 
-	/**
-	 * Enum which describes the FFileManifestList data version.
-	 */
-	enum class EFileManifestListVersion : uint8
-	{
-		Original = 0,
-
-		// Always after the latest version, signifies the latest version plus 1 to allow initialization simplicity.
-		LatestPlusOne,
-		Latest = (LatestPlusOne - 1)
-	};
-
 	FFileManifestList::FFileManifestList()
 	{
 	}
diff --git a/Engine/Source/Runtime/Online/BuildPatchServices/Private/Diffing/DiffManifests.cpp b/Engine/Source/Runtime/Online/BuildPatchServices/Private/Diffing/DiffManifests.cpp
--- a/Engine/Source/Runtime/Online/BuildPatchServices/Private/Diffing/DiffManifests.cpp
+++ b/Engine/Source/Runtime/Online/BuildPatchServices/Private/Diffing/DiffManifests.cpp
@@ -418,7 +418,7 @@ namespace BuildPatchServices
 			FOptimisedDeltaDependencies OptimisedDeltaDependencies;
 			OptimisedDeltaDependencies.DownloadService = DownloadService.Get();
 			TUniquePtr<IOptimisedDelta> OptimisedDelta(FOptimisedDeltaFactory::Create(OptimisedDeltaConfiguration, OptimisedDeltaDependencies));
-			ManifestB = OptimisedDelta->GetDestinationManifest();
+			ManifestB = OptimisedDelta->GetResult().GetValue();
 			const int32 MetaDownloadBytes = OptimisedDelta->GetMetaDownloadSize();
 
 			TSet<FString> TagsA, TagsB;
diff --git a/Engine/Source/Runtime/Online/BuildPatchServices/Private/Generation/BuildStreamer.h b/Engine/Source/Runtime/Online/BuildPatchServices/Private/Generation/BuildStreamer.h
--- a/Engine/Source/Runtime/Online/BuildPatchServices/Private/Generation/BuildStreamer.h
+++ b/Engine/Source/Runtime/Online/BuildPatchServices/Private/Generation/BuildStreamer.h
@@ -66,7 +66,7 @@ namespace BuildPatchServices
 		 * MUST be called only after IsEndOfData returns true.
 		 * @return the number of bytes in the streamed build.
 		 */
-		virtual uint64 GetBuildSize() const = 0;
+		virtual uint64 GetBuildSize() const = 0; 
 
 		/**
 		 * Get the list of file spans for each file in the build, including empty files.
@@ -74,6 +74,12 @@ namespace BuildPatchServices
 		 * @return the list of files in the build and their details.
 		 */
 		virtual TArray<FFileSpan> GetAllFiles() const = 0;
+
+		/**
+		 * Gets if the streamer has aborted and is no longer streaming data
+		 * @return True if it has aborted
+		 */
+		virtual bool HasAborted() const = 0;
 	};
 
 	class IManifestBuildStreamer
diff --git a/Engine/Source/Runtime/Online/BuildPatchServices/Private/Generation/ChunkDeltaOptimiser.cpp b/Engine/Source/Runtime/Online/BuildPatchServices/Private/Generation/ChunkDeltaOptimiser.cpp
--- a/Engine/Source/Runtime/Online/BuildPatchServices/Private/Generation/ChunkDeltaOptimiser.cpp
+++ b/Engine/Source/Runtime/Online/BuildPatchServices/Private/Generation/ChunkDeltaOptimiser.cpp
@@ -1196,9 +1196,10 @@ namespace BuildPatchServices
 					}
 				}
 				DeltaManifest.InitLookups();
-				// Currently we just save out as the first version for delta support. If we change the delta version later we'd take this as commandline selection.
+				// For save format, we'll pick the newest of two input manifests, or EFeatureLevel::FirstOptimisedDelta if manifests are older.
+				EFeatureLevel DeltaOutputFormat = FMath::Max3(EFeatureLevel::FirstOptimisedDelta, ManifestA->GetFeatureLevel(), ManifestB->GetFeatureLevel());
 				const FString TmpOutputDeltaFilename = OutputDeltaFilename + TEXT("tmp");
-				DeltaManifest.SaveToFile(TmpOutputDeltaFilename, EFeatureLevel::FirstOptimisedDelta);
+				DeltaManifest.SaveToFile(TmpOutputDeltaFilename, DeltaOutputFormat);
 				FileSystem->MoveFile(*OutputDeltaFilename, *TmpOutputDeltaFilename);
 				FinalStatLogs.Add(FString::Printf(TEXT("Saved new optimised delta file %s"), *OutputDeltaFilename));
 			}
diff --git a/Engine/Source/Runtime/Online/BuildPatchServices/Private/Generation/DirectoryBuildStreamer.cpp b/Engine/Source/Runtime/Online/BuildPatchServices/Private/Generation/DirectoryBuildStreamer.cpp
--- a/Engine/Source/Runtime/Online/BuildPatchServices/Private/Generation/DirectoryBuildStreamer.cpp
+++ b/Engine/Source/Runtime/Online/BuildPatchServices/Private/Generation/DirectoryBuildStreamer.cpp
@@ -84,6 +84,7 @@ namespace BuildPatchServices
 		virtual TArray<FString> GetAllFilenames() const override;
 		virtual uint64 GetBuildSize() const override;
 		virtual TArray<FFileSpan> GetAllFiles() const override;
+		virtual bool HasAborted() const override;
 		// IDirectoryBuildStreamer interface end.
 
 	private:
@@ -274,6 +275,11 @@ namespace BuildPatchServices
 		return AllFiles;
 	}
 
+	bool FDirectoryBuildStreamer::HasAborted() const
+	{
+		return bShouldAbort;
+	}
+
 	void FDirectoryBuildStreamer::ReadData()
 	{
 		// Stats
@@ -331,7 +337,13 @@ namespace BuildPatchServices
 			FString SymlinkTarget = GetSymlinkTarget(*SourceFile);
 			FStatsCollector::AccumulateTimeEnd(StatFileOpenTime, TempValue);
 			// Not being able to load a required file from the build would be fatal, hard fault.
-			checkf(FileReader.IsValid(), TEXT("Could not open file from build! %s"), *SourceFile);
+			if(!FileReader.IsValid())
+			{
+				UE_LOG(LogBuildStreamer, Error, TEXT("Failed to open file reader for %s"), *SourceFile);
+				bShouldAbort = true;
+				break;
+			}
+
 			// Make SourceFile the format we want it in and start a new file.
 			FPaths::MakePathRelativeTo(SourceFile, *(Config.BuildRoot + TEXT("/")));
 			int64 FileSize = FileReader->TotalSize();
diff --git a/Engine/Source/Runtime/Online/BuildPatchServices/Private/Generation/PackageChunkData.cpp b/Engine/Source/Runtime/Online/BuildPatchServices/Private/Generation/PackageChunkData.cpp
--- a/Engine/Source/Runtime/Online/BuildPatchServices/Private/Generation/PackageChunkData.cpp
+++ b/Engine/Source/Runtime/Online/BuildPatchServices/Private/Generation/PackageChunkData.cpp
@@ -115,15 +115,15 @@
 	private:
 		void HandleDownloadComplete(int32 RequestId, const FDownloadRef& Download);
 		void HandleManifestComplete();
-		void HandleManifestSelection(FBuildPatchAppManifestPtr DeltaManifest);
+		void HandleManifestSelection(const IOptimisedDelta::FResultValueOrError& ResultValueOrError);
 		void BeginPackageProcess();
 		void OnPackageComplete(bool bInSuccess);
 
 		typedef void(FPackageChunks::*PromiseCompleteFunc)();
 		TFunction<void()> MakePromiseCompleteDelegate(PromiseCompleteFunc OnComplete);
 
-		typedef void(FPackageChunks::*OptimiseCompleteFunc)(FBuildPatchAppManifestPtr);
-		TFunction<void(FBuildPatchAppManifestPtr)> MakeOptimiseCompleteDelegate(OptimiseCompleteFunc OnComplete);
+		typedef void(FPackageChunks::*OptimiseCompleteFunc)(const IOptimisedDelta::FResultValueOrError&);
+		TFunction<void(const IOptimisedDelta::FResultValueOrError&)> MakeOptimiseCompleteDelegate(OptimiseCompleteFunc OnComplete);
 
 	private:
 		// Configuration.
@@ -133,28 +133,6 @@
 		FTicker& CoreTicker;
 		FDownloadCompleteDelegate DownloadCompleteDelegate;
 		FDownloadProgressDelegate DownloadProgressDelegate;
-		TUniquePtr<IPlatform> Platform;
-		TUniquePtr<IHttpManager> HttpManager;
-		TUniquePtr<IFileSystem> FileSystem;
-		TUniquePtr<IMessagePump> MessagePump;
-		TUniquePtr<IInstallerError> InstallerError;
-		TUniquePtr<ISpeedRecorder> DownloadSpeedRecorder;
-		TUniquePtr<IChunkDataSizeProvider> ChunkDataSizeProvider;
-		TUniquePtr<IInstallerAnalytics> InstallerAnalytics;
-		TUniquePtr<IDownloadServiceStatistics> DownloadServiceStatistics;
-		TUniquePtr<IDownloadService> DownloadService;
-		TUniquePtr<IChunkReferenceTracker> ChunkReferenceTracker;
-		TUniquePtr<IFileOperationTracker> FileOperationTracker;
-		TUniquePtr<IOptimisedDelta> OptimisedDelta;
-		FBuildPatchProgress BuildProgress;
-		TUniquePtr<IMemoryChunkStoreStatistics> MemoryChunkStoreStatistics;
-		TUniquePtr<ICloudChunkSourceStatistics> CloudChunkSourceStatistics;
-		TUniquePtr<IChunkDataSerialization> ChunkDataSerialization;
-		TUniquePtr<IChunkEvictionPolicy> MemoryEvictionPolicy;
-		TUniquePtr<IMemoryChunkStore> CloudChunkStore;
-		TUniquePtr<IDownloadConnectionCount> DownloadConnectionCount;
-		TUniquePtr<ICloudChunkSource> CloudChunkSource;
-		TUniquePtr<IChunkDatabaseWriter> ChunkDatabaseWriter;
 
 		// Process control.
 		TArray<FMessageHandler*> MessageHandlers;
@@ -177,6 +155,30 @@
 		// Packaging.
 		TArray<FChunkDatabaseFile> ChunkDbFiles;
 		TArray<TArray<int32>> TagSetLookupTable;
+
+		// Declare constructed systems last, to ensure they are destructed before our data members.
+		TUniquePtr<IPlatform> Platform;
+		TUniquePtr<IHttpManager> HttpManager;
+		TUniquePtr<IFileSystem> FileSystem;
+		TUniquePtr<IMessagePump> MessagePump;
+		TUniquePtr<IInstallerError> InstallerError;
+		TUniquePtr<ISpeedRecorder> DownloadSpeedRecorder;
+		TUniquePtr<IChunkDataSizeProvider> ChunkDataSizeProvider;
+		TUniquePtr<IInstallerAnalytics> InstallerAnalytics;
+		TUniquePtr<IDownloadServiceStatistics> DownloadServiceStatistics;
+		TUniquePtr<IDownloadService> DownloadService;
+		TUniquePtr<IChunkReferenceTracker> ChunkReferenceTracker;
+		TUniquePtr<IFileOperationTracker> FileOperationTracker;
+		TUniquePtr<IOptimisedDelta> OptimisedDelta;
+		FBuildPatchProgress BuildProgress;
+		TUniquePtr<IMemoryChunkStoreStatistics> MemoryChunkStoreStatistics;
+		TUniquePtr<ICloudChunkSourceStatistics> CloudChunkSourceStatistics;
+		TUniquePtr<IChunkDataSerialization> ChunkDataSerialization;
+		TUniquePtr<IChunkEvictionPolicy> MemoryEvictionPolicy;
+		TUniquePtr<IMemoryChunkStore> CloudChunkStore;
+		TUniquePtr<IDownloadConnectionCount> DownloadConnectionCount;
+		TUniquePtr<ICloudChunkSource> CloudChunkSource;
+		TUniquePtr<IChunkDatabaseWriter> ChunkDatabaseWriter;
 	};
 
 	FPackageChunks::FPackageChunks(const FPackageChunksConfiguration& InConfiguration)
@@ -184,6 +186,16 @@
 		, CoreTicker(FTicker::GetCoreTicker())
 		, DownloadCompleteDelegate(FDownloadCompleteDelegate::CreateRaw(this, &FPackageChunks::HandleDownloadComplete))
 		, DownloadProgressDelegate()
+		, bManifestsProcessed(false)
+		, bShouldRun(true)
+		, bSuccess(true)
+		, RequestIdManifestFile(INDEX_NONE)
+		, RequestIdPrevManifestFile(INDEX_NONE)
+		, PromiseManifestFile(MakePromiseCompleteDelegate(&FPackageChunks::HandleManifestComplete))
+		, PromisePrevManifestFile(MakePromiseCompleteDelegate(&FPackageChunks::HandleManifestComplete))
+		, FutureManifestFile(PromiseManifestFile.GetFuture())
+		, FuturePrevManifestFile(PromisePrevManifestFile.GetFuture())
+		, bUsingOptimisedDelta(false)
 		, Platform(FPlatformFactory::Create())
 		, HttpManager(FHttpManagerFactory::Create())
 		, FileSystem(FFileSystemFactory::Create())
@@ -195,17 +207,18 @@
 		, DownloadServiceStatistics(FDownloadServiceStatisticsFactory::Create(DownloadSpeedRecorder.Get(), ChunkDataSizeProvider.Get(), InstallerAnalytics.Get()))
 		, DownloadService(FDownloadServiceFactory::Create(CoreTicker, HttpManager.Get(), FileSystem.Get(), DownloadServiceStatistics.Get(), InstallerAnalytics.Get()))
 		, FileOperationTracker(FFileOperationTrackerFactory::Create(CoreTicker))
-		, bManifestsProcessed(false)
-		, bShouldRun(true)
-		, bSuccess(true)
-		, RequestIdManifestFile(INDEX_NONE)
-		, RequestIdPrevManifestFile(INDEX_NONE)
-		, PromiseManifestFile(MakePromiseCompleteDelegate(&FPackageChunks::HandleManifestComplete))
-		, PromisePrevManifestFile(MakePromiseCompleteDelegate(&FPackageChunks::HandleManifestComplete))
-		, FutureManifestFile(PromiseManifestFile.GetFuture())
-		, FuturePrevManifestFile(PromisePrevManifestFile.GetFuture())
-		, bUsingOptimisedDelta(false)
 	{
+		// Make sure the cloud chunk source gets the abort signal if an error occurred.
+		InstallerError->RegisterForErrors([this]()
+		{
+			AsyncHelpers::ExecuteOnGameThread<void>([this]()
+			{
+				if (CloudChunkSource.IsValid())
+				{
+					CloudChunkSource->Abort();
+				}
+			}).Wait();
+		});
 	}
 
 	FPackageChunks::~FPackageChunks()
@@ -320,8 +333,9 @@
 		}
 	}
 
-	void FPackageChunks::HandleManifestSelection(FBuildPatchAppManifestPtr DeltaManifest)
+	void FPackageChunks::HandleManifestSelection(const IOptimisedDelta::FResultValueOrError& ResultValueOrError)
 	{
+		FBuildPatchAppManifestPtr DeltaManifest = ResultValueOrError.IsValid() ? ResultValueOrError.GetValue() : nullptr;
 		bUsingOptimisedDelta = Manifest.Get() != DeltaManifest.Get();
 		if (DeltaManifest.IsValid())
 		{
@@ -599,11 +613,16 @@
 		return GameThreadWrapper;
 	}
 
-	TFunction<void(FBuildPatchAppManifestPtr)> FPackageChunks::MakeOptimiseCompleteDelegate(OptimiseCompleteFunc OnComplete)
+	TFunction<void(const IOptimisedDelta::FResultValueOrError&)> FPackageChunks::MakeOptimiseCompleteDelegate(OptimiseCompleteFunc OnComplete)
 	{
-		typedef TMemberFunctionCaller<FPackageChunks, OptimiseCompleteFunc> FOptimiseCompleteCaller;
-		TFunction<void(FBuildPatchAppManifestPtr)> OnCompleteDelegate = [this, OnComplete](FBuildPatchAppManifestPtr ManifestPtr) { FOptimiseCompleteCaller(this, OnComplete)(ManifestPtr); };
-		TFunction<void(FBuildPatchAppManifestPtr)> GameThreadWrapper = [OnCompleteDelegate](FBuildPatchAppManifestPtr ManifestPtr) { AsyncHelpers::ExecuteOnGameThread<void, FBuildPatchAppManifestPtr>(OnCompleteDelegate, ManifestPtr); };
+		TFunction<void(const IOptimisedDelta::FResultValueOrError&)> GameThreadWrapper = [this, OnComplete](const IOptimisedDelta::FResultValueOrError& Result)
+		{
+			// This is boiler plate overcoming IOptimisedDelta::FResultValueOrError being non-copyable, and interface between async thread and main game thread.
+			typedef TMemberFunctionCaller<FPackageChunks, OptimiseCompleteFunc> FOptimiseCompleteCaller;
+			TSharedPtr<IOptimisedDelta::FResultValueOrError> NewResult = Result.IsValid() ? MakeShared<IOptimisedDelta::FResultValueOrError>(MakeValue(Result.GetValue())) : MakeShared<IOptimisedDelta::FResultValueOrError>(MakeError(Result.GetError()));
+			TFunction<void()> OnCompleteDelegate = [this, OnComplete, NewResult = MoveTemp(NewResult)]() { FOptimiseCompleteCaller(this, OnComplete)(*NewResult); };
+			AsyncHelpers::ExecuteOnGameThread<void>(OnCompleteDelegate);
+		};
 		return GameThreadWrapper;
 	}
 
diff --git a/Engine/Source/Runtime/Online/BuildPatchServices/Private/Installer/BuildStatistics.cpp b/Engine/Source/Runtime/Online/BuildPatchServices/Private/Installer/BuildStatistics.cpp
--- a/Engine/Source/Runtime/Online/BuildPatchServices/Private/Installer/BuildStatistics.cpp
+++ b/Engine/Source/Runtime/Online/BuildPatchServices/Private/Installer/BuildStatistics.cpp
@@ -53,6 +53,7 @@ namespace BuildPatchServices
 		virtual double GetDiskReadByteSpeed() const override;
 		virtual double GetChunkDbReadByteSpeed() const override;
 		virtual double GetDiskWriteByteSpeed() const override;
+		virtual int32 GetNumVerifyErrors() const override;
 		virtual TMap<EVerifyError, int32> GetVerifyErrorCounts() const override;
 		// IBuildStatistics interface end.
 
@@ -167,6 +168,11 @@ namespace BuildPatchServices
 		return Installer->GetDiskWriteSpeedRecorder()->GetAverageSpeed(10.0f);
 	}
 
+	int32 FBuildStatistics::GetNumVerifyErrors() const
+	{
+		return Installer->GetVerifierStatistics()->GetNumFailedFilesVerified();
+	}
+
 	TMap<EVerifyError, int32> FBuildStatistics::GetVerifyErrorCounts() const
 	{
 		return Installer->GetVerifierStatistics()->GetVerifyErrorCounts();
diff --git a/Engine/Source/Runtime/Online/BuildPatchServices/Private/Installer/InstallerError.h b/Engine/Source/Runtime/Online/BuildPatchServices/Private/Installer/InstallerError.h
--- a/Engine/Source/Runtime/Online/BuildPatchServices/Private/Installer/InstallerError.h
+++ b/Engine/Source/Runtime/Online/BuildPatchServices/Private/Installer/InstallerError.h
@@ -44,7 +44,9 @@ namespace BuildPatchServices
 	namespace DownloadErrorCodes
 	{
 		static const TCHAR* OutOfChunkRetries = TEXT("01");
-		static const TCHAR* MissingDestinationManifest = TEXT("02");
+		static const TCHAR* MissingDeltaFile = TEXT("02");
+		static const TCHAR* UnserialisableDeltaFile = TEXT("03");
+		static const TCHAR* RejectedDeltaFile = TEXT("04");
 	}
 
 	/**
@@ -58,6 +60,7 @@ namespace BuildPatchServices
 		static const TCHAR* MissingFileInfo = TEXT("04");
 		static const TCHAR* OutboundCorrupt = TEXT("05");
 		static const TCHAR* SerializationError = TEXT("06");
+		static const TCHAR* TrackingError = TEXT("07");
 	}
 
 	/**
diff --git a/Engine/Source/Runtime/Online/BuildPatchServices/Private/Installer/OptimisedDelta.cpp b/Engine/Source/Runtime/Online/BuildPatchServices/Private/Installer/OptimisedDelta.cpp
--- a/Engine/Source/Runtime/Online/BuildPatchServices/Private/Installer/OptimisedDelta.cpp
+++ b/Engine/Source/Runtime/Online/BuildPatchServices/Private/Installer/OptimisedDelta.cpp
@@ -6,6 +6,7 @@
 #include "Misc/ConfigCacheIni.h"
 
 #include "Installer/DownloadService.h"
+#include "Installer/InstallerError.h"
 #include "BuildPatchUtil.h"
 #include "BuildPatchMergeManifests.h"
 
@@ -31,14 +32,15 @@ namespace BuildPatchServices
 		FOptimisedDelta(const FOptimisedDeltaConfiguration& Configuration, const FOptimisedDeltaDependencies& Dependencies);
 
 		// IOptimisedDelta interface begin.
-		virtual FBuildPatchAppManifestPtr GetDestinationManifest() override;
-		virtual int32 GetMetaDownloadSize() override;
+		virtual const IOptimisedDelta::FResultValueOrError& GetResult() const override;
+		virtual int32 GetMetaDownloadSize() const override;
 		// IOptimisedDelta interface end.
 
 	private:
 		bool RequiresChunkDownload();
 		void OnDownloadComplete(int32 RequestId, const FDownloadRef& Download);
 		bool ShouldRetry(const FDownloadRef& Download);
+		void SetResult(const FBuildPatchAppManifestPtr& ResultManifest);
 		void SetFailedDownload();
 
 	private:
@@ -51,9 +53,11 @@ namespace BuildPatchServices
 		int32 RetryCount;
 		FDownloadProgressDelegate ChunkDeltaProgress;
 		FDownloadCompleteDelegate ChunkDeltaComplete;
-		TPromise<FBuildPatchAppManifestPtr> ChunkDeltaPromise;
-		TFuture<FBuildPatchAppManifestPtr> ChunkDeltaFuture;
+		// TPromise and TFuture are restricted to types that are copyable, assignable, and default constructible. Value should not be exposed.
+		TPromise<TSharedPtr<FResultValueOrError>> ChunkDeltaPromise;
+		TFuture<TSharedPtr<FResultValueOrError>> ChunkDeltaFuture;
 		FThreadSafeCounter DownloadedBytes;
+		const TCHAR* ErrorCode;
 	};
 
 	FOptimisedDelta::FOptimisedDelta(const FOptimisedDeltaConfiguration& InConfiguration, const FOptimisedDeltaDependencies& InDependencies)
@@ -68,6 +72,7 @@ namespace BuildPatchServices
 		, ChunkDeltaPromise()
 		, ChunkDeltaFuture(ChunkDeltaPromise.GetFuture())
 		, DownloadedBytes(0)
+		, ErrorCode(nullptr)
 	{
 		// There are some conditions in which we do not use a delta.
 		const bool bNoSourceManifest = Configuration.SourceManifest.IsValid() == false;
@@ -87,17 +92,19 @@ namespace BuildPatchServices
 		// Otherwise we provide the standard destination manifest.
 		else
 		{
-			ChunkDeltaPromise.SetValue(Configuration.DestinationManifest);
-			Dependencies.OnComplete(Configuration.DestinationManifest);
+			SetResult(Configuration.DestinationManifest);
 		}
 	}
 
-	FBuildPatchAppManifestPtr FOptimisedDelta::GetDestinationManifest()
+	const IOptimisedDelta::FResultValueOrError& FOptimisedDelta::GetResult() const
 	{
-		return ChunkDeltaFuture.Get();
+		TSharedPtr<FResultValueOrError> Result = ChunkDeltaFuture.Get();
+		// Result should never be set nullptr.
+		check(Result.IsValid());
+		return *Result;
 	}
 
-	int32 FOptimisedDelta::GetMetaDownloadSize()
+	int32 FOptimisedDelta::GetMetaDownloadSize() const
 	{
 		ChunkDeltaFuture.Wait();
 		return DownloadedBytes.GetValue();
@@ -134,14 +141,57 @@ namespace BuildPatchServices
 			FBuildPatchAppManifestRef DeltaManifest = MakeShareable(new FBuildPatchAppManifest());
 			if (DeltaManifest->DeserializeFromData(Download->GetData()))
 			{
-				NewManifest = FBuildMergeManifests::MergeDeltaManifest(Configuration.DestinationManifest, DeltaManifest);
+				// Verify each file received matches what we expect.
+				TSet<FString> DeltaFilenameList;
+				DeltaManifest->GetFileList(DeltaFilenameList);
+				bool bDeltaAccepted = true;
+				for (const FString& DeltaFilename : DeltaFilenameList)
+				{
+					const FFileManifest* OriginalFileManifest = Configuration.DestinationManifest->GetFileManifest(DeltaFilename);
+					const FFileManifest* DeltaFileManifest = DeltaManifest->GetFileManifest(DeltaFilename);
+					if (OriginalFileManifest != nullptr && DeltaFileManifest != nullptr)
+					{
+						FSHAHash OriginalFileHash;
+						FSHAHash DeltaFileHash;
+						Configuration.DestinationManifest->GetFileHash(DeltaFilename, OriginalFileHash);
+						DeltaManifest->GetFileHash(DeltaFilename, DeltaFileHash);
+						if (OriginalFileHash != DeltaFileHash)
+						{
+							UE_LOG(LogOptimisedDelta, Log, TEXT("Rejected optimised delta due to file SHA1 mismatch."));
+							ErrorCode = DownloadErrorCodes::RejectedDeltaFile;
+							bDeltaAccepted = false;
+							break;
+						}
+					}
+					else
+					{
+						UE_LOG(LogOptimisedDelta, Log, TEXT("Rejected optimised delta due to file list mismatch."));
+						ErrorCode = DownloadErrorCodes::RejectedDeltaFile;
+						bDeltaAccepted = false;
+						break;
+					}
+				}
+				// Only use if accepted so far.
+				if (bDeltaAccepted)
+				{
+					NewManifest = FBuildMergeManifests::MergeDeltaManifest(Configuration.DestinationManifest, DeltaManifest);
+					if (!NewManifest.IsValid())
+					{
+						UE_LOG(LogOptimisedDelta, Log, TEXT("Rejected optimised delta due to merge failure."));
+						ErrorCode = DownloadErrorCodes::RejectedDeltaFile;
+						bDeltaAccepted = false;
+					}
+				}
+			}
+			else
+			{
+				ErrorCode = DownloadErrorCodes::UnserialisableDeltaFile;
 			}
 			if (NewManifest.IsValid())
 			{
 				UE_LOG(LogOptimisedDelta, Log, TEXT("Received optimised delta file successfully %s"), *RelativeDeltaFilePath);
 				DownloadedBytes.Set(Download->GetData().Num());
-				ChunkDeltaPromise.SetValue(NewManifest);
-				Dependencies.OnComplete(NewManifest);
+				SetResult(NewManifest);
 			}
 			else if (ShouldRetry(Download))
 			{
@@ -162,6 +212,7 @@ namespace BuildPatchServices
 		}
 		else
 		{
+			ErrorCode = DownloadErrorCodes::MissingDeltaFile;
 			SetFailedDownload();
 		}
 	}
@@ -176,19 +227,24 @@ namespace BuildPatchServices
 		return bCanRetry && bMayRetry;
 	}
 
+	void FOptimisedDelta::SetResult(const FBuildPatchAppManifestPtr& ResultManifest)
+	{
+		ChunkDeltaPromise.SetValue(MakeShared<FResultValueOrError>(MakeValue(ResultManifest)));
+		Dependencies.OnComplete(GetResult());
+	}
+
 	void FOptimisedDelta::SetFailedDownload()
 	{
 		if (DeltaPolicy == EDeltaPolicy::TryFetchContinueWithout)
 		{
 			UE_LOG(LogOptimisedDelta, Log, TEXT("Skipping optimised delta file."));
-			ChunkDeltaPromise.SetValue(Configuration.DestinationManifest);
-			Dependencies.OnComplete(Configuration.DestinationManifest);
+			SetResult(Configuration.DestinationManifest);
 		}
 		else
 		{
 			UE_LOG(LogOptimisedDelta, Log, TEXT("Failed optimised delta file fetch %s"), *RelativeDeltaFilePath);
-			ChunkDeltaPromise.SetValue(nullptr);
-			Dependencies.OnComplete(nullptr);
+			ChunkDeltaPromise.SetValue(MakeShared<FResultValueOrError>(MakeError(ErrorCode)));
+			Dependencies.OnComplete(GetResult());
 		}
 	}
 
@@ -201,7 +257,7 @@ namespace BuildPatchServices
 
 	FOptimisedDeltaDependencies::FOptimisedDeltaDependencies()
 		: DownloadService(nullptr)
-		, OnComplete([](FBuildPatchAppManifestPtr) {})
+		, OnComplete([](const IOptimisedDelta::FResultValueOrError&) {})
 	{
 	}
 
diff --git a/Engine/Source/Runtime/Online/BuildPatchServices/Private/Installer/OptimisedDelta.h b/Engine/Source/Runtime/Online/BuildPatchServices/Private/Installer/OptimisedDelta.h
--- a/Engine/Source/Runtime/Online/BuildPatchServices/Private/Installer/OptimisedDelta.h
+++ b/Engine/Source/Runtime/Online/BuildPatchServices/Private/Installer/OptimisedDelta.h
@@ -1,6 +1,7 @@
 // Copyright Epic Games, Inc. All Rights Reserved.
 #pragma once
 
+#include "Templates/ValueOrError.h"
 #include "BuildPatchManifest.h"
 #include "BuildPatchSettings.h"
 
@@ -15,18 +16,21 @@ namespace BuildPatchServices
 	{
 	public:
 		virtual ~IOptimisedDelta() {}
+		typedef TValueOrError<FBuildPatchAppManifestPtr, FString> FResultValueOrError;
 
 		/**
-		 * Gets the manifest that should be used as the destination manifest. Invalid if there was an error.
+		 * Gets the manifest that should be used as the destination manifest, or the error instead if there was one.
+		 * If DeltaPolicy configuration is EDeltaPolicy::TryFetchContinueWithout then the result is guarenteed to always contain a valid manifest, which
+		 * may equal the original DestinationManifest provided.
 		 * @return the destination manifest.
 		 */
-		virtual FBuildPatchAppManifestPtr GetDestinationManifest() = 0;
+		virtual const IOptimisedDelta::FResultValueOrError& GetResult() const = 0;
 
 		/**
 		 * Gets the size of the metadata downloaded to create the optimised manifest.
 		 * @return the downloaded size in bytes.
 		 */
-		virtual int32 GetMetaDownloadSize() = 0;
+		virtual int32 GetMetaDownloadSize() const = 0;
 	};
 
 	/**
@@ -67,8 +71,8 @@ namespace BuildPatchServices
 	public:
 		// A download service instance.
 		IDownloadService* DownloadService;
-		// Function to call once the destination manifest has been selected.
-		TFunction<void(FBuildPatchAppManifestPtr)> OnComplete;
+		// Function to call once the destination manifest has been selected. Receives manifest or error.
+		TFunction<void(const IOptimisedDelta::FResultValueOrError&)> OnComplete;
 	};
 
 	class FOptimisedDeltaFactory
diff --git a/Engine/Source/Runtime/Online/BuildPatchServices/Public/BuildPatchFeatureLevel.h b/Engine/Source/Runtime/Online/BuildPatchServices/Public/BuildPatchFeatureLevel.h
--- a/Engine/Source/Runtime/Online/BuildPatchServices/Public/BuildPatchFeatureLevel.h
+++ b/Engine/Source/Runtime/Online/BuildPatchServices/Public/BuildPatchFeatureLevel.h
@@ -43,8 +43,10 @@ namespace BuildPatchServices
 		VariableSizeChunksWithoutWindowSizeChunkInfo,
 		// Manifest can reference chunks with dynamic window size, and also serializes them.
 		VariableSizeChunks,
-		// Manifest stores a unique build id for exact matching of build data.
-		StoresUniqueBuildId,
+		// Manifest uses a build id generated from its metadata.
+		UsesRuntimeGeneratedBuildId,
+		// Manifest uses a build id generated unique at build time, and stored in manifest.
+		UsesBuildTimeGeneratedBuildId,
 
 		// !! Always after the latest version entry, signifies the latest version plus 1 to allow the following Latest alias.
 		LatestPlusOne,
@@ -55,7 +57,10 @@ namespace BuildPatchServices
 		// An alias to provide the latest version of a manifest supported by a json serialized format.
 		LatestJson = StoresPrerequisiteIds,
 		// An alias to provide the first available version of optimised delta manifest saving.
-		FirstOptimisedDelta = StoresUniqueBuildId,
+		FirstOptimisedDelta = UsesRuntimeGeneratedBuildId,
+
+		// More aliases, but this time for values that have been renamed
+		StoresUniqueBuildId = UsesRuntimeGeneratedBuildId,
 
 		// JSON manifests were stored with a version of 255 during a certain CL range due to a bug.
 		// We will treat this as being StoresChunkFileSizes in code.
diff --git a/Engine/Source/Runtime/Online/BuildPatchServices/Public/Interfaces/IBuildStatistics.h b/Engine/Source/Runtime/Online/BuildPatchServices/Public/Interfaces/IBuildStatistics.h
--- a/Engine/Source/Runtime/Online/BuildPatchServices/Public/Interfaces/IBuildStatistics.h
+++ b/Engine/Source/Runtime/Online/BuildPatchServices/Public/Interfaces/IBuildStatistics.h
@@ -209,9 +209,14 @@ namespace BuildPatchServices
 		 * @return the current disk write speed in bytes per second.
 		 */
 		virtual double GetDiskWriteByteSpeed() const = 0;
+
+		/**
+		 * @return the total number of verify errors experienced during this installation.
+		 */
+		virtual int32 GetNumVerifyErrors() const = 0;
 		
 		/**
-		 * @return the verify errors experienced during this installation.
+		 * @return the verify error counts per error type experienced during this installation.
 		 */
 		virtual TMap<EVerifyError, int32> GetVerifyErrorCounts() const = 0;
 
