CurlHttp: unset the CURLOPT_SHARE option as soon as we are done processing a curl request. Fixes an ensure when we are shutting down curl if something is holding onto a request object

diff --git a/Engine/Source/Runtime/Online/HTTP/Private/Curl/CurlHttp.cpp b/Engine/Source/Runtime/Online/HTTP/Private/Curl/CurlHttp.cpp
--- a/Engine/Source/Runtime/Online/HTTP/Private/Curl/CurlHttp.cpp
+++ b/Engine/Source/Runtime/Online/HTTP/Private/Curl/CurlHttp.cpp
@@ -88,8 +88,6 @@ FCurlHttpRequest::FCurlHttpRequest()
 
 	curl_easy_setopt(EasyHandle, CURLOPT_BUFFERSIZE, FCurlHttpManager::CurlRequestOptions.BufferSize);
 
-	curl_easy_setopt(EasyHandle, CURLOPT_SHARE, FCurlHttpManager::GShareHandle);
-
 	curl_easy_setopt(EasyHandle, CURLOPT_USE_SSL, CURLUSESSL_ALL);
 
 	// set certificate verification (disable to allow self-signed certificates)
@@ -888,6 +886,8 @@ bool FCurlHttpRequest::SetupRequest()
 		curl_easy_setopt(EasyHandle, CURLOPT_SEEKFUNCTION, StaticSeekCallback);
 	}
 
+	curl_easy_setopt(EasyHandle, CURLOPT_SHARE, FCurlHttpManager::GShareHandle);
+
 	UE_LOG(LogHttp, Log, TEXT("%p: Starting %s request to URL='%s'"), this, *Verb, *URL);
 	return true;
 }
@@ -1081,6 +1081,8 @@ void FCurlHttpRequest::FinishedRequest()
 {
 	check(IsInGameThread());
 	QUICK_SCOPE_CYCLE_COUNTER(STAT_FCurlHttpRequest_FinishedRequest);
+
+	curl_easy_setopt(EasyHandle, CURLOPT_SHARE, nullptr);
 	
 	CheckProgressDelegate();
 	// if completed, get more info
