Added cpu scope to FHttpManager::Flush

diff --git a/Engine/Source/Runtime/Online/HTTP/Private/HttpManager.cpp b/Engine/Source/Runtime/Online/HTTP/Private/HttpManager.cpp
--- a/Engine/Source/Runtime/Online/HTTP/Private/HttpManager.cpp
+++ b/Engine/Source/Runtime/Online/HTTP/Private/HttpManager.cpp
@@ -120,6 +120,8 @@ FHttpThread* FHttpManager::CreateHttpThread()
 
 void FHttpManager::Flush(bool bShutdown)
 {
+	QUICK_SCOPE_CYCLE_COUNTER(STAT_FHttpManager_Flush);
+
 	FScopeLock ScopeLock(&RequestLock);
 	double MaxFlushTimeSeconds = -1.0; // default to no limit
 	GConfig->GetDouble(TEXT("HTTP"), TEXT("MaxFlushTimeSeconds"), MaxFlushTimeSeconds, GEngineIni);
