STAT_FOnlineSubsystemMcp_Tick_UserRequests hitching

I couldn't find a place where I can say is stalling. So I added a few more CYCLE_COUNTER to get more information on this hitch in a future perf pass.

diff --git a/Engine/Source/Runtime/Online/HTTP/Private/Apple/AppleHTTP.cpp b/Engine/Source/Runtime/Online/HTTP/Private/Apple/AppleHTTP.cpp
--- a/Engine/Source/Runtime/Online/HTTP/Private/Apple/AppleHTTP.cpp
+++ b/Engine/Source/Runtime/Online/HTTP/Private/Apple/AppleHTTP.cpp
@@ -281,6 +281,8 @@ void FAppleHttpRequest::SetVerb(const FString& Verb)
 
 bool FAppleHttpRequest::ProcessRequest()
 {
+	QUICK_SCOPE_CYCLE_COUNTER(STAT_FAppleHttpRequest_ProcessRequest);
+
 	SCOPED_AUTORELEASE_POOL;
 	UE_LOG(LogHttp, Verbose, TEXT("FAppleHttpRequest::ProcessRequest()"));
 	bool bStarted = false;
@@ -424,6 +426,8 @@ void FAppleHttpRequest::CleanupRequest()
 
 void FAppleHttpRequest::CancelRequest()
 {
+	QUICK_SCOPE_CYCLE_COUNTER(STAT_FAppleHttpRequest_CancelRequest);
+	
 	UE_LOG(LogHttp, Verbose, TEXT("FAppleHttpRequest::CancelRequest()"));
 	if(Connection != nullptr)
 	{
diff --git a/Engine/Source/Runtime/Online/HTTP/Private/Curl/CurlHttp.cpp b/Engine/Source/Runtime/Online/HTTP/Private/Curl/CurlHttp.cpp
--- a/Engine/Source/Runtime/Online/HTTP/Private/Curl/CurlHttp.cpp
+++ b/Engine/Source/Runtime/Online/HTTP/Private/Curl/CurlHttp.cpp
@@ -998,6 +998,8 @@ void FCurlHttpRequest::TickThreadedRequest(float DeltaSeconds)
 
 void FCurlHttpRequest::CancelRequest()
 {
+	QUICK_SCOPE_CYCLE_COUNTER(STAT_FCurlHttpRequest_CancelRequest); 
+	
 	bCanceled = true;
 	UE_LOG(LogHttp, Verbose, TEXT("%p: HTTP request canceled.  URL=%s"), this, *GetURL());
 	
diff --git a/Engine/Source/Runtime/Online/HTTP/Private/HttpRetrySystem.cpp b/Engine/Source/Runtime/Online/HTTP/Private/HttpRetrySystem.cpp
--- a/Engine/Source/Runtime/Online/HTTP/Private/HttpRetrySystem.cpp
+++ b/Engine/Source/Runtime/Online/HTTP/Private/HttpRetrySystem.cpp
@@ -535,6 +535,8 @@ FHttpRetrySystem::FManager::FHttpRetryRequestEntry::FHttpRetryRequestEntry(TShar
 
 bool FHttpRetrySystem::FManager::ProcessRequest(TSharedRef<FHttpRetrySystem::FRequest>& HttpRetryRequest)
 {
+	QUICK_SCOPE_CYCLE_COUNTER(STAT_FHttpRetrySystem_FManager_ProcessRequest);
+
 	bool bResult = HttpRetryRequest->HttpRequest->ProcessRequest();
 
 	if (bResult)
@@ -547,6 +549,8 @@ bool FHttpRetrySystem::FManager::ProcessRequest(TSharedRef<FHttpRetrySystem::FRe
 
 void FHttpRetrySystem::FManager::CancelRequest(TSharedRef<FHttpRetrySystem::FRequest>& HttpRetryRequest)
 {
+	QUICK_SCOPE_CYCLE_COUNTER(STAT_FHttpRetrySystem_FManager_CancelRequest);
+
 	// Find the existing request entry if is was previously processed.
 	bool bFound = false;
 	for (int32 i = 0; i < RequestList.Num(); ++i)
