Removing two CYCLE_COUNTERs on AppleHTTP.cpp file that were blocking the iOS build.

diff --git a/Engine/Source/Runtime/Online/HTTP/Private/Apple/AppleHTTP.cpp b/Engine/Source/Runtime/Online/HTTP/Private/Apple/AppleHTTP.cpp
--- a/Engine/Source/Runtime/Online/HTTP/Private/Apple/AppleHTTP.cpp
+++ b/Engine/Source/Runtime/Online/HTTP/Private/Apple/AppleHTTP.cpp
@@ -281,8 +281,6 @@ void FAppleHttpRequest::SetVerb(const FString& Verb)
 
 bool FAppleHttpRequest::ProcessRequest()
 {
-	QUICK_SCOPE_CYCLE_COUNTER(STAT_FAppleHttpRequest_ProcessRequest);
-
 	SCOPED_AUTORELEASE_POOL;
 	UE_LOG(LogHttp, Verbose, TEXT("FAppleHttpRequest::ProcessRequest()"));
 	bool bStarted = false;
@@ -426,7 +424,6 @@ void FAppleHttpRequest::CleanupRequest()
 
 void FAppleHttpRequest::CancelRequest()
 {
-	QUICK_SCOPE_CYCLE_COUNTER(STAT_FAppleHttpRequest_CancelRequest);
 	
 	UE_LOG(LogHttp, Verbose, TEXT("FAppleHttpRequest::CancelRequest()"));
 	if(Connection != nullptr)
