Adding more quick scope cycle counters to narrow down which part is causing the hitch.

diff --git a/Engine/Source/Runtime/Online/HTTP/Private/Curl/CurlHttp.cpp b/Engine/Source/Runtime/Online/HTTP/Private/Curl/CurlHttp.cpp
--- a/Engine/Source/Runtime/Online/HTTP/Private/Curl/CurlHttp.cpp
+++ b/Engine/Source/Runtime/Online/HTTP/Private/Curl/CurlHttp.cpp
@@ -707,8 +707,11 @@ bool FCurlHttpRequest::SetupRequest()
 	bCanceled = false;
 	CurlAddToMultiResult = CURLM_OK;
 
-	curl_slist_free_all(HeaderList);
-	HeaderList = nullptr;
+	{
+		QUICK_SCOPE_CYCLE_COUNTER(STAT_FCurlHttpRequest_SetupRequest_SLIST_FREE_HEADERS);
+		curl_slist_free_all(HeaderList);
+		HeaderList = nullptr;
+	}	
 
 	// default no verb to a GET
 	if (Verb.IsEmpty())
@@ -741,153 +744,166 @@ bool FCurlHttpRequest::SetupRequest()
 		return false;
 	}	
 
-	curl_easy_setopt(EasyHandle, CURLOPT_URL, TCHAR_TO_ANSI(*URL));
-
-	if (!FCurlHttpManager::CurlRequestOptions.LocalHostAddr.IsEmpty())
 	{
-		// Set the local address to use for making these requests
-		CURLcode ErrCode = curl_easy_setopt(EasyHandle, CURLOPT_INTERFACE, TCHAR_TO_ANSI(*FCurlHttpManager::CurlRequestOptions.LocalHostAddr));
-	}
+		QUICK_SCOPE_CYCLE_COUNTER(STAT_FCurlHttpRequest_SetupRequest_EASY_SETOPT);
 
-	bool bUseReadFunction = false;
+		curl_easy_setopt(EasyHandle, CURLOPT_URL, TCHAR_TO_ANSI(*URL));
 
-	// set up verb (note that Verb is expected to be uppercase only)
-	if (Verb == TEXT("POST"))
-	{
-		// If we don't pass any other Content-Type, RequestPayload is assumed to be URL-encoded by this time
-		// In the case of using a streamed file, you must explicitly set the Content-Type, because RequestPayload->IsURLEncoded returns false.
-		check(!GetHeader(TEXT("Content-Type")).IsEmpty() || RequestPayload->IsURLEncoded());
-		curl_easy_setopt(EasyHandle, CURLOPT_POST, 1L);
-		curl_easy_setopt(EasyHandle, CURLOPT_POSTFIELDS, NULL);
-		curl_easy_setopt(EasyHandle, CURLOPT_POSTFIELDSIZE, RequestPayload->GetContentLength());
-		bUseReadFunction = true;
-	}
-	else if (Verb == TEXT("PUT") || Verb == TEXT("PATCH"))
-	{
-		curl_easy_setopt(EasyHandle, CURLOPT_UPLOAD, 1L);
-		curl_easy_setopt(EasyHandle, CURLOPT_INFILESIZE, RequestPayload->GetContentLength());
-		if (Verb != TEXT("PUT"))
+		if (!FCurlHttpManager::CurlRequestOptions.LocalHostAddr.IsEmpty())
 		{
-			curl_easy_setopt(EasyHandle, CURLOPT_CUSTOMREQUEST, TCHAR_TO_UTF8(*Verb));
+			// Set the local address to use for making these requests
+			CURLcode ErrCode = curl_easy_setopt(EasyHandle, CURLOPT_INTERFACE, TCHAR_TO_ANSI(*FCurlHttpManager::CurlRequestOptions.LocalHostAddr));
 		}
 
-		bUseReadFunction = true;
-	}
-	else if (Verb == TEXT("GET"))
-	{
-		// technically might not be needed unless we reuse the handles
-		curl_easy_setopt(EasyHandle, CURLOPT_HTTPGET, 1L);
-	}
-	else if (Verb == TEXT("HEAD"))
-	{
-		curl_easy_setopt(EasyHandle, CURLOPT_NOBODY, 1L);
-	}
-	else if (Verb == TEXT("DELETE"))
-	{
-		// If we don't pass any other Content-Type, RequestPayload is assumed to be URL-encoded by this time
-		// (if we pass, don't check here and trust the request)
-		check(!GetHeader(TEXT("Content-Type")).IsEmpty() || RequestPayload->IsURLEncoded());
-
-		curl_easy_setopt(EasyHandle, CURLOPT_POST, 1L);
-		curl_easy_setopt(EasyHandle, CURLOPT_CUSTOMREQUEST, "DELETE");
-		curl_easy_setopt(EasyHandle, CURLOPT_POSTFIELDSIZE, RequestPayload->GetContentLength());
-		bUseReadFunction = true;
-	}
-	else
-	{
-		UE_LOG(LogHttp, Fatal, TEXT("Unsupported verb '%s', can be perhaps added with CURLOPT_CUSTOMREQUEST"), *Verb);
-		UE_DEBUG_BREAK();
-	}
+		bool bUseReadFunction = false;
 
-	if (bUseReadFunction)
-	{
-		BytesSent.Reset();
-		TotalBytesSent.Reset();
-		curl_easy_setopt(EasyHandle, CURLOPT_READDATA, this);
-		curl_easy_setopt(EasyHandle, CURLOPT_READFUNCTION, StaticUploadCallback);
-	}
+		// set up verb (note that Verb is expected to be uppercase only)
+		if (Verb == TEXT("POST"))
+		{
+			// If we don't pass any other Content-Type, RequestPayload is assumed to be URL-encoded by this time
+			// In the case of using a streamed file, you must explicitly set the Content-Type, because RequestPayload->IsURLEncoded returns false.
+			check(!GetHeader(TEXT("Content-Type")).IsEmpty() || RequestPayload->IsURLEncoded());
+			curl_easy_setopt(EasyHandle, CURLOPT_POST, 1L);
+			curl_easy_setopt(EasyHandle, CURLOPT_POSTFIELDS, NULL);
+			curl_easy_setopt(EasyHandle, CURLOPT_POSTFIELDSIZE, RequestPayload->GetContentLength());
+			bUseReadFunction = true;
+		}
+		else if (Verb == TEXT("PUT") || Verb == TEXT("PATCH"))
+		{
+			curl_easy_setopt(EasyHandle, CURLOPT_UPLOAD, 1L);
+			curl_easy_setopt(EasyHandle, CURLOPT_INFILESIZE, RequestPayload->GetContentLength());
+			if (Verb != TEXT("PUT"))
+			{
+				curl_easy_setopt(EasyHandle, CURLOPT_CUSTOMREQUEST, TCHAR_TO_UTF8(*Verb));
+			}
 
-	// set up header function to receive response headers
-	curl_easy_setopt(EasyHandle, CURLOPT_HEADERDATA, this);
-	curl_easy_setopt(EasyHandle, CURLOPT_HEADERFUNCTION, StaticReceiveResponseHeaderCallback);
+			bUseReadFunction = true;
+		}
+		else if (Verb == TEXT("GET"))
+		{
+			// technically might not be needed unless we reuse the handles
+			curl_easy_setopt(EasyHandle, CURLOPT_HTTPGET, 1L);
+		}
+		else if (Verb == TEXT("HEAD"))
+		{
+			curl_easy_setopt(EasyHandle, CURLOPT_NOBODY, 1L);
+		}
+		else if (Verb == TEXT("DELETE"))
+		{
+			// If we don't pass any other Content-Type, RequestPayload is assumed to be URL-encoded by this time
+			// (if we pass, don't check here and trust the request)
+			check(!GetHeader(TEXT("Content-Type")).IsEmpty() || RequestPayload->IsURLEncoded());
+
+			curl_easy_setopt(EasyHandle, CURLOPT_POST, 1L);
+			curl_easy_setopt(EasyHandle, CURLOPT_CUSTOMREQUEST, "DELETE");
+			curl_easy_setopt(EasyHandle, CURLOPT_POSTFIELDSIZE, RequestPayload->GetContentLength());
+			bUseReadFunction = true;
+		}
+		else
+		{
+			UE_LOG(LogHttp, Fatal, TEXT("Unsupported verb '%s', can be perhaps added with CURLOPT_CUSTOMREQUEST"), *Verb);
+			UE_DEBUG_BREAK();
+		}
 
-	// set up write function to receive response payload
-	curl_easy_setopt(EasyHandle, CURLOPT_WRITEDATA, this);
-	curl_easy_setopt(EasyHandle, CURLOPT_WRITEFUNCTION, StaticReceiveResponseBodyCallback);
+		if (bUseReadFunction)
+		{
+			BytesSent.Reset();
+			TotalBytesSent.Reset();
+			curl_easy_setopt(EasyHandle, CURLOPT_READDATA, this);
+			curl_easy_setopt(EasyHandle, CURLOPT_READFUNCTION, StaticUploadCallback);
+		}
 
-	// set up headers
+		// set up header function to receive response headers
+		curl_easy_setopt(EasyHandle, CURLOPT_HEADERDATA, this);
+		curl_easy_setopt(EasyHandle, CURLOPT_HEADERFUNCTION, StaticReceiveResponseHeaderCallback);
 
-	// Empty string here tells Curl to list all supported encodings, allowing servers to send compressed content.
-	if (FCurlHttpManager::CurlRequestOptions.bAcceptCompressedContent)
-	{
-		curl_easy_setopt(EasyHandle, CURLOPT_ACCEPT_ENCODING, "");
-	}
+		// set up write function to receive response payload
+		curl_easy_setopt(EasyHandle, CURLOPT_WRITEDATA, this);
+		curl_easy_setopt(EasyHandle, CURLOPT_WRITEFUNCTION, StaticReceiveResponseBodyCallback);
 
-	if (GetHeader(TEXT("User-Agent")).IsEmpty())
-	{
-		SetHeader(TEXT("User-Agent"), FPlatformHttp::GetDefaultUserAgent());
-	}
+		// set up headers
 
-	// content-length should be present http://www.w3.org/Protocols/rfc2616/rfc2616-sec4.html#sec4.4
-	if (GetHeader(TEXT("Content-Length")).IsEmpty())
-	{
-		SetHeader(TEXT("Content-Length"), FString::Printf(TEXT("%d"), RequestPayload->GetContentLength()));
-	}
+		// Empty string here tells Curl to list all supported encodings, allowing servers to send compressed content.
+		if (FCurlHttpManager::CurlRequestOptions.bAcceptCompressedContent)
+		{
+			curl_easy_setopt(EasyHandle, CURLOPT_ACCEPT_ENCODING, "");
+		}
+	
+		if (GetHeader(TEXT("User-Agent")).IsEmpty())
+		{
+			SetHeader(TEXT("User-Agent"), FPlatformHttp::GetDefaultUserAgent());
+		}
 
-	// Remove "Expect: 100-continue" since this is supposed to cause problematic behavior on Amazon ELB (and WinInet doesn't send that either)
-	// (also see http://www.iandennismiller.com/posts/curl-http1-1-100-continue-and-multipartform-data-post.html , http://the-stickman.com/web-development/php-and-curl-disabling-100-continue-header/ )
-	if (GetHeader(TEXT("Expect")).IsEmpty())
-	{
-		SetHeader(TEXT("Expect"), TEXT(""));
-	}
+		// content-length should be present http://www.w3.org/Protocols/rfc2616/rfc2616-sec4.html#sec4.4
+		if (GetHeader(TEXT("Content-Length")).IsEmpty())
+		{
+			SetHeader(TEXT("Content-Length"), FString::Printf(TEXT("%d"), RequestPayload->GetContentLength()));
+		}
 
-	TArray<FString> AllHeaders = GetAllHeaders();
-	const int32 NumAllHeaders = AllHeaders.Num();
-	for (int32 Idx = 0; Idx < NumAllHeaders; ++Idx)
-	{
-		const bool bCanLogHeaderValue = !AllHeaders[Idx].Contains(TEXT("Authorization"));
-		if (bCanLogHeaderValue)
+		// Remove "Expect: 100-continue" since this is supposed to cause problematic behavior on Amazon ELB (and WinInet doesn't send that either)
+		// (also see http://www.iandennismiller.com/posts/curl-http1-1-100-continue-and-multipartform-data-post.html , http://the-stickman.com/web-development/php-and-curl-disabling-100-continue-header/ )
+		if (GetHeader(TEXT("Expect")).IsEmpty())
 		{
-			UE_LOG(LogHttp, Verbose, TEXT("%p: Adding header '%s'"), this, *AllHeaders[Idx]);
+			SetHeader(TEXT("Expect"), TEXT(""));
 		}
 
-		curl_slist* NewHeaderList = curl_slist_append(HeaderList, TCHAR_TO_UTF8(*AllHeaders[Idx]));
-		if (!NewHeaderList)
 		{
-			if (bCanLogHeaderValue)
-			{
-				UE_LOG(LogHttp, Warning, TEXT("Failed to append header '%s'"), *AllHeaders[Idx]);
-			}
-			else
+			QUICK_SCOPE_CYCLE_COUNTER(STAT_FCurlHttpRequest_SetupRequest_SLIST_APPEND_HEADERS);
+
+			TArray<FString> AllHeaders = GetAllHeaders();
+			const int32 NumAllHeaders = AllHeaders.Num();
+			for (int32 Idx = 0; Idx < NumAllHeaders; ++Idx)
 			{
-				UE_LOG(LogHttp, Warning, TEXT("Failed to append header 'Authorization'"));
+				const bool bCanLogHeaderValue = !AllHeaders[Idx].Contains(TEXT("Authorization"));
+				if (bCanLogHeaderValue)
+				{
+					UE_LOG(LogHttp, Verbose, TEXT("%p: Adding header '%s'"), this, *AllHeaders[Idx]);
+				}
+
+				curl_slist* NewHeaderList = curl_slist_append(HeaderList, TCHAR_TO_UTF8(*AllHeaders[Idx]));
+				if (!NewHeaderList)
+				{
+					if (bCanLogHeaderValue)
+					{
+						UE_LOG(LogHttp, Warning, TEXT("Failed to append header '%s'"), *AllHeaders[Idx]);
+					}
+					else
+					{
+						UE_LOG(LogHttp, Warning, TEXT("Failed to append header 'Authorization'"));
+					}
+				}
+				else
+				{
+					HeaderList = NewHeaderList;
+				}
 			}
 		}
-		else
+
+		if (HeaderList)
 		{
-			HeaderList = NewHeaderList;
+			curl_easy_setopt(EasyHandle, CURLOPT_HTTPHEADER, HeaderList);
 		}
-	}
 
-	if (HeaderList)
-	{
-		curl_easy_setopt(EasyHandle, CURLOPT_HTTPHEADER, HeaderList);
-	}
+		// Set connection timeout in seconds
+		int32 HttpConnectionTimeout = FHttpModule::Get().GetHttpConnectionTimeout();
+		if (HttpConnectionTimeout >= 0)
+		{
+			curl_easy_setopt(EasyHandle, CURLOPT_CONNECTTIMEOUT, HttpConnectionTimeout);
+		}
 
-	// Set connection timeout in seconds
-	int32 HttpConnectionTimeout = FHttpModule::Get().GetHttpConnectionTimeout();
-	if (HttpConnectionTimeout >= 0)
-	{
-		curl_easy_setopt(EasyHandle, CURLOPT_CONNECTTIMEOUT, HttpConnectionTimeout);
-	}
+		if (FCurlHttpManager::CurlRequestOptions.bAllowSeekFunction && bIsRequestPayloadSeekable)
+		{
+			curl_easy_setopt(EasyHandle, CURLOPT_SEEKDATA, this);
+			curl_easy_setopt(EasyHandle, CURLOPT_SEEKFUNCTION, StaticSeekCallback);
+		}
 
-	if (FCurlHttpManager::CurlRequestOptions.bAllowSeekFunction && bIsRequestPayloadSeekable)
-	{
-		curl_easy_setopt(EasyHandle, CURLOPT_SEEKDATA, this);
-		curl_easy_setopt(EasyHandle, CURLOPT_SEEKFUNCTION, StaticSeekCallback);
-	}
+		{
+			//Tracking the locking in the CURLOPT_SHARE branch of the curl_easy_setopt implementation
+			QUICK_SCOPE_CYCLE_COUNTER(STAT_FCurlHttpRequest_SetupRequest_EASY_CURLOPT_SHARE);
 
-	curl_easy_setopt(EasyHandle, CURLOPT_SHARE, FCurlHttpManager::GShareHandle);
+			curl_easy_setopt(EasyHandle, CURLOPT_SHARE, FCurlHttpManager::GShareHandle);
+		}
+	}
 
 	UE_LOG(LogHttp, Log, TEXT("%p: Starting %s request to URL='%s'"), this, *Verb, *URL);
 	return true;
