Allow GetChunkDeltaFilename to be called from outside BPS.  Patch team needs this to get accurate download size estimates using delta manifests.

diff --git a/Engine/Source/Runtime/Online/BuildPatchServices/Private/BuildPatchDelta.cpp b/Engine/Source/Runtime/Online/BuildPatchServices/Private/BuildPatchDelta.cpp
new file mode 100644
--- /dev/null
+++ b/Engine/Source/Runtime/Online/BuildPatchServices/Private/BuildPatchDelta.cpp
@@ -0,0 +1,9 @@
+// Copyright Epic Games, Inc. All Rights Reserved.
+
+#include "BuildPatchDelta.h"
+#include "BuildPatchUtil.h"
+
+BUILDPATCHSERVICES_API FString BuildPatchServices::GetChunkDeltaFilename(const IBuildManifestRef& SourceManifest, const IBuildManifestRef& DestinationManifest)
+{
+	return FBuildPatchUtils::GetChunkDeltaFilename(StaticCastSharedRef<FBuildPatchAppManifest>(SourceManifest).Get(), StaticCastSharedRef<FBuildPatchAppManifest>(DestinationManifest).Get());
+}
diff --git a/Engine/Source/Runtime/Online/BuildPatchServices/Public/BuildPatchDelta.h b/Engine/Source/Runtime/Online/BuildPatchServices/Public/BuildPatchDelta.h
--- a/Engine/Source/Runtime/Online/BuildPatchServices/Public/BuildPatchDelta.h
+++ b/Engine/Source/Runtime/Online/BuildPatchServices/Public/BuildPatchDelta.h
@@ -3,6 +3,8 @@
 
 #include "CoreMinimal.h"
 
+#include "Interfaces/IBuildManifest.h"
+
 namespace BuildPatchServices
 {
 	/**
@@ -40,4 +42,12 @@ namespace BuildPatchServices
 			default: return InvalidOrMax;
 		}
 	}
+
+	/**
+	 * Based on the source and destination manifests, get the filename for the delta that optimises patching from source to destination.
+	 * @param SourceManifest        The source manifest.
+	 * @param DestinationManifest   The destination manifest.
+	 * @return the CloudDir relative delta filename.
+	 */
+	BUILDPATCHSERVICES_API FString GetChunkDeltaFilename(const IBuildManifestRef& SourceManifest, const IBuildManifestRef& DestinationManifest);
 }
