Allow MergeDeltaManifest to be called from outside BPS.  Patch team needs this to get accurate download size estimates using delta manifests.

diff --git a/Engine/Source/Runtime/Online/BuildPatchServices/Private/BuildPatchDelta.cpp b/Engine/Source/Runtime/Online/BuildPatchServices/Private/BuildPatchDelta.cpp
--- a/Engine/Source/Runtime/Online/BuildPatchServices/Private/BuildPatchDelta.cpp
+++ b/Engine/Source/Runtime/Online/BuildPatchServices/Private/BuildPatchDelta.cpp
@@ -2,8 +2,14 @@
 
 #include "BuildPatchDelta.h"
 #include "BuildPatchUtil.h"
+#include "BuildPatchMergeManifests.h"
 
 BUILDPATCHSERVICES_API FString BuildPatchServices::GetChunkDeltaFilename(const IBuildManifestRef& SourceManifest, const IBuildManifestRef& DestinationManifest)
 {
 	return FBuildPatchUtils::GetChunkDeltaFilename(StaticCastSharedRef<FBuildPatchAppManifest>(SourceManifest).Get(), StaticCastSharedRef<FBuildPatchAppManifest>(DestinationManifest).Get());
 }
+
+BUILDPATCHSERVICES_API IBuildManifestPtr BuildPatchServices::MergeDeltaManifest(const IBuildManifestRef& Manifest, const IBuildManifestRef& Delta)
+{
+	return FBuildMergeManifests::MergeDeltaManifest(StaticCastSharedRef<FBuildPatchAppManifest>(Manifest), StaticCastSharedRef<FBuildPatchAppManifest>(Delta));
+}
diff --git a/Engine/Source/Runtime/Online/BuildPatchServices/Public/BuildPatchDelta.h b/Engine/Source/Runtime/Online/BuildPatchServices/Public/BuildPatchDelta.h
--- a/Engine/Source/Runtime/Online/BuildPatchServices/Public/BuildPatchDelta.h
+++ b/Engine/Source/Runtime/Online/BuildPatchServices/Public/BuildPatchDelta.h
@@ -50,4 +50,6 @@ namespace BuildPatchServices
 	 * @return the CloudDir relative delta filename.
 	 */
 	BUILDPATCHSERVICES_API FString GetChunkDeltaFilename(const IBuildManifestRef& SourceManifest, const IBuildManifestRef& DestinationManifest);
+
+	BUILDPATCHSERVICES_API IBuildManifestPtr MergeDeltaManifest(const IBuildManifestRef& Manifest, const IBuildManifestRef& Delta);
 }
