Add new BuildPatchErrorUtils to add patching error utilities.

GetUnrecoverablePatchTypesStrings populates an array with the error values that signify an unrecoverable patch did happen, for example, BuildVerifyFail.

diff --git a/Engine/Source/Runtime/Online/BuildPatchServices/Public/Interfaces/IBuildInstaller.h b/Engine/Source/Runtime/Online/BuildPatchServices/Public/Interfaces/IBuildInstaller.h
--- a/Engine/Source/Runtime/Online/BuildPatchServices/Public/Interfaces/IBuildInstaller.h
+++ b/Engine/Source/Runtime/Online/BuildPatchServices/Public/Interfaces/IBuildInstaller.h
@@ -466,3 +466,31 @@ public:
 	 */
 	virtual const BuildPatchServices::FBuildInstallerConfiguration& GetConfiguration() const = 0;
 };
+
+/**
+ * @return the stringified version of the EBuildPatchInstallError enum passed in
+ */
+inline const TCHAR* LexToString(EBuildPatchInstallError EnumVal)
+{
+	static_assert((int32)EBuildPatchInstallError::NumInstallErrors == 12, "Please add support for the extra values to the function below.");
+	#define CASE_ENUM_SET(State) case State: return TEXT(#State); break;
+	switch (EnumVal)
+	{
+		CASE_ENUM_SET(EBuildPatchInstallError::NoError)
+		CASE_ENUM_SET(EBuildPatchInstallError::DownloadError)
+		CASE_ENUM_SET(EBuildPatchInstallError::FileConstructionFail)
+		CASE_ENUM_SET(EBuildPatchInstallError::MoveFileToInstall)
+		CASE_ENUM_SET(EBuildPatchInstallError::BuildVerifyFail)
+		CASE_ENUM_SET(EBuildPatchInstallError::ApplicationClosing)
+		CASE_ENUM_SET(EBuildPatchInstallError::ApplicationError)
+		CASE_ENUM_SET(EBuildPatchInstallError::UserCanceled)
+		CASE_ENUM_SET(EBuildPatchInstallError::PrerequisiteError)
+		CASE_ENUM_SET(EBuildPatchInstallError::InitializationError)
+		CASE_ENUM_SET(EBuildPatchInstallError::PathLengthExceeded)
+		CASE_ENUM_SET(EBuildPatchInstallError::OutOfDiskSpace)
+			
+		default:
+			return TEXT("EBuildPatchInstallError::Unknown");
+	}
+	#undef CASE_ENUM_SET
+};
