Adding a Tick function in XmppModule to call ProcessPendingRemovals in order to clean closed xmpp connections.

diff --git a/Engine/Source/Runtime/Online/XMPP/Private/XmppModule.cpp b/Engine/Source/Runtime/Online/XMPP/Private/XmppModule.cpp
--- a/Engine/Source/Runtime/Online/XMPP/Private/XmppModule.cpp
+++ b/Engine/Source/Runtime/Online/XMPP/Private/XmppModule.cpp
@@ -768,6 +768,12 @@ void FXmppModule::RemoveConnection(const FString& UserId)
 	ActiveConnections.Remove(UserId);
 }
 
+bool FXmppModule::Tick(float DeltaTime)
+{
+	ProcessPendingRemovals();
+	return true;
+}
+
 void FXmppModule::ProcessPendingRemovals()
 {
 	if (PendingRemovals.Num() > 0)
diff --git a/Engine/Source/Runtime/Online/XMPP/Public/XmppModule.h b/Engine/Source/Runtime/Online/XMPP/Public/XmppModule.h
--- a/Engine/Source/Runtime/Online/XMPP/Public/XmppModule.h
+++ b/Engine/Source/Runtime/Online/XMPP/Public/XmppModule.h
@@ -15,7 +15,7 @@ class Error;
  * Use CreateConnection to create a new Xmpp connection
  */
 class XMPP_API FXmppModule :
-	public IModuleInterface, public FSelfRegisteringExec
+	public IModuleInterface, public FSelfRegisteringExec, public FTickerObjectBase
 {
 
 public:
@@ -96,6 +96,9 @@ public:
 		return bEnabled;
 	}
 
+	// FTickerObjectBase
+	virtual bool Tick(float DeltaTime) override;
+
 	/**
 	 * Delegate callback when a system acquires ownership over an XMPP connection
 	 *
