Add missing include
#rnx

diff --git a/Engine/Source/Runtime/Online/XMPP/Public/XmppModule.h b/Engine/Source/Runtime/Online/XMPP/Public/XmppModule.h
--- a/Engine/Source/Runtime/Online/XMPP/Public/XmppModule.h
+++ b/Engine/Source/Runtime/Online/XMPP/Public/XmppModule.h
@@ -4,6 +4,7 @@
 
 #include "CoreMinimal.h"
 #include "Misc/CoreMisc.h"
+#include "Containers/Ticker.h"
 #include "XmppConnection.h"
 #include "Modules/ModuleInterface.h"
 #include "XmppMultiUserChat.h"
