Fix ensure

diff --git a/Engine/Source/Runtime/Online/BackgroundHTTP/Private/GenericPlatform/GenericPlatformBackgroundHttpRequest.cpp b/Engine/Source/Runtime/Online/BackgroundHTTP/Private/GenericPlatform/GenericPlatformBackgroundHttpRequest.cpp
--- a/Engine/Source/Runtime/Online/BackgroundHTTP/Private/GenericPlatform/GenericPlatformBackgroundHttpRequest.cpp
+++ b/Engine/Source/Runtime/Online/BackgroundHTTP/Private/GenericPlatform/GenericPlatformBackgroundHttpRequest.cpp
@@ -173,7 +173,7 @@ void FGenericPlatformBackgroundHttpRequest::FGenericPlatformBackgroundHttpWrappe
 	if (ensureAlwaysMsgf(OriginalRequest.IsValid(), TEXT("Recieved UpdateHttpProgress callback with invalid OriginalRequest pointer! Can not update request progress!")))
 	{
 		const int32 ByteDifference = (LastProgressUpdateBytes > 0) ? BytesReceived - LastProgressUpdateBytes : BytesReceived;
-		ensureAlwaysMsgf((ByteDifference > 0), TEXT("Invalid Byte Difference in UpdateHttpProgress -- ByteDifference:%d | LastProgressUpdateBytes:%d | BytesSent:%d | BytesReceived:%d"), ByteDifference, LastProgressUpdateBytes, BytesSent, BytesReceived);
+		ensureAlwaysMsgf((ByteDifference >= 0), TEXT("Invalid Byte Difference in UpdateHttpProgress -- ByteDifference:%d | LastProgressUpdateBytes:%d | BytesSent:%d | BytesReceived:%d"), ByteDifference, LastProgressUpdateBytes, BytesSent, BytesReceived);
 		LastProgressUpdateBytes = BytesReceived;
 
 		UE_LOG(LogBackgroundHttpRequest, VeryVerbose, TEXT("HttpRequest Progress Update- RequestID:%s | BytesSent: %d | BytesReceived:%d | BytesReceivedSinceLastUpdate:%d"), *OriginalRequest.Pin()->GetRequestID(), BytesSent, BytesReceived, ByteDifference);
