- Allow custom thread sizes for ping on IOS

diff --git a/Engine/Source/Runtime/Online/ICMP/Private/IcmpPrivate.h b/Engine/Source/Runtime/Online/ICMP/Private/IcmpPrivate.h
--- a/Engine/Source/Runtime/Online/ICMP/Private/IcmpPrivate.h
+++ b/Engine/Source/Runtime/Online/ICMP/Private/IcmpPrivate.h
@@ -32,6 +32,12 @@
 		#define PLATFORM_USES_POSIX_ICMP 1
 		#define PING_ALLOWS_CUSTOM_THREAD_SIZE 0
 
+	#elif PLATFORM_IOS
+
+		#define PLATFORM_SUPPORTS_ICMP 0
+		#define PLATFORM_USES_POSIX_ICMP 0
+		#define PING_ALLOWS_CUSTOM_THREAD_SIZE 1
+
 	#else
 
 		#define PLATFORM_SUPPORTS_ICMP 0
