Client crash in FXmppConnectionStrophe::Logout.
Only call FXmppPresenceStrophe::GetRosterMembers on the game thread.

diff --git a/Engine/Source/Runtime/Online/XMPP/Private/XmppStrophe/XmppPresenceStrophe.cpp b/Engine/Source/Runtime/Online/XMPP/Private/XmppStrophe/XmppPresenceStrophe.cpp
--- a/Engine/Source/Runtime/Online/XMPP/Private/XmppStrophe/XmppPresenceStrophe.cpp
+++ b/Engine/Source/Runtime/Online/XMPP/Private/XmppStrophe/XmppPresenceStrophe.cpp
@@ -211,6 +211,7 @@ TArray<TSharedPtr<FXmppUserPresence>> FXmppPresenceStrophe::GetRosterPresence(co
 
 void FXmppPresenceStrophe::GetRosterMembers(TArray<FXmppUserJid>& Members)
 {
+	check(IsInGameThread());
 	Members.Empty(RosterMembers.Num());
 	for (const TMap<FString, TSharedRef<FXmppUserPresence>>::ElementType& Pair : RosterMembers)
 	{
diff --git a/Engine/Source/Runtime/Online/XMPP/Private/XmppStrophe/XmppPrivateChatStrophe.cpp b/Engine/Source/Runtime/Online/XMPP/Private/XmppStrophe/XmppPrivateChatStrophe.cpp
--- a/Engine/Source/Runtime/Online/XMPP/Private/XmppStrophe/XmppPrivateChatStrophe.cpp
+++ b/Engine/Source/Runtime/Online/XMPP/Private/XmppStrophe/XmppPrivateChatStrophe.cpp
@@ -43,22 +43,6 @@ bool FXmppPrivateChatStrophe::ReceiveStanza(const FStropheStanza& IncomingStanza
 		return false;
 	}
 
-	// Potentially filter out non-friends/non-admins
-	if (ConnectionManager.GetServer().bPrivateChatFriendsOnly && ConnectionManager.Presence().IsValid())
-	{
-		FXmppUserJid FromJid = IncomingStanza.GetFrom();
-		if (FromJid.Id != TEXT("xmpp-admin"))
-		{
-			TArray<FXmppUserJid> RosterMembers;
-			ConnectionManager.Presence()->GetRosterMembers(RosterMembers);
-			if (!RosterMembers.Contains(FromJid))
-			{
-				// This was meant for us, but we don't want to see it
-				return true;
-			}
-		}
-	}
-
 	TOptional<FString> BodyText = IncomingStanza.GetBodyText();
 	if (!BodyText.IsSet())
 	{
@@ -134,6 +118,21 @@ bool FXmppPrivateChatStrophe::Tick(float DeltaTime)
 void FXmppPrivateChatStrophe::OnChatReceived(TUniquePtr<FXmppChatMessage>&& Chat)
 {
 	TSharedRef<FXmppChatMessage> ChatRef = MakeShareable(Chat.Release());
+
+	// Potentially filter out non-friends/non-admins
+	if (ConnectionManager.GetServer().bPrivateChatFriendsOnly && ConnectionManager.Presence().IsValid())
+	{
+		if (ChatRef->FromJid.Id != TEXT("xmpp-admin"))
+		{
+			TArray<FXmppUserJid> RosterMembers;
+			ConnectionManager.Presence()->GetRosterMembers(RosterMembers);
+			if (!RosterMembers.Contains(ChatRef->FromJid))
+			{
+				return;
+			}
+		}
+	}
+
 	OnChatReceivedDelegate.Broadcast(ConnectionManager.AsShared(), ChatRef->FromJid, ChatRef);
 }
 
