- MMS MCP DNS resolver.

diff --git a/Engine/Source/Runtime/Online/WebSockets/Private/Lws/LwsWebSocket.cpp b/Engine/Source/Runtime/Online/WebSockets/Private/Lws/LwsWebSocket.cpp
--- a/Engine/Source/Runtime/Online/WebSockets/Private/Lws/LwsWebSocket.cpp
+++ b/Engine/Source/Runtime/Online/WebSockets/Private/Lws/LwsWebSocket.cpp
@@ -487,15 +487,26 @@ int FLwsWebSocket::LwsCallback(lws* Instance, lws_callback_reasons Reason, void*
 	{
 		// in FLwsWebSocketsManager::CallbackWrapper, we copied UserData to Data, so Data is the X509_STORE_CTX* instead of the SSL*
 		X509_STORE_CTX* Context = static_cast<X509_STORE_CTX*>(Data);
-		int PreverifyOk = Length;
-		if (PreverifyOk == 1)
+		int PreverifyOk = 1;
+		bool bDisableCertValidation = false;
+		GConfig->GetBool(TEXT("LwsWebSocket"), TEXT("bDisableCertValidation"), bDisableCertValidation, GEngineIni);
+		if (!bDisableCertValidation)
 		{
-			const FString Domain = FGenericPlatformHttp::GetUrlDomain(Url);
-			if (!FSslModule::Get().GetCertificateManager().VerifySslCertificates(Context, Domain))
+			PreverifyOk = Length;
+			if (PreverifyOk == 1)
 			{
-				PreverifyOk = 0;
+				const FString Domain = FGenericPlatformHttp::GetUrlDomain(Url);
+				if (!FSslModule::Get().GetCertificateManager().VerifySslCertificates(Context, Domain))
+				{
+					PreverifyOk = 0;
+				}
 			}
 		}
+		else
+		{
+			X509_STORE_CTX_set_error(Context, X509_V_OK);
+		}
+
 		return PreverifyOk == 1 ? 0 : 1;
 	}
 #endif
