- Fix non-unity compile error
bob.tellez

diff --git a/Engine/Source/Runtime/Online/WebSockets/Private/Lws/LwsWebSocket.cpp b/Engine/Source/Runtime/Online/WebSockets/Private/Lws/LwsWebSocket.cpp
--- a/Engine/Source/Runtime/Online/WebSockets/Private/Lws/LwsWebSocket.cpp
+++ b/Engine/Source/Runtime/Online/WebSockets/Private/Lws/LwsWebSocket.cpp
@@ -12,6 +12,7 @@
 #endif
 #include "Misc/EmbeddedCommunication.h"
 #include "Misc/ScopeLock.h"
+#include "Misc/ConfigCacheIni.h"
 #include "HttpModule.h"
 #include "HttpManager.h"
 #include "PlatformHttp.h"
