add support for skipping domain whitelisting in LWS connect

#rnx

diff --git a/Engine/Source/Runtime/Online/WebSockets/Private/Lws/LwsWebSocket.cpp b/Engine/Source/Runtime/Online/WebSockets/Private/Lws/LwsWebSocket.cpp
--- a/Engine/Source/Runtime/Online/WebSockets/Private/Lws/LwsWebSocket.cpp
+++ b/Engine/Source/Runtime/Online/WebSockets/Private/Lws/LwsWebSocket.cpp
@@ -100,13 +100,22 @@ void FLwsWebSocket::Connect()
 		return;
 	}
 
-	FHttpManager& HttpManager = FHttpModule::Get().GetHttpManager();
-	if (!HttpManager.IsDomainAllowed(Url))
+	bool bDisableDomainWhitelist = false;
+	GConfig->GetBool(TEXT("LwsWebSocket"), TEXT("bDisableDomainWhitelist"), bDisableDomainWhitelist, GEngineIni);
+	if (!bDisableDomainWhitelist)
 	{
-		State = EState::Error;
-		UE_LOG(LogWebSockets, Warning, TEXT("FLwsWebSocket[%d]::Connect: %s is not whitelisted. Refusing to connect."), Identifier, *Url);
-		OnConnectionError().Broadcast(TEXT("Invalid Domain"));
-		return;
+		FHttpManager& HttpManager = FHttpModule::Get().GetHttpManager();
+		if (!HttpManager.IsDomainAllowed(Url))
+		{
+			State = EState::Error;
+			UE_LOG(LogWebSockets, Warning, TEXT("FLwsWebSocket[%d]::Connect: %s is not whitelisted. Refusing to connect."), Identifier, *Url);
+			OnConnectionError().Broadcast(TEXT("Invalid Domain"));
+			return;
+		}
+	}
+	else
+	{
+		UE_LOG(LogWebSockets, Log, TEXT("FLwsWebSocket[%d]::Connect: Domain whitelisting has been disabled by config.", Identifier));
 	}
 
 	// No lock, we are not being processed on the websockets thread yet
