CIS fix against LwsWebSocket.cpp

diff --git a/Engine/Source/Runtime/Online/WebSockets/Private/Lws/LwsWebSocket.cpp b/Engine/Source/Runtime/Online/WebSockets/Private/Lws/LwsWebSocket.cpp
--- a/Engine/Source/Runtime/Online/WebSockets/Private/Lws/LwsWebSocket.cpp
+++ b/Engine/Source/Runtime/Online/WebSockets/Private/Lws/LwsWebSocket.cpp
@@ -115,7 +115,7 @@ void FLwsWebSocket::Connect()
 	}
 	else
 	{
-		UE_LOG(LogWebSockets, Log, TEXT("FLwsWebSocket[%d]::Connect: Domain whitelisting has been disabled by config.", Identifier));
+		UE_LOG(LogWebSockets, Log, TEXT("FLwsWebSocket[%d]::Connect: Domain whitelisting has been disabled by config."), Identifier);
 	}
 
 	// No lock, we are not being processed on the websockets thread yet
