Run the LwsWebSocketsModule in its unique thread on forked processes that support multithreading

diff --git a/Engine/Source/Runtime/Online/WebSockets/Private/Lws/LwsWebSocketsManager.cpp b/Engine/Source/Runtime/Online/WebSockets/Private/Lws/LwsWebSocketsManager.cpp
--- a/Engine/Source/Runtime/Online/WebSockets/Private/Lws/LwsWebSocketsManager.cpp
+++ b/Engine/Source/Runtime/Online/WebSockets/Private/Lws/LwsWebSocketsManager.cpp
@@ -17,6 +17,7 @@
 #include "HAL/LowLevelMemTracker.h"
 #include "Misc/ConfigCacheIni.h"
 #include "Misc/CoreDelegates.h"
+#include "Misc/Fork.h"
 #include "Stats/Stats.h"
 
 namespace {
@@ -63,6 +64,7 @@ FLwsWebSocketsManager& FLwsWebSocketsManager::Get()
 
 void FLwsWebSocketsManager::InitWebSockets(TArrayView<const FString> Protocols)
 {
+	QUICK_SCOPE_CYCLE_COUNTER(STAT_FLwsWebSocketsManager_InitWebSockets);
 	check(!Thread && LwsProtocols.Num() == 0);
 
 	LwsProtocols.Reserve(Protocols.Num() + 1);
@@ -150,7 +152,7 @@ void FLwsWebSocketsManager::InitWebSockets(TArrayView<const FString> Protocols)
 	
 	int32 ThreadStackSize = 128 * 1024;
 	GConfig->GetInt(TEXT("WebSockets.LibWebSockets"), TEXT("ThreadStackSize"), ThreadStackSize, GEngineIni);
-	Thread = FRunnableThread::Create(this, TEXT("LibwebsocketsThread"), 128 * 1024, TPri_Normal);
+	Thread = FForkProcessHelper::CreateForkableThread(this, TEXT("LibwebsocketsThread"), 128 * 1024, TPri_Normal, FPlatformAffinity::GetNoAffinityMask());
 	if (!Thread)
 	{
 		UE_LOG(LogWebSockets, Error, TEXT("FLwsWebSocketsManager failed to initialize thread!"));
@@ -166,6 +168,7 @@ void FLwsWebSocketsManager::InitWebSockets(TArrayView<const FString> Protocols)
 
 void FLwsWebSocketsManager::ShutdownWebSockets()
 {
+	QUICK_SCOPE_CYCLE_COUNTER(STAT_FLwsWebSocketsManager_ShutdownWebSockets);
 	if (TickHandle.IsValid())
 	{
 		FBackgroundableTicker::GetCoreTicker().RemoveTicker(TickHandle);
@@ -350,6 +353,7 @@ int FLwsWebSocketsManager::CallbackWrapper(lws* Connection, lws_callback_reasons
 
 void FLwsWebSocketsManager::Tick()
 {
+	QUICK_SCOPE_CYCLE_COUNTER(STAT_FLwsWebSocketsManager_Tick);
 	LLM_SCOPE(ELLMTag::Networking);
 
 	{
