HTTPManager

* Allow HTTPThread to be a real thread on multithread forked processes
* Added final to Tick() andd Run() of HTTPThread to ensure derived classes do not redefine the internal implementation.

diff --git a/Engine/Source/Runtime/Online/HTTP/Private/HttpThread.cpp b/Engine/Source/Runtime/Online/HTTP/Private/HttpThread.cpp
--- a/Engine/Source/Runtime/Online/HTTP/Private/HttpThread.cpp
+++ b/Engine/Source/Runtime/Online/HTTP/Private/HttpThread.cpp
@@ -5,6 +5,7 @@
 #include "HAL/PlatformTime.h"
 #include "HAL/PlatformProcess.h"
 #include "HAL/RunnableThread.h"
+#include "Misc/Fork.h"
 #include "HttpModule.h"
 #include "Http.h"
 
@@ -28,7 +29,7 @@ FHttpThread::~FHttpThread()
 
 void FHttpThread::StartThread()
 {
-	Thread = FRunnableThread::Create(this, TEXT("HttpManagerThread"), 128 * 1024, TPri_Normal);
+	Thread = FForkProcessHelper::CreateForkableThread(this, TEXT("HttpManagerThread"), 128 * 1024, TPri_Normal);
 }
 
 void FHttpThread::StopThread()
@@ -67,8 +68,6 @@ bool FHttpThread::Init()
 	return true;
 }
 
-// @todo samz - needed due to InnerLoopEnd,InnerLoopBegin 
-PRAGMA_DISABLE_OPTIMIZATION
 uint32 FHttpThread::Run()
 {
 	// Arrays declared outside of loop to re-use memory
@@ -77,12 +76,12 @@ uint32 FHttpThread::Run()
 	TArray<IHttpThreadedRequest*> RequestsToComplete;
 	while (!ExitRequest.GetValue())
 	{
-		double OuterLoopBegin = FPlatformTime::Seconds();
+		const double OuterLoopBegin = FPlatformTime::Seconds();
 		double OuterLoopEnd = 0.0;
 		bool bKeepProcessing = true;
 		while (bKeepProcessing)
 		{
-			double InnerLoopBegin = FPlatformTime::Seconds();
+			const double InnerLoopBegin = FPlatformTime::Seconds();
 			
 			Process(RequestsToCancel, RequestsToStart, RequestsToComplete);
 			
@@ -91,7 +90,7 @@ uint32 FHttpThread::Run()
 				bKeepProcessing = false;
 			}
 
-			double InnerLoopEnd = FPlatformTime::Seconds();
+			const double InnerLoopEnd = FPlatformTime::Seconds();
 			if (bKeepProcessing)
 			{
 				double InnerLoopTime = InnerLoopEnd - InnerLoopBegin;
@@ -109,7 +108,6 @@ uint32 FHttpThread::Run()
 	}
 	return 0;
 }
-PRAGMA_ENABLE_OPTIMIZATION
 
 
 void FHttpThread::Tick()
diff --git a/Engine/Source/Runtime/Online/HTTP/Private/HttpThread.h b/Engine/Source/Runtime/Online/HTTP/Private/HttpThread.h
--- a/Engine/Source/Runtime/Online/HTTP/Private/HttpThread.h
+++ b/Engine/Source/Runtime/Online/HTTP/Private/HttpThread.h
@@ -55,7 +55,8 @@ public:
 	void GetCompletedRequests(TArray<IHttpThreadedRequest*>& OutCompletedRequests);
 
 	//~ Begin FSingleThreadRunnable Interface
-	virtual void Tick() override;
+	// Cannot be overriden to ensure identical behavior with the threaded tick
+	virtual void Tick() override final;
 	//~ End FSingleThreadRunnable Interface
 
 protected:
@@ -81,7 +82,7 @@ protected:
 
 	//~ Begin FRunnable Interface
 	virtual bool Init() override;
-	virtual uint32 Run() override;
+	virtual uint32 Run() override final;
 	virtual void Stop() override;
 	virtual void Exit() override;
 	//~ End FRunnable Interface
