Revert HTTP thread to use a fake thread on forkable servers

diff --git a/Engine/Source/Runtime/Online/HTTP/Private/HttpThread.cpp b/Engine/Source/Runtime/Online/HTTP/Private/HttpThread.cpp
--- a/Engine/Source/Runtime/Online/HTTP/Private/HttpThread.cpp
+++ b/Engine/Source/Runtime/Online/HTTP/Private/HttpThread.cpp
@@ -29,7 +29,8 @@ FHttpThread::~FHttpThread()
 
 void FHttpThread::StartThread()
 {
-	Thread = FForkProcessHelper::CreateForkableThread(this, TEXT("HttpManagerThread"), 128 * 1024, TPri_Normal);
+	//Thread = FForkProcessHelper::CreateForkableThread(this, TEXT("HttpManagerThread"), 128 * 1024, TPri_Normal);
+	Thread = FRunnableThread::Create(this, TEXT("HttpManagerThread"), 128 * 1024, TPri_Normal);
 }
 
 void FHttpThread::StopThread()
