Prevent http flush call from ticking the runnable thread manually when PostForkThreading is enabled.

diff --git a/Engine/Source/Runtime/Online/HTTP/Private/HttpManager.cpp b/Engine/Source/Runtime/Online/HTTP/Private/HttpManager.cpp
--- a/Engine/Source/Runtime/Online/HTTP/Private/HttpManager.cpp
+++ b/Engine/Source/Runtime/Online/HTTP/Private/HttpManager.cpp
@@ -7,6 +7,7 @@
 #include "Misc/ScopeLock.h"
 #include "Http.h"
 #include "Misc/Guid.h"
+#include "Misc/Fork.h"
 
 #include "HttpThread.h"
 #include "Misc/ConfigCacheIni.h"
@@ -167,7 +168,7 @@ void FHttpManager::Flush(bool bShutdown)
 		LastTime = AppTime;
 		if (Requests.Num() > 0)
 		{
-			if (FPlatformProcess::SupportsMultithreading())
+			if (FPlatformProcess::SupportsMultithreading() || FForkProcessHelper::IsForkedMultithreadInstance())
 			{
 				UE_LOG(LogHttp, Display, TEXT("Sleeping 0.5s to wait for %d outstanding Http requests."), Requests.Num());
 				FPlatformProcess::Sleep(0.5f);
