Revert all changes and ensure HTTPManager is running single-threaded on forked multithread servers

diff --git a/Engine/Source/Runtime/Online/HTTP/Private/HttpManager.cpp b/Engine/Source/Runtime/Online/HTTP/Private/HttpManager.cpp
--- a/Engine/Source/Runtime/Online/HTTP/Private/HttpManager.cpp
+++ b/Engine/Source/Runtime/Online/HTTP/Private/HttpManager.cpp
@@ -168,7 +168,7 @@ void FHttpManager::Flush(bool bShutdown)
 		LastTime = AppTime;
 		if (Requests.Num() > 0)
 		{
-			if (FPlatformProcess::SupportsMultithreading() || FForkProcessHelper::IsForkedMultithreadInstance())
+			if (FPlatformProcess::SupportsMultithreading())
 			{
 				UE_LOG(LogHttp, Display, TEXT("Sleeping 0.5s to wait for %d outstanding Http requests."), Requests.Num());
 				FPlatformProcess::Sleep(0.5f);
diff --git a/Engine/Source/Runtime/Online/HTTP/Private/HttpThread.cpp b/Engine/Source/Runtime/Online/HTTP/Private/HttpThread.cpp
--- a/Engine/Source/Runtime/Online/HTTP/Private/HttpThread.cpp
+++ b/Engine/Source/Runtime/Online/HTTP/Private/HttpThread.cpp
@@ -29,7 +29,6 @@ FHttpThread::~FHttpThread()
 
 void FHttpThread::StartThread()
 {
-	//Thread = FForkProcessHelper::CreateForkableThread(this, TEXT("HttpManagerThread"), 128 * 1024, TPri_Normal);
 	Thread = FRunnableThread::Create(this, TEXT("HttpManagerThread"), 128 * 1024, TPri_Normal);
 }
 
