- Move curl setopt operations to http thread

diff --git a/Engine/Source/Runtime/Online/HTTP/Private/Curl/CurlHttp.cpp b/Engine/Source/Runtime/Online/HTTP/Private/Curl/CurlHttp.cpp
--- a/Engine/Source/Runtime/Online/HTTP/Private/Curl/CurlHttp.cpp
+++ b/Engine/Source/Runtime/Online/HTTP/Private/Curl/CurlHttp.cpp
@@ -281,25 +281,48 @@ const TArray<uint8>& FCurlHttpRequest::GetContent() const
 
 void FCurlHttpRequest::SetVerb(const FString& InVerb)
 {
-	check(EasyHandle);
+	if (CompletionStatus == EHttpRequestStatus::Processing)
+	{
+		UE_LOG(LogHttp, Warning, TEXT("FCurlHttpRequest::SetVerb() - attempted to set verb on a request that is inflight"));
+		return;
+	}
 
+	check(EasyHandle);
 	Verb = InVerb.ToUpper();
 }
 
 void FCurlHttpRequest::SetURL(const FString& InURL)
 {
+	if (CompletionStatus == EHttpRequestStatus::Processing)
+	{
+		UE_LOG(LogHttp, Warning, TEXT("FCurlHttpRequest::SetURL() - attempted to set url on a request that is inflight"));
+		return;
+	}
+
 	check(EasyHandle);
 	URL = InURL;
 }
 
 void FCurlHttpRequest::SetContent(const TArray<uint8>& ContentPayload)
 {
+	if (CompletionStatus == EHttpRequestStatus::Processing)
+	{
+		UE_LOG(LogHttp, Warning, TEXT("FCurlHttpRequest::SetContent() - attempted to set content on a request that is inflight"));
+		return;
+	}
+
 	RequestPayload = MakeUnique<FRequestPayloadInMemory>(ContentPayload);
 	bIsRequestPayloadSeekable = true;
 }
 
 void FCurlHttpRequest::SetContentAsString(const FString& ContentString)
 {
+	if (CompletionStatus == EHttpRequestStatus::Processing)
+	{
+		UE_LOG(LogHttp, Warning, TEXT("FCurlHttpRequest::SetContentAsString() - attempted to set content on a request that is inflight"));
+		return;
+	}
+
 	FTCHARToUTF8 Converter(*ContentString);
 	TArray<uint8> Buffer;
 	Buffer.SetNum(Converter.Length());
@@ -349,11 +372,23 @@ bool FCurlHttpRequest::SetContentFromStream(TSharedRef<FArchive, ESPMode::Thread
 
 void FCurlHttpRequest::SetHeader(const FString& HeaderName, const FString& HeaderValue)
 {
+	if (CompletionStatus == EHttpRequestStatus::Processing)
+	{
+		UE_LOG(LogHttp, Warning, TEXT("FCurlHttpRequest::SetHeader() - attempted to set header on a request that is inflight"));
+		return;
+	}
+
 	Headers.Add(HeaderName, HeaderValue);
 }
 
 void FCurlHttpRequest::AppendToHeader(const FString& HeaderName, const FString& AdditionalHeaderValue)
 {
+	if (CompletionStatus == EHttpRequestStatus::Processing)
+	{
+		UE_LOG(LogHttp, Warning, TEXT("FCurlHttpRequest::AppendToHeader() - attempted to append to header on a request that is inflight"));
+		return;
+	}
+
 	if (!HeaderName.IsEmpty() && !AdditionalHeaderValue.IsEmpty())
 	{
 		FString* PreviousValue = Headers.Find(HeaderName);
@@ -691,7 +726,6 @@ size_t FCurlHttpRequest::DebugCallback(CURL * Handle, curl_infotype DebugInfoTyp
 	return 0;
 }
 
-
 bool FCurlHttpRequest::SetupRequest()
 {
 	QUICK_SCOPE_CYCLE_COUNTER(STAT_FCurlHttpRequest_SetupRequest);
@@ -707,12 +741,6 @@ bool FCurlHttpRequest::SetupRequest()
 	bCanceled = false;
 	CurlAddToMultiResult = CURLM_OK;
 
-	{
-		QUICK_SCOPE_CYCLE_COUNTER(STAT_FCurlHttpRequest_SetupRequest_SLIST_FREE_HEADERS);
-		curl_slist_free_all(HeaderList);
-		HeaderList = nullptr;
-	}	
-
 	// default no verb to a GET
 	if (Verb.IsEmpty())
 	{
@@ -742,7 +770,37 @@ bool FCurlHttpRequest::SetupRequest()
 	{
 		UE_LOG(LogHttp, Log, TEXT("Cannot process HTTP request: URL is empty"));
 		return false;
-	}	
+	}
+
+	if (GetHeader(TEXT("User-Agent")).IsEmpty())
+	{
+		SetHeader(TEXT("User-Agent"), FPlatformHttp::GetDefaultUserAgent());
+	}
+
+	// content-length should be present http://www.w3.org/Protocols/rfc2616/rfc2616-sec4.html#sec4.4
+	if (GetHeader(TEXT("Content-Length")).IsEmpty())
+	{
+		SetHeader(TEXT("Content-Length"), FString::Printf(TEXT("%d"), RequestPayload->GetContentLength()));
+	}
+
+	// Remove "Expect: 100-continue" since this is supposed to cause problematic behavior on Amazon ELB (and WinInet doesn't send that either)
+	// (also see http://www.iandennismiller.com/posts/curl-http1-1-100-continue-and-multipartform-data-post.html , http://the-stickman.com/web-development/php-and-curl-disabling-100-continue-header/ )
+	if (GetHeader(TEXT("Expect")).IsEmpty())
+	{
+		SetHeader(TEXT("Expect"), TEXT(""));
+	}
+
+	return true;
+}
+
+bool FCurlHttpRequest::SetupRequestHttpThread()
+{
+	QUICK_SCOPE_CYCLE_COUNTER(STAT_FCurlHttpRequest_SetupRequestHttpThread);
+	{
+		QUICK_SCOPE_CYCLE_COUNTER(STAT_FCurlHttpRequest_SetupRequest_SLIST_FREE_HEADERS);
+		curl_slist_free_all(HeaderList);
+		HeaderList = nullptr;
+	}
 
 	{
 		QUICK_SCOPE_CYCLE_COUNTER(STAT_FCurlHttpRequest_SetupRequest_EASY_SETOPT);
@@ -828,24 +886,6 @@ bool FCurlHttpRequest::SetupRequest()
 		{
 			curl_easy_setopt(EasyHandle, CURLOPT_ACCEPT_ENCODING, "");
 		}
-	
-		if (GetHeader(TEXT("User-Agent")).IsEmpty())
-		{
-			SetHeader(TEXT("User-Agent"), FPlatformHttp::GetDefaultUserAgent());
-		}
-
-		// content-length should be present http://www.w3.org/Protocols/rfc2616/rfc2616-sec4.html#sec4.4
-		if (GetHeader(TEXT("Content-Length")).IsEmpty())
-		{
-			SetHeader(TEXT("Content-Length"), FString::Printf(TEXT("%d"), RequestPayload->GetContentLength()));
-		}
-
-		// Remove "Expect: 100-continue" since this is supposed to cause problematic behavior on Amazon ELB (and WinInet doesn't send that either)
-		// (also see http://www.iandennismiller.com/posts/curl-http1-1-100-continue-and-multipartform-data-post.html , http://the-stickman.com/web-development/php-and-curl-disabling-100-continue-header/ )
-		if (GetHeader(TEXT("Expect")).IsEmpty())
-		{
-			SetHeader(TEXT("Expect"), TEXT(""));
-		}
 
 		{
 			QUICK_SCOPE_CYCLE_COUNTER(STAT_FCurlHttpRequest_SetupRequest_SLIST_APPEND_HEADERS);
@@ -908,6 +948,10 @@ bool FCurlHttpRequest::SetupRequest()
 	curl_easy_setopt(EasyHandle, CURLOPT_SHARE, FCurlHttpManager::GShareHandle);
 
 	UE_LOG(LogHttp, Log, TEXT("%p: Starting %s request to URL='%s'"), this, *Verb, *URL);
+
+	// Response object to handle data that comes back after starting this request
+	Response = MakeShared<FCurlHttpResponse, ESPMode::ThreadSafe>(*this);
+
 	return true;
 }
 
@@ -929,7 +973,7 @@ bool FCurlHttpRequest::ProcessRequest()
 	}
 	else if (!SetupRequest())
 	{
-		UE_LOG(LogHttp, Warning, TEXT("Could not set libcurl options for easy handle, processing HTTP request failed. Increase verbosity for additional information."));
+		UE_LOG(LogHttp, Warning, TEXT("Could not perform game thread setup, processing HTTP request failed. Increase verbosity for additional information."));
 	}
 	else
 	{
@@ -948,8 +992,6 @@ bool FCurlHttpRequest::ProcessRequest()
 		QUICK_SCOPE_CYCLE_COUNTER(STAT_CurlHttpAddThreadedRequest);
 		// Mark as in-flight to prevent overlapped requests using the same object
 		CompletionStatus = EHttpRequestStatus::Processing;
-		// Response object to handle data that comes back after starting this request
-		Response = MakeShareable(new FCurlHttpResponse(*this));
 		// Add to global list while being processed so that the ref counted request does not get deleted
 		FHttpModule::Get().GetHttpManager().AddThreadedRequest(SharedThis(this));
 
diff --git a/Engine/Source/Runtime/Online/HTTP/Private/Curl/CurlHttp.h b/Engine/Source/Runtime/Online/HTTP/Private/Curl/CurlHttp.h
--- a/Engine/Source/Runtime/Online/HTTP/Private/Curl/CurlHttp.h
+++ b/Engine/Source/Runtime/Online/HTTP/Private/Curl/CurlHttp.h
@@ -160,6 +160,13 @@ public:
 	virtual void TickThreadedRequest(float DeltaSeconds) override;
 	//~ End IHttpRequestThreaded Interface
 
+	/**
+	 * Perform the http-thread setup of the request
+	 *
+	 * @return true if the request was successfully setup
+	 */
+	bool SetupRequestHttpThread();
+
 	/**
 	 * Returns libcurl's easy handle - needed for HTTP manager.
 	 *
@@ -312,7 +319,7 @@ private:
 	size_t DebugCallback(CURL * Handle, curl_infotype DebugInfoType, char * DebugInfo, size_t DebugInfoSize);
 
 	/**
-	 * Setup the request
+	 * Perform the game-thread setup of the request
 	 *
 	 * @return true if the request was successfully setup
 	 */
diff --git a/Engine/Source/Runtime/Online/HTTP/Private/Curl/CurlHttpThread.cpp b/Engine/Source/Runtime/Online/HTTP/Private/Curl/CurlHttpThread.cpp
--- a/Engine/Source/Runtime/Online/HTTP/Private/Curl/CurlHttpThread.cpp
+++ b/Engine/Source/Runtime/Online/HTTP/Private/Curl/CurlHttpThread.cpp
@@ -17,52 +17,48 @@ void FCurlHttpThread::HttpThreadTick(float DeltaSeconds)
 	QUICK_SCOPE_CYCLE_COUNTER(STAT_FCurlHttpThread_HttpThreadTick);
 	check(FCurlHttpManager::GMultiHandle);
 
-	const int NumRequestsToTick = RunningThreadedRequests.Num();
-
+	if (RunningThreadedRequests.Num() > 0)
 	{
-		if (RunningThreadedRequests.Num() > 0)
+		int RunningRequests = -1;
 		{
-			int RunningRequests = -1;
-			{
-				QUICK_SCOPE_CYCLE_COUNTER(STAT_FCurlHttpThread_HttpThreadTick_Perform);
-				curl_multi_perform(FCurlHttpManager::GMultiHandle, &RunningRequests);
-			}
+			QUICK_SCOPE_CYCLE_COUNTER(STAT_FCurlHttpThread_HttpThreadTick_Perform);
+			curl_multi_perform(FCurlHttpManager::GMultiHandle, &RunningRequests);
+		}
 
-			// read more info if number of requests changed or if there's zero running
-			// (note that some requests might have never be "running" from libcurl's point of view)
-			if (RunningRequests == 0 || RunningRequests != RunningThreadedRequests.Num())
+		// read more info if number of requests changed or if there's zero running
+		// (note that some requests might have never be "running" from libcurl's point of view)
+		if (RunningRequests == 0 || RunningRequests != RunningThreadedRequests.Num())
+		{
+			for (;;)
 			{
-				for (;;)
+				QUICK_SCOPE_CYCLE_COUNTER(STAT_FCurlHttpThread_HttpThreadTick_Loop);
+				int MsgsStillInQueue = 0;	// may use that to impose some upper limit we may spend in that loop
+				CURLMsg * Message = curl_multi_info_read(FCurlHttpManager::GMultiHandle, &MsgsStillInQueue);
+
+				if (Message == NULL)
+				{
+					break;
+				}
+
+				// find out which requests have completed
+				if (Message->msg == CURLMSG_DONE)
 				{
-					QUICK_SCOPE_CYCLE_COUNTER(STAT_FCurlHttpThread_HttpThreadTick_Loop);
-					int MsgsStillInQueue = 0;	// may use that to impose some upper limit we may spend in that loop
-					CURLMsg * Message = curl_multi_info_read(FCurlHttpManager::GMultiHandle, &MsgsStillInQueue);
+					CURL* CompletedHandle = Message->easy_handle;
+					curl_multi_remove_handle(FCurlHttpManager::GMultiHandle, CompletedHandle);
 
-					if (Message == NULL)
+					IHttpThreadedRequest** Request = HandlesToRequests.Find(CompletedHandle);
+					if (Request)
 					{
-						break;
-					}
+						FCurlHttpRequest* CurlRequest = static_cast<FCurlHttpRequest*>(*Request);
+						CurlRequest->MarkAsCompleted(Message->data.result);
 
-					// find out which requests have completed
-					if (Message->msg == CURLMSG_DONE)
+						UE_LOG(LogHttp, Verbose, TEXT("Request %p (easy handle:%p) has completed (code:%d) and has been marked as such"), CurlRequest, CompletedHandle, (int32)Message->data.result);
+
+						HandlesToRequests.Remove(CompletedHandle);
+					}
+					else
 					{
-						CURL* CompletedHandle = Message->easy_handle;
-						curl_multi_remove_handle(FCurlHttpManager::GMultiHandle, CompletedHandle);
-
-						IHttpThreadedRequest** Request = HandlesToRequests.Find(CompletedHandle);
-						if (Request)
-						{
-							FCurlHttpRequest* CurlRequest = static_cast< FCurlHttpRequest* >(*Request);
-							CurlRequest->MarkAsCompleted(Message->data.result);
-
-							UE_LOG(LogHttp, Verbose, TEXT("Request %p (easy handle:%p) has completed (code:%d) and has been marked as such"), CurlRequest, CompletedHandle, (int32)Message->data.result);
-
-							HandlesToRequests.Remove(CompletedHandle);
-						}
-						else
-						{
-							UE_LOG(LogHttp, Warning, TEXT("Could not find mapping for completed request (easy handle: %p)"), CompletedHandle);
-						}
+						UE_LOG(LogHttp, Warning, TEXT("Could not find mapping for completed request (easy handle: %p)"), CompletedHandle);
 					}
 				}
 			}
@@ -78,6 +74,12 @@ bool FCurlHttpThread::StartThreadedRequest(IHttpThreadedRequest* Request)
 	CURL* EasyHandle = CurlRequest->GetEasyHandle();
 	ensure(!HandlesToRequests.Contains(EasyHandle));
 
+	if (!CurlRequest->SetupRequestHttpThread())
+	{
+		UE_LOG(LogHttp, Warning, TEXT("Could not set libcurl options for easy handle, processing HTTP request failed. Increase verbosity for additional information."));
+		return false;
+	}
+
 	CURLMcode AddResult = curl_multi_add_handle(FCurlHttpManager::GMultiHandle, EasyHandle);
 	CurlRequest->SetAddToCurlMultiResult(AddResult);
 
