- Some minor additions to HTTP stat captures

diff --git a/Engine/Source/Runtime/Online/HTTP/Private/GenericPlatform/HttpRequestImpl.cpp b/Engine/Source/Runtime/Online/HTTP/Private/GenericPlatform/HttpRequestImpl.cpp
--- a/Engine/Source/Runtime/Online/HTTP/Private/GenericPlatform/HttpRequestImpl.cpp
+++ b/Engine/Source/Runtime/Online/HTTP/Private/GenericPlatform/HttpRequestImpl.cpp
@@ -43,6 +43,8 @@ void FHttpRequestImpl::BroadcastResponseHeadersReceived()
 				if (Header.Split(TEXT(":"), &HeaderName, &HeaderValue))
 				{
 					HeaderValue.TrimStartInline();
+
+					QUICK_SCOPE_CYCLE_COUNTER(STAT_FHttpRequestImpl_BroadcastResponseHeadersReceived_OnHeaderReceived);
 					OnHeaderReceived().ExecuteIfBound(ThisPtr, HeaderName, HeaderValue);
 				}
 			}
diff --git a/Engine/Source/Runtime/Online/HTTP/Private/HttpRetrySystem.cpp b/Engine/Source/Runtime/Online/HTTP/Private/HttpRetrySystem.cpp
--- a/Engine/Source/Runtime/Online/HTTP/Private/HttpRetrySystem.cpp
+++ b/Engine/Source/Runtime/Online/HTTP/Private/HttpRetrySystem.cpp
@@ -424,6 +424,7 @@ bool FHttpRetrySystem::FManager::Update(uint32* FileCount, uint32* FailingCount,
 							HttpRetryRequestEntry.LockoutEndTimeAbsoluteSeconds = NowAbsoluteSeconds + LockoutPeriod;
 							HttpRetryRequest->Status = FHttpRetrySystem::FRequest::EStatus::ProcessingLockout;
 							
+							QUICK_SCOPE_CYCLE_COUNTER(STAT_FHttpRetrySystem_FManager_Update_OnRequestWillRetry);
 							HttpRetryRequest->OnRequestWillRetry().ExecuteIfBound(HttpRetryRequest, HttpRetryRequest->GetResponse(), LockoutPeriod);
 						}
 						else
