- Fix non unity build
- Additional stat in FManager::Update OnProcessRequestComplete

diff --git a/Engine/Source/Runtime/Online/HTTP/Private/GenericPlatform/HttpRequestImpl.cpp b/Engine/Source/Runtime/Online/HTTP/Private/GenericPlatform/HttpRequestImpl.cpp
--- a/Engine/Source/Runtime/Online/HTTP/Private/GenericPlatform/HttpRequestImpl.cpp
+++ b/Engine/Source/Runtime/Online/HTTP/Private/GenericPlatform/HttpRequestImpl.cpp
@@ -1,6 +1,7 @@
 // Copyright Epic Games, Inc. All Rights Reserved.
 
 #include "GenericPlatform/HttpRequestImpl.h"
+#include "Stats/Stats.h"
 #include "Http.h"
 
 FHttpRequestCompleteDelegate& FHttpRequestImpl::OnProcessRequestComplete()
diff --git a/Engine/Source/Runtime/Online/HTTP/Private/HttpRetrySystem.cpp b/Engine/Source/Runtime/Online/HTTP/Private/HttpRetrySystem.cpp
--- a/Engine/Source/Runtime/Online/HTTP/Private/HttpRetrySystem.cpp
+++ b/Engine/Source/Runtime/Online/HTTP/Private/HttpRetrySystem.cpp
@@ -504,6 +504,8 @@ bool FHttpRetrySystem::FManager::Update(uint32* FileCount, uint32* FailingCount,
 			{
 				HttpRetryRequest->BroadcastResponseHeadersReceived();
 			}
+			
+			QUICK_SCOPE_CYCLE_COUNTER(STAT_FHttpRetrySystem_FManager_Update_OnProcessRequestComplete);
 			HttpRetryRequest->OnProcessRequestComplete().ExecuteIfBound(HttpRetryRequest, HttpRetryRequest->GetResponse(), bWasSuccessful);
 		}
 
