- Add scope cycle counter to FManager Update to show clearly how many items are being iterated, and how long they each take.

diff --git a/Engine/Source/Runtime/Online/HTTP/Private/HttpRetrySystem.cpp b/Engine/Source/Runtime/Online/HTTP/Private/HttpRetrySystem.cpp
--- a/Engine/Source/Runtime/Online/HTTP/Private/HttpRetrySystem.cpp
+++ b/Engine/Source/Runtime/Online/HTTP/Private/HttpRetrySystem.cpp
@@ -355,6 +355,8 @@ bool FHttpRetrySystem::FManager::Update(uint32* FileCount, uint32* FailingCount,
 	int32 index = 0;
 	while (index < RequestList.Num())
 	{
+		QUICK_SCOPE_CYCLE_COUNTER(STAT_FHttpRetrySystem_FManager_Update_RequestListItem);
+
 		FHttpRetryRequestEntry& HttpRetryRequestEntry = RequestList[index];
 		TSharedRef<FHttpRetrySystem::FRequest>& HttpRetryRequest = HttpRetryRequestEntry.Request;
 
