- Remove unnecessary copies in SetContentAsString

diff --git a/Engine/Source/Runtime/Online/HTTP/Private/Curl/CurlHttp.cpp b/Engine/Source/Runtime/Online/HTTP/Private/Curl/CurlHttp.cpp
--- a/Engine/Source/Runtime/Online/HTTP/Private/Curl/CurlHttp.cpp
+++ b/Engine/Source/Runtime/Online/HTTP/Private/Curl/CurlHttp.cpp
@@ -323,12 +323,12 @@ void FCurlHttpRequest::SetContentAsString(const FString& ContentString)
 		return;
 	}
 
-	FTCHARToUTF8 Converter(*ContentString);
+	int32 Utf8Length = FTCHARToUTF8_Convert::ConvertedLength(*ContentString, ContentString.Len());
 	TArray<uint8> Buffer;
-	Buffer.SetNum(Converter.Length());
-	FMemory::Memcpy(Buffer.GetData(), Converter.Get(), Buffer.Num());
-	bIsRequestPayloadSeekable = true;
+	Buffer.SetNumUninitialized(Utf8Length);
+	FTCHARToUTF8_Convert::Convert((ANSICHAR*)Buffer.GetData(), Buffer.Num(), *ContentString, ContentString.Len());
 	RequestPayload = MakeUnique<FRequestPayloadInMemory>(MoveTemp(Buffer));
+	bIsRequestPayloadSeekable = true;
 }
 
 bool FCurlHttpRequest::SetContentAsStreamedFile(const FString& Filename)
diff --git a/Engine/Source/Runtime/Online/HTTP/Private/IXML/HttpIXml.cpp b/Engine/Source/Runtime/Online/HTTP/Private/IXML/HttpIXml.cpp
--- a/Engine/Source/Runtime/Online/HTTP/Private/IXML/HttpIXml.cpp
+++ b/Engine/Source/Runtime/Online/HTTP/Private/IXML/HttpIXml.cpp
@@ -129,11 +129,11 @@ void FHttpRequestIXML::SetContentAsString(const FString& ContentString)
 {
 	if ( ContentString.Len() )
 	{
-		FTCHARToUTF8 Converter(*ContentString);
+		int32 Utf8Length = FTCHARToUTF8_Convert::ConvertedLength(*ContentString, ContentString.Len());
 		TArray<uint8> Buffer;
-		Buffer.SetNum(Converter.Length());
-		FMemory::Memcpy(Buffer.GetData(), Converter.Get(), Buffer.Num());
-		Payload = MakeUnique<FRequestPayloadInMemory>(Buffer);
+		Buffer.SetNumUninitialized(Utf8Length);
+		FTCHARToUTF8_Convert::Convert((ANSICHAR*)Buffer.GetData(), Buffer.Num(), *ContentString, ContentString.Len());
+		Payload = MakeUnique<FRequestPayloadInMemory>(MoveTemp(Buffer));
 	}
 }
 
diff --git a/Engine/Source/Runtime/Online/HTTP/Private/NullHttp.cpp b/Engine/Source/Runtime/Online/HTTP/Private/NullHttp.cpp
--- a/Engine/Source/Runtime/Online/HTTP/Private/NullHttp.cpp
+++ b/Engine/Source/Runtime/Online/HTTP/Private/NullHttp.cpp
@@ -74,9 +74,9 @@ void FNullHttpRequest::SetContent(const TArray<uint8>& ContentPayload)
 
 void FNullHttpRequest::SetContentAsString(const FString& ContentString)
 {
-	FTCHARToUTF8 Converter(*ContentString);
-	Payload.SetNum(Converter.Length());
-	FMemory::Memcpy(Payload.GetData(), (uint8*)(ANSICHAR*)Converter.Get(), Payload.Num());
+	int32 Utf8Length = FTCHARToUTF8_Convert::ConvertedLength(*ContentString, ContentString.Len());
+	Payload.SetNumUninitialized(Utf8Length);
+	FTCHARToUTF8_Convert::Convert((ANSICHAR*)Payload.GetData(), Payload.Num(), *ContentString, ContentString.Len());
 }
 
 bool FNullHttpRequest::SetContentAsStreamedFile(const FString& Filename)
