[HttpServer] HttpServerConfig yields a mechanism for users to customize http listeners on a per-port basis

+ Configuration resolution is hierarchical and processed in order:  Per-Port Configuration, Default Configuration, Hard-Coded configuration

+ Per-port overriders can be declared with:
[HttpServer.Listeners]
+ListenerOverrides=(Port=..., BindAddress=...., BufferSize=..., ConnectionsBacklogSize=..., MaxConnectionsAcceptPerFrame=...)

+Default overrides can be declared with
[HttpServer.Listeners]
DefaultBindAddress=...
DefaultBufferSize=...
DefaultConnectionsBacklogSize=...
DefaultMaxConnectionsAcceptPerFrame=...

+ Lastly, if no per-port/default overrides are set, the in-code defaults are the last fallback

diff --git a/Engine/Source/Runtime/Online/HTTPServer/Private/HttpListener.cpp b/Engine/Source/Runtime/Online/HTTPServer/Private/HttpListener.cpp
--- a/Engine/Source/Runtime/Online/HTTPServer/Private/HttpListener.cpp
+++ b/Engine/Source/Runtime/Online/HTTPServer/Private/HttpListener.cpp
@@ -4,21 +4,18 @@
 #include "HttpConnection.h"
 #include "HttpRequestHandler.h"
 #include "HttpRouter.h"
+#include "HttpServerConfig.h"
 
 #include "Sockets.h"
 #include "SocketSubsystem.h"
 #include "IPAddress.h"
 
-constexpr uint32 FHttpListener::MaxConnectionsToAcceptPerFrame;
-constexpr uint32 FHttpListener::ListenerConnectionBacklogSize;
-constexpr uint32 FHttpListener::ListenerBufferSize;
-
 DEFINE_LOG_CATEGORY(LogHttpListener)
 
-FHttpListener::FHttpListener(uint32 Port)
+FHttpListener::FHttpListener(uint32 InListenPort)
 { 
-	check(Port > 0);
-	ListenPort = Port;
+	check(InListenPort > 0);
+	ListenPort = InListenPort;
 	Router = MakeShared<FHttpRouter>();
 }
 
@@ -61,28 +58,45 @@ bool FHttpListener::StartListening()
 	}
 	ListenSocket->SetNonBlocking(true);
 
-	// Bind to Localhost/Caller-defined port
-	TSharedRef<FInternetAddr> LocalhostAddr = SocketSubsystem->CreateInternetAddr();
-	LocalhostAddr->SetAnyAddress();
-	LocalhostAddr->SetPort(ListenPort);
-	if (!ListenSocket->Bind(*LocalhostAddr))
+	// Bind to config-driven address
+	TSharedRef<FInternetAddr> BindAddress = SocketSubsystem->CreateInternetAddr();
+	Config = FHttpServerConfig::GetListenerConfig(ListenPort);
+	if (0 == Config.BindAddress.Compare(TEXT("any"), ESearchCase::IgnoreCase))
+	{
+		BindAddress->SetAnyAddress();
+	}
+	else
+	{
+		bool bIsValidAddress = false;
+		BindAddress->SetIp(*(Config.BindAddress), bIsValidAddress);
+		if (!bIsValidAddress)
+		{
+			UE_LOG(LogHttpListener, Error,
+				TEXT("HttpListener detected invalid bind address (%s:%u)"),
+				*Config.BindAddress, ListenPort);
+			return false;
+		}
+	}
+
+	BindAddress->SetPort(ListenPort);
+	if (!ListenSocket->Bind(*BindAddress))
 	{
 		UE_LOG(LogHttpListener, Error, 
-			TEXT("HttpListener unable to bind to %s"), 
-			*LocalhostAddr->ToString(true));
+			TEXT("HttpListener unable to bind to %s:%u"),
+			*BindAddress->ToString(true), ListenPort);
 		return false;
 	}
 
 	int32 ActualBufferSize;
-	ListenSocket->SetSendBufferSize(ListenerBufferSize, ActualBufferSize);
-	if (ActualBufferSize != ListenerBufferSize)
+	ListenSocket->SetSendBufferSize(Config.BufferSize, ActualBufferSize);
+	if (ActualBufferSize != Config.BufferSize)
 	{
 		UE_LOG(LogHttpListener, Warning, 
 			TEXT("HttpListener unable to set desired buffer size (%d): Limited to %d"),
-			ListenerBufferSize, ActualBufferSize);
+			Config.BufferSize, ActualBufferSize);
 	}
 
-	if (!ListenSocket->Listen(ListenerConnectionBacklogSize))
+	if (!ListenSocket->Listen(Config.ConnectionsBacklogSize))
 	{
 		UE_LOG(LogHttpListener, Error, 
 			TEXT("HttpListener unable to listen on socket"));
@@ -90,7 +104,8 @@ bool FHttpListener::StartListening()
 	}
 
 	UE_LOG(LogHttpListener, Log, 
-		TEXT("Created new HttpListener on port %u"), ListenPort);
+		TEXT("Created new HttpListener on %s:%u"), 
+		*BindAddress->ToString(true), ListenPort);
 	return true;
 }
 
@@ -123,7 +138,7 @@ void FHttpListener::StopListening()
 void FHttpListener::Tick(float DeltaTime)
 {
 	// Accept new connections
-	AcceptConnections(MaxConnectionsToAcceptPerFrame);
+	AcceptConnections();
 
 	// Tick Connections
 	TickConnections(DeltaTime);
@@ -150,11 +165,11 @@ bool FHttpListener::HasPendingConnections() const
 // --------------------------------------------------------------------------------------------
 // Private Implementation
 // --------------------------------------------------------------------------------------------
-void FHttpListener::AcceptConnections(uint32 MaxConnectionsToAccept)
+void FHttpListener::AcceptConnections()
 {
 	check(ListenSocket);
 
-	for (uint32 i = 0; i < MaxConnectionsToAccept; ++i)
+	for (int32 i = 0; i < Config.MaxConnectionsAcceptPerFrame; ++i)
 	{
 		// Check pending prior to Accept()ing
 		bool bHasPendingConnection = false;
diff --git a/Engine/Source/Runtime/Online/HTTPServer/Private/HttpListener.h b/Engine/Source/Runtime/Online/HTTPServer/Private/HttpListener.h
--- a/Engine/Source/Runtime/Online/HTTPServer/Private/HttpListener.h
+++ b/Engine/Source/Runtime/Online/HTTPServer/Private/HttpListener.h
@@ -3,6 +3,7 @@
 
 #include "HttpConnectionTypes.h"
 #include "HttpRouter.h"
+#include "HttpServerConfig.h"
 
 struct FHttpConnection;
 struct FHttpPath;
@@ -18,9 +19,9 @@ public:
 	/**
 	 * Constructor
 	 *
-	 *  @param ListenPort The port on which to listen for incoming connections
+	 *  @param InListenPort The port on which to listen for incoming connections
 	 */
-	FHttpListener(uint32 ListenPort);
+	FHttpListener(uint32 InListenPort);
 
 	/**
 	 * Destructor
@@ -29,7 +30,7 @@ public:
 
 	/**
 	 * Starts listening for and accepting incoming connections
-	 *
+	 * 
 	 * @return true if the listener was able to start listening, false otherwise
 	 */
 	bool StartListening();
@@ -74,11 +75,9 @@ public:
 private:
 
 	/**
-	 * Accepts a single available connection
-	 * 
-	 * @param MaxConnectionsToAccept The maximum number of connections to accept
+	 * Accepts available connection(s)
 	 */
-	void AcceptConnections(uint32 MaxConnectionsToAccept);
+	void AcceptConnections();
 
 	/**
 	 * Ticks connections in reading/writing phases
@@ -112,10 +111,6 @@ private:
 	/** The total number of connections accepted by this listener */
 	uint32 NumConnectionsAccepted = 0;
 
-	/** Maximum number of connections to accept per frame */
-	static constexpr uint32 MaxConnectionsToAcceptPerFrame = 1;
-	/** Maximum number of pending connections to queue */
-	static constexpr uint32 ListenerConnectionBacklogSize = 16;
-	/** Maximum send buffer size */
-	static constexpr uint32 ListenerBufferSize = 512 * 1024;
+	/** Listener configuration data */
+	FHttpServerListenerConfig Config;
 };
\ No newline at end of file
diff --git a/Engine/Source/Runtime/Online/HTTPServer/Private/HttpServerConfig.cpp b/Engine/Source/Runtime/Online/HTTPServer/Private/HttpServerConfig.cpp
new file mode 100644
--- /dev/null
+++ b/Engine/Source/Runtime/Online/HTTPServer/Private/HttpServerConfig.cpp
@@ -0,0 +1,53 @@
+// Copyright Epic Games, Inc. All Rights Reserved.
+
+#include "HttpServerConfig.h"
+#include "CoreMinimal.h"
+
+DEFINE_LOG_CATEGORY(LogHttpServerConfig)
+
+const FHttpServerListenerConfig FHttpServerConfig::GetListenerConfig(uint32 Port) 
+{
+	static const FString IniSectionName(TEXT("HTTPServer.Listeners"));
+
+	// Code default values
+	FHttpServerListenerConfig Config;
+
+	// Apply default ini configuration
+	GConfig->GetString(*IniSectionName, TEXT("DefaultBindAddress"), Config.BindAddress, GEngineIni);
+	GConfig->GetInt(*IniSectionName, TEXT("DefaultBufferSize"), Config.BufferSize, GEngineIni);
+	GConfig->GetInt(*IniSectionName, TEXT("DefaultConnectionsBacklogSize"), Config.ConnectionsBacklogSize, GEngineIni);
+	GConfig->GetInt(*IniSectionName, TEXT("DefaultMaxConnectionsAcceptPerFrame"), Config.MaxConnectionsAcceptPerFrame, GEngineIni);
+
+	// Apply per-port ini overrides
+	TArray<FString> ListenerConfigs;
+	if (GConfig->GetArray(*IniSectionName, TEXT("ListenerOverrides"), ListenerConfigs, GEngineIni))
+	{
+		for (FString ListenerConfigStr : ListenerConfigs)
+		{
+			ListenerConfigStr.TrimStartAndEndInline();
+			ListenerConfigStr.ReplaceInline(TEXT("("), TEXT(""));
+			ListenerConfigStr.ReplaceInline(TEXT(")"), TEXT(""));
+
+			// Listener config overrides must specify a port
+			uint32 ConfiguredPort = 0;
+			if (!FParse::Value(*ListenerConfigStr, TEXT("Port="), ConfiguredPort))
+			{
+				UE_LOG(LogHttpServerConfig, Error,
+					TEXT("ListenerOverride: %s does not specify required Port parameter"),
+					*ListenerConfigStr);
+				continue;
+			}
+
+			if (Port == ConfiguredPort)
+			{
+				// override defaults with config values
+				FParse::Value(*ListenerConfigStr, TEXT("BindAddress="), Config.BindAddress);
+				FParse::Value(*ListenerConfigStr, TEXT("BufferSize="), Config.BufferSize);
+				FParse::Value(*ListenerConfigStr, TEXT("ConnectionsBacklogSize="), Config.ConnectionsBacklogSize);
+				FParse::Value(*ListenerConfigStr, TEXT("MaxConnectionsAcceptPerFrame="), Config.MaxConnectionsAcceptPerFrame);
+				break;
+			}
+		}
+	}
+	return Config;
+}
diff --git a/Engine/Source/Runtime/Online/HTTPServer/Private/HttpServerConfig.h b/Engine/Source/Runtime/Online/HTTPServer/Private/HttpServerConfig.h
new file mode 100644
--- /dev/null
+++ b/Engine/Source/Runtime/Online/HTTPServer/Private/HttpServerConfig.h
@@ -0,0 +1,30 @@
+// Copyright Epic Games, Inc. All Rights Reserved.
+#pragma once
+
+#include "CoreMinimal.h"
+#include "Containers/Map.h"
+
+DECLARE_LOG_CATEGORY_EXTERN(LogHttpServerConfig, Log, All);
+
+struct FHttpServerListenerConfig
+{
+	/** The address to bind */
+	FString BindAddress = FString(TEXT("127.0.0.1"));
+	/** Maximum send buffer size */
+	int32 BufferSize = 512 * 1024;
+	/** Number of pending connections to queue */
+	int32 ConnectionsBacklogSize = 16;
+	/** Max Number of connections to accept per frame */
+	int32 MaxConnectionsAcceptPerFrame = 1;
+};
+
+struct FHttpServerConfig
+{
+	/** 
+	* Gets the listener configuration for the caller-supplied Port
+	*
+	* @param Port The respective listener port
+	* @return The per-port configuration if configured, or default configuration
+	*/
+	static const FHttpServerListenerConfig GetListenerConfig(uint32 Port);
+};
\ No newline at end of file
