[HttpServer] CIS fix
+ explicitly include "Misc/ConfigCacheIni.h"

diff --git a/Engine/Source/Runtime/Online/HTTPServer/Private/HttpServerConfig.cpp b/Engine/Source/Runtime/Online/HTTPServer/Private/HttpServerConfig.cpp
--- a/Engine/Source/Runtime/Online/HTTPServer/Private/HttpServerConfig.cpp
+++ b/Engine/Source/Runtime/Online/HTTPServer/Private/HttpServerConfig.cpp
@@ -2,6 +2,7 @@
 
 #include "HttpServerConfig.h"
 #include "CoreMinimal.h"
+#include "Misc/ConfigCacheIni.h"
 
 DEFINE_LOG_CATEGORY(LogHttpServerConfig)
 
