Multithread Servers
* Allow HTTPManager to run on a autonomous thread on forked processes
   * That option can be turned off via
HTTPManager::SetAllowThreadsOnForkedProcess

#Moved HTTPManager::OnAfterFork call from right after Fork occurred to the end of the Frame where the Fork occured.

diff --git a/Engine/Source/Runtime/Online/HTTP/Private/HttpManager.cpp b/Engine/Source/Runtime/Online/HTTP/Private/HttpManager.cpp
--- a/Engine/Source/Runtime/Online/HTTP/Private/HttpManager.cpp
+++ b/Engine/Source/Runtime/Online/HTTP/Private/HttpManager.cpp
@@ -109,6 +109,17 @@ void FHttpManager::OnAfterFork()
 
 }
 
+void FHttpManager::SetAllowThreadsOnForkedProcess(bool bAllow)
+{
+	if (FHttpThread::bAllowForkableThreads != bAllow)
+	{
+		if (ensureMsgf(Thread == nullptr || Thread->IsStopped(), TEXT("Only modify the forked process thread behavior when the HTTPThread does not exist")))
+		{
+			FHttpThread::bAllowForkableThreads = bAllow;
+		}
+	}
+}
+
 void FHttpManager::UpdateConfigs()
 {
 	// empty
@@ -168,14 +179,17 @@ void FHttpManager::Flush(bool bShutdown)
 		LastTime = AppTime;
 		if (Requests.Num() > 0)
 		{
-			if (FPlatformProcess::SupportsMultithreading())
+			if (Thread)
 			{
-				UE_LOG(LogHttp, Display, TEXT("Sleeping 0.5s to wait for %d outstanding Http requests."), Requests.Num());
-				FPlatformProcess::Sleep(0.5f);
-			}
-			else if (Thread)
-			{
-				Thread->Tick();
+				if( Thread->NeedsSingleThreadTick() )
+				{
+					Thread->Tick();
+				}
+				else
+				{
+					UE_LOG(LogHttp, Display, TEXT("Sleeping 0.5s to wait for %d outstanding Http requests."), Requests.Num());
+					FPlatformProcess::Sleep(0.5f);
+				}
 			}
 			else
 			{
diff --git a/Engine/Source/Runtime/Online/HTTP/Private/HttpThread.cpp b/Engine/Source/Runtime/Online/HTTP/Private/HttpThread.cpp
--- a/Engine/Source/Runtime/Online/HTTP/Private/HttpThread.cpp
+++ b/Engine/Source/Runtime/Online/HTTP/Private/HttpThread.cpp
@@ -11,8 +11,12 @@
 
 // FHttpThread
 
+bool FHttpThread::bAllowForkableThreads = true;
+
 FHttpThread::FHttpThread()
 	:	Thread(nullptr)
+	,	bNeedsSingleThreadTick(false)
+	,	bIsStopped(true)
 {
 	HttpThreadActiveFrameTimeInSeconds = FHttpModule::Get().GetHttpThreadActiveFrameTimeInSeconds();
 	HttpThreadActiveMinimumSleepTimeInSeconds = FHttpModule::Get().GetHttpThreadActiveMinimumSleepTimeInSeconds();
@@ -29,7 +33,25 @@ FHttpThread::~FHttpThread()
 
 void FHttpThread::StartThread()
 {
-	Thread = FRunnableThread::Create(this, TEXT("HttpManagerThread"), 128 * 1024, TPri_Normal);
+	bNeedsSingleThreadTick = false;
+
+	if (bAllowForkableThreads && FForkProcessHelper::IsForkedMultithreadInstance() )
+	{
+		// We only create forkable threads on the forked instance since the HTTPManager cannot safely transition from fake to real seamlessly
+		Thread = FForkProcessHelper::CreateForkableThread(this, TEXT("HttpManagerThread"), 128 * 1024, TPri_Normal);
+	}
+	else
+	{
+		Thread = FRunnableThread::Create(this, TEXT("HttpManagerThread"), 128 * 1024, TPri_Normal);
+
+		// If the runnable thread is fake.
+		if (FGenericPlatformProcess::SupportsMultithreading() == false)
+		{
+			bNeedsSingleThreadTick = true;
+		}
+	}
+
+	bIsStopped = false;
 }
 
 void FHttpThread::StopThread()
@@ -40,6 +62,8 @@ void FHttpThread::StopThread()
 		delete Thread;
 		Thread = nullptr;
 	}
+
+	bIsStopped = true;
 }
 
 void FHttpThread::AddRequest(IHttpThreadedRequest* Request)
@@ -65,6 +89,7 @@ void FHttpThread::GetCompletedRequests(TArray<IHttpThreadedRequest*>& OutComplet
 bool FHttpThread::Init()
 {
 	LastTime = FPlatformTime::Seconds();
+	ExitRequest.Set(false);
 	return true;
 }
 
@@ -76,35 +101,43 @@ uint32 FHttpThread::Run()
 	TArray<IHttpThreadedRequest*> RequestsToComplete;
 	while (!ExitRequest.GetValue())
 	{
-		const double OuterLoopBegin = FPlatformTime::Seconds();
-		double OuterLoopEnd = 0.0;
-		bool bKeepProcessing = true;
-		while (bKeepProcessing)
+		const bool bIsMultithreaded = NeedsSingleThreadTick() == false;
+		if (ensureMsgf(bIsMultithreaded, TEXT("HTTP Thread was set to singlethread mode while it was running autonomously!")))
 		{
-			const double InnerLoopBegin = FPlatformTime::Seconds();
+			const double OuterLoopBegin = FPlatformTime::Seconds();
+			double OuterLoopEnd = 0.0;
+			bool bKeepProcessing = true;
+			while (bKeepProcessing)
+			{
+				const double InnerLoopBegin = FPlatformTime::Seconds();
 			
-			Process(RequestsToCancel, RequestsToStart, RequestsToComplete);
+				Process(RequestsToCancel, RequestsToStart, RequestsToComplete);
 			
-			if (RunningThreadedRequests.Num() == 0)
-			{
-				bKeepProcessing = false;
-			}
+				if (RunningThreadedRequests.Num() == 0)
+				{
+					bKeepProcessing = false;
+				}
 
-			const double InnerLoopEnd = FPlatformTime::Seconds();
-			if (bKeepProcessing)
-			{
-				double InnerLoopTime = InnerLoopEnd - InnerLoopBegin;
-				double InnerSleep = FMath::Max(HttpThreadActiveFrameTimeInSeconds - InnerLoopTime, HttpThreadActiveMinimumSleepTimeInSeconds);
-				FPlatformProcess::SleepNoStats(InnerSleep);
-			}
-			else
-			{
-				OuterLoopEnd = InnerLoopEnd;
+				const double InnerLoopEnd = FPlatformTime::Seconds();
+				if (bKeepProcessing)
+				{
+					double InnerLoopTime = InnerLoopEnd - InnerLoopBegin;
+					double InnerSleep = FMath::Max(HttpThreadActiveFrameTimeInSeconds - InnerLoopTime, HttpThreadActiveMinimumSleepTimeInSeconds);
+					FPlatformProcess::SleepNoStats(InnerSleep);
+				}
+				else
+				{
+					OuterLoopEnd = InnerLoopEnd;
+				}
 			}
+			double OuterLoopTime = OuterLoopEnd - OuterLoopBegin;
+			double OuterSleep = FMath::Max(HttpThreadIdleFrameTimeInSeconds - OuterLoopTime, HttpThreadIdleMinimumSleepTimeInSeconds);
+			FPlatformProcess::SleepNoStats(OuterSleep);
+		}
+		else
+		{
+			break;
 		}
-		double OuterLoopTime = OuterLoopEnd - OuterLoopBegin;
-		double OuterSleep = FMath::Max(HttpThreadIdleFrameTimeInSeconds - OuterLoopTime, HttpThreadIdleMinimumSleepTimeInSeconds);
-		FPlatformProcess::SleepNoStats(OuterSleep);
 	}
 	return 0;
 }
@@ -112,10 +145,20 @@ uint32 FHttpThread::Run()
 
 void FHttpThread::Tick()
 {
-	TArray<IHttpThreadedRequest*> RequestsToCancel;
-	TArray<IHttpThreadedRequest*> RequestsToStart;
-	TArray<IHttpThreadedRequest*> RequestsToComplete;
-	Process(RequestsToCancel, RequestsToStart, RequestsToComplete);
+	if (ensure(NeedsSingleThreadTick()))
+	{
+		TArray<IHttpThreadedRequest*> RequestsToCancel;
+		TArray<IHttpThreadedRequest*> RequestsToStart;
+		TArray<IHttpThreadedRequest*> RequestsToComplete;
+		Process(RequestsToCancel, RequestsToStart, RequestsToComplete);
+	}
+}
+
+bool FHttpThread::NeedsSingleThreadTick() const
+{
+	checkf(Thread != nullptr, TEXT("NeedsSingleThreadTick should never be called after StopThread"));
+
+	return bNeedsSingleThreadTick;
 }
 
 void FHttpThread::HttpThreadTick(float DeltaSeconds)
diff --git a/Engine/Source/Runtime/Online/HTTP/Private/HttpThread.h b/Engine/Source/Runtime/Online/HTTP/Private/HttpThread.h
--- a/Engine/Source/Runtime/Online/HTTP/Private/HttpThread.h
+++ b/Engine/Source/Runtime/Online/HTTP/Private/HttpThread.h
@@ -19,6 +19,7 @@ class FHttpThread
 	: FRunnable, FSingleThreadRunnable
 {
 public:
+	friend class FHttpManager;
 
 	FHttpThread();
 	virtual ~FHttpThread();
@@ -33,6 +34,10 @@ public:
 	 */
 	void StopThread();
 
+	/**
+	 * Is the HTTP thread started or stopped.
+	 */
+	bool IsStopped() const { return bIsStopped; }
 	/** 
 	 * Add a request to begin processing on HTTP thread.
 	 *
@@ -59,6 +64,12 @@ public:
 	virtual void Tick() override final;
 	//~ End FSingleThreadRunnable Interface
 
+	/**
+	 * When true the owner of the HTTPThread needs to manually call Tick() since no automomous threads are
+	 * executing the runnable object
+	 */
+	bool NeedsSingleThreadTick() const;
+
 protected:
 
 	/**
@@ -136,4 +147,17 @@ protected:
 
 	/** Pointer to Runnable Thread */
 	FRunnableThread* Thread;
+
+private:
+
+	/** Are we holding a fake thread and we need to be ticked manually when Flushing */
+	bool bNeedsSingleThreadTick;
+
+	/** Tells if the runnable thread is running or stopped */
+	bool bIsStopped;
+
+	/**
+	 * Determines if we allow the creation of real threads on forked processes
+	 */
+	static bool bAllowForkableThreads;
 };
diff --git a/Engine/Source/Runtime/Online/HTTP/Public/HttpManager.h b/Engine/Source/Runtime/Online/HTTP/Public/HttpManager.h
--- a/Engine/Source/Runtime/Online/HTTP/Public/HttpManager.h
+++ b/Engine/Source/Runtime/Online/HTTP/Public/HttpManager.h
@@ -146,6 +146,15 @@ public:
 	 */
 	virtual void OnAfterFork();
 
+
+	/**
+	 * Tells if we allow real threads to be created on forked processes
+	 * Should be called when the manager is shutdown and has finished all processing
+
+	 * @param bAllow Do we allow real threads or not
+	 */
+	void SetAllowThreadsOnForkedProcess(bool bAllow);
+
 	/**
 	 * Update configuration. Called when config has been updated and we need to apply any changes.
 	 */
