Fix crash in HTTPThread due to FRunnableThread::Create starting thread execution immediately.

diff --git a/Engine/Source/Runtime/Online/HTTP/Private/HttpThread.cpp b/Engine/Source/Runtime/Online/HTTP/Private/HttpThread.cpp
--- a/Engine/Source/Runtime/Online/HTTP/Private/HttpThread.cpp
+++ b/Engine/Source/Runtime/Online/HTTP/Private/HttpThread.cpp
@@ -42,13 +42,13 @@ void FHttpThread::StartThread()
 	}
 	else
 	{
-		Thread = FRunnableThread::Create(this, TEXT("HttpManagerThread"), 128 * 1024, TPri_Normal);
-
 		// If the runnable thread is fake.
 		if (FGenericPlatformProcess::SupportsMultithreading() == false)
 		{
 			bNeedsSingleThreadTick = true;
 		}
+
+		Thread = FRunnableThread::Create(this, TEXT("HttpManagerThread"), 128 * 1024, TPri_Normal);
 	}
 
 	bIsStopped = false;
@@ -156,8 +156,6 @@ void FHttpThread::Tick()
 
 bool FHttpThread::NeedsSingleThreadTick() const
 {
-	checkf(Thread != nullptr, TEXT("NeedsSingleThreadTick should never be called after StopThread"));
-
 	return bNeedsSingleThreadTick;
 }
 
