Added checks to CurlHTTPRequests to detect if they are created or destroyed outside the lifetime of the curl module.

diff --git a/Engine/Source/Runtime/Online/HTTP/Private/Curl/CurlHttp.cpp b/Engine/Source/Runtime/Online/HTTP/Private/Curl/CurlHttp.cpp
--- a/Engine/Source/Runtime/Online/HTTP/Private/Curl/CurlHttp.cpp
+++ b/Engine/Source/Runtime/Online/HTTP/Private/Curl/CurlHttp.cpp
@@ -79,6 +79,8 @@ FCurlHttpRequest::FCurlHttpRequest()
 	,	LastReportedBytesSent(0)
 	,   LeastRecentlyCachedInfoMessageIndex(0)
 {
+	checkf(FCurlHttpManager::IsInit(), TEXT("Curl request was created while the library is shutdown"));
+
 	EasyHandle = curl_easy_init();
 
 	// Always setup the debug function to allow for activity to be tracked
@@ -168,6 +170,7 @@ FCurlHttpRequest::FCurlHttpRequest()
 
 FCurlHttpRequest::~FCurlHttpRequest()
 {
+	checkf(FCurlHttpManager::IsInit(), TEXT("Curl request was held after the library was shutdown."));
 	if (EasyHandle)
 	{
 		// clear to prevent crashing in debug callback when this handle is part of an asynchronous curl_multi_perform()
diff --git a/Engine/Source/Runtime/Online/HTTP/Private/Curl/CurlHttpManager.cpp b/Engine/Source/Runtime/Online/HTTP/Private/Curl/CurlHttpManager.cpp
--- a/Engine/Source/Runtime/Online/HTTP/Private/Curl/CurlHttpManager.cpp
+++ b/Engine/Source/Runtime/Online/HTTP/Private/Curl/CurlHttpManager.cpp
@@ -105,9 +105,14 @@ namespace LibCryptoMemHooks
 	}
 }
 
+bool FCurlHttpManager::IsInit()
+{
+	return GMultiHandle != nullptr;
+}
+
 void FCurlHttpManager::InitCurl()
 {
-	if (GMultiHandle != NULL)
+	if (IsInit())
 	{
 		UE_LOG(LogInit, Warning, TEXT("Already initialized multi handle"));
 		return;
@@ -346,11 +351,18 @@ void FCurlHttpManager::OnBeforeFork()
 void FCurlHttpManager::OnAfterFork()
 {
 	InitCurl();
+
+	//todo: not for mut supported instead
 	Thread->StartThread();
 
 	FHttpManager::OnAfterFork();
 }
 
+//void FCurlHttpManager::OnEndFramePostFork()
+//{
+//	Thread->StartThread();
+//}
+
 void FCurlHttpManager::UpdateConfigs()
 {
 	// Update configs - update settings that are safe to update after initialize 
diff --git a/Engine/Source/Runtime/Online/HTTP/Private/Curl/CurlHttpManager.h b/Engine/Source/Runtime/Online/HTTP/Private/Curl/CurlHttpManager.h
--- a/Engine/Source/Runtime/Online/HTTP/Private/Curl/CurlHttpManager.h
+++ b/Engine/Source/Runtime/Online/HTTP/Private/Curl/CurlHttpManager.h
@@ -17,6 +17,7 @@ class FCurlHttpManager : public FHttpManager
 public:
 	static void InitCurl();
 	static void ShutdownCurl();
+	static bool IsInit();
 	static CURLSH* GShareHandle;
 	static CURLM * GMultiHandle;
 
diff --git a/Engine/Source/Runtime/Online/HTTP/Private/Curl/CurlHttpThread.cpp b/Engine/Source/Runtime/Online/HTTP/Private/Curl/CurlHttpThread.cpp
--- a/Engine/Source/Runtime/Online/HTTP/Private/Curl/CurlHttpThread.cpp
+++ b/Engine/Source/Runtime/Online/HTTP/Private/Curl/CurlHttpThread.cpp
@@ -15,7 +15,7 @@ FCurlHttpThread::FCurlHttpThread()
 void FCurlHttpThread::HttpThreadTick(float DeltaSeconds)
 {
 	QUICK_SCOPE_CYCLE_COUNTER(STAT_FCurlHttpThread_HttpThreadTick);
-	check(FCurlHttpManager::GMultiHandle);
+	check(FCurlHttpManager::IsInit());
 
 	if (RunningThreadedRequests.Num() > 0)
 	{
