- Moved XboxOne source code into PlatformExtension location (not C#, ThirdParty, or plugins, yet)

diff --git a/Engine/Source/Runtime/AVEncoder/Public/AVEncoder.cpp b/Engine/Source/Runtime/AVEncoder/Public/AVEncoder.cpp
--- a/Engine/Source/Runtime/AVEncoder/Public/AVEncoder.cpp
+++ b/Engine/Source/Runtime/AVEncoder/Public/AVEncoder.cpp
@@ -9,7 +9,7 @@
 	#include "Microsoft/WmfAudioEncoder.h"
 #elif PLATFORM_XBOXONE
 	#include "Microsoft/WmfAudioEncoder.h"
-	#include "Microsoft/XBoxOne/XboxOneVideoEncoder.h"
+	#include "XboxOneVideoEncoder.h"
 #endif
 
 namespace AVEncoder
diff --git a/Engine/Source/Runtime/Online/WebSockets/WebSockets.Build.cs b/Engine/Source/Runtime/Online/WebSockets/WebSockets.Build.cs
--- a/Engine/Source/Runtime/Online/WebSockets/WebSockets.Build.cs
+++ b/Engine/Source/Runtime/Online/WebSockets/WebSockets.Build.cs
@@ -31,10 +31,9 @@ public class WebSockets : ModuleRules
 	{
 		get
 		{
-			bool bPlatformSupportsXboxWebsockets = Target.Platform == UnrealTargetPlatform.XboxOne;
 			bool bPlatformSupportsWinRTWebsockets = Target.Platform == UnrealTargetPlatform.HoloLens;
 
-			return PlatformSupportsLibWebsockets || bPlatformSupportsXboxWebsockets || bPlatformSupportsWinRTWebsockets;
+			return PlatformSupportsLibWebsockets || bPlatformSupportsWinRTWebsockets;
 		}
 	}
 
