Added a http check in curlhttpRequest to avoid Download Image BP to fail when using a local file instead of a web url.

diff --git a/Engine/Source/Runtime/Online/HTTP/Private/Curl/CurlHttp.cpp b/Engine/Source/Runtime/Online/HTTP/Private/Curl/CurlHttp.cpp
--- a/Engine/Source/Runtime/Online/HTTP/Private/Curl/CurlHttp.cpp
+++ b/Engine/Source/Runtime/Online/HTTP/Private/Curl/CurlHttp.cpp
@@ -1178,7 +1178,7 @@ void FCurlHttpRequest::FinishedRequest()
 				Response->ContentLength = Response->TotalBytesRead.GetValue();
 			}
 
-			if (Response->HttpCode <= 0)
+			if (Response->HttpCode <= 0 && URL.StartsWith(TEXT("Http"), ESearchCase::IgnoreCase))
 			{
 				UE_LOG(LogHttp, Warning, TEXT("%p: invalid HTTP response code received. URL: %s, HTTP code: %d, content length: %d, actual payload size: %d"),
 					this, *GetURL(), Response->HttpCode, Response->ContentLength, Response->Payload.Num());
