Allow CurlHTTPManager to be multithreaded on forked servers
* Can be disabled via -DisableForkedHTTPThread on the commandline

Added FForkedProcessHelper::IsForkedChildProcess() that tells if you are a forked process or not.

diff --git a/Engine/Source/Runtime/Core/Private/HAL/ThreadingBase.cpp b/Engine/Source/Runtime/Core/Private/HAL/ThreadingBase.cpp
--- a/Engine/Source/Runtime/Core/Private/HAL/ThreadingBase.cpp
+++ b/Engine/Source/Runtime/Core/Private/HAL/ThreadingBase.cpp
@@ -1091,6 +1091,17 @@ protected:
 //-------------------------------------------------------------------------------
 
 bool FForkProcessHelper::bIsForkedMultithreadInstance = false;
+bool FForkProcessHelper::bIsForkedChildProcess = false;
+
+bool FForkProcessHelper::IsForkedChildProcess()
+{
+	return bIsForkedChildProcess;
+}
+
+void FForkProcessHelper::SetIsForkedChildProcess()
+{
+	bIsForkedChildProcess = true;
+}
 
 void FForkProcessHelper::OnForkingOccured()
 {
diff --git a/Engine/Source/Runtime/Core/Private/Unix/UnixPlatformProcess.cpp b/Engine/Source/Runtime/Core/Private/Unix/UnixPlatformProcess.cpp
--- a/Engine/Source/Runtime/Core/Private/Unix/UnixPlatformProcess.cpp
+++ b/Engine/Source/Runtime/Core/Private/Unix/UnixPlatformProcess.cpp
@@ -28,6 +28,7 @@
 #include "Misc/OutputDeviceRedirector.h"
 #include "Misc/CoreDelegates.h"
 #include "Misc/App.h"
+#include "Misc/Fork.h"
 
 namespace PlatformProcessLimits
 {
@@ -1422,6 +1423,8 @@ FGenericPlatformProcess::EWaitAndForkResult FUnixPlatformProcess::WaitAndFork()
 			}
 			else if (ChildPID == 0)
 			{
+				FForkProcessHelper::SetIsForkedChildProcess();
+
 				// Child
 				uint16 Cookie = (SignalValue >> 16) & 0xffff;
 				uint16 ChildIdx = SignalValue & 0xffff;
diff --git a/Engine/Source/Runtime/Core/Public/Misc/Fork.h b/Engine/Source/Runtime/Core/Public/Misc/Fork.h
--- a/Engine/Source/Runtime/Core/Public/Misc/Fork.h
+++ b/Engine/Source/Runtime/Core/Public/Misc/Fork.h
@@ -24,11 +24,23 @@ public:
 
 	/**
 	 * Are we a forked process that supports multithreading
+	 * This only becomes true after its safe to be multithread.
+	 * Since a process can be forked mid-tick, there is a period of time where IsForkedChildProcess is true but IsForkedMultithreadInstance will be false
 	 */
 	static bool IsForkedMultithreadInstance();
 
 	/**
-	 * Event triggered when the fork occurred on the child process
+	 * Is this a process that was forked
+	 */
+	static bool IsForkedChildProcess();
+
+	/**
+	 * Sets the forked child process flag 
+	 */
+	static void SetIsForkedChildProcess();
+
+	/**
+	 * Event triggered when a fork occurred on the child process and its safe to create real threads
 	 */
 	static void OnForkingOccured();
 
@@ -58,6 +70,8 @@ public:
 private:
 
 	static bool bIsForkedMultithreadInstance;
+
+	static bool bIsForkedChildProcess;
 };
 
 
diff --git a/Engine/Source/Runtime/Online/HTTP/Private/Curl/CurlHttpManager.cpp b/Engine/Source/Runtime/Online/HTTP/Private/Curl/CurlHttpManager.cpp
--- a/Engine/Source/Runtime/Online/HTTP/Private/Curl/CurlHttpManager.cpp
+++ b/Engine/Source/Runtime/Online/HTTP/Private/Curl/CurlHttpManager.cpp
@@ -11,6 +11,7 @@
 #include "Misc/FileHelper.h"
 #include "Misc/LocalTimestampDirectoryVisitor.h"
 #include "Misc/Paths.h"
+#include "Misc/Fork.h"
 
 #include "Curl/CurlHttpThread.h"
 #include "Curl/CurlHttp.h"
@@ -352,16 +353,26 @@ void FCurlHttpManager::OnAfterFork()
 {
 	InitCurl();
 
-	//todo: not for mut supported instead
-	Thread->StartThread();
+	if (FForkProcessHelper::IsForkedChildProcess() == false || FForkProcessHelper::SupportsMultithreadingPostFork() == false)
+	{
+		// Since this will create a fake thread its safe to create it immediately here
+		Thread->StartThread();
+	}
 
 	FHttpManager::OnAfterFork();
 }
 
-//void FCurlHttpManager::OnEndFramePostFork()
-//{
-//	Thread->StartThread();
-//}
+void FCurlHttpManager::OnEndFramePostFork()
+{
+	if (FForkProcessHelper::SupportsMultithreadingPostFork())
+	{
+		// We forked and the frame is done, time to start the autonomous thread
+		check(FForkProcessHelper::IsForkedMultithreadInstance());
+		Thread->StartThread();
+	}
+
+	FHttpManager::OnEndFramePostFork();
+}
 
 void FCurlHttpManager::UpdateConfigs()
 {
diff --git a/Engine/Source/Runtime/Online/HTTP/Private/Curl/CurlHttpManager.h b/Engine/Source/Runtime/Online/HTTP/Private/Curl/CurlHttpManager.h
--- a/Engine/Source/Runtime/Online/HTTP/Private/Curl/CurlHttpManager.h
+++ b/Engine/Source/Runtime/Online/HTTP/Private/Curl/CurlHttpManager.h
@@ -60,6 +60,7 @@ public:
 	//~ Begin HttpManager Interface
 	virtual void OnBeforeFork() override;
 	virtual void OnAfterFork() override;
+	virtual void OnEndFramePostFork() override;
 	virtual void UpdateConfigs() override;
 
 public:
diff --git a/Engine/Source/Runtime/Online/HTTP/Private/HttpManager.cpp b/Engine/Source/Runtime/Online/HTTP/Private/HttpManager.cpp
--- a/Engine/Source/Runtime/Online/HTTP/Private/HttpManager.cpp
+++ b/Engine/Source/Runtime/Online/HTTP/Private/HttpManager.cpp
@@ -109,6 +109,12 @@ void FHttpManager::OnAfterFork()
 
 }
 
+void FHttpManager::OnEndFramePostFork()
+{
+	// nothing
+}
+
+
 void FHttpManager::UpdateConfigs()
 {
 	// empty
@@ -147,6 +153,7 @@ void FHttpManager::Flush(bool bShutdown)
 	// block until all active requests have completed
 	double BeginWaitTime = FPlatformTime::Seconds();
 	double LastTime = BeginWaitTime;
+	double StallWarnTime = BeginWaitTime + 0.5;
 	while (Requests.Num() > 0)
 	{
 		const double AppTime = FPlatformTime::Seconds();
@@ -168,14 +175,22 @@ void FHttpManager::Flush(bool bShutdown)
 		LastTime = AppTime;
 		if (Requests.Num() > 0)
 		{
-			if (FPlatformProcess::SupportsMultithreading())
-			{
-				UE_LOG(LogHttp, Display, TEXT("Sleeping 0.5s to wait for %d outstanding Http requests."), Requests.Num());
-				FPlatformProcess::Sleep(0.5f);
-			}
-			else if (Thread)
+			if (Thread)
 			{
-				Thread->Tick();
+				if( Thread->NeedsSingleThreadTick() )
+				{
+					if (AppTime >= StallWarnTime)
+					{
+						UE_LOG(LogHttp, Display, TEXT("Ticking HTTPThread for %d outstanding Http requests."), Requests.Num());
+						StallWarnTime = AppTime + 0.5;
+					}
+					Thread->Tick();
+				}
+				else
+				{
+					UE_LOG(LogHttp, Display, TEXT("Sleeping 0.5s to wait for %d outstanding Http requests."), Requests.Num());
+					FPlatformProcess::Sleep(0.5f);
+				}
 			}
 			else
 			{
diff --git a/Engine/Source/Runtime/Online/HTTP/Private/HttpThread.cpp b/Engine/Source/Runtime/Online/HTTP/Private/HttpThread.cpp
--- a/Engine/Source/Runtime/Online/HTTP/Private/HttpThread.cpp
+++ b/Engine/Source/Runtime/Online/HTTP/Private/HttpThread.cpp
@@ -13,6 +13,8 @@
 
 FHttpThread::FHttpThread()
 	:	Thread(nullptr)
+	,	bIsSingleThread(false)
+	,	bIsStopped(true)
 {
 	HttpThreadActiveFrameTimeInSeconds = FHttpModule::Get().GetHttpThreadActiveFrameTimeInSeconds();
 	HttpThreadActiveMinimumSleepTimeInSeconds = FHttpModule::Get().GetHttpThreadActiveMinimumSleepTimeInSeconds();
@@ -29,7 +31,27 @@ FHttpThread::~FHttpThread()
 
 void FHttpThread::StartThread()
 {
-	Thread = FRunnableThread::Create(this, TEXT("HttpManagerThread"), 128 * 1024, TPri_Normal);
+	bIsSingleThread = false;
+
+	const bool bDisableForkedHTTPThread = FParse::Param(FCommandLine::Get(), TEXT("DisableForkedHTTPThread"));
+
+	if (FForkProcessHelper::IsForkedMultithreadInstance() && bDisableForkedHTTPThread == false)
+	{
+		// We only create forkable threads on the forked instance since the HTTPManager cannot safely transition from fake to real seamlessly
+		Thread = FForkProcessHelper::CreateForkableThread(this, TEXT("HttpManagerThread"), 128 * 1024, TPri_Normal);
+	}
+	else
+	{
+		// If the runnable thread is fake.
+		if (FGenericPlatformProcess::SupportsMultithreading() == false)
+		{
+			bIsSingleThread = true;
+		}
+
+		Thread = FRunnableThread::Create(this, TEXT("HttpManagerThread"), 128 * 1024, TPri_Normal);
+	}
+
+	bIsStopped = false;
 }
 
 void FHttpThread::StopThread()
@@ -40,6 +62,9 @@ void FHttpThread::StopThread()
 		delete Thread;
 		Thread = nullptr;
 	}
+
+	bIsStopped = true;
+	bIsSingleThread = true;
 }
 
 void FHttpThread::AddRequest(IHttpThreadedRequest* Request)
@@ -65,6 +90,7 @@ void FHttpThread::GetCompletedRequests(TArray<IHttpThreadedRequest*>& OutComplet
 bool FHttpThread::Init()
 {
 	LastTime = FPlatformTime::Seconds();
+	ExitRequest.Set(false);
 	return true;
 }
 
@@ -76,35 +102,42 @@ uint32 FHttpThread::Run()
 	TArray<IHttpThreadedRequest*> RequestsToComplete;
 	while (!ExitRequest.GetValue())
 	{
-		const double OuterLoopBegin = FPlatformTime::Seconds();
-		double OuterLoopEnd = 0.0;
-		bool bKeepProcessing = true;
-		while (bKeepProcessing)
+		if (ensureMsgf(!bIsSingleThread, TEXT("HTTP Thread was set to singlethread mode while it was running autonomously!")))
 		{
-			const double InnerLoopBegin = FPlatformTime::Seconds();
+			const double OuterLoopBegin = FPlatformTime::Seconds();
+			double OuterLoopEnd = 0.0;
+			bool bKeepProcessing = true;
+			while (bKeepProcessing)
+			{
+				const double InnerLoopBegin = FPlatformTime::Seconds();
 			
-			Process(RequestsToCancel, RequestsToStart, RequestsToComplete);
+				Process(RequestsToCancel, RequestsToStart, RequestsToComplete);
 			
-			if (RunningThreadedRequests.Num() == 0)
-			{
-				bKeepProcessing = false;
-			}
+				if (RunningThreadedRequests.Num() == 0)
+				{
+					bKeepProcessing = false;
+				}
 
-			const double InnerLoopEnd = FPlatformTime::Seconds();
-			if (bKeepProcessing)
-			{
-				double InnerLoopTime = InnerLoopEnd - InnerLoopBegin;
-				double InnerSleep = FMath::Max(HttpThreadActiveFrameTimeInSeconds - InnerLoopTime, HttpThreadActiveMinimumSleepTimeInSeconds);
-				FPlatformProcess::SleepNoStats(InnerSleep);
-			}
-			else
-			{
-				OuterLoopEnd = InnerLoopEnd;
+				const double InnerLoopEnd = FPlatformTime::Seconds();
+				if (bKeepProcessing)
+				{
+					double InnerLoopTime = InnerLoopEnd - InnerLoopBegin;
+					double InnerSleep = FMath::Max(HttpThreadActiveFrameTimeInSeconds - InnerLoopTime, HttpThreadActiveMinimumSleepTimeInSeconds);
+					FPlatformProcess::SleepNoStats(InnerSleep);
+				}
+				else
+				{
+					OuterLoopEnd = InnerLoopEnd;
+				}
 			}
+			double OuterLoopTime = OuterLoopEnd - OuterLoopBegin;
+			double OuterSleep = FMath::Max(HttpThreadIdleFrameTimeInSeconds - OuterLoopTime, HttpThreadIdleMinimumSleepTimeInSeconds);
+			FPlatformProcess::SleepNoStats(OuterSleep);
+		}
+		else
+		{
+			break;
 		}
-		double OuterLoopTime = OuterLoopEnd - OuterLoopBegin;
-		double OuterSleep = FMath::Max(HttpThreadIdleFrameTimeInSeconds - OuterLoopTime, HttpThreadIdleMinimumSleepTimeInSeconds);
-		FPlatformProcess::SleepNoStats(OuterSleep);
 	}
 	return 0;
 }
@@ -112,10 +145,18 @@ uint32 FHttpThread::Run()
 
 void FHttpThread::Tick()
 {
-	TArray<IHttpThreadedRequest*> RequestsToCancel;
-	TArray<IHttpThreadedRequest*> RequestsToStart;
-	TArray<IHttpThreadedRequest*> RequestsToComplete;
-	Process(RequestsToCancel, RequestsToStart, RequestsToComplete);
+	if (ensure(bIsSingleThread))
+	{
+		TArray<IHttpThreadedRequest*> RequestsToCancel;
+		TArray<IHttpThreadedRequest*> RequestsToStart;
+		TArray<IHttpThreadedRequest*> RequestsToComplete;
+		Process(RequestsToCancel, RequestsToStart, RequestsToComplete);
+	}
+}
+
+bool FHttpThread::NeedsSingleThreadTick() const
+{
+	return bIsSingleThread;
 }
 
 void FHttpThread::HttpThreadTick(float DeltaSeconds)
diff --git a/Engine/Source/Runtime/Online/HTTP/Private/HttpThread.h b/Engine/Source/Runtime/Online/HTTP/Private/HttpThread.h
--- a/Engine/Source/Runtime/Online/HTTP/Private/HttpThread.h
+++ b/Engine/Source/Runtime/Online/HTTP/Private/HttpThread.h
@@ -33,6 +33,10 @@ public:
 	 */
 	void StopThread();
 
+	/**
+	 * Is the HTTP thread started or stopped.
+	 */
+	bool IsStopped() const { return bIsStopped; }
 	/** 
 	 * Add a request to begin processing on HTTP thread.
 	 *
@@ -59,6 +63,12 @@ public:
 	virtual void Tick() override final;
 	//~ End FSingleThreadRunnable Interface
 
+	/**
+	 * When true the owner of the HTTPThread needs to manually call Tick() since no automomous threads are
+	 * executing the runnable object
+	 */
+	bool NeedsSingleThreadTick() const;
+
 protected:
 
 	/**
@@ -136,4 +146,12 @@ protected:
 
 	/** Pointer to Runnable Thread */
 	FRunnableThread* Thread;
+
+private:
+
+	/** Are we holding a fake thread and we need to be ticked manually when Flushing */
+	bool bIsSingleThread;
+
+	/** Tells if the runnable thread is running or stopped */
+	bool bIsStopped;
 };
diff --git a/Engine/Source/Runtime/Online/HTTP/Public/HttpManager.h b/Engine/Source/Runtime/Online/HTTP/Public/HttpManager.h
--- a/Engine/Source/Runtime/Online/HTTP/Public/HttpManager.h
+++ b/Engine/Source/Runtime/Online/HTTP/Public/HttpManager.h
@@ -146,6 +146,12 @@ public:
 	 */
 	virtual void OnAfterFork();
 
+	/**
+	 * Inform the HTTP Manager that we finished ticking right after forking. Only called on the forked process
+	 */
+	virtual void OnEndFramePostFork();
+
+
 	/**
 	 * Update configuration. Called when config has been updated and we need to apply any changes.
 	 */
