Add missing includes to HTTPThread.cpp

diff --git a/Engine/Source/Runtime/Online/HTTP/Private/HttpThread.cpp b/Engine/Source/Runtime/Online/HTTP/Private/HttpThread.cpp
--- a/Engine/Source/Runtime/Online/HTTP/Private/HttpThread.cpp
+++ b/Engine/Source/Runtime/Online/HTTP/Private/HttpThread.cpp
@@ -5,7 +5,9 @@
 #include "HAL/PlatformTime.h"
 #include "HAL/PlatformProcess.h"
 #include "HAL/RunnableThread.h"
+#include "Misc/CommandLine.h"
 #include "Misc/Fork.h"
+#include "Misc/Parse.h"
 #include "HttpModule.h"
 #include "Http.h"
 
