Set a default time limit for HTTP flush

diff --git a/Engine/Source/Runtime/Online/HTTP/Private/HttpManager.cpp b/Engine/Source/Runtime/Online/HTTP/Private/HttpManager.cpp
--- a/Engine/Source/Runtime/Online/HTTP/Private/HttpManager.cpp
+++ b/Engine/Source/Runtime/Online/HTTP/Private/HttpManager.cpp
@@ -138,7 +138,7 @@ void FHttpManager::Flush(bool bShutdown)
 	QUICK_SCOPE_CYCLE_COUNTER(STAT_FHttpManager_Flush);
 
 	FScopeLock ScopeLock(&RequestLock);
-	double MaxFlushTimeSeconds = -1.0; // default to no limit
+	double MaxFlushTimeSeconds = 5.0; // default to generous limit
 	GConfig->GetDouble(TEXT("HTTP"), TEXT("MaxFlushTimeSeconds"), MaxFlushTimeSeconds, GEngineIni);
 
 	bool bAlwaysCancelRequestsOnFlush = false; // Default to not immediately cancelling
