Fix new errors/warnings that are no longer being accidentally suppressed in clang on Windows (see CL 15880121).
- Member variable shadowing inside lambdas
- Deprecation warnings due to implicit copy constructors/operator being generated because of dllexports (this is more of a clang bug than an actual issue in our code).

#rnx

diff --git a/Engine/Source/Developer/DerivedDataCache/Private/S3DerivedDataBackend.cpp b/Engine/Source/Developer/DerivedDataCache/Private/S3DerivedDataBackend.cpp
--- a/Engine/Source/Developer/DerivedDataCache/Private/S3DerivedDataBackend.cpp
+++ b/Engine/Source/Developer/DerivedDataCache/Private/S3DerivedDataBackend.cpp
@@ -702,13 +702,13 @@ struct FS3DerivedDataBackend::FRootManifest
 	FString SecretKey;
 	TArray<FString> Keys;
 
-	bool Load(const FString& RootManifestPath)
+	bool Load(const FString& InRootManifestPath)
 	{
 		// Read the root manifest from disk
 		FString RootManifestText;
-		if (!FFileHelper::LoadFileToString(RootManifestText, *RootManifestPath))
+		if (!FFileHelper::LoadFileToString(RootManifestText, *InRootManifestPath))
 		{
-			UE_LOG(LogDerivedDataCache, Warning, TEXT("Unable to read manifest from %s"), *RootManifestPath);
+			UE_LOG(LogDerivedDataCache, Warning, TEXT("Unable to read manifest from %s"), *InRootManifestPath);
 			return false;
 		}
 
@@ -716,19 +716,19 @@ struct FS3DerivedDataBackend::FRootManifest
 		TSharedPtr<FJsonObject> RootManifestObject;
 		if (!FJsonSerializer::Deserialize(TJsonReaderFactory<>::Create(RootManifestText), RootManifestObject) || !RootManifestObject.IsValid())
 		{
-			UE_LOG(LogDerivedDataCache, Warning, TEXT("Unable to parse manifest from %s"), *RootManifestPath);
+			UE_LOG(LogDerivedDataCache, Warning, TEXT("Unable to parse manifest from %s"), *InRootManifestPath);
 			return false;
 		}
 
 		// Read the access and secret keys
 		if (!RootManifestObject->TryGetStringField(TEXT("AccessKey"), AccessKey))
 		{
-			UE_LOG(LogDerivedDataCache, Warning, TEXT("Root manifest %s does not specify AccessKey"), *RootManifestPath);
+			UE_LOG(LogDerivedDataCache, Warning, TEXT("Root manifest %s does not specify AccessKey"), *InRootManifestPath);
 			return false;
 		}
 		if (!RootManifestObject->TryGetStringField(TEXT("SecretKey"), SecretKey))
 		{
-			UE_LOG(LogDerivedDataCache, Warning, TEXT("Root manifest %s does not specify SecretKey"), *RootManifestPath);
+			UE_LOG(LogDerivedDataCache, Warning, TEXT("Root manifest %s does not specify SecretKey"), *InRootManifestPath);
 			return false;
 		}
 
@@ -736,7 +736,7 @@ struct FS3DerivedDataBackend::FRootManifest
 		const TArray<TSharedPtr<FJsonValue>>* RootEntriesArray;
 		if (!RootManifestObject->TryGetArrayField(TEXT("Entries"), RootEntriesArray))
 		{
-			UE_LOG(LogDerivedDataCache, Warning, TEXT("Root manifest from %s is missing entries array"), *RootManifestPath);
+			UE_LOG(LogDerivedDataCache, Warning, TEXT("Root manifest from %s is missing entries array"), *InRootManifestPath);
 			return false;
 		}
 		for (const TSharedPtr<FJsonValue>& Value : *RootEntriesArray)
diff --git a/Engine/Source/Editor/Kismet/Public/BlueprintEditor.h b/Engine/Source/Editor/Kismet/Public/BlueprintEditor.h
--- a/Engine/Source/Editor/Kismet/Public/BlueprintEditor.h
+++ b/Engine/Source/Editor/Kismet/Public/BlueprintEditor.h
@@ -155,7 +155,8 @@ namespace ENodeCreateAction
 // FBlueprintEditor
 
 /** Main Kismet asset editor */
-class KISMET_API FBlueprintEditor : public IBlueprintEditor, public FGCObject, public FNotifyHook, public FTickableEditorObject, public FEditorUndoClient
+class KISMET_API FBlueprintEditor : public IBlueprintEditor, public FGCObject, public FNotifyHook, public 
+FTickableEditorObject, public FEditorUndoClient, public FNoncopyable
 {
 	DECLARE_MULTICAST_DELEGATE_OneParam(FOnSetPinVisibility, SGraphEditor::EPinVisibility);
 
diff --git a/Engine/Source/Editor/UnrealEd/Public/Toolkits/AssetEditorToolkit.h b/Engine/Source/Editor/UnrealEd/Public/Toolkits/AssetEditorToolkit.h
--- a/Engine/Source/Editor/UnrealEd/Public/Toolkits/AssetEditorToolkit.h
+++ b/Engine/Source/Editor/UnrealEd/Public/Toolkits/AssetEditorToolkit.h
@@ -115,6 +115,9 @@ public:
 	 */
 	void InitAssetEditor(const EToolkitMode::Type Mode, const TSharedPtr<IToolkitHost>& InitToolkitHost, const FName AppIdentifier, const TSharedRef<FTabManager::FLayout>& StandaloneDefaultLayout, const bool bCreateDefaultStandaloneMenu, const bool bCreateDefaultToolbar, const TArray<UObject*>& ObjectsToEdit, const bool bInIsToolbarFocusable = false, const bool bInUseSmallToolbarIcons = false);
 	void InitAssetEditor(const EToolkitMode::Type Mode, const TSharedPtr<IToolkitHost>& InitToolkitHost, const FName AppIdentifier, const TSharedRef<FTabManager::FLayout>& StandaloneDefaultLayout, const bool bCreateDefaultStandaloneMenu, const bool bCreateDefaultToolbar, UObject* ObjectToEdit, const bool bInIsToolbarFocusable = false, const bool bInUseSmallToolbarIcons = false);
+
+	FAssetEditorToolkit(const FAssetEditorToolkit&) = delete;
+	FAssetEditorToolkit& operator=(const FAssetEditorToolkit&) = delete;
 	
 	/** Virtual destructor, so that we can clean up our app when destroyed */
 	virtual ~FAssetEditorToolkit();
diff --git a/Engine/Source/Runtime/D3D12RHI/Private/D3D12Buffer.cpp b/Engine/Source/Runtime/D3D12RHI/Private/D3D12Buffer.cpp
--- a/Engine/Source/Runtime/D3D12RHI/Private/D3D12Buffer.cpp
+++ b/Engine/Source/Runtime/D3D12RHI/Private/D3D12Buffer.cpp
@@ -787,7 +787,7 @@ void FD3D12CommandContext::RHICopyBufferRegions(const TArrayView<const FCopyBuff
 		FinalizeDest,
 	};
 
-	auto TransitionResources = [](FD3D12CommandListHandle& CommandListHandle, FLocalResourceArray& SortedResources, EBatchCopyState State)
+	auto TransitionResources = [](FD3D12CommandListHandle& InCommandListHandle, FLocalResourceArray& SortedResources, EBatchCopyState State)
 	{
 		const uint32 Subresource = 0; // Buffers only have one subresource
 
@@ -826,11 +826,11 @@ void FD3D12CommandContext::RHICopyBufferRegions(const TArrayView<const FCopyBuff
 			if (bUseDefaultState)
 			{
 				check(CurrentState != D3D12_RESOURCE_STATE_CORRUPT);
-				CommandListHandle.AddTransitionBarrier(Resource, CurrentState, DesiredState, Subresource);
+				InCommandListHandle.AddTransitionBarrier(Resource, CurrentState, DesiredState, Subresource);
 			}
 			else
 			{
-				FD3D12DynamicRHI::TransitionResource(CommandListHandle, Resource, D3D12_RESOURCE_STATE_TBD, DesiredState, Subresource, FD3D12DynamicRHI::ETransitionMode::Validate);
+				FD3D12DynamicRHI::TransitionResource(InCommandListHandle, Resource, D3D12_RESOURCE_STATE_TBD, DesiredState, Subresource, FD3D12DynamicRHI::ETransitionMode::Validate);
 			}
 
 			PrevResource = Resource;
diff --git a/Engine/Source/Runtime/Engine/Private/Animation/AnimCompressionTypes.cpp b/Engine/Source/Runtime/Engine/Private/Animation/AnimCompressionTypes.cpp
--- a/Engine/Source/Runtime/Engine/Private/Animation/AnimCompressionTypes.cpp
+++ b/Engine/Source/Runtime/Engine/Private/Animation/AnimCompressionTypes.cpp
@@ -78,8 +78,10 @@ void StripFramesOdd(TArray<ArrayValue>& Keys, const int32 NumFrames)
 }
 
 PRAGMA_DISABLE_DEPRECATION_WARNINGS
+FCompressibleAnimData::FCompressibleAnimData(const FCompressibleAnimData&) = default;
 FCompressibleAnimData& FCompressibleAnimData::operator=(const FCompressibleAnimData&) = default;
 ICompressedAnimData::ICompressedAnimData(const ICompressedAnimData&) = default;
+ICompressedAnimData& ICompressedAnimData::operator=(const ICompressedAnimData&) = default;
 PRAGMA_ENABLE_DEPRECATION_WARNINGS
 
 FCompressibleAnimData::FCompressibleAnimData(class UAnimSequence* InSeq, const bool bPerformStripping)
diff --git a/Engine/Source/Runtime/Engine/Public/Animation/AnimCompressionTypes.h b/Engine/Source/Runtime/Engine/Public/Animation/AnimCompressionTypes.h
--- a/Engine/Source/Runtime/Engine/Public/Animation/AnimCompressionTypes.h
+++ b/Engine/Source/Runtime/Engine/Public/Animation/AnimCompressionTypes.h
@@ -195,6 +195,7 @@ public:
 
 	FCompressibleAnimData(class UAnimSequence* InSeq, const bool bPerformStripping);
 
+	FCompressibleAnimData(const FCompressibleAnimData&);
 	FCompressibleAnimData& operator=(const FCompressibleAnimData&);
 
 	UAnimCurveCompressionSettings* CurveCompressionSettings;
@@ -410,7 +411,8 @@ struct ENGINE_API ICompressedAnimData
 
 	ICompressedAnimData() = default;
 	ICompressedAnimData(const ICompressedAnimData&);
-
+	ICompressedAnimData& operator=(const ICompressedAnimData&);
+	
 	/* Virtual interface codecs must implement */
 	virtual ~ICompressedAnimData() {}
 
diff --git a/Engine/Source/Runtime/Online/HTTP/Private/WinHttp/WinHttpHttpRequest.cpp b/Engine/Source/Runtime/Online/HTTP/Private/WinHttp/WinHttpHttpRequest.cpp
--- a/Engine/Source/Runtime/Online/HTTP/Private/WinHttp/WinHttpHttpRequest.cpp
+++ b/Engine/Source/Runtime/Online/HTTP/Private/WinHttp/WinHttpHttpRequest.cpp
@@ -258,11 +258,11 @@ bool FWinHttpHttpRequest::ProcessRequest()
 			return;
 		}
 
-		FWinHttpHttpRequestData& RequestData = StrongThis->RequestData;
+		FWinHttpHttpRequestData& LocalRequestData = StrongThis->RequestData;
 
 		// Create connection object
-		TSharedPtr<FWinHttpConnectionHttp, ESPMode::ThreadSafe> Connection = FWinHttpConnectionHttp::CreateHttpConnection(*SessionPtr, RequestData.Verb, RequestData.Url, RequestData.Headers, RequestData.Payload);
-		if (!Connection.IsValid())
+		TSharedPtr<FWinHttpConnectionHttp, ESPMode::ThreadSafe> LocalConnection = FWinHttpConnectionHttp::CreateHttpConnection(*SessionPtr, LocalRequestData.Verb, LocalRequestData.Url, LocalRequestData.Headers, LocalRequestData.Payload);
+		if (!LocalConnection.IsValid())
 		{
 			UE_LOG(LogHttp, Warning, TEXT("Unable to create WinHttp Connection, failing request"));
 			StrongThis->FinishRequestOnGameThread();
@@ -271,13 +271,13 @@ bool FWinHttpHttpRequest::ProcessRequest()
 
 		// Bind listeners
 		TSharedRef<FWinHttpHttpRequest, ESPMode::ThreadSafe> StrongThisRef = StrongThis.ToSharedRef();
-		Connection->SetDataTransferredHandler(FWinHttpConnectionHttpOnDataTransferred::CreateThreadSafeSP(StrongThisRef, &FWinHttpHttpRequest::HandleDataTransferred));
-		Connection->SetHeaderReceivedHandler(FWinHttpConnectionHttpOnHeaderReceived::CreateThreadSafeSP(StrongThisRef, &FWinHttpHttpRequest::HandleHeaderReceived));
-		Connection->SetRequestCompletedHandler(FWinHttpConnectionHttpOnRequestComplete::CreateThreadSafeSP(StrongThisRef, &FWinHttpHttpRequest::HandleRequestComplete));
+		LocalConnection->SetDataTransferredHandler(FWinHttpConnectionHttpOnDataTransferred::CreateThreadSafeSP(StrongThisRef, &FWinHttpHttpRequest::HandleDataTransferred));
+		LocalConnection->SetHeaderReceivedHandler(FWinHttpConnectionHttpOnHeaderReceived::CreateThreadSafeSP(StrongThisRef, &FWinHttpHttpRequest::HandleHeaderReceived));
+		LocalConnection->SetRequestCompletedHandler(FWinHttpConnectionHttpOnRequestComplete::CreateThreadSafeSP(StrongThisRef, &FWinHttpHttpRequest::HandleRequestComplete));
 
 		// Start request!
 		StrongThisRef->RequestStartTimeSeconds = FPlatformTime::Seconds();
-		if (!Connection->StartRequest())
+		if (!LocalConnection->StartRequest())
 		{
 			UE_LOG(LogHttp, Warning, TEXT("Unable to start WinHttp Connection, failing request"));
 			StrongThisRef->FinishRequestOnGameThread();
@@ -285,7 +285,7 @@ bool FWinHttpHttpRequest::ProcessRequest()
 		}
 
 		// Save object
-		StrongThisRef->Connection = MoveTemp(Connection);
+		StrongThisRef->Connection = MoveTemp(LocalConnection);
 	}));
 
 	// Store our request so it doesn't die if the requester doesn't store it (common use case)
