Removal of UE4 reference in comments

diff --git a/Engine/Source/Runtime/Online/BuildPatchServices/Private/Tests/Fake/FileSystem.fake.h b/Engine/Source/Runtime/Online/BuildPatchServices/Private/Tests/Fake/FileSystem.fake.h
--- a/Engine/Source/Runtime/Online/BuildPatchServices/Private/Tests/Fake/FileSystem.fake.h
+++ b/Engine/Source/Runtime/Online/BuildPatchServices/Private/Tests/Fake/FileSystem.fake.h
@@ -10,7 +10,7 @@
 namespace BuildPatchServices
 {
 	/**
-	 * Using a fake file reader allows us to simulate UE4's file reader behavior, where if the file is written to
+	 * Using a fake file reader allows us to simulate Unreal Engine's file reader behavior, where if the file is written to
 	 * after the handle is opened, you will cause an assert if you try to read the new data at the end as the total size is cached.
 	 */
 	class FFakeFileReader
