#online Resolve a non threadsafe Execute pattern when handling request progress updates in AppleHTTP.

diff --git a/Engine/Source/Runtime/Online/HTTP/Private/Apple/AppleHTTP.cpp b/Engine/Source/Runtime/Online/HTTP/Private/Apple/AppleHTTP.cpp
--- a/Engine/Source/Runtime/Online/HTTP/Private/Apple/AppleHTTP.cpp
+++ b/Engine/Source/Runtime/Online/HTTP/Private/Apple/AppleHTTP.cpp
@@ -534,7 +534,7 @@ void FAppleHttpRequest::Tick(float DeltaSeconds)
 			const int32 BytesRead = Response->GetNumBytesReceived();
 			if (BytesWritten > 0 || BytesRead > 0)
 			{
-				OnRequestProgress().Execute(SharedThis(this), BytesWritten, BytesRead);
+				OnRequestProgress().ExecuteIfBound(SharedThis(this), BytesWritten, BytesRead);
 			}
 		}
 		if (Response->IsReady())
