deprecated FTicker and family and replaced by thread-safe FTSTicker

diff --git a/Engine/Source/Runtime/Engine/Private/UnrealEngine.cpp b/Engine/Source/Runtime/Engine/Private/UnrealEngine.cpp
--- a/Engine/Source/Runtime/Engine/Private/UnrealEngine.cpp
+++ b/Engine/Source/Runtime/Engine/Private/UnrealEngine.cpp
@@ -9945,7 +9945,7 @@ void UEngine::FErrorsAndWarningsCollector::Initialize()
 	if (DisplayTime > 0.0f)
 	{
 		SetVerbosity((GSupressWarningsInOnScreenDisplay != 0) ? ELogVerbosity::Error : ELogVerbosity::Warning);
-		TickerHandle = FTicker::GetCoreTicker().AddTicker(FTickerDelegate::CreateRaw(this, &UEngine::FErrorsAndWarningsCollector::Tick), DisplayTime);
+		TickerHandle = FTSTicker::GetCoreTicker().AddTicker(FTickerDelegate::CreateRaw(this, &UEngine::FErrorsAndWarningsCollector::Tick), DisplayTime);
 		FOutputDeviceRedirector::Get()->AddOutputDevice(this);
 	}
 }
@@ -9955,7 +9955,7 @@ UEngine::FErrorsAndWarningsCollector::~FErrorsAndWarningsCollector()
 	if (TickerHandle.IsValid())
 	{
 		FOutputDeviceRedirector::Get()->RemoveOutputDevice(this);
-		FTicker::GetCoreTicker().RemoveTicker(TickerHandle);
+		FTSTicker::GetCoreTicker().RemoveTicker(TickerHandle);
 	}
 }
 
@@ -16197,7 +16197,7 @@ static void DoUAssetLoadTest(const TArray<FString>& Args)
 		UE_LOG(LogConsoleResponse, Display, TEXT("gc.TimeBetweenPurgingPendingKillObjects set to 10, you might want to adjust that or do gc.CollectGarbageEveryFrame 1"));
 
 		UE_LOG(LogConsoleResponse, Display, TEXT("Testing %d uasset files in %s"), AssetLoadTest_Filenames.Num(), *MountPoint);
-		FTicker::GetCoreTicker().AddTicker(FTickerDelegate::CreateLambda([](float Delta) -> bool { TickAssetLoadTest(); return true; }));
+		FTSTicker::GetCoreTicker().AddTicker(FTickerDelegate::CreateLambda([](float Delta) -> bool { TickAssetLoadTest(); return true; }));
 	}
 
 }
diff --git a/Engine/Source/Runtime/Engine/Public/ODSC/ODSCManager.h b/Engine/Source/Runtime/Engine/Public/ODSC/ODSCManager.h
--- a/Engine/Source/Runtime/Engine/Public/ODSC/ODSCManager.h
+++ b/Engine/Source/Runtime/Engine/Public/ODSC/ODSCManager.h
@@ -5,6 +5,7 @@
 #include "CoreMinimal.h"
 #include "Containers/Ticker.h"
 #include "RHIDefinitions.h"
+#include "Containers/Ticker.h"
 
 class FODSCThread;
 
@@ -13,7 +14,7 @@ class FODSCThread;
  * Interface for submitting shader compile requests to the ODSC Thread.
  */
 class ENGINE_API FODSCManager
-	: public FTickerObjectBase
+	: public FTSTickerObjectBase
 {
 public:
 
@@ -29,10 +30,10 @@ public:
 	 */
 	virtual ~FODSCManager();
 
-	// FTickerObjectBase
+	// FTSTickerObjectBase
 
 	/**
-	 * FTicker callback
+	 * FTSTicker callback
 	 *
 	 * @param DeltaSeconds - time in seconds since the last tick
 	 *
diff --git a/Engine/Source/Runtime/Engine/Public/ProfilingDebugging/MallocLeakReporter.h b/Engine/Source/Runtime/Engine/Public/ProfilingDebugging/MallocLeakReporter.h
--- a/Engine/Source/Runtime/Engine/Public/ProfilingDebugging/MallocLeakReporter.h
+++ b/Engine/Source/Runtime/Engine/Public/ProfilingDebugging/MallocLeakReporter.h
@@ -6,6 +6,7 @@
 #include "HAL/MallocLeakDetection.h"
 #include "Delegates/Delegate.h"
 #include "Delegates/DelegateCombinations.h"
+#include "Containers/Ticker.h"
 
 
 /**
@@ -124,8 +125,8 @@ protected:
 
 	bool						Enabled;
 	int32						ReportCount;
-	FDelegateHandle				CheckpointTicker;
-	FDelegateHandle				ReportTicker;
+	FTSTicker::FDelegateHandle				CheckpointTicker;
+	FTSTicker::FDelegateHandle				ReportTicker;
 	FMallocLeakReportDelegate	ReportDelegate;
 
 	// Report vars
diff --git a/Engine/Source/Runtime/HardwareSurvey/Private/HardwareSurveyModule.cpp b/Engine/Source/Runtime/HardwareSurvey/Private/HardwareSurveyModule.cpp
--- a/Engine/Source/Runtime/HardwareSurvey/Private/HardwareSurveyModule.cpp
+++ b/Engine/Source/Runtime/HardwareSurvey/Private/HardwareSurveyModule.cpp
@@ -31,7 +31,7 @@ public:
 	{
 		if (bPendingHardwareSurveyResults)
 		{
-			FTicker::GetCoreTicker().RemoveTicker(TickerHandle);
+			FTSTicker::GetCoreTicker().RemoveTicker(TickerHandle);
 			bPendingHardwareSurveyResults = false;
 		}
 	}
@@ -53,7 +53,7 @@ public:
 		if (IsHardwareSurveyRequired())
 		{
 			bPendingHardwareSurveyResults = true;
-			TickerHandle = FTicker::GetCoreTicker().AddTicker(FTickerDelegate::CreateRaw(this, &FHardwareSurveyModule::TickHardwareSurvey));
+			TickerHandle = FTSTicker::GetCoreTicker().AddTicker(FTickerDelegate::CreateRaw(this, &FHardwareSurveyModule::TickHardwareSurvey));
 		}
 	}
 
@@ -378,7 +378,7 @@ private:
 	// Holds the analytics provider if available
 	IAnalyticsProvider* Analytics;
 
-	FDelegateHandle TickerHandle;
+	FTSTicker::FDelegateHandle TickerHandle;
 };
 
 
diff --git a/Engine/Source/Runtime/InstallBundleManager/Private/BundlePrereqCombinedStatusHelper.cpp b/Engine/Source/Runtime/InstallBundleManager/Private/BundlePrereqCombinedStatusHelper.cpp
--- a/Engine/Source/Runtime/InstallBundleManager/Private/BundlePrereqCombinedStatusHelper.cpp
+++ b/Engine/Source/Runtime/InstallBundleManager/Private/BundlePrereqCombinedStatusHelper.cpp
@@ -75,7 +75,7 @@ void FInstallBundleCombinedProgressTracker::SetupDelegates(bool bAutoTick)
 	IInstallBundleManager::PausedBundleDelegate.AddRaw(this, &FInstallBundleCombinedProgressTracker::OnBundleInstallPauseChanged);
 	if (bAutoTick)
 	{
-		TickHandle = FTicker::GetCoreTicker().AddTicker(FTickerDelegate::CreateRaw(this, &FInstallBundleCombinedProgressTracker::Tick));
+		TickHandle = FTSTicker::GetCoreTicker().AddTicker(FTickerDelegate::CreateRaw(this, &FInstallBundleCombinedProgressTracker::Tick));
 	}
 }
 
@@ -85,7 +85,7 @@ void FInstallBundleCombinedProgressTracker::CleanUpDelegates()
 	IInstallBundleManager::PausedBundleDelegate.RemoveAll(this);
 	if (TickHandle.IsValid())
 	{
-		FTicker::GetCoreTicker().RemoveTicker(TickHandle);
+		FTSTicker::GetCoreTicker().RemoveTicker(TickHandle);
 		TickHandle.Reset();
 	}
 }
diff --git a/Engine/Source/Runtime/InstallBundleManager/Private/InstallBundleUtils.cpp b/Engine/Source/Runtime/InstallBundleManager/Private/InstallBundleUtils.cpp
--- a/Engine/Source/Runtime/InstallBundleManager/Private/InstallBundleUtils.cpp
+++ b/Engine/Source/Runtime/InstallBundleManager/Private/InstallBundleUtils.cpp
@@ -830,7 +830,7 @@ namespace InstallBundleUtil
 				//Only setup a tick function if we would use it
 				if ((bShouldAutoUpdateTimersInTick || bShouldSaveDirtyStatsOnTick) && !TickHandle.IsValid())
 				{
-					TickHandle = FTicker::GetCoreTicker().AddTicker(FTickerDelegate::CreateRaw(this, &FPersistentStatContainerBase::Tick));
+					TickHandle = FTSTicker::GetCoreTicker().AddTicker(FTickerDelegate::CreateRaw(this, &FPersistentStatContainerBase::Tick));
 				}
 
 				//Only setup Foreground/Background delegates if we should be using them to swap stats
@@ -852,7 +852,7 @@ namespace InstallBundleUtil
 		{
 			if (TickHandle.IsValid())
 			{
-				FTicker::GetCoreTicker().RemoveTicker(TickHandle);
+				FTSTicker::GetCoreTicker().RemoveTicker(TickHandle);
 				TickHandle.Reset();
 			}
 
diff --git a/Engine/Source/Runtime/Interchange/Engine/Private/InterchangeManager.cpp b/Engine/Source/Runtime/Interchange/Engine/Private/InterchangeManager.cpp
--- a/Engine/Source/Runtime/Interchange/Engine/Private/InterchangeManager.cpp
+++ b/Engine/Source/Runtime/Interchange/Engine/Private/InterchangeManager.cpp
@@ -1038,7 +1038,7 @@ void UInterchangeManager::SetActiveMode(bool IsActive)
 	if (bIsActive)
 	{
 		ensure(!NotificationTickHandle.IsValid());
-		NotificationTickHandle = FTicker::GetCoreTicker().AddTicker(TEXT("InterchangeManagerTickHandle"), 0.1f, [this](float)
+		NotificationTickHandle = FTSTicker::GetCoreTicker().AddTicker(TEXT("InterchangeManagerTickHandle"), 0.1f, [this](float)
 		{
 			if (Notification.IsValid() && Notification->GetPromptAction() == EAsyncTaskNotificationPromptAction::Cancel)
 			{
@@ -1060,7 +1060,7 @@ void UInterchangeManager::SetActiveMode(bool IsActive)
 	}
 	else
 	{
-		FTicker::GetCoreTicker().RemoveTicker(NotificationTickHandle);
+		FTSTicker::GetCoreTicker().RemoveTicker(NotificationTickHandle);
 		NotificationTickHandle.Reset();
 
 		if (GcGuardThread.IsJoinable())
diff --git a/Engine/Source/Runtime/Interchange/Engine/Public/InterchangeManager.h b/Engine/Source/Runtime/Interchange/Engine/Public/InterchangeManager.h
--- a/Engine/Source/Runtime/Interchange/Engine/Public/InterchangeManager.h
+++ b/Engine/Source/Runtime/Interchange/Engine/Public/InterchangeManager.h
@@ -22,6 +22,7 @@
 #include "UObject/Object.h"
 #include "UObject/ObjectMacros.h"
 #include "UObject/StrongObjectPtr.h"
+#include "Containers/Ticker.h"
 
 #include "InterchangeManager.generated.h"
 
@@ -458,7 +459,7 @@ private:
 	TArray<TSharedPtr<UE::Interchange::FImportAsyncHelper, ESPMode::ThreadSafe> > ImportTasks;
 
 	TSharedPtr<FAsyncTaskNotification> Notification = nullptr;
-	FDelegateHandle NotificationTickHandle;
+	FTSTicker::FDelegateHandle NotificationTickHandle;
 
 	// Caching the registered translator classes to avoid double registration fast
 	UPROPERTY()
diff --git a/Engine/Source/Runtime/Launch/Private/IOS/LaunchIOS.cpp b/Engine/Source/Runtime/Launch/Private/IOS/LaunchIOS.cpp
--- a/Engine/Source/Runtime/Launch/Private/IOS/LaunchIOS.cpp
+++ b/Engine/Source/Runtime/Launch/Private/IOS/LaunchIOS.cpp
@@ -379,7 +379,7 @@ void FAppEntry::PlatformInit()
 #if BUILD_EMBEDDED_APP
 		// while embedded, the native app may be waiting on some processing to happen before showing the view, so we have to let
 		// processing occur here
-		FTicker::GetCoreTicker().Tick(0.005f);
+		FTSTicker::GetCoreTicker().Tick(0.005f);
 		FThreadManager::Get().Tick();
 #endif
 		FPlatformProcess::Sleep(0.005f);
diff --git a/Engine/Source/Runtime/Messaging/Private/Bus/MessageTracer.cpp b/Engine/Source/Runtime/Messaging/Private/Bus/MessageTracer.cpp
--- a/Engine/Source/Runtime/Messaging/Private/Bus/MessageTracer.cpp
+++ b/Engine/Source/Runtime/Messaging/Private/Bus/MessageTracer.cpp
@@ -16,13 +16,13 @@ FMessageTracer::FMessageTracer()
 	, Running(false)
 {
 	ContinueEvent = FPlatformProcess::GetSynchEventFromPool();
-	TickDelegateHandle = FTicker::GetCoreTicker().AddTicker(FTickerDelegate::CreateRaw(this, &FMessageTracer::Tick), 0.0f);
+	TickDelegateHandle = FTSTicker::GetCoreTicker().AddTicker(FTickerDelegate::CreateRaw(this, &FMessageTracer::Tick), 0.0f);
 }
 
 
 FMessageTracer::~FMessageTracer()
 {
-	FTicker::GetCoreTicker().RemoveTicker(TickDelegateHandle);
+	FTSTicker::GetCoreTicker().RemoveTicker(TickDelegateHandle);
 	FPlatformProcess::ReturnSynchEventToPool(ContinueEvent);
 	ContinueEvent = nullptr;
 }
diff --git a/Engine/Source/Runtime/Messaging/Private/Bus/MessageTracer.h b/Engine/Source/Runtime/Messaging/Private/Bus/MessageTracer.h
--- a/Engine/Source/Runtime/Messaging/Private/Bus/MessageTracer.h
+++ b/Engine/Source/Runtime/Messaging/Private/Bus/MessageTracer.h
@@ -7,6 +7,7 @@
 #include "Containers/Queue.h"
 #include "IMessageContext.h"
 #include "IMessageTracer.h"
+#include "Containers/Ticker.h"
 
 class IMessageInterceptor;
 class IMessageReceiver;
@@ -195,7 +196,7 @@ private:
 	bool Running;
 
 	/** Handle to the registered TickDelegate. */
-	FDelegateHandle TickDelegateHandle;
+	FTSTicker::FDelegateHandle TickDelegateHandle;
 
 	/** Holds the trace actions queue. */
 	TQueue<TFunction<void()>, EQueueMode::Mpsc> Traces;
diff --git a/Engine/Source/Runtime/MessagingRpc/Private/MessageRpcClient.cpp b/Engine/Source/Runtime/MessagingRpc/Private/MessageRpcClient.cpp
--- a/Engine/Source/Runtime/MessagingRpc/Private/MessageRpcClient.cpp
+++ b/Engine/Source/Runtime/MessagingRpc/Private/MessageRpcClient.cpp
@@ -32,12 +32,12 @@ FMessageRpcClient::FMessageRpcClient(FMessageEndpointBuilder&& InEndpointBuilder
 		.Handling<FMessageRpcProgress>(this, &FMessageRpcClient::HandleProgressMessage)
 		.WithCatchall(this, &FMessageRpcClient::HandleRpcMessages);
 
-	TickerHandle = FTicker::GetCoreTicker().AddTicker(FTickerDelegate::CreateRaw(this, &FMessageRpcClient::HandleTicker), MESSAGE_RPC_RETRY_INTERVAL);
+	TickerHandle = FTSTicker::GetCoreTicker().AddTicker(FTickerDelegate::CreateRaw(this, &FMessageRpcClient::HandleTicker), MESSAGE_RPC_RETRY_INTERVAL);
 }
 
 FMessageRpcClient::~FMessageRpcClient()
 {
-	FTicker::GetCoreTicker().RemoveTicker(TickerHandle);
+	FTSTicker::GetCoreTicker().RemoveTicker(TickerHandle);
 }
 
 
diff --git a/Engine/Source/Runtime/MessagingRpc/Private/MessageRpcServer.cpp b/Engine/Source/Runtime/MessagingRpc/Private/MessageRpcServer.cpp
--- a/Engine/Source/Runtime/MessagingRpc/Private/MessageRpcServer.cpp
+++ b/Engine/Source/Runtime/MessagingRpc/Private/MessageRpcServer.cpp
@@ -30,12 +30,12 @@ FMessageRpcServer::FMessageRpcServer(const FString& InDebugName, const TSharedRe
 FMessageRpcServer::FMessageRpcServer(FMessageEndpointBuilder&& InEndpointBuilder)
 {
 	MessageEndpoint = InEndpointBuilder.WithCatchall(this, &FMessageRpcServer::HandleMessage);
-	TickerHandle = FTicker::GetCoreTicker().AddTicker(FTickerDelegate::CreateRaw(this, &FMessageRpcServer::HandleTicker), MESSAGE_RPC_TICK_DELAY);
+	TickerHandle = FTSTicker::GetCoreTicker().AddTicker(FTickerDelegate::CreateRaw(this, &FMessageRpcServer::HandleTicker), MESSAGE_RPC_TICK_DELAY);
 }
 
 FMessageRpcServer::~FMessageRpcServer()
 {
-	FTicker::GetCoreTicker().RemoveTicker(TickerHandle);
+	FTSTicker::GetCoreTicker().RemoveTicker(TickerHandle);
 }
 
 
diff --git a/Engine/Source/Runtime/NVidia/GeForceNOW/Private/GeForceNOWActionZoneProcessor.cpp b/Engine/Source/Runtime/NVidia/GeForceNOW/Private/GeForceNOWActionZoneProcessor.cpp
--- a/Engine/Source/Runtime/NVidia/GeForceNOW/Private/GeForceNOWActionZoneProcessor.cpp
+++ b/Engine/Source/Runtime/NVidia/GeForceNOW/Private/GeForceNOWActionZoneProcessor.cpp
@@ -210,13 +210,13 @@ void GeForceNOWActionZoneProcessor::StartProcess()
 {
 	if (!ProcessDelegateHandle.IsValid())
 	{
-		ProcessDelegateHandle = FTicker::GetCoreTicker().AddTicker(FTickerDelegate::CreateRaw(this, &GeForceNOWActionZoneProcessor::ProcessGFNWidgetActionZones), GFNWidgetActionZonesProcessDelay);
+		ProcessDelegateHandle = FTSTicker::GetCoreTicker().AddTicker(FTickerDelegate::CreateRaw(this, &GeForceNOWActionZoneProcessor::ProcessGFNWidgetActionZones), GFNWidgetActionZonesProcessDelay);
 	}
 }
 
 void GeForceNOWActionZoneProcessor::StopProcess()
 {
-	FTicker::GetCoreTicker().RemoveTicker(ProcessDelegateHandle);
+	FTSTicker::GetCoreTicker().RemoveTicker(ProcessDelegateHandle);
 	ProcessDelegateHandle.Reset();
 }
 
diff --git a/Engine/Source/Runtime/NVidia/GeForceNOW/Public/GeForceNOWActionZoneProcessor.h b/Engine/Source/Runtime/NVidia/GeForceNOW/Public/GeForceNOWActionZoneProcessor.h
--- a/Engine/Source/Runtime/NVidia/GeForceNOW/Public/GeForceNOWActionZoneProcessor.h
+++ b/Engine/Source/Runtime/NVidia/GeForceNOW/Public/GeForceNOWActionZoneProcessor.h
@@ -6,6 +6,7 @@
 
 #include "CoreMinimal.h"
 #include "Layout/SlateRect.h"
+#include "Containers/Ticker.h"
 
 class SWidget;
 class SWindow;
@@ -85,7 +86,7 @@ private:
 
 	TArray<FWidgetGFNActionZone> GFNWidgetActionZones;
 
-	FDelegateHandle ProcessDelegateHandle;
+	FTSTicker::FDelegateHandle ProcessDelegateHandle;
 };
 
 #endif // NV_GEFORCENOW
diff --git a/Engine/Source/Runtime/Online/BackgroundHTTP/Public/BackgroundHttpManagerImpl.h b/Engine/Source/Runtime/Online/BackgroundHTTP/Public/BackgroundHttpManagerImpl.h
--- a/Engine/Source/Runtime/Online/BackgroundHTTP/Public/BackgroundHttpManagerImpl.h
+++ b/Engine/Source/Runtime/Online/BackgroundHTTP/Public/BackgroundHttpManagerImpl.h
@@ -18,7 +18,7 @@ class FBackgroundHttpFileHashHelper;
  */
 class BACKGROUNDHTTP_API FBackgroundHttpManagerImpl 
 	: public IBackgroundHttpManager
-	, public FTickerObjectBase
+	, public FTSTickerObjectBase
 {
 public:
 	FBackgroundHttpManagerImpl();
@@ -40,7 +40,7 @@ public:
 	
 	virtual void CleanUpDataAfterCompletingRequest(const FBackgroundHttpRequestPtr Request) override;
 	
-	//FTickerObjectBase implementation
+	//FTSTickerObjectBase implementation
 	virtual bool Tick(float DeltaTime) override;
 	
 protected:
diff --git a/Engine/Source/Runtime/Online/BackgroundHTTP/Public/IOS/ApplePlatformBackgroundHttpManager.h b/Engine/Source/Runtime/Online/BackgroundHTTP/Public/IOS/ApplePlatformBackgroundHttpManager.h
--- a/Engine/Source/Runtime/Online/BackgroundHTTP/Public/IOS/ApplePlatformBackgroundHttpManager.h
+++ b/Engine/Source/Runtime/Online/BackgroundHTTP/Public/IOS/ApplePlatformBackgroundHttpManager.h
@@ -38,7 +38,7 @@ public:
 	 */
 	virtual ~FApplePlatformBackgroundHttpManager();
 
-    //FTickerObjectBase implementation
+    //FTSTickerObjectBase implementation
     virtual bool Tick(float DeltaTime) override;
     
     //This is the identifier we expect to be associated with our NSURLDownloadSession for background downloads
diff --git a/Engine/Source/Runtime/Online/BuildPatchServices/Private/BuildPatchInstaller.cpp b/Engine/Source/Runtime/Online/BuildPatchServices/Private/BuildPatchInstaller.cpp
--- a/Engine/Source/Runtime/Online/BuildPatchServices/Private/BuildPatchInstaller.cpp
+++ b/Engine/Source/Runtime/Online/BuildPatchServices/Private/BuildPatchInstaller.cpp
@@ -434,7 +434,7 @@ namespace BuildPatchServices
 		, InstallerError(FInstallerErrorFactory::Create())
 		, Analytics(MoveTemp(InAnalytics))
 		, InstallerAnalytics(FInstallerAnalyticsFactory::Create(Analytics.Get()))
-		, FileOperationTracker(FFileOperationTrackerFactory::Create(FTicker::GetCoreTicker()))
+		, FileOperationTracker(FFileOperationTrackerFactory::Create(FTSTicker::GetCoreTicker()))
 		, MemoryChunkStoreStatistics(FMemoryChunkStoreStatisticsFactory::Create(FileOperationTracker.Get()))
 		, DiskChunkStoreStatistics(FDiskChunkStoreStatisticsFactory::Create(InstallerAnalytics.Get(), FileOperationTracker.Get()))
 		, DownloadSpeedRecorder(FSpeedRecorderFactory::Create())
@@ -448,7 +448,7 @@ namespace BuildPatchServices
 		, CloudChunkSourceStatistics(FCloudChunkSourceStatisticsFactory::Create(InstallerAnalytics.Get(), &BuildProgress, FileOperationTracker.Get()))
 		, FileConstructorStatistics(FFileConstructorStatisticsFactory::Create(DiskReadSpeedRecorder.Get(), DiskWriteSpeedRecorder.Get(), &BuildProgress, FileOperationTracker.Get()))
 		, VerifierStatistics(FVerifierStatisticsFactory::Create(DiskReadSpeedRecorder.Get(), &BuildProgress, FileOperationTracker.Get()))
-		, DownloadService(FDownloadServiceFactory::Create(FTicker::GetCoreTicker(), HttpManager.Get(), FileSystem.Get(), DownloadServiceStatistics.Get(), InstallerAnalytics.Get()))
+		, DownloadService(FDownloadServiceFactory::Create(FTSTicker::GetCoreTicker(), HttpManager.Get(), FileSystem.Get(), DownloadServiceStatistics.Get(), InstallerAnalytics.Get()))
 		, MessagePump(FMessagePumpFactory::Create())
 		, Controllables()
 	{
diff --git a/Engine/Source/Runtime/Online/BuildPatchServices/Private/BuildPatchServicesModule.cpp b/Engine/Source/Runtime/Online/BuildPatchServices/Private/BuildPatchServicesModule.cpp
--- a/Engine/Source/Runtime/Online/BuildPatchServices/Private/BuildPatchServicesModule.cpp
+++ b/Engine/Source/Runtime/Online/BuildPatchServices/Private/BuildPatchServicesModule.cpp
@@ -80,7 +80,7 @@ void FBuildPatchServicesModule::StartupModule()
 	bForceSkipPrereqs = bForceSkipPrereqsCmdline || bForceSkipPrereqsConfig;
 
 	// Add our ticker
-	TickDelegateHandle = FTicker::GetCoreTicker().AddTicker( FTickerDelegate::CreateRaw( this, &FBuildPatchServicesModule::Tick ) );
+	TickDelegateHandle = FTSTicker::GetCoreTicker().AddTicker( FTickerDelegate::CreateRaw( this, &FBuildPatchServicesModule::Tick ) );
 
 	// Register core PreExit
 	FCoreDelegates::OnPreExit.AddRaw(this, &FBuildPatchServicesModule::PreExit);
@@ -108,7 +108,7 @@ void FBuildPatchServicesModule::ShutdownModule()
 
 	// Remove our ticker
 	GLog->Log(ELogVerbosity::VeryVerbose, TEXT( "BuildPatchServicesModule: Removing Ticker" ) );
-	FTicker::GetCoreTicker().RemoveTicker( TickDelegateHandle );
+	FTSTicker::GetCoreTicker().RemoveTicker( TickDelegateHandle );
 
 	GLog->Log(ELogVerbosity::VeryVerbose, TEXT( "BuildPatchServicesModule: Finished shutting down" ) );
 }
diff --git a/Engine/Source/Runtime/Online/BuildPatchServices/Private/Common/SpeedRecorder.cpp b/Engine/Source/Runtime/Online/BuildPatchServices/Private/Common/SpeedRecorder.cpp
--- a/Engine/Source/Runtime/Online/BuildPatchServices/Private/Common/SpeedRecorder.cpp
+++ b/Engine/Source/Runtime/Online/BuildPatchServices/Private/Common/SpeedRecorder.cpp
@@ -86,7 +86,7 @@ namespace BuildPatchServices
 
 	class FSpeedRecorder
 		: public ISpeedRecorder
-		, public FTickerObjectBase
+		, public FTSTickerObjectBase
 	{
 	public:
 		FSpeedRecorder();
@@ -98,9 +98,9 @@ namespace BuildPatchServices
 		virtual double GetPeakSpeed() const;
 		// ISpeedRecorder interface end.
 
-		// FTickerObjectBase interface begin.
+		// FTSTickerObjectBase interface begin.
 		virtual bool Tick(float DeltaTime);
-		// FTickerObjectBase interface end.
+		// FTSTickerObjectBase interface end.
 
 	private:
 		void PutRecordsInTemp(float OverSeconds) const;
diff --git a/Engine/Source/Runtime/Online/BuildPatchServices/Private/Diffing/DiffManifests.cpp b/Engine/Source/Runtime/Online/BuildPatchServices/Private/Diffing/DiffManifests.cpp
--- a/Engine/Source/Runtime/Online/BuildPatchServices/Private/Diffing/DiffManifests.cpp
+++ b/Engine/Source/Runtime/Online/BuildPatchServices/Private/Diffing/DiffManifests.cpp
@@ -299,7 +299,7 @@ namespace BuildPatchServices
 
 	private:
 		const FDiffManifestsConfiguration Configuration;
-		FTicker& CoreTicker;
+		FTSTicker& CoreTicker;
 		FDownloadCompleteDelegate DownloadCompleteDelegate;
 		FDownloadProgressDelegate DownloadProgressDelegate;
 		TUniquePtr<IFileSystem> FileSystem;
@@ -323,7 +323,7 @@ namespace BuildPatchServices
 
 	FDiffManifests::FDiffManifests(const FDiffManifestsConfiguration& InConfiguration)
 		: Configuration(InConfiguration)
-		, CoreTicker(FTicker::GetCoreTicker())
+		, CoreTicker(FTSTicker::GetCoreTicker())
 		, DownloadCompleteDelegate(FDownloadCompleteDelegate::CreateRaw(this, &FDiffManifests::HandleDownloadComplete))
 		, DownloadProgressDelegate()
 		, FileSystem(FFileSystemFactory::Create())
@@ -376,7 +376,7 @@ namespace BuildPatchServices
 
 			// Application tick.
 			FTaskGraphInterface::Get().ProcessThreadUntilIdle(ENamedThreads::GameThread);
-			FTicker::GetCoreTicker().Tick(DeltaTime);
+			FTSTicker::GetCoreTicker().Tick(DeltaTime);
 			GLog->FlushThreadedLogs();
 
 			// Control frame rate.
diff --git a/Engine/Source/Runtime/Online/BuildPatchServices/Private/Generation/ChunkDeltaOptimiser.cpp b/Engine/Source/Runtime/Online/BuildPatchServices/Private/Generation/ChunkDeltaOptimiser.cpp
--- a/Engine/Source/Runtime/Online/BuildPatchServices/Private/Generation/ChunkDeltaOptimiser.cpp
+++ b/Engine/Source/Runtime/Online/BuildPatchServices/Private/Generation/ChunkDeltaOptimiser.cpp
@@ -574,7 +574,7 @@ namespace BuildPatchServices
 
 	private:
 		const FChunkDeltaOptimiserConfiguration Configuration;
-		FTicker& CoreTicker;
+		FTSTicker& CoreTicker;
 		FDownloadCompleteDelegate DownloadCompleteDelegate;
 		FDownloadProgressDelegate DownloadProgressDelegate;
 		TUniquePtr<IFileSystem> FileSystem;
@@ -601,7 +601,7 @@ namespace BuildPatchServices
 
 	FChunkDeltaOptimiser::FChunkDeltaOptimiser(const FChunkDeltaOptimiserConfiguration& InConfiguration)
 		: Configuration(InConfiguration)
-		, CoreTicker(FTicker::GetCoreTicker())
+		, CoreTicker(FTSTicker::GetCoreTicker())
 		, DownloadCompleteDelegate(FDownloadCompleteDelegate::CreateRaw(this, &FChunkDeltaOptimiser::HandleDownloadComplete))
 		, DownloadProgressDelegate()
 		, FileSystem(FFileSystemFactory::Create())
@@ -666,7 +666,7 @@ namespace BuildPatchServices
 
 			// Application tick.
 			FTaskGraphInterface::Get().ProcessThreadUntilIdle(ENamedThreads::GameThread);
-			FTicker::GetCoreTicker().Tick(DeltaTime);
+			FTSTicker::GetCoreTicker().Tick(DeltaTime);
 
 			// Message pump.
 			MessagePump->PumpMessages(MessageHandlers);
diff --git a/Engine/Source/Runtime/Online/BuildPatchServices/Private/Installer/DownloadService.cpp b/Engine/Source/Runtime/Online/BuildPatchServices/Private/Installer/DownloadService.cpp
--- a/Engine/Source/Runtime/Online/BuildPatchServices/Private/Installer/DownloadService.cpp
+++ b/Engine/Source/Runtime/Online/BuildPatchServices/Private/Installer/DownloadService.cpp
@@ -11,6 +11,7 @@
 #include "Common/HttpManager.h"
 #include "Common/FileSystem.h"
 #include "Stats/Stats.h"
+#include "Containers/Ticker.h"
 
 DECLARE_LOG_CATEGORY_EXTERN(LogDownloadService, Warning, All);
 DEFINE_LOG_CATEGORY(LogDownloadService);
@@ -212,7 +213,7 @@ namespace BuildPatchServices
 		};
 
 	public:
-		FDownloadService(FTicker& Ticker, IHttpManager* HttpManager, IFileSystem* FileSystem, IDownloadServiceStat* DownloadServiceStat, IInstallerAnalytics* InstallerAnalytics);
+		FDownloadService(FTSTicker& Ticker, IHttpManager* HttpManager, IFileSystem* FileSystem, IDownloadServiceStat* DownloadServiceStat, IInstallerAnalytics* InstallerAnalytics);
 		~FDownloadService();
 
 		// IDownloadService interface begin.
@@ -242,7 +243,7 @@ namespace BuildPatchServices
 
 	private:
 		TSharedRef<FHttpDelegates, ESPMode::ThreadSafe> HttpDelegates;
-		FTicker& Ticker;
+		FTSTicker& Ticker;
 		IHttpManager* HttpManager;
 		IFileSystem* FileSystem;
 		IDownloadServiceStat* DownloadServiceStat;
@@ -269,7 +270,7 @@ namespace BuildPatchServices
 		FCriticalSection CompletedRequestsCS;
 		TMap<int32, FDownloadBaseRef> CompletedRequests;
 
-		FDelegateHandle TickerHandle;
+		FTSTicker::FDelegateHandle TickerHandle;
 	};
 
 	FDownloadService::FHttpDelegates::FHttpDelegates(FDownloadService& InDownloadService)
@@ -287,7 +288,7 @@ namespace BuildPatchServices
 		DownloadService.HttpRequestComplete(MoveTemp(Request), MoveTemp(Response), bSucceeded, MoveTemp(DownloadRecord));
 	}
 
-	FDownloadService::FDownloadService(FTicker& InTicker, IHttpManager* InHttpManager, IFileSystem* InFileSystem, IDownloadServiceStat* InDownloadServiceStat, IInstallerAnalytics* InInstallerAnalytics)
+	FDownloadService::FDownloadService(FTSTicker& InTicker, IHttpManager* InHttpManager, IFileSystem* InFileSystem, IDownloadServiceStat* InDownloadServiceStat, IInstallerAnalytics* InInstallerAnalytics)
 		: HttpDelegates(MakeShareable(new FHttpDelegates(*this)))
 		, Ticker(InTicker)
 		, HttpManager(InHttpManager)
@@ -633,7 +634,7 @@ namespace BuildPatchServices
 		CompletedRequests.Emplace(RequestId, MakeShareable(new FHttpDownload(MoveTemp(Response), bSuccess, MoveTemp(DownloadRecord))));
 	}
 
-	IDownloadService* FDownloadServiceFactory::Create(FTicker& Ticker, IHttpManager* HttpManager, IFileSystem* FileSystem, IDownloadServiceStat* DownloadServiceStat, IInstallerAnalytics* InstallerAnalytics)
+	IDownloadService* FDownloadServiceFactory::Create(FTSTicker& Ticker, IHttpManager* HttpManager, IFileSystem* FileSystem, IDownloadServiceStat* DownloadServiceStat, IInstallerAnalytics* InstallerAnalytics)
 	{
 		check(HttpManager != nullptr);
 		check(FileSystem != nullptr);
diff --git a/Engine/Source/Runtime/Online/BuildPatchServices/Private/Installer/DownloadService.h b/Engine/Source/Runtime/Online/BuildPatchServices/Private/Installer/DownloadService.h
--- a/Engine/Source/Runtime/Online/BuildPatchServices/Private/Installer/DownloadService.h
+++ b/Engine/Source/Runtime/Online/BuildPatchServices/Private/Installer/DownloadService.h
@@ -99,7 +99,7 @@ namespace BuildPatchServices
 		 * @param InstallerAnalytics    The analytics implementation for tracking HTTP requests.
 		 * @return the new IDownloadService instance created.
 		 */
-		static IDownloadService* Create(FTicker& Ticker, IHttpManager* HttpManager, IFileSystem* FileSystem, IDownloadServiceStat* DownloadServiceStat, IInstallerAnalytics* InstallerAnalytics);
+		static IDownloadService* Create(FTSTicker& Ticker, IHttpManager* HttpManager, IFileSystem* FileSystem, IDownloadServiceStat* DownloadServiceStat, IInstallerAnalytics* InstallerAnalytics);
 	};
 
 	/**
diff --git a/Engine/Source/Runtime/Online/BuildPatchServices/Private/Installer/InstallerAnalytics.cpp b/Engine/Source/Runtime/Online/BuildPatchServices/Private/Installer/InstallerAnalytics.cpp
--- a/Engine/Source/Runtime/Online/BuildPatchServices/Private/Installer/InstallerAnalytics.cpp
+++ b/Engine/Source/Runtime/Online/BuildPatchServices/Private/Installer/InstallerAnalytics.cpp
@@ -62,7 +62,7 @@ namespace BuildPatchServices
 		FThreadSafeCounter ConstructionErrors;
 		FCriticalSection AnalyticsEventQueueCS;
 		TArray<FAnalyticsEventInfo> AnalyticsEventQueue;
-		FDelegateHandle TickerHandle;
+		FTSTicker::FDelegateHandle TickerHandle;
 	};
 
 	FInstallerAnalytics::FInstallerAnalytics(IAnalyticsProvider* InAnalyticsProvider)
@@ -73,13 +73,13 @@ namespace BuildPatchServices
 		, AnalyticsEventQueueCS()
 		, AnalyticsEventQueue()
 	{
-		TickerHandle = FTicker::GetCoreTicker().AddTicker(FTickerDelegate::CreateRaw(this, &FInstallerAnalytics::Tick));
+		TickerHandle = FTSTicker::GetCoreTicker().AddTicker(FTickerDelegate::CreateRaw(this, &FInstallerAnalytics::Tick));
 	}
 
 	FInstallerAnalytics::~FInstallerAnalytics()
 	{
 		// Remove ticker.
-		FTicker::GetCoreTicker().RemoveTicker(TickerHandle);
+		FTSTicker::GetCoreTicker().RemoveTicker(TickerHandle);
 	}
 
 	void FInstallerAnalytics::RecordChunkDownloadError(const FString& ChunkUrl, int32 ResponseCode, const FString& ErrorString)
diff --git a/Engine/Source/Runtime/Online/BuildPatchServices/Private/Installer/Statistics/FileOperationTracker.cpp b/Engine/Source/Runtime/Online/BuildPatchServices/Private/Installer/Statistics/FileOperationTracker.cpp
--- a/Engine/Source/Runtime/Online/BuildPatchServices/Private/Installer/Statistics/FileOperationTracker.cpp
+++ b/Engine/Source/Runtime/Online/BuildPatchServices/Private/Installer/Statistics/FileOperationTracker.cpp
@@ -10,6 +10,7 @@
 #include "BuildPatchManifest.h"
 #include "IBuildManifestSet.h"
 #include "Stats/Stats.h"
+#include "Containers/Ticker.h"
 
 DECLARE_LOG_CATEGORY_EXTERN(LogFileOperationTracker, Warning, All);
 DEFINE_LOG_CATEGORY(LogFileOperationTracker);
@@ -26,7 +27,7 @@ namespace BuildPatchServices
 		: public IFileOperationTracker
 	{
 	public:
-		FFileOperationTracker(FTicker& Ticker);
+		FFileOperationTracker(FTSTicker& Ticker);
 		~FFileOperationTracker();
 
 	public:
@@ -53,8 +54,8 @@ namespace BuildPatchServices
 		static FOperationInitialiser BuildOperationInitialiser(IBuildManifestSet* ManifestSet);
 
 	private:
-		FTicker& Ticker;
-		FDelegateHandle TickerHandle;
+		FTSTicker& Ticker;
+		FTSTicker::FDelegateHandle TickerHandle;
 		TArray<FFileOperation> FileOperationStates;
 		TArray<FFileOperation> DummyOperationStates;
 		TMap<FGuid, TArray<FFileOperation*>> FileOperationStatesDataIdLookup;
@@ -63,7 +64,7 @@ namespace BuildPatchServices
 		IBuildManifestSet* LastUsedManifestSet;
 	};
 
-	FFileOperationTracker::FFileOperationTracker(FTicker& InTicker)
+	FFileOperationTracker::FFileOperationTracker(FTSTicker& InTicker)
 		: Ticker(InTicker)
 		, LastUsedManifestSet(nullptr)
 	{
@@ -264,7 +265,7 @@ namespace BuildPatchServices
 		return Result;
 	}
 
-	IFileOperationTracker* FFileOperationTrackerFactory::Create(FTicker& Ticker)
+	IFileOperationTracker* FFileOperationTrackerFactory::Create(FTSTicker& Ticker)
 	{
 		return new FFileOperationTracker(Ticker);
 	}
diff --git a/Engine/Source/Runtime/Online/BuildPatchServices/Private/Installer/Statistics/FileOperationTracker.h b/Engine/Source/Runtime/Online/BuildPatchServices/Private/Installer/Statistics/FileOperationTracker.h
--- a/Engine/Source/Runtime/Online/BuildPatchServices/Private/Installer/Statistics/FileOperationTracker.h
+++ b/Engine/Source/Runtime/Online/BuildPatchServices/Private/Installer/Statistics/FileOperationTracker.h
@@ -95,6 +95,6 @@ namespace BuildPatchServices
 		 * @param   Manifest    The manifest for the build being installed.
 		 * @return the new IFileOperationTracker instance created.
 		 */
-		static IFileOperationTracker* Create(FTicker& Ticker);
+		static IFileOperationTracker* Create(FTSTicker& Ticker);
 	};
 }
\ No newline at end of file
diff --git a/Engine/Source/Runtime/Online/BuildPatchServices/Private/Tests/Fake/HttpManager.fake.h b/Engine/Source/Runtime/Online/BuildPatchServices/Private/Tests/Fake/HttpManager.fake.h
--- a/Engine/Source/Runtime/Online/BuildPatchServices/Private/Tests/Fake/HttpManager.fake.h
+++ b/Engine/Source/Runtime/Online/BuildPatchServices/Private/Tests/Fake/HttpManager.fake.h
@@ -38,7 +38,7 @@ namespace BuildPatchServices
 		: public FMockHttpManager
 	{
 	public:
-		FFakeHttpManager(FTicker& Ticker);
+		FFakeHttpManager(FTSTicker& Ticker);
 		virtual TSharedRef<IHttpRequest, ESPMode::ThreadSafe> CreateRequest() override;
 		virtual bool Tick(float Delta);
 		virtual bool OnProcessRequest(FFakeHttpRequest* FakeHttpRequest);
@@ -78,7 +78,7 @@ namespace BuildPatchServices
 		return Data.Num();
 	}
 
-	FFakeHttpManager::FFakeHttpManager(FTicker& Ticker)
+	FFakeHttpManager::FFakeHttpManager(FTSTicker& Ticker)
 	{
 		Ticker.AddTicker(FTickerDelegate::CreateRaw(this, &FFakeHttpManager::Tick));
 	}
diff --git a/Engine/Source/Runtime/Online/BuildPatchServices/Private/Tests/Unit/DownloadService.spec.cpp b/Engine/Source/Runtime/Online/BuildPatchServices/Private/Tests/Unit/DownloadService.spec.cpp
--- a/Engine/Source/Runtime/Online/BuildPatchServices/Private/Tests/Unit/DownloadService.spec.cpp
+++ b/Engine/Source/Runtime/Online/BuildPatchServices/Private/Tests/Unit/DownloadService.spec.cpp
@@ -16,7 +16,7 @@ BEGIN_DEFINE_SPEC(FDownloadServiceSpec, "BuildPatchServices.Unit", EAutomationTe
 // Unit.
 TUniquePtr<BuildPatchServices::IDownloadService> DownloadService;
 // Mock/Fake.
-FTicker Ticker;
+FTSTicker Ticker;
 TUniquePtr<BuildPatchServices::FFakeHttpManager> FakeHttpModule;
 TUniquePtr<BuildPatchServices::FMockFileSystem> MockFileSystem;
 TUniquePtr<BuildPatchServices::FMockDownloadServiceStat> MockDownloadServiceStat;
@@ -57,7 +57,7 @@ void FDownloadServiceSpec::Define()
 	// Specs.
 	BeforeEach([this]()
 	{
-		Ticker = FTicker();
+		Ticker.Reset();
 		FakeHttpModule.Reset(new FFakeHttpManager(Ticker));
 		MockFileSystem.Reset(new FMockFileSystem());
 		MockDownloadServiceStat.Reset(new FMockDownloadServiceStat());
diff --git a/Engine/Source/Runtime/Online/HTTP/Private/HttpManager.cpp b/Engine/Source/Runtime/Online/HTTP/Private/HttpManager.cpp
--- a/Engine/Source/Runtime/Online/HTTP/Private/HttpManager.cpp
+++ b/Engine/Source/Runtime/Online/HTTP/Private/HttpManager.cpp
@@ -21,7 +21,7 @@
 FCriticalSection FHttpManager::RequestLock;
 
 FHttpManager::FHttpManager()
-	: FTickerObjectBase(0.0f, FBackgroundableTicker::GetCoreTicker())
+	: FTSTickerObjectBase(0.0f, FTSBackgroundableTicker::GetCoreTicker())
 	, Thread(nullptr)
 	, CorrelationIdMethod(FHttpManager::GetDefaultCorrelationIdMethod())
 {
diff --git a/Engine/Source/Runtime/Online/HTTP/Public/HttpManager.h b/Engine/Source/Runtime/Online/HTTP/Public/HttpManager.h
--- a/Engine/Source/Runtime/Online/HTTP/Public/HttpManager.h
+++ b/Engine/Source/Runtime/Online/HTTP/Public/HttpManager.h
@@ -8,6 +8,7 @@
 #include "Containers/Ticker.h"
 #include "Containers/Queue.h"
 #include "HttpPackage.h"
+#include "Containers/Ticker.h"
 
 class FHttpThread;
 
@@ -15,7 +16,7 @@ class FHttpThread;
  * Manages Http request that are currently being processed
  */
 class HTTP_API FHttpManager
-	: public FTickerObjectBase
+	: public FTSTickerObjectBase
 {
 public:
 
@@ -69,7 +70,7 @@ public:
 	void Flush(bool bShutdown);
 
 	/**
-	 * FTicker callback
+	 * FTSTicker callback
 	 *
 	 * @param DeltaSeconds - time in seconds since the last tick
 	 *
diff --git a/Engine/Source/Runtime/Online/HTTPServer/Public/HttpServerModule.h b/Engine/Source/Runtime/Online/HTTPServer/Public/HttpServerModule.h
--- a/Engine/Source/Runtime/Online/HTTPServer/Public/HttpServerModule.h
+++ b/Engine/Source/Runtime/Online/HTTPServer/Public/HttpServerModule.h
@@ -17,7 +17,7 @@ DECLARE_LOG_CATEGORY_EXTERN(LogHttpServerModule, Log, All);
  */
 class FHttpServerModule : 
 	public IModuleInterface
-	,public FTickerObjectBase
+	,public FTSTickerObjectBase
 {
 
 public:
@@ -46,7 +46,7 @@ public:
 	HTTPSERVER_API TSharedPtr<IHttpRouter> GetHttpRouter(uint32 Port);
 
 	/**
-	 * FTicker callback
+	 * FTSTicker callback
 	 * 
 	 * @param DeltaTime  The time in seconds since the last tick
 	 * @return           false if no longer needs ticking, true otherwise
diff --git a/Engine/Source/Runtime/Online/ICMP/Private/Icmp.cpp b/Engine/Source/Runtime/Online/ICMP/Private/Icmp.cpp
--- a/Engine/Source/Runtime/Online/ICMP/Private/Icmp.cpp
+++ b/Engine/Source/Runtime/Online/ICMP/Private/Icmp.cpp
@@ -67,12 +67,12 @@ bool ResolveIp(ISocketSubsystem* SocketSub, const FString& HostName, FString& Ou
 }
 
 class FIcmpAsyncResult
-	: public FTickerObjectBase
+	: public FTSTickerObjectBase
 {
 public:
 
 	FIcmpAsyncResult(ISocketSubsystem* InSocketSub, const FString& TargetAddress, float Timeout, uint32 StackSize, FIcmpEchoResultCallback InCallback)
-		: FTickerObjectBase(0)
+		: FTSTickerObjectBase(0)
 		, SocketSub(InSocketSub)
 		, Callback(InCallback)
 		, bThreadCompleted(false)
diff --git a/Engine/Source/Runtime/Online/ICMP/Private/UDPPing.cpp b/Engine/Source/Runtime/Online/ICMP/Private/UDPPing.cpp
--- a/Engine/Source/Runtime/Online/ICMP/Private/UDPPing.cpp
+++ b/Engine/Source/Runtime/Online/ICMP/Private/UDPPing.cpp
@@ -233,7 +233,7 @@ private:
 // ----------------------------------------------------------------------------
 
 class FUdpPingManyAsync
-	: public FTickerObjectBase
+	: public FTSTickerObjectBase
 {
 public:
 	FUdpPingManyAsync(ISocketSubsystem* const InSocketSub, const FIcmpEchoManyCompleteDelegate& InCompletionDelegate);
@@ -951,7 +951,7 @@ uint32 FUdpPingWorker::FProgressKey::CalcHash() const
 // ----------------------------------------------------------------------------
 
 FUdpPingManyAsync::FUdpPingManyAsync(ISocketSubsystem* const InSocketSub, const FIcmpEchoManyCompleteDelegate& InCompletionDelegate)
-	: FTickerObjectBase(0)
+	: FTSTickerObjectBase(0)
 	, SocketSub(InSocketSub)
 	, CompletionDelegate(InCompletionDelegate)
 	, bThreadCompleted(false)
@@ -1436,12 +1436,12 @@ FIcmpEchoResult UDPEchoImpl(ISocketSubsystem* SocketSub, const FString& TargetAd
 }
 
 class FUDPPingAsyncResult
-	: public FTickerObjectBase
+	: public FTSTickerObjectBase
 {
 public:
 
 	FUDPPingAsyncResult(ISocketSubsystem* InSocketSub, const FString& TargetAddress, float Timeout, uint32 StackSize, FIcmpEchoResultCallback InCallback)
-		: FTickerObjectBase(0)
+		: FTSTickerObjectBase(0)
 		, SocketSub(InSocketSub)
 		, Callback(InCallback)
 		, bThreadCompleted(false)
diff --git a/Engine/Source/Runtime/Online/Stomp/Private/StompClient.h b/Engine/Source/Runtime/Online/Stomp/Private/StompClient.h
--- a/Engine/Source/Runtime/Online/Stomp/Private/StompClient.h
+++ b/Engine/Source/Runtime/Online/Stomp/Private/StompClient.h
@@ -15,14 +15,14 @@ class Error;
 
 class FStompClient
 	: public IStompClient
-	, public FTickerObjectBase
+	, public FTSTickerObjectBase
 	, public TSharedFromThis<FStompClient>
 {
 public:
 
 	virtual ~FStompClient();
 
-	// FTickerObjectBase
+	// FTSTickerObjectBase
 	virtual bool Tick(float DeltaTime) override;
 
 	// IStompClient
diff --git a/Engine/Source/Runtime/Online/Voice/Private/Windows/VoiceCaptureWindows.h b/Engine/Source/Runtime/Online/Voice/Private/Windows/VoiceCaptureWindows.h
--- a/Engine/Source/Runtime/Online/Voice/Private/Windows/VoiceCaptureWindows.h
+++ b/Engine/Source/Runtime/Online/Voice/Private/Windows/VoiceCaptureWindows.h
@@ -104,7 +104,7 @@ public:
 /**
  * Windows implementation of voice capture using DirectSound
  */
-class FVoiceCaptureWindows : public IVoiceCapture, public FTickerObjectBase
+class FVoiceCaptureWindows : public IVoiceCapture, public FTSTickerObjectBase
 {
 public:
 
@@ -125,7 +125,7 @@ public:
 	virtual void DumpState() const override;
 	virtual float GetCurrentAmplitude() const override;
 
-	// FTickerObjectBase
+	// FTSTickerObjectBase
 	virtual bool Tick(float DeltaTime) override;
 	
 private:
diff --git a/Engine/Source/Runtime/Online/WebSockets/Private/HoloLens/HoloLensWebSocket.h b/Engine/Source/Runtime/Online/WebSockets/Private/HoloLens/HoloLensWebSocket.h
--- a/Engine/Source/Runtime/Online/WebSockets/Private/HoloLens/HoloLensWebSocket.h
+++ b/Engine/Source/Runtime/Online/WebSockets/Private/HoloLens/HoloLensWebSocket.h
@@ -12,7 +12,7 @@
 
 class FHoloLensWebSocket :
 	public IWebSocket,
-	public FTickerObjectBase,
+	public FTSTickerObjectBase,
 	public TSharedFromThis<FHoloLensWebSocket>
 {
 public:
diff --git a/Engine/Source/Runtime/Online/WebSockets/Private/Lws/LwsWebSocketsManager.cpp b/Engine/Source/Runtime/Online/WebSockets/Private/Lws/LwsWebSocketsManager.cpp
--- a/Engine/Source/Runtime/Online/WebSockets/Private/Lws/LwsWebSocketsManager.cpp
+++ b/Engine/Source/Runtime/Online/WebSockets/Private/Lws/LwsWebSocketsManager.cpp
@@ -166,7 +166,7 @@ void FLwsWebSocketsManager::InitWebSockets(TArrayView<const FString> Protocols)
 
 	// Setup our game thread tick
 	FTickerDelegate TickDelegate = FTickerDelegate::CreateRaw(this, &FLwsWebSocketsManager::GameThreadTick);
-	TickHandle = FBackgroundableTicker::GetCoreTicker().AddTicker(TickDelegate, 0.0f);
+	TickHandle = FTSBackgroundableTicker::GetCoreTicker().AddTicker(TickDelegate, 0.0f);
 }
 
 void FLwsWebSocketsManager::ShutdownWebSockets()
@@ -174,7 +174,7 @@ void FLwsWebSocketsManager::ShutdownWebSockets()
 	QUICK_SCOPE_CYCLE_COUNTER(STAT_FLwsWebSocketsManager_ShutdownWebSockets);
 	if (TickHandle.IsValid())
 	{
-		FBackgroundableTicker::GetCoreTicker().RemoveTicker(TickHandle);
+		FTSBackgroundableTicker::GetCoreTicker().RemoveTicker(TickHandle);
 		TickHandle.Reset();
 	}
 
diff --git a/Engine/Source/Runtime/Online/WebSockets/Private/WinHttp/WinHttpWebSocketsManager.cpp b/Engine/Source/Runtime/Online/WebSockets/Private/WinHttp/WinHttpWebSocketsManager.cpp
--- a/Engine/Source/Runtime/Online/WebSockets/Private/WinHttp/WinHttpWebSocketsManager.cpp
+++ b/Engine/Source/Runtime/Online/WebSockets/Private/WinHttp/WinHttpWebSocketsManager.cpp
@@ -15,7 +15,7 @@ void FWinHttpWebSocketsManager::InitWebSockets(TArrayView<const FString> Protoco
 
 	if (ensure(!TickHandle.IsValid()))
 	{
-		TickHandle = FBackgroundableTicker::GetCoreTicker().AddTicker(FTickerDelegate::CreateRaw(this, &FWinHttpWebSocketsManager::GameThreadTick), 0.0f);
+		TickHandle = FTSBackgroundableTicker::GetCoreTicker().AddTicker(FTickerDelegate::CreateRaw(this, &FWinHttpWebSocketsManager::GameThreadTick), 0.0f);
 	}
 }
 
@@ -32,7 +32,7 @@ void FWinHttpWebSocketsManager::ShutdownWebSockets()
 
 	if (ensure(TickHandle.IsValid()))
 	{
-		FBackgroundableTicker::GetCoreTicker().RemoveTicker(TickHandle);
+		FTSBackgroundableTicker::GetCoreTicker().RemoveTicker(TickHandle);
 		TickHandle.Reset();
 	}
 }
diff --git a/Engine/Source/Runtime/Online/XMPP/Private/XmppJingle/XmppChatJingle.h b/Engine/Source/Runtime/Online/XMPP/Private/XmppJingle/XmppChatJingle.h
--- a/Engine/Source/Runtime/Online/XMPP/Private/XmppJingle/XmppChatJingle.h
+++ b/Engine/Source/Runtime/Online/XMPP/Private/XmppJingle/XmppChatJingle.h
@@ -7,6 +7,7 @@
 #include "Containers/Ticker.h"
 #include "XmppJingle/XmppJingle.h"
 #include "XmppChat.h"
+#include "Containers/Ticker.h"
 
 #if WITH_XMPP_JINGLE
 
@@ -16,7 +17,7 @@
 class FXmppChatJingle
 	: public IXmppChat
 	, public sigslot::has_slots<>
-	, public FTickerObjectBase
+	, public FTSTickerObjectBase
 {
 public:
 
@@ -25,7 +26,7 @@ public:
 	virtual bool SendChat(const FXmppUserJid& RecipientId, const FString& Message) override;
 	virtual FOnXmppChatReceived& OnReceiveChat() override { return OnXmppChatReceivedDelegate; }
 
-	// FTickerObjectBase
+	// FTSTickerObjectBase
 	
 	virtual bool Tick(float DeltaTime) override;
 
diff --git a/Engine/Source/Runtime/Online/XMPP/Private/XmppJingle/XmppConnectionJingle.h b/Engine/Source/Runtime/Online/XMPP/Private/XmppJingle/XmppConnectionJingle.h
--- a/Engine/Source/Runtime/Online/XMPP/Private/XmppJingle/XmppConnectionJingle.h
+++ b/Engine/Source/Runtime/Online/XMPP/Private/XmppJingle/XmppConnectionJingle.h
@@ -6,6 +6,7 @@
 #include "Containers/Ticker.h"
 #include "XmppConnection.h"
 #include "XmppJingle/XmppChatJingle.h"
+#include "Containers/Ticker.h"
 
 class FXmppMessagesJingle;
 class FXmppMultiUserChatJingle;
@@ -19,7 +20,7 @@ class FXmppPresenceJingle;
  */
 class FXmppConnectionJingle
 	: public IXmppConnection
-	, public FTickerObjectBase
+	, public FTSTickerObjectBase
 	, public sigslot::has_slots<>
 {
 public:
@@ -48,7 +49,7 @@ public:
 
 	virtual void DumpState() const override;
 
-	// FTickerObjectBase
+	// FTSTickerObjectBase
 
 	virtual bool Tick(float DeltaTime) override;
 	
diff --git a/Engine/Source/Runtime/Online/XMPP/Private/XmppJingle/XmppMessagesJingle.h b/Engine/Source/Runtime/Online/XMPP/Private/XmppJingle/XmppMessagesJingle.h
--- a/Engine/Source/Runtime/Online/XMPP/Private/XmppJingle/XmppMessagesJingle.h
+++ b/Engine/Source/Runtime/Online/XMPP/Private/XmppJingle/XmppMessagesJingle.h
@@ -16,7 +16,7 @@
 class FXmppMessagesJingle
 	: public IXmppMessages
 	, public sigslot::has_slots<>
-	, public FTickerObjectBase
+	, public FTSTickerObjectBase
 {
 public:
 
@@ -26,7 +26,7 @@ public:
 	virtual bool SendMessage(const FXmppUserJid& RecipientId, const FString& Type, const TSharedRef<class FJsonObject>& Payload) override;
 	virtual FOnXmppMessageReceived& OnReceiveMessage() override { return OnXmppMessageReceivedDelegate; }
 
-	// FTickerObjectBase
+	// FTSTickerObjectBase
 	
 	virtual bool Tick(float DeltaTime) override;
 
diff --git a/Engine/Source/Runtime/Online/XMPP/Private/XmppStrophe/XmppPingStrophe.h b/Engine/Source/Runtime/Online/XMPP/Private/XmppStrophe/XmppPingStrophe.h
--- a/Engine/Source/Runtime/Online/XMPP/Private/XmppStrophe/XmppPingStrophe.h
+++ b/Engine/Source/Runtime/Online/XMPP/Private/XmppStrophe/XmppPingStrophe.h
@@ -33,14 +33,14 @@ public:
 };
 
 class FXmppPingStrophe
-	: public FTickerObjectBase
+	: public FTSTickerObjectBase
 {
 public:
 	// FXmppPingStrophe
 	FXmppPingStrophe(FXmppConnectionStrophe& InConnectionManager);
 	virtual ~FXmppPingStrophe();
 
-	// FTickerObjectBase
+	// FTSTickerObjectBase
 	virtual bool Tick(float DeltaTime) override;
 
 	/** Clear any caches on disconnect */
diff --git a/Engine/Source/Runtime/Online/XMPP/Public/XmppModule.h b/Engine/Source/Runtime/Online/XMPP/Public/XmppModule.h
--- a/Engine/Source/Runtime/Online/XMPP/Public/XmppModule.h
+++ b/Engine/Source/Runtime/Online/XMPP/Public/XmppModule.h
@@ -16,7 +16,7 @@ class Error;
  * Use CreateConnection to create a new Xmpp connection
  */
 class XMPP_API FXmppModule :
-	public IModuleInterface, public FSelfRegisteringExec, public FTickerObjectBase
+	public IModuleInterface, public FSelfRegisteringExec, public FTSTickerObjectBase
 {
 
 public:
@@ -97,7 +97,7 @@ public:
 		return bEnabled;
 	}
 
-	// FTickerObjectBase
+	// FTSTickerObjectBase
 	virtual bool Tick(float DeltaTime) override;
 
 	/**
diff --git a/Engine/Source/Runtime/PerfCounters/Private/PerfCounters.h b/Engine/Source/Runtime/PerfCounters/Private/PerfCounters.h
--- a/Engine/Source/Runtime/PerfCounters/Private/PerfCounters.h
+++ b/Engine/Source/Runtime/PerfCounters/Private/PerfCounters.h
@@ -20,7 +20,7 @@ class FInternetAddr;
 DECLARE_LOG_CATEGORY_EXTERN(LogPerfCounters, Log, All);
 
 class FPerfCounters 
-	: public FTickerObjectBase
+	: public FTSTickerObjectBase
 	, public FSelfRegisteringExec
 	, public IPerfCounters
 	, public TSharedFromThis<FPerfCounters>
@@ -33,7 +33,7 @@ public:
 	/** Initializes this instance from JSON config. */
 	bool Initialize();
 
-	/** FTickerObjectBase */
+	/** FTSTickerObjectBase */
 	virtual bool Tick(float DeltaTime) override;
 
 	//~ Begin Exec Interface
diff --git a/Engine/Source/Runtime/Portal/Rpc/Private/PortalRpcLocator.cpp b/Engine/Source/Runtime/Portal/Rpc/Private/PortalRpcLocator.cpp
--- a/Engine/Source/Runtime/Portal/Rpc/Private/PortalRpcLocator.cpp
+++ b/Engine/Source/Runtime/Portal/Rpc/Private/PortalRpcLocator.cpp
@@ -20,7 +20,7 @@ public:
 
 	virtual ~FPortalRpcLocatorImpl()
 	{
-		FTicker::GetCoreTicker().RemoveTicker(TickerHandle);
+		FTSTicker::GetCoreTicker().RemoveTicker(TickerHandle);
 	}
 
 public:
@@ -83,7 +83,7 @@ private:
 		MessageEndpoint = FMessageEndpoint::Builder("FPortalRpcLocator")
 			.Handling<FPortalRpcServer>(this, &FPortalRpcLocatorImpl::HandleMessage);
 
-		TickerHandle = FTicker::GetCoreTicker().AddTicker(FTickerDelegate::CreateRaw(this, &FPortalRpcLocatorImpl::HandleTicker), PORTAL_RPC_LOCATE_INTERVAL);
+		TickerHandle = FTSTicker::GetCoreTicker().AddTicker(FTickerDelegate::CreateRaw(this, &FPortalRpcLocatorImpl::HandleTicker), PORTAL_RPC_LOCATE_INTERVAL);
 	}
 	PRAGMA_ENABLE_DEPRECATION_WARNINGS
 
@@ -109,7 +109,7 @@ private:
 	FSimpleDelegate ServerLostDelegate;
 
 	/** Handle to the registered ticker. */
-	FDelegateHandle TickerHandle;
+	FTSTicker::FDelegateHandle TickerHandle;
 
 	friend FPortalRpcLocatorFactory;
 };
diff --git a/Engine/Source/Runtime/SessionServices/Private/SessionManager.cpp b/Engine/Source/Runtime/SessionServices/Private/SessionManager.cpp
--- a/Engine/Source/Runtime/SessionServices/Private/SessionManager.cpp
+++ b/Engine/Source/Runtime/SessionServices/Private/SessionManager.cpp
@@ -43,7 +43,7 @@ FSessionManager::FSessionManager(const TSharedRef<IMessageBus, ESPMode::ThreadSa
 		.Handling<FSessionServicePong>(this, &FSessionManager::HandleSessionPongMessage);
 
 	// initialize ticker
-	TickDelegateHandle = FTicker::GetCoreTicker().AddTicker(FTickerDelegate::CreateRaw(this, &FSessionManager::HandleTicker), SESSION_MANAGER_PING_INTERVAL);
+	TickDelegateHandle = FTSTicker::GetCoreTicker().AddTicker(FTickerDelegate::CreateRaw(this, &FSessionManager::HandleTicker), SESSION_MANAGER_PING_INTERVAL);
 
 	SendPing();
 }
@@ -51,7 +51,7 @@ FSessionManager::FSessionManager(const TSharedRef<IMessageBus, ESPMode::ThreadSa
 
 FSessionManager::~FSessionManager()
 {
-	FTicker::GetCoreTicker().RemoveTicker(TickDelegateHandle);
+	FTSTicker::GetCoreTicker().RemoveTicker(TickDelegateHandle);
 }
 
 
diff --git a/Engine/Source/Runtime/SlateCore/Private/Animation/CurveSequence.cpp b/Engine/Source/Runtime/SlateCore/Private/Animation/CurveSequence.cpp
--- a/Engine/Source/Runtime/SlateCore/Private/Animation/CurveSequence.cpp
+++ b/Engine/Source/Runtime/SlateCore/Private/Animation/CurveSequence.cpp
@@ -42,7 +42,7 @@ FCurveSequence::~FCurveSequence()
 
 	if (TickerHandle.IsValid())
 	{
-		FTicker::GetCoreTicker().RemoveTicker(TickerHandle);
+		FTSTicker::GetCoreTicker().RemoveTicker(TickerHandle);
 	}
 }
 
@@ -94,7 +94,7 @@ void FCurveSequence::Play(const FTickerDelegate& InDelegate, bool bPlayLooped, c
 	SetStartTime(FSlateApplicationBase::Get().GetCurrentTime() - StartAtTime);
 
 	// Register the core ticker so we can start giving callbacks
-	TickerHandle = FTicker::GetCoreTicker().AddTicker(FTickerDelegate::CreateRaw(this, &FCurveSequence::TickPlay, InDelegate));
+	TickerHandle = FTSTicker::GetCoreTicker().AddTicker(FTickerDelegate::CreateRaw(this, &FCurveSequence::TickPlay, InDelegate));
 }
 
 void FCurveSequence::Reverse( )
diff --git a/Engine/Source/Runtime/UMG/Private/Components/ScrollBox.cpp b/Engine/Source/Runtime/UMG/Private/Components/ScrollBox.cpp
--- a/Engine/Source/Runtime/UMG/Private/Components/ScrollBox.cpp
+++ b/Engine/Source/Runtime/UMG/Private/Components/ScrollBox.cpp
@@ -435,7 +435,7 @@ void UScrollBox::OnDescendantSelectedByDesigner( UWidget* DescendantWidget )
 
 		if ( TickHandle.IsValid() )
 		{
-			FTicker::GetCoreTicker().RemoveTicker( TickHandle );
+			FTSTicker::GetCoreTicker().RemoveTicker( TickHandle );
 			TickHandle.Reset();
 		}
 	}
@@ -445,12 +445,12 @@ void UScrollBox::OnDescendantDeselectedByDesigner( UWidget* DescendantWidget )
 {
 	if ( TickHandle.IsValid() )
 	{
-		FTicker::GetCoreTicker().RemoveTicker( TickHandle );
+		FTSTicker::GetCoreTicker().RemoveTicker( TickHandle );
 		TickHandle.Reset();
 	}
 
 	// because we get a deselect before we get a select, we need to delay this call until we're sure we didn't scroll to another widget.
-	TickHandle = FTicker::GetCoreTicker().AddTicker( FTickerDelegate::CreateLambda( [=]( float ) -> bool
+	TickHandle = FTSTicker::GetCoreTicker().AddTicker( FTickerDelegate::CreateLambda( [=]( float ) -> bool
 	                                                                                {
                                                                                         QUICK_SCOPE_CYCLE_COUNTER(STAT_UScrollBox_ScrollToStart_LambdaTick);
 		                                                                                this->ScrollToStart();
diff --git a/Engine/Source/Runtime/WebBrowser/Private/WebBrowserSingleton.h b/Engine/Source/Runtime/WebBrowser/Private/WebBrowserSingleton.h
--- a/Engine/Source/Runtime/WebBrowser/Private/WebBrowserSingleton.h
+++ b/Engine/Source/Runtime/WebBrowser/Private/WebBrowserSingleton.h
@@ -49,7 +49,7 @@ PRAGMA_DISABLE_DEPRECATION_WARNINGS
  */
 class FWebBrowserSingleton
 	: public IWebBrowserSingleton
-	, public FTickerObjectBase
+	, public FTSTickerObjectBase
 {
 public:
 
@@ -155,7 +155,7 @@ public:
 
 public:
 
-	// FTickerObjectBase Interface
+	// FTSTickerObjectBase Interface
 
 	virtual bool Tick(float DeltaTime) override;
 
