Integration of XCurl support provided by Microsoft, and necessary additional adjustments.
Added commandline parameter to enable the use of previous WinHttp implementation.

diff --git a/Engine/Source/Runtime/Online/HTTP/HTTP.Build.cs b/Engine/Source/Runtime/Online/HTTP/HTTP.Build.cs
--- a/Engine/Source/Runtime/Online/HTTP/HTTP.Build.cs
+++ b/Engine/Source/Runtime/Online/HTTP/HTTP.Build.cs
@@ -1,7 +1,6 @@
 // Copyright Epic Games, Inc. All Rights Reserved.
 
 using UnrealBuildTool;
-using EpicGames.Core;
 
 public class HTTP : ModuleRules
 {
@@ -23,6 +22,9 @@ public class HTTP : ModuleRules
 				Target.Platform == UnrealTargetPlatform.Switch;
 		}
 	}
+	protected virtual bool bPlatformSupportsXCurl { get { return false; } }
+
+	private bool bPlatformSupportsCurl { get { return bPlatformSupportsLibCurl || bPlatformSupportsXCurl; } }
 
 	protected virtual bool bPlatformRequiresOpenSSL
 	{
@@ -56,25 +58,29 @@ public class HTTP : ModuleRules
 			}
 			);
 
-		if (bPlatformSupportsLibCurl)
+		if (bPlatformSupportsCurl)
 		{
-			AddEngineThirdPartyPrivateStaticDependencies(Target, "libcurl");
 			PrivateDependencyModuleNames.AddRange(
 				new string[] {
 					"Sockets",
 				}
 			);
 
-			PublicDefinitions.Add("CURL_ENABLE_DEBUG_CALLBACK=1");
-			if (Target.Configuration != UnrealTargetConfiguration.Shipping)
+			if (bPlatformSupportsLibCurl && !bPlatformSupportsXCurl)
 			{
-				PublicDefinitions.Add("CURL_ENABLE_NO_TIMEOUTS_OPTION=1");
+				AddEngineThirdPartyPrivateStaticDependencies(Target, "libcurl");
+
+				PublicDefinitions.Add("CURL_ENABLE_DEBUG_CALLBACK=1");
+				if (Target.Configuration != UnrealTargetConfiguration.Shipping)
+				{
+					PublicDefinitions.Add("CURL_ENABLE_NO_TIMEOUTS_OPTION=1");
+				}
 			}
 		}
-		else
-		{
-			PublicDefinitions.Add("WITH_LIBCURL=0");
-		}
+
+		PrivateDefinitions.Add("WITH_CURL_LIBCURL =" + (bPlatformSupportsLibCurl ? "1" : "0"));
+		PrivateDefinitions.Add("WITH_CURL_XCURL=" + (bPlatformSupportsXCurl ? "1" : "0"));
+		PrivateDefinitions.Add("WITH_CURL= " + ((bPlatformSupportsLibCurl || bPlatformSupportsXCurl) ? "1" : "0"));
 
 		// Use Curl over WinHttp on platforms that support it (until WinHttp client security is in a good place at the least)
 		if (bPlatformSupportsWinHttp)
diff --git a/Engine/Source/Runtime/Online/HTTP/Private/Curl/CurlHttp.cpp b/Engine/Source/Runtime/Online/HTTP/Private/Curl/CurlHttp.cpp
--- a/Engine/Source/Runtime/Online/HTTP/Private/Curl/CurlHttp.cpp
+++ b/Engine/Source/Runtime/Online/HTTP/Private/Curl/CurlHttp.cpp
@@ -1,6 +1,6 @@
 // Copyright Epic Games, Inc. All Rights Reserved.
 
-#if WITH_LIBCURL
+#if WITH_CURL
 
 #include "Curl/CurlHttp.h"
 #include "Stats/Stats.h"
@@ -833,7 +833,11 @@ bool FCurlHttpRequest::SetupRequestHttpThread()
 			check(!GetHeader(TEXT("Content-Type")).IsEmpty() || RequestPayload->IsURLEncoded());
 			curl_easy_setopt(EasyHandle, CURLOPT_POST, 1L);
 			curl_easy_setopt(EasyHandle, CURLOPT_POSTFIELDS, NULL);
+#if WITH_CURL_XCURL
+			curl_easy_setopt(EasyHandle, CURLOPT_INFILESIZE, RequestPayload->GetContentLength());
+#else
 			curl_easy_setopt(EasyHandle, CURLOPT_POSTFIELDSIZE, RequestPayload->GetContentLength());
+#endif
 			bUseReadFunction = true;
 		}
 		else if (Verb == TEXT("PUT") || Verb == TEXT("PATCH"))
@@ -946,16 +950,16 @@ bool FCurlHttpRequest::SetupRequestHttpThread()
 			curl_easy_setopt(EasyHandle, CURLOPT_SEEKDATA, this);
 			curl_easy_setopt(EasyHandle, CURLOPT_SEEKFUNCTION, StaticSeekCallback);
 		}
+	}
 
-		{
-			//Tracking the locking in the CURLOPT_SHARE branch of the curl_easy_setopt implementation
-			QUICK_SCOPE_CYCLE_COUNTER(STAT_FCurlHttpRequest_SetupRequest_EASY_CURLOPT_SHARE);
+#if !WITH_CURL_XCURL
+	{
+		//Tracking the locking in the CURLOPT_SHARE branch of the curl_easy_setopt implementation
+		QUICK_SCOPE_CYCLE_COUNTER(STAT_FCurlHttpRequest_SetupRequest_EASY_CURLOPT_SHARE);
 
-			curl_easy_setopt(EasyHandle, CURLOPT_SHARE, FCurlHttpManager::GShareHandle);
-		}
+		curl_easy_setopt(EasyHandle, CURLOPT_SHARE, FCurlHttpManager::GShareHandle);
 	}
-
-	curl_easy_setopt(EasyHandle, CURLOPT_SHARE, FCurlHttpManager::GShareHandle);
+#endif
 
 	UE_LOG(LogHttp, Log, TEXT("%p: Starting %s request to URL='%s'"), this, *Verb, *URL);
 
@@ -1438,4 +1442,4 @@ FString FCurlHttpResponse::GetContentAsString() const
 	return FString(TCHARData.Length(), TCHARData.Get());
 }
 
-#endif //WITH_LIBCURL
+#endif //WITH_CURL
diff --git a/Engine/Source/Runtime/Online/HTTP/Private/Curl/CurlHttp.h b/Engine/Source/Runtime/Online/HTTP/Private/Curl/CurlHttp.h
--- a/Engine/Source/Runtime/Online/HTTP/Private/Curl/CurlHttp.h
+++ b/Engine/Source/Runtime/Online/HTTP/Private/Curl/CurlHttp.h
@@ -12,12 +12,26 @@
 
 class FCurlHttpResponse;
 
-#if WITH_LIBCURL
+#if WITH_CURL
 #if PLATFORM_WINDOWS || PLATFORM_HOLOLENS
 #include "Windows/WindowsHWrapper.h"
 #include "Windows/AllowWindowsPlatformTypes.h"
 #endif
+
+#if WITH_CURL_XCURL
+//We copied this template to include the windows file from WindowsHWrapper's way if including MinWindows.h, since including xcurl.h directly caused gnarly build errors
+#include "CoreTypes.h"
+#include "HAL/PlatformMemory.h"
+#include "Windows/PreWindowsApi.h"
+#ifndef STRICT
+#define STRICT
+#endif
+#include "xcurl.h"
+#include "Windows/PostWindowsApi.h"
+#else
 	#include "curl/curl.h"
+#endif
+
 #if PLATFORM_WINDOWS || PLATFORM_HOLOLENS
 #include "Windows/HideWindowsPlatformTypes.h"
 #endif
@@ -459,4 +473,4 @@ private:
 	int32 volatile bSucceeded;
 };
 
-#endif //WITH_LIBCURL
+#endif //WITH_CURL
diff --git a/Engine/Source/Runtime/Online/HTTP/Private/Curl/CurlHttpManager.cpp b/Engine/Source/Runtime/Online/HTTP/Private/Curl/CurlHttpManager.cpp
--- a/Engine/Source/Runtime/Online/HTTP/Private/Curl/CurlHttpManager.cpp
+++ b/Engine/Source/Runtime/Online/HTTP/Private/Curl/CurlHttpManager.cpp
@@ -2,7 +2,7 @@
 
 #include "Curl/CurlHttpManager.h"
 
-#if WITH_LIBCURL
+#if WITH_CURL
 
 #include "HAL/PlatformFileManager.h"
 #include "HAL/FileManager.h"
@@ -34,7 +34,9 @@
 #endif
 
 CURLM* FCurlHttpManager::GMultiHandle = nullptr;
+#if !WITH_CURL_XCURL
 CURLSH* FCurlHttpManager::GShareHandle = nullptr;
+#endif
 
 FCurlHttpManager::FCurlRequestOptions FCurlHttpManager::CurlRequestOptions;
 
@@ -196,6 +198,7 @@ void FCurlHttpManager::InitCurl()
 			}
 		}
 
+#if !WITH_CURL_XCURL
 		GShareHandle = curl_share_init();
 		if (NULL != GShareHandle)
 		{
@@ -207,6 +210,7 @@ void FCurlHttpManager::InitCurl()
 		{
 			UE_LOG(LogInit, Fatal, TEXT("Could not initialize libcurl share handle!"));
 		}
+#endif
 	}
 	else
 	{
@@ -316,12 +320,14 @@ void FCurlHttpManager::FCurlRequestOptions::Log()
 
 void FCurlHttpManager::ShutdownCurl()
 {
+#if !WITH_CURL_XCURL
 	if (GShareHandle != nullptr)
 	{
 		CURLSHcode ShareCleanupCode = curl_share_cleanup(GShareHandle);
 		ensureMsgf(ShareCleanupCode == CURLSHE_OK, TEXT("CurlShareCleanup failed. ReturnValue=[%d]"), static_cast<int32>(ShareCleanupCode));
 		GShareHandle = nullptr;
 	}
+#endif
 
 	if (GMultiHandle != nullptr)
 	{
@@ -425,4 +431,4 @@ bool FCurlHttpManager::SupportsDynamicProxy() const
 {
 	return true;
 }
-#endif //WITH_LIBCURL
+#endif //WITH_CURL
diff --git a/Engine/Source/Runtime/Online/HTTP/Private/Curl/CurlHttpManager.h b/Engine/Source/Runtime/Online/HTTP/Private/Curl/CurlHttpManager.h
--- a/Engine/Source/Runtime/Online/HTTP/Private/Curl/CurlHttpManager.h
+++ b/Engine/Source/Runtime/Online/HTTP/Private/Curl/CurlHttpManager.h
@@ -7,9 +7,11 @@
 
 class FHttpThread;
 
-#if WITH_LIBCURL
+#if WITH_CURL
 
+#if !WITH_CURL_XCURL
 typedef void CURLSH;
+#endif
 typedef void CURLM;
 
 class FCurlHttpManager : public FHttpManager
@@ -18,7 +20,9 @@ public:
 	static void InitCurl();
 	static void ShutdownCurl();
 	static bool IsInit();
+#if !WITH_CURL_XCURL
 	static CURLSH* GShareHandle;
+#endif
 	static CURLM * GMultiHandle;
 
 	static struct FCurlRequestOptions
@@ -70,4 +74,4 @@ protected:
 	//~ End HttpManager Interface
 };
 
-#endif //WITH_LIBCURL
+#endif //WITH_CURL 
diff --git a/Engine/Source/Runtime/Online/HTTP/Private/Curl/CurlHttpThread.cpp b/Engine/Source/Runtime/Online/HTTP/Private/Curl/CurlHttpThread.cpp
--- a/Engine/Source/Runtime/Online/HTTP/Private/Curl/CurlHttpThread.cpp
+++ b/Engine/Source/Runtime/Online/HTTP/Private/Curl/CurlHttpThread.cpp
@@ -6,7 +6,7 @@
 #include "Curl/CurlHttp.h"
 #include "Curl/CurlHttpManager.h"
 
-#if WITH_LIBCURL
+#if WITH_CURL
 
 FCurlHttpThread::FCurlHttpThread()
 {
diff --git a/Engine/Source/Runtime/Online/HTTP/Private/Curl/CurlHttpThread.h b/Engine/Source/Runtime/Online/HTTP/Private/Curl/CurlHttpThread.h
--- a/Engine/Source/Runtime/Online/HTTP/Private/Curl/CurlHttpThread.h
+++ b/Engine/Source/Runtime/Online/HTTP/Private/Curl/CurlHttpThread.h
@@ -4,7 +4,7 @@
 
 #include "CoreMinimal.h"
 
-#if WITH_LIBCURL
+#if WITH_CURL
 
 #include "HttpThread.h"
 
@@ -12,16 +12,28 @@
 #include "Windows/WindowsHWrapper.h"
 #include "Windows/AllowWindowsPlatformTypes.h"
 #endif
-	#include "curl/curl.h"
+#if WITH_CURL_XCURL
+//We copied this template to include the windows file from WindowsHWrapper's way if including MinWindows.h, since including xcurl.h directly caused gnarly build errors
+#include "CoreTypes.h"
+#include "HAL/PlatformMemory.h"
+#include "Windows/PreWindowsApi.h"
+#ifndef STRICT
+#define STRICT
+#endif
+#include "xcurl.h"
+#include "Windows/PostWindowsApi.h"
+#else
+#include "curl/curl.h"
+#endif
 #if PLATFORM_WINDOWS
 #include "Windows/HideWindowsPlatformTypes.h"
 #endif
 
-#endif //WITH_LIBCURL
+#endif //WITH_CURL
 
 class IHttpThreadedRequest;
 
-#if WITH_LIBCURL
+#if WITH_CURL
 
 class FCurlHttpThread
 	: public FHttpThread
@@ -43,4 +55,4 @@ protected:
 };
 
 
-#endif //WITH_LIBCURL
+#endif //WITH_CURL
diff --git a/Engine/Source/Runtime/Online/HTTP/Private/GenericPlatform/GenericPlatformHttp.cpp b/Engine/Source/Runtime/Online/HTTP/Private/GenericPlatform/GenericPlatformHttp.cpp
--- a/Engine/Source/Runtime/Online/HTTP/Private/GenericPlatform/GenericPlatformHttp.cpp
+++ b/Engine/Source/Runtime/Online/HTTP/Private/GenericPlatform/GenericPlatformHttp.cpp
@@ -125,7 +125,7 @@ bool FGenericPlatformHttp::UsesThreadedHttp()
 {
 	// Many platforms use libcurl.  Our libcurl implementation uses threading.
 	// Platforms that don't use libcurl but have threading will need to override UsesThreadedHttp to return true for their platform.
-	return WITH_LIBCURL;
+	return WITH_CURL;
 }
 
 static bool IsAllowedChar(UTF8CHAR LookupChar)
diff --git a/Engine/Source/Runtime/Online/HTTP/Private/Windows/WindowsPlatformHttp.cpp b/Engine/Source/Runtime/Online/HTTP/Private/Windows/WindowsPlatformHttp.cpp
--- a/Engine/Source/Runtime/Online/HTTP/Private/Windows/WindowsPlatformHttp.cpp
+++ b/Engine/Source/Runtime/Online/HTTP/Private/Windows/WindowsPlatformHttp.cpp
@@ -4,13 +4,13 @@
 #include "Misc/CommandLine.h"
 #include "Misc/Parse.h"
 #include "Misc/ConfigCacheIni.h"
-#if WITH_LIBCURL
+#if WITH_CURL
 	#include "Curl/CurlHttp.h"
 	#include "Curl/CurlHttpManager.h"
-#else // ^^^ WITH_LIBCURL ^^^ // vvv !WITH_LIBCURL vvv
+#else // ^^^ WITH_CURL  ^^^ // vvv !WITH_CURL  vvv
 	#include "WinHttp/WinHttpHttpManager.h"
 	#include "WinHttp/WinHttpHttpRequest.h"
-#endif // !WITH_LIBCURL
+#endif // !WITH_CURL 
 #include "Http.h"
 #include "WinHttp/Support/WinHttpTypes.h" // Always include for OS proxy settings
 
@@ -82,36 +82,36 @@ void FWindowsPlatformHttp::Init()
 
 	FGenericPlatformHttp::Init();
 
-#if WITH_LIBCURL
+#if WITH_CURL
 	FCurlHttpManager::InitCurl();
 #endif // WITH_WINHTTP
 }
 
 void FWindowsPlatformHttp::Shutdown()
 {
-#if WITH_LIBCURL
+#if WITH_CURL
 	FCurlHttpManager::ShutdownCurl();
-#endif // WITH_LIBCURL
+#endif // WITH_CURL 
 
 	FGenericPlatformHttp::Shutdown();
 }
 
 FHttpManager * FWindowsPlatformHttp::CreatePlatformHttpManager()
 {
-#if WITH_LIBCURL
+#if WITH_CURL
 	return new FCurlHttpManager();
-#else // ^^^ WITH_LIBCURL ^^^ // vvv WITH_LIBCURL vvv
+#else // ^^^ WITH_CURL  ^^^ // vvv WITH_CURL  vvv
 	return new FWinHttpHttpManager();
-#endif // !WITH_LIBCURL
+#endif // !WITH_CURL 
 }
 
 IHttpRequest* FWindowsPlatformHttp::ConstructRequest()
 {
-#if WITH_LIBCURL
+#if WITH_CURL
 		return new FCurlHttpRequest();
-#else // ^^^ WITH_LIBCURL ^^^ // vvv WITH_LIBCURL vvv
+#else // ^^^ WITH_CURL  ^^^ // vvv WITH_CURL  vvv
 		return new FWinHttpHttpRequest();
-#endif // !WITH_LIBCURL
+#endif // !WITH_CURL 
 }
 
 FString FWindowsPlatformHttp::GetMimeType(const FString& FilePath)
@@ -204,9 +204,9 @@ bool FWindowsPlatformHttp::IsOperatingSystemProxyInformationSupported()
 bool FWindowsPlatformHttp::VerifyPeerSslCertificate(bool verify)
 {
 	bool prev = false;
-#if WITH_LIBCURL
+#if WITH_CURL
 	prev = FCurlHttpManager::CurlRequestOptions.bVerifyPeer;
 	FCurlHttpManager::CurlRequestOptions.bVerifyPeer = verify;
-#endif // #if WITH_LIBCURL
+#endif // #if WITH_CURL 
 	return prev;
 }
