fix non-unity compile on non-windows microsoft platforms.

#rnx

diff --git a/Engine/Source/Runtime/Online/HTTP/Private/Curl/CurlHttp.h b/Engine/Source/Runtime/Online/HTTP/Private/Curl/CurlHttp.h
--- a/Engine/Source/Runtime/Online/HTTP/Private/Curl/CurlHttp.h
+++ b/Engine/Source/Runtime/Online/HTTP/Private/Curl/CurlHttp.h
@@ -13,27 +13,27 @@
 class FCurlHttpResponse;
 
 #if WITH_CURL
-#if PLATFORM_WINDOWS || PLATFORM_HOLOLENS
-#include "Windows/WindowsHWrapper.h"
-#include "Windows/AllowWindowsPlatformTypes.h"
+#if PLATFORM_MICROSOFT
+#include "Microsoft/WindowsHWrapper.h"
+#include "Microsoft/AllowMicrosoftPlatformTypes.h"
 #endif
 
 #if WITH_CURL_XCURL
 //We copied this template to include the windows file from WindowsHWrapper's way if including MinWindows.h, since including xcurl.h directly caused gnarly build errors
 #include "CoreTypes.h"
 #include "HAL/PlatformMemory.h"
-#include "Windows/PreWindowsApi.h"
+#include "Microsoft/PreWindowsApi.h"
 #ifndef STRICT
 #define STRICT
 #endif
 #include "xcurl.h"
-#include "Windows/PostWindowsApi.h"
+#include "Microsoft/PostWindowsApi.h"
 #else
 	#include "curl/curl.h"
 #endif
 
-#if PLATFORM_WINDOWS || PLATFORM_HOLOLENS
-#include "Windows/HideWindowsPlatformTypes.h"
+#if PLATFORM_MICROSOFT
+#include "Microsoft/HideMicrosoftPlatformTypes.h"
 #endif
 
 #if !defined(CURL_ENABLE_DEBUG_CALLBACK)
diff --git a/Engine/Source/Runtime/Online/HTTP/Private/Curl/CurlHttpThread.h b/Engine/Source/Runtime/Online/HTTP/Private/Curl/CurlHttpThread.h
--- a/Engine/Source/Runtime/Online/HTTP/Private/Curl/CurlHttpThread.h
+++ b/Engine/Source/Runtime/Online/HTTP/Private/Curl/CurlHttpThread.h
@@ -8,25 +8,25 @@
 
 #include "HttpThread.h"
 
-#if PLATFORM_WINDOWS
-#include "Windows/WindowsHWrapper.h"
-#include "Windows/AllowWindowsPlatformTypes.h"
+#if PLATFORM_MICROSOFT
+#include "Microsoft/WindowsHWrapper.h"
+#include "Microsoft/AllowMicrosoftPlatformTypes.h"
 #endif
 #if WITH_CURL_XCURL
 //We copied this template to include the windows file from WindowsHWrapper's way if including MinWindows.h, since including xcurl.h directly caused gnarly build errors
 #include "CoreTypes.h"
 #include "HAL/PlatformMemory.h"
-#include "Windows/PreWindowsApi.h"
+#include "Microsoft/PreWindowsApi.h"
 #ifndef STRICT
 #define STRICT
 #endif
 #include "xcurl.h"
-#include "Windows/PostWindowsApi.h"
+#include "Microsoft/PostWindowsApi.h"
 #else
 #include "curl/curl.h"
 #endif
-#if PLATFORM_WINDOWS
-#include "Windows/HideWindowsPlatformTypes.h"
+#if PLATFORM_MICROSOFT
+#include "Microsoft/HideMicrosoftPlatformTypes.h"
 #endif
 
 #endif //WITH_CURL
