Adding logic to ensure initialisation of WinHttpHttpManager to enable WebSocket functionality

diff --git a/Engine/Source/Runtime/Online/WebSockets/Private/WinHttp/WinHttpWebSocketsManager.cpp b/Engine/Source/Runtime/Online/WebSockets/Private/WinHttp/WinHttpWebSocketsManager.cpp
--- a/Engine/Source/Runtime/Online/WebSockets/Private/WinHttp/WinHttpWebSocketsManager.cpp
+++ b/Engine/Source/Runtime/Online/WebSockets/Private/WinHttp/WinHttpWebSocketsManager.cpp
@@ -4,6 +4,8 @@
 
 #include "WinHttp/WinHttpWebSocketsManager.h"
 #include "WinHttp/WinHttpWebSocket.h"
+#include "WinHttp/WinHttpHttpManager.h"
+#include "Http.h"
 #include "Modules/ModuleManager.h"
 #include "HttpModule.h"
 #include "Containers/BackgroundableTicker.h"
@@ -13,14 +15,32 @@ void FWinHttpWebSocketsManager::InitWebSockets(TArrayView<const FString> Protoco
 {
 	(void)FModuleManager::LoadModuleChecked<FHttpModule>(TEXT("Http"));
 
+	if (FWinHttpHttpManager::GetManager() == nullptr)
+	{
+		InitHttpManager();
+	}
+
 	if (ensure(!TickHandle.IsValid()))
 	{
 		TickHandle = FTSBackgroundableTicker::GetCoreTicker().AddTicker(FTickerDelegate::CreateRaw(this, &FWinHttpWebSocketsManager::GameThreadTick), 0.0f);
 	}
 }
 
+void FWinHttpWebSocketsManager::InitHttpManager()
+{
+	HttpManager = FPlatformHttp::CreateWinHttpHttpManager();
+	check(HttpManager);
+	HttpManager->Initialize();
+}
+
 void FWinHttpWebSocketsManager::ShutdownWebSockets()
 {
+	if (HttpManager)
+	{
+		delete HttpManager;
+		HttpManager = nullptr;
+	}
+
 	for (TWeakPtr<FWinHttpWebSocket>& WeakWebSocket : ActiveWebSockets)
 	{
 		if (TSharedPtr<FWinHttpWebSocket> StrongWebSocket = WeakWebSocket.Pin())
diff --git a/Engine/Source/Runtime/Online/WebSockets/Private/WinHttp/WinHttpWebSocketsManager.h b/Engine/Source/Runtime/Online/WebSockets/Private/WinHttp/WinHttpWebSocketsManager.h
--- a/Engine/Source/Runtime/Online/WebSockets/Private/WinHttp/WinHttpWebSocketsManager.h
+++ b/Engine/Source/Runtime/Online/WebSockets/Private/WinHttp/WinHttpWebSocketsManager.h
@@ -8,6 +8,7 @@
 #include "IWebSocketsManager.h"
 #include "Containers/Ticker.h"
 
+class FHttpManager;
 class FWinHttpWebSocket;
 
 class FWinHttpWebSocketsManager
@@ -25,9 +26,13 @@ public:
 protected:
 	bool GameThreadTick(float DeltaTime);
 
+	void InitHttpManager();
+
 protected:
 	FTSTicker::FDelegateHandle TickHandle;
 
+	FHttpManager* HttpManager = nullptr;
+
 	TArray<TWeakPtr<FWinHttpWebSocket>> ActiveWebSockets;
 };
 
