Addition of checks to prevent unwanted assertions in Curl

diff --git a/Engine/Source/Runtime/Online/HTTP/Private/Curl/CurlHttp.h b/Engine/Source/Runtime/Online/HTTP/Private/Curl/CurlHttp.h
--- a/Engine/Source/Runtime/Online/HTTP/Private/Curl/CurlHttp.h
+++ b/Engine/Source/Runtime/Online/HTTP/Private/Curl/CurlHttp.h
@@ -54,8 +54,14 @@ namespace
 	*/
 	void* CurlMalloc(size_t Size)
 	{
-		check(Size);
-		return FMemory::Malloc(Size);
+		void* Return = NULL;
+
+		if (Size)
+		{
+			Return = FMemory::Malloc(Size);
+		}
+		
+		return Return;
 	}
 
 	/**
@@ -77,8 +83,14 @@ namespace
 	*/
 	void* CurlRealloc(void* Ptr, size_t Size)
 	{
-		check(Size);
-		return FMemory::Realloc(Ptr, Size);
+		void* Return = NULL;
+
+		if (Size)
+		{
+			Return = FMemory::Realloc(Ptr, Size);
+		}
+
+		return Return;
 	}
 
 	/**
