Logic to implement a limit of running threaded http requests, to avoid memory errors.
Added command to launch http requests defined by the user.

diff --git a/Engine/Config/BaseEngine.ini b/Engine/Config/BaseEngine.ini
--- a/Engine/Config/BaseEngine.ini
+++ b/Engine/Config/BaseEngine.ini
@@ -29,6 +29,9 @@ HttpConnectionTimeout=60
 HttpReceiveTimeout=30
 HttpSendTimeout=30
 
+[HTTP.HttpThread]
+RunningThreadedRequestLimit=11
+
 [HTTP.Curl]
 bAllowSeekFunction=false
 
diff --git a/Engine/Source/Runtime/Online/HTTP/Private/HttpManager.cpp b/Engine/Source/Runtime/Online/HTTP/Private/HttpManager.cpp
--- a/Engine/Source/Runtime/Online/HTTP/Private/HttpManager.cpp
+++ b/Engine/Source/Runtime/Online/HTTP/Private/HttpManager.cpp
@@ -118,7 +118,10 @@ void FHttpManager::OnEndFramePostFork()
 
 void FHttpManager::UpdateConfigs()
 {
-	// empty
+	if (Thread)
+	{
+		Thread->UpdateConfigs();
+	}
 }
 
 void FHttpManager::AddGameThreadTask(TFunction<void()>&& Task)
diff --git a/Engine/Source/Runtime/Online/HTTP/Private/HttpModule.cpp b/Engine/Source/Runtime/Online/HTTP/Private/HttpModule.cpp
--- a/Engine/Source/Runtime/Online/HTTP/Private/HttpModule.cpp
+++ b/Engine/Source/Runtime/Online/HTTP/Private/HttpModule.cpp
@@ -185,6 +185,35 @@ bool FHttpModule::HandleHTTPCommand(const TCHAR* Cmd, FOutputDevice& Ar)
 		}
 	}
 #endif
+	else if (FParse::Command(&Cmd, TEXT("LAUNCHREQUESTS")))
+	{
+		FString Verb = FParse::Token(Cmd, false);
+		FString Url = FParse::Token(Cmd, false);
+		int32 NumRequests = FCString::Atoi(*FParse::Token(Cmd, false));
+		bool bCancelRequests = FCString::ToBool(*FParse::Token(Cmd, false));
+
+		TArray<TSharedRef<IHttpRequest, ESPMode::ThreadSafe>> Requests;
+
+		for (int32 i = 0; i < NumRequests; ++i)
+		{
+			TSharedRef<IHttpRequest, ESPMode::ThreadSafe> HttpRequest = FHttpModule::Get().CreateRequest();
+			HttpRequest->SetURL(*Url);
+			HttpRequest->SetVerb(*Verb);
+			HttpRequest->OnProcessRequestComplete().BindLambda([](FHttpRequestPtr HttpRequest, FHttpResponsePtr HttpResponse, bool bSucceeded) {});
+			HttpRequest->ProcessRequest();
+
+			Requests.Add(HttpRequest);
+		}
+
+		if (bCancelRequests)
+		{
+			for (auto Request : Requests)
+			{
+				Request->CancelRequest();
+			}
+		}
+	}
+
 	return true;
 }
 
diff --git a/Engine/Source/Runtime/Online/HTTP/Private/HttpThread.cpp b/Engine/Source/Runtime/Online/HTTP/Private/HttpThread.cpp
--- a/Engine/Source/Runtime/Online/HTTP/Private/HttpThread.cpp
+++ b/Engine/Source/Runtime/Online/HTTP/Private/HttpThread.cpp
@@ -5,6 +5,7 @@
 #include "HAL/PlatformTime.h"
 #include "HAL/PlatformProcess.h"
 #include "HAL/RunnableThread.h"
+#include "Misc/ConfigCacheIni.h"
 #include "Misc/CommandLine.h"
 #include "Misc/Fork.h"
 #include "Misc/Parse.h"
@@ -80,7 +81,7 @@ void FHttpThread::StopThread()
 
 void FHttpThread::AddRequest(IHttpThreadedRequest* Request)
 {
-	PendingThreadedRequests.Enqueue(Request);
+	NewThreadedRequests.Enqueue(Request);
 }
 
 void FHttpThread::CancelRequest(IHttpThreadedRequest* Request)
@@ -102,6 +103,9 @@ bool FHttpThread::Init()
 {
 	LastTime = FPlatformTime::Seconds();
 	ExitRequest.Set(false);
+
+	UpdateConfigs();
+
 	return true;
 }
 
@@ -109,7 +113,6 @@ uint32 FHttpThread::Run()
 {
 	// Arrays declared outside of loop to re-use memory
 	TArray<IHttpThreadedRequest*> RequestsToCancel;
-	TArray<IHttpThreadedRequest*> RequestsToStart;
 	TArray<IHttpThreadedRequest*> RequestsToComplete;
 	while (!ExitRequest.GetValue())
 	{
@@ -122,7 +125,7 @@ uint32 FHttpThread::Run()
 			{
 				const double InnerLoopBegin = FPlatformTime::Seconds();
 			
-				Process(RequestsToCancel, RequestsToStart, RequestsToComplete);
+				Process(RequestsToCancel, RequestsToComplete);
 			
 				if (RunningThreadedRequests.Num() == 0)
 				{
@@ -159,9 +162,8 @@ void FHttpThread::Tick()
 	if (ensure(bIsSingleThread))
 	{
 		TArray<IHttpThreadedRequest*> RequestsToCancel;
-		TArray<IHttpThreadedRequest*> RequestsToStart;
 		TArray<IHttpThreadedRequest*> RequestsToComplete;
-		Process(RequestsToCancel, RequestsToStart, RequestsToComplete);
+		Process(RequestsToCancel, RequestsToComplete);
 	}
 }
 
@@ -170,6 +172,16 @@ bool FHttpThread::NeedsSingleThreadTick() const
 	return bIsSingleThread;
 }
 
+void FHttpThread::UpdateConfigs()
+{
+	GConfig->GetInt(TEXT("HTTP.HttpThread"), TEXT("RunningThreadedRequestLimit"), RunningThreadedRequestLimit, GEngineIni);
+	if (RunningThreadedRequestLimit < 1)
+	{
+		UE_LOG(LogHttp, Warning, TEXT("RunningThreadedRequestLimit must be configured as a number greater than 0. Current value is %d."), RunningThreadedRequestLimit);
+		RunningThreadedRequestLimit = INT_MAX;
+	}
+}
+
 void FHttpThread::HttpThreadTick(float DeltaSeconds)
 {
 	// empty
@@ -185,11 +197,11 @@ void FHttpThread::CompleteThreadedRequest(IHttpThreadedRequest* Request)
 	// empty
 }
 
-void FHttpThread::Process(TArray<IHttpThreadedRequest*>& RequestsToCancel, TArray<IHttpThreadedRequest*>& RequestsToStart, TArray<IHttpThreadedRequest*>& RequestsToComplete)
+void FHttpThread::Process(TArray<IHttpThreadedRequest*>& RequestsToCancel, TArray<IHttpThreadedRequest*>& RequestsToComplete)
 {
 	SCOPE_CYCLE_COUNTER(STAT_HTTPThread_Process);
 
-	// cache all cancelled and pending requests
+	// cache all cancelled and new requests
 	{
 		IHttpThreadedRequest* Request = nullptr;
 
@@ -199,10 +211,9 @@ void FHttpThread::Process(TArray<IHttpThreadedRequest*>& RequestsToCancel, TArra
 			RequestsToCancel.Add(Request);
 		}
 
-		RequestsToStart.Reset();
-		while (PendingThreadedRequests.Dequeue(Request))
+		while (NewThreadedRequests.Dequeue(Request))
 		{
-			RequestsToStart.Add(Request);
+			RateLimitedThreadedRequests.Add(Request);
 		}
 	}
 
@@ -213,6 +224,14 @@ void FHttpThread::Process(TArray<IHttpThreadedRequest*>& RequestsToCancel, TArra
 		{
 			RequestsToComplete.AddUnique(Request);
 		}
+		else if (RateLimitedThreadedRequests.Remove(Request) > 0)
+		{
+			RequestsToComplete.AddUnique(Request);
+		}
+		else
+		{
+			UE_LOG(LogHttp, Warning, TEXT("Unable to find request (%p) in HttpThread"), Request);
+		}
 	}
 
 	const double AppTime = FPlatformTime::Seconds();
@@ -228,21 +247,30 @@ void FHttpThread::Process(TArray<IHttpThreadedRequest*>& RequestsToCancel, TArra
 		Request->TickThreadedRequest(ElapsedTime);
 	}
 
-	// Start any pending requests
+	// We'll start rate limited requests until we hit the limit
 	// Tick new requests separately from existing RunningThreadedRequests so they get a chance 
 	// to send unaffected by possibly large ElapsedTime above
-	for (IHttpThreadedRequest* Request : RequestsToStart)
+	int32 RunningThreadedRequestsCounter = RunningThreadedRequests.Num();
+	if (RunningThreadedRequestsCounter < RunningThreadedRequestLimit)
 	{
-		SCOPE_CYCLE_COUNTER(STAT_HTTPThread_StartThreadedRequest);
-
-		if (StartThreadedRequest(Request))
-		{
-			RunningThreadedRequests.Add(Request);
-			Request->TickThreadedRequest(0.0f);
-		}
-		else
+		for (RunningThreadedRequestsCounter; RunningThreadedRequestsCounter < RunningThreadedRequestLimit && RateLimitedThreadedRequests.Num();)
 		{
-			RequestsToComplete.AddUnique(Request);
+			SCOPE_CYCLE_COUNTER(STAT_HTTPThread_StartThreadedRequest);
+
+			IHttpThreadedRequest* ReadyThreadedRequest = RateLimitedThreadedRequests[0];
+			RateLimitedThreadedRequests.RemoveAt(0);
+
+			if (StartThreadedRequest(ReadyThreadedRequest))
+			{
+				RunningThreadedRequestsCounter++;
+				RunningThreadedRequests.Add(ReadyThreadedRequest);
+				ReadyThreadedRequest->TickThreadedRequest(0.0f);
+				UE_LOG(LogHttp, Verbose, TEXT("Started running threaded request (%p). Running threaded requests (%d) Rate limited threaded requests (%d)"), ReadyThreadedRequest, RunningThreadedRequests.Num(), RateLimitedThreadedRequests.Num());
+			}
+			else
+			{
+				RequestsToComplete.AddUnique(ReadyThreadedRequest);
+			}
 		}
 	}
 
@@ -266,6 +294,7 @@ void FHttpThread::Process(TArray<IHttpThreadedRequest*>& RequestsToCancel, TArra
 			RequestsToComplete.AddUnique(Request);
 			RunningThreadedRequests.RemoveAtSwap(Index);
 			--Index;
+			UE_LOG(LogHttp, Verbose, TEXT("Threaded request (%p) completed. Running threaded requests (%d)"), Request, RunningThreadedRequests.Num());
 		}
 	}
 
diff --git a/Engine/Source/Runtime/Online/HTTP/Private/HttpThread.h b/Engine/Source/Runtime/Online/HTTP/Private/HttpThread.h
--- a/Engine/Source/Runtime/Online/HTTP/Private/HttpThread.h
+++ b/Engine/Source/Runtime/Online/HTTP/Private/HttpThread.h
@@ -69,6 +69,11 @@ public:
 	 */
 	bool NeedsSingleThreadTick() const;
 
+	/**
+	 * Update configuration. Called when config has been updated and we need to apply any changes.
+	 */
+	virtual void UpdateConfigs();
+
 protected:
 
 	/**
@@ -103,7 +108,7 @@ protected:
 	*/
 	virtual class FSingleThreadRunnable* GetSingleThreadInterface() override { return this; }
 
-	void Process(TArray<IHttpThreadedRequest*>& RequestsToCancel, TArray<IHttpThreadedRequest*>& RequestsToStart, TArray<IHttpThreadedRequest*>& RequestsToComplete);
+	void Process(TArray<IHttpThreadedRequest*>& RequestsToCancel, TArray<IHttpThreadedRequest*>& RequestsToComplete);
 
 	/** signal request to stop and exit thread */
 	FThreadSafeCounter ExitRequest;
@@ -124,7 +129,7 @@ protected:
 	 * Threaded requests that are waiting to be processed on the http thread.
 	 * Added to on (any) non-HTTP thread, processed then cleared on HTTP thread.
 	 */
-	TQueue<IHttpThreadedRequest*, EQueueMode::Mpsc> PendingThreadedRequests;
+	TQueue<IHttpThreadedRequest*, EQueueMode::Mpsc> NewThreadedRequests;
 
 	/**
 	 * Threaded requests that are waiting to be cancelled on the http thread.
@@ -132,6 +137,12 @@ protected:
 	 */
 	TQueue<IHttpThreadedRequest*, EQueueMode::Mpsc> CancelledThreadedRequests;
 
+	/**
+	 * Threaded requests that are ready to run, but waiting due to the running request limit (not in any of the other lists, except potentially CancelledThreadedRequests).
+	 * Only accessed on the HTTP thread.
+	 */
+	TArray<IHttpThreadedRequest*> RateLimitedThreadedRequests;
+
 	/**
 	 * Currently running threaded requests (not in any of the other lists, except potentially CancelledThreadedRequests).
 	 * Only accessed on the HTTP thread.
@@ -154,4 +165,7 @@ private:
 
 	/** Tells if the runnable thread is running or stopped */
 	bool bIsStopped;
+
+	/** Limit for threaded http requests running at the same time. If not specified through configuration values, there will be no limit */
+	int32 RunningThreadedRequestLimit = INT_MAX;
 };
