Refactor to fix build errors in some platforms

diff --git a/Engine/Source/Runtime/Online/HTTP/Private/HttpThread.cpp b/Engine/Source/Runtime/Online/HTTP/Private/HttpThread.cpp
--- a/Engine/Source/Runtime/Online/HTTP/Private/HttpThread.cpp
+++ b/Engine/Source/Runtime/Online/HTTP/Private/HttpThread.cpp
@@ -253,7 +253,7 @@ void FHttpThread::Process(TArray<IHttpThreadedRequest*>& RequestsToCancel, TArra
 	int32 RunningThreadedRequestsCounter = RunningThreadedRequests.Num();
 	if (RunningThreadedRequestsCounter < RunningThreadedRequestLimit)
 	{
-		for (RunningThreadedRequestsCounter; RunningThreadedRequestsCounter < RunningThreadedRequestLimit && RateLimitedThreadedRequests.Num();)
+		while(RunningThreadedRequestsCounter < RunningThreadedRequestLimit && RateLimitedThreadedRequests.Num())
 		{
 			SCOPE_CYCLE_COUNTER(STAT_HTTPThread_StartThreadedRequest);
 
