Null out CurlHttp response in ProcessRequest
Fixes race condition where a reused (via HTTP retry manager) request has a reference to the old Response and attempts to access it, and the HTTP thread can re-assign Response, deleting the instance referenced on the game thread.
Fixes crash observed in FCurlHttpRequest::BroadcastNewlyReceivedHeaders

diff --git a/Engine/Source/Runtime/Online/HTTP/Private/Curl/CurlHttp.cpp b/Engine/Source/Runtime/Online/HTTP/Private/Curl/CurlHttp.cpp
--- a/Engine/Source/Runtime/Online/HTTP/Private/Curl/CurlHttp.cpp
+++ b/Engine/Source/Runtime/Online/HTTP/Private/Curl/CurlHttp.cpp
@@ -970,6 +970,9 @@ bool FCurlHttpRequest::ProcessRequest()
 	QUICK_SCOPE_CYCLE_COUNTER(STAT_FCurlHttpRequest_ProcessRequest);
 	check(EasyHandle);
 
+	// Clear out response. If this is a re-used request, Response could point to a stale response until SetupRequestHttpThread is called
+	Response = nullptr;
+
 	bool bStarted = false;
 	if (!FHttpModule::Get().GetHttpManager().IsDomainAllowed(URL))
 	{
@@ -993,9 +996,6 @@ bool FCurlHttpRequest::ProcessRequest()
 
 	if (!bStarted)
 	{
-		// No response since connection failed
-		Response = nullptr;
-
 		if (!IsInGameThread())
 		{
 			// Always finish on the game thread
@@ -1155,10 +1155,6 @@ void FCurlHttpRequest::BroadcastNewlyReceivedHeaders()
 	check(IsInGameThread());
 	if (Response.IsValid())
 	{
-		// temporary to help debugging of crashes due to memory corruption issues
-		TCHAR RequestUrl[256];
-		FCString::Strncpy(RequestUrl, *URL, 256);
-
 		// Process the headers received on the HTTP thread and merge them into our master list and then broadcast the new headers
 		TPair<FString, FString> NewHeader;
 		while (Response->NewlyReceivedHeaders.Dequeue(NewHeader))
@@ -1338,7 +1334,7 @@ void FCurlHttpRequest::FinishedRequest()
 		OnProcessRequestComplete().ExecuteIfBound(SharedThis(this), Response, false);
 
 		//Delegate needs to know about the errors -- so nuke Response (since connection failed) afterwards...
-		Response = NULL;
+		Response = nullptr;
 	}
 }
 
