Get Build Patch Services Installer running on a servers as long as we are past the fork point.

diff --git a/Engine/Source/Runtime/Online/BuildPatchServices/Private/BuildPatchFileConstructor.cpp b/Engine/Source/Runtime/Online/BuildPatchServices/Private/BuildPatchFileConstructor.cpp
--- a/Engine/Source/Runtime/Online/BuildPatchServices/Private/BuildPatchFileConstructor.cpp
+++ b/Engine/Source/Runtime/Online/BuildPatchServices/Private/BuildPatchFileConstructor.cpp
@@ -275,9 +275,21 @@ FBuildPatchFileConstructor::FBuildPatchFileConstructor(FFileConstructorConfig In
 		}
 		ConstructionStack[(ConstructListNum - 1) - ConstructListIdx] = ConstructListElem;
 	}
+
 	// Start thread!
 	const TCHAR* ThreadName = TEXT("FileConstructorThread");
-	Thread = FRunnableThread::Create(this, ThreadName);
+
+	// Ideally this would check if we were forkable or a forked child process but there is 
+	// currently no in-engine way to check that.  Since BPS does not currently support
+	// FRunnableThread::ThreadType::Fake or FRunnableThread::ThreadType::Forkable 
+	// this check ends up being equivalent for now.  We most likely *never* want support 
+	// forking while an installer is running.
+	if (FPlatformProcess::SupportsMultithreading() || FForkProcessHelper::IsForkedMultithreadInstance())
+	{
+		Thread = FForkProcessHelper::CreateForkableThread(this, ThreadName);
+		check(Thread != nullptr);
+		check(Thread->GetThreadType() == FRunnableThread::ThreadType::Real);
+	}
 }
 
 FBuildPatchFileConstructor::~FBuildPatchFileConstructor()
diff --git a/Engine/Source/Runtime/Online/BuildPatchServices/Private/BuildPatchInstaller.cpp b/Engine/Source/Runtime/Online/BuildPatchServices/Private/BuildPatchInstaller.cpp
--- a/Engine/Source/Runtime/Online/BuildPatchServices/Private/BuildPatchInstaller.cpp
+++ b/Engine/Source/Runtime/Online/BuildPatchServices/Private/BuildPatchInstaller.cpp
@@ -580,7 +580,19 @@ namespace BuildPatchServices
 		{
 			// Start thread!
 			const TCHAR* ThreadName = TEXT("BuildPatchInstallerThread");
-			Thread = FRunnableThread::Create(this, ThreadName);
+
+			// Ideally this would check if we were forkable or a forked child process but there is 
+			// currently no in-engine way to check that.  Since BPS does not currently support
+			// FRunnableThread::ThreadType::Fake or FRunnableThread::ThreadType::Forkable 
+			// this check ends up being equivalent for now.  We most likely *never* want support 
+			// forking while an installer is running.
+			if (FPlatformProcess::SupportsMultithreading() || FForkProcessHelper::IsForkedMultithreadInstance())
+			{
+				Thread = FForkProcessHelper::CreateForkableThread(this, ThreadName);
+				check(Thread != nullptr);
+				check(Thread->GetThreadType() == FRunnableThread::ThreadType::Real);
+			}
+
 			StartDelegate.ExecuteIfBound(AsShared());
 		}
 		return Thread != nullptr;
diff --git a/Engine/Source/Runtime/Online/BuildPatchServices/Private/Installer/ChunkDbChunkSource.cpp b/Engine/Source/Runtime/Online/BuildPatchServices/Private/Installer/ChunkDbChunkSource.cpp
--- a/Engine/Source/Runtime/Online/BuildPatchServices/Private/Installer/ChunkDbChunkSource.cpp
+++ b/Engine/Source/Runtime/Online/BuildPatchServices/Private/Installer/ChunkDbChunkSource.cpp
@@ -223,7 +223,7 @@ namespace BuildPatchServices
 
 		// Start threaded load worker.
 		TFunction<void()> Task = [this]() { return ThreadRun(); };
-		Future = Async(EAsyncExecution::Thread, MoveTemp(Task));
+		Future = Async(EAsyncExecution::ThreadIfForkSafe, MoveTemp(Task));
 	}
 
 	FChunkDbChunkSource::~FChunkDbChunkSource()
diff --git a/Engine/Source/Runtime/Online/BuildPatchServices/Private/Installer/CloudChunkSource.cpp b/Engine/Source/Runtime/Online/BuildPatchServices/Private/Installer/CloudChunkSource.cpp
--- a/Engine/Source/Runtime/Online/BuildPatchServices/Private/Installer/CloudChunkSource.cpp
+++ b/Engine/Source/Runtime/Online/BuildPatchServices/Private/Installer/CloudChunkSource.cpp
@@ -247,7 +247,7 @@ namespace BuildPatchServices
 		, DownloadCount(InDownloadConnectionCount)
 	{
 		TFunction<void()> Task = [this]() { return ThreadRun(); };
-		Future = Async(EAsyncExecution::Thread, MoveTemp(Task));
+		Future = Async(EAsyncExecution::ThreadIfForkSafe, MoveTemp(Task));
 	}
 
 	FCloudChunkSource::~FCloudChunkSource()
diff --git a/Engine/Source/Runtime/Online/BuildPatchServices/Private/Installer/DiskChunkStore.cpp b/Engine/Source/Runtime/Online/BuildPatchServices/Private/Installer/DiskChunkStore.cpp
--- a/Engine/Source/Runtime/Online/BuildPatchServices/Private/Installer/DiskChunkStore.cpp
+++ b/Engine/Source/Runtime/Online/BuildPatchServices/Private/Installer/DiskChunkStore.cpp
@@ -78,7 +78,7 @@ namespace BuildPatchServices
 	{
 		QueueTrigger = FPlatformProcess::GetSynchEventFromPool(true);
 		IoThreadTrigger = FPlatformProcess::GetSynchEventFromPool(true);
-		IoThreadFuture = Async(EAsyncExecution::Thread, [this]()
+		IoThreadFuture = Async(EAsyncExecution::ThreadIfForkSafe, [this]()
 		{
 			return IoThread();
 		});
