Fix race condition between HTTP Retry Manager and HTTP Thread
Separate public facing status from internal status

diff --git a/Engine/Source/Runtime/Online/HTTP/Private/WinHttp/WinHttpHttpRequest.cpp b/Engine/Source/Runtime/Online/HTTP/Private/WinHttp/WinHttpHttpRequest.cpp
--- a/Engine/Source/Runtime/Online/HTTP/Private/WinHttp/WinHttpHttpRequest.cpp
+++ b/Engine/Source/Runtime/Online/HTTP/Private/WinHttp/WinHttpHttpRequest.cpp
@@ -233,6 +233,7 @@ bool FWinHttpHttpRequest::ProcessRequest()
 	RequestFinishTimeSeconds.Reset();
 	bRequestCancelled = false;
 
+	CompletionStatus = EHttpRequestStatus::Processing;
 	State = EHttpRequestStatus::Processing;
 
 	TSharedRef<FWinHttpHttpRequest, ESPMode::ThreadSafe> LocalStrongThis = StaticCastSharedRef<FWinHttpHttpRequest>(AsShared());
@@ -254,7 +255,7 @@ bool FWinHttpHttpRequest::ProcessRequest()
 		{
 			// Could not create session
 			UE_LOG(LogHttp, Warning, TEXT("Unable to create WinHttp Session, failing request"));
-			StrongThis->FinishRequestOnGameThread();
+			StrongThis->OnWinHttpRequestComplete();
 			return;
 		}
 
@@ -265,7 +266,7 @@ bool FWinHttpHttpRequest::ProcessRequest()
 		if (!LocalConnection.IsValid())
 		{
 			UE_LOG(LogHttp, Warning, TEXT("Unable to create WinHttp Connection, failing request"));
-			StrongThis->FinishRequestOnGameThread();
+			StrongThis->OnWinHttpRequestComplete();
 			return;
 		}
 
@@ -280,7 +281,7 @@ bool FWinHttpHttpRequest::ProcessRequest()
 		if (!LocalConnection->StartRequest())
 		{
 			UE_LOG(LogHttp, Warning, TEXT("Unable to start WinHttp Connection, failing request"));
-			StrongThisRef->FinishRequestOnGameThread();
+			StrongThisRef->OnWinHttpRequestComplete();
 			return;
 		}
 
@@ -308,14 +309,31 @@ void FWinHttpHttpRequest::CancelRequest()
 		return;
 	}
 
+	// FinishRequest will cleanup connection
 	bRequestCancelled = true;
 
-	FinishRequestOnGameThread();
+	FHttpManager& HttpManager = FHttpModule::Get().GetHttpManager();
+	if (HttpManager.IsValidRequest(this))
+	{
+		HttpManager.CancelThreadedRequest(SharedThis(this));
+	}
+	else if (!IsInGameThread())
+	{
+		// Always finish on the game thread
+		FHttpModule::Get().GetHttpManager().AddGameThreadTask([StrongThis = StaticCastSharedRef<FWinHttpHttpRequest>(AsShared())]()
+		{
+			StrongThis->FinishRequest();
+		});
+	}
+	else
+	{
+		FinishRequest();
+	}
 }
 
 EHttpRequestStatus::Type FWinHttpHttpRequest::GetStatus() const
 {
-	return State;
+	return CompletionStatus;
 }
 
 const FHttpResponsePtr FWinHttpHttpRequest::GetResponse() const
@@ -358,19 +376,11 @@ bool FWinHttpHttpRequest::StartThreadedRequest()
 
 bool FWinHttpHttpRequest::IsThreadedRequestComplete()
 {
-	if(bRequestCancelled)
+	if (bRequestCancelled)
 	{
 		return true;
 	}
-
-	if (EHttpRequestStatus::IsFinished(State))
-	{
-		if(RequestStartTimeSeconds.IsSet())
-		{
-			return (FPlatformTime::Seconds() - RequestStartTimeSeconds.GetValue()) >= FHttpModule::Get().GetHttpDelayTime();
-		}		
-	}
-	return false;	
+	return EHttpRequestStatus::IsFinished(State);
 }
 
 void FWinHttpHttpRequest::TickThreadedRequest(float DeltaSeconds)
@@ -383,24 +393,8 @@ void FWinHttpHttpRequest::TickThreadedRequest(float DeltaSeconds)
 	}
 }
 
-void FWinHttpHttpRequest::FinishRequestOnGameThread()
-{	
-	if (!IsInGameThread())
-	{
-		FHttpModule::Get().GetHttpManager().AddGameThreadTask([StrongThis = StaticCastSharedRef<FWinHttpHttpRequest>(AsShared())]()
-		{
-			StrongThis->FinishRequest();
-		});
-	}
-	else
-	{
-		FinishRequest();
-	}
-}
-
-void FWinHttpHttpRequest::FinishRequest()
+void FWinHttpHttpRequest::OnWinHttpRequestComplete()
 {
-	UE_LOG(LogHttp, Verbose, TEXT("FWinHttpHttpRequest::FinishRequest()  req %p State==[%s]"), this,  EHttpRequestStatus::ToString(State));
 	if (RequestFinishTimeSeconds.IsSet())
 	{
 		// Already finished
@@ -409,15 +403,17 @@ void FWinHttpHttpRequest::FinishRequest()
 	RequestFinishTimeSeconds = FPlatformTime::Seconds();
 
 	// Set our final state if it's not set yet
-	if (!EHttpRequestStatus::IsFinished(State) && !EHttpRequestStatus::IsFinished(CompletionStatus))
+	if (!EHttpRequestStatus::IsFinished(State))
 	{
 		State = EHttpRequestStatus::Failed;
 	}
-	else if(!EHttpRequestStatus::IsFinished(State))
-	{
-		State = CompletionStatus;
-	}
-	UE_LOG(LogHttp, Verbose, TEXT("FWinHttpHttpRequest::FinishRequest() req %p State becomes %s]"), this, EHttpRequestStatus::ToString(State));
+}
+
+void FWinHttpHttpRequest::FinishRequest()
+{
+	check(IsInGameThread());
+	check(IsThreadedRequestComplete());
+
 	// Shutdown our connection
 	if (Connection.IsValid())
 	{
@@ -428,6 +424,8 @@ void FWinHttpHttpRequest::FinishRequest()
 		Connection.Reset();
 	}
 
+	CompletionStatus = State;
+
 	TSharedRef<IHttpRequest, ESPMode::ThreadSafe> KeepAlive = AsShared();
 	OnProcessRequestComplete().ExecuteIfBound(KeepAlive, Response, Response.IsValid());
 }
@@ -466,19 +464,19 @@ void FWinHttpHttpRequest::HandleHeaderReceived(const FString& HeaderKey, const F
 	OnHeaderReceived().ExecuteIfBound(AsShared(), HeaderKey, HeaderValue);
 }
 
-void FWinHttpHttpRequest::HandleRequestComplete(EHttpRequestStatus::Type CompletionStatusUpdate)
+void FWinHttpHttpRequest::HandleRequestComplete(EHttpRequestStatus::Type RequestCompletionStatus)
 {
 	check(IsInGameThread());
-	check(EHttpRequestStatus::IsFinished(CompletionStatusUpdate));
+	check(EHttpRequestStatus::IsFinished(RequestCompletionStatus));
 
-	CompletionStatus = CompletionStatusUpdate;
+	State = RequestCompletionStatus;
 
-	if (CompletionStatus == EHttpRequestStatus::Succeeded)
+	if (RequestCompletionStatus == EHttpRequestStatus::Succeeded)
 	{
 		UpdateResponseBody(true);
 	}
 
-	FinishRequestOnGameThread();
+	OnWinHttpRequestComplete();
 }
 
 void FWinHttpHttpRequest::UpdateResponseBody(bool bForceResponseExist)
diff --git a/Engine/Source/Runtime/Online/HTTP/Private/WinHttp/WinHttpHttpRequest.h b/Engine/Source/Runtime/Online/HTTP/Private/WinHttp/WinHttpHttpRequest.h
--- a/Engine/Source/Runtime/Online/HTTP/Private/WinHttp/WinHttpHttpRequest.h
+++ b/Engine/Source/Runtime/Online/HTTP/Private/WinHttp/WinHttpHttpRequest.h
@@ -70,8 +70,7 @@ protected:
 
 	void UpdateResponseBody(bool bForceResponseExist = false);
 
-	void FinishRequestOnGameThread();
-
+	void OnWinHttpRequestComplete();
 private:
 	struct FWinHttpHttpRequestData
 	{
@@ -102,7 +101,7 @@ private:
 	/** Current status of request being processed */
 	EHttpRequestStatus::Type State = EHttpRequestStatus::NotStarted;
 
-	/** Final status of request after being processed. Will be copied to current state in FinishRequest */
+	/** Status of request available via GetStatus */
 	EHttpRequestStatus::Type CompletionStatus = EHttpRequestStatus::NotStarted;
 
 	/** */
