Fix HTTP headers being double null terminated.  This was causing MD5 checks fail in VK.

diff --git a/Engine/Source/Runtime/Online/HTTP/Private/WinHttp/Support/WinHttpConnectionHttp.cpp b/Engine/Source/Runtime/Online/HTTP/Private/WinHttp/Support/WinHttpConnectionHttp.cpp
--- a/Engine/Source/Runtime/Online/HTTP/Private/WinHttp/Support/WinHttpConnectionHttp.cpp
+++ b/Engine/Source/Runtime/Online/HTTP/Private/WinHttp/Support/WinHttpConnectionHttp.cpp
@@ -818,6 +818,11 @@ bool FWinHttpConnectionHttp::ProcessResponseHeaders()
 			{
 				FWideStringView HeaderKey(CompleteHeader.Left(OutIndex));
 				FWideStringView HeaderValue(CompleteHeader.RightChop(OutIndex + 1));
+
+				// Remove trailing NULL terminator from the view. Views should not contain the terminator or
+				// the resulting string will end up double terminated.
+				HeaderValue.RemoveSuffix(1);
+
 				HeaderValue.TrimStartAndEndInline();
 
 				HeadersReceived.Emplace(HeaderKey, HeaderValue);
