Allow ProcessRequest to be called from completion delegate
Formerly would crash if the ref count of the HTTP request went to 0

diff --git a/Engine/Source/Runtime/Online/HTTP/Private/HttpManager.cpp b/Engine/Source/Runtime/Online/HTTP/Private/HttpManager.cpp
--- a/Engine/Source/Runtime/Online/HTTP/Private/HttpManager.cpp
+++ b/Engine/Source/Runtime/Online/HTTP/Private/HttpManager.cpp
@@ -158,9 +158,9 @@ void FHttpManager::Flush(bool bShutdown)
 			UE_LOG(LogHttp, Display, TEXT("Http module shutting down, but needs to wait on %d outstanding Http requests:"), Requests.Num());
 		}
 		// Clear delegates since they may point to deleted instances
-		for (TArray<TSharedRef<IHttpRequest, ESPMode::ThreadSafe>>::TIterator It(Requests); It; ++It)
+		for (TArray<FHttpRequestRef>::TIterator It(Requests); It; ++It)
 		{
-			TSharedRef<IHttpRequest, ESPMode::ThreadSafe>& Request = *It;
+			FHttpRequestRef& Request = *It;
 			Request->OnProcessRequestComplete().Unbind();
 			Request->OnRequestProgress().Unbind();
 			Request->OnHeaderReceived().Unbind();
@@ -190,9 +190,9 @@ void FHttpManager::Flush(bool bShutdown)
 				UE_LOG(LogHttp, Display, TEXT("Canceling remaining HTTP requests after waiting %0.2f seconds"), (AppTime - BeginWaitTime));
 			}
 
-			for (TArray<TSharedRef<IHttpRequest, ESPMode::ThreadSafe>>::TIterator It(Requests); It; ++It)
+			for (TArray<FHttpRequestRef>::TIterator It(Requests); It; ++It)
 			{
-				TSharedRef<IHttpRequest, ESPMode::ThreadSafe>& Request = *It;
+				FHttpRequestRef& Request = *It;
 				FScopedEnterBackgroundEvent(*Request->GetURL());
 				if (IsEngineExitRequested())
 				{
@@ -249,9 +249,9 @@ bool FHttpManager::Tick(float DeltaSeconds)
 	FScopeLock ScopeLock(&RequestLock);
 
 	// Tick each active request
-	for (TArray<TSharedRef<IHttpRequest, ESPMode::ThreadSafe>>::TIterator It(Requests); It; ++It)
+	for (TArray<FHttpRequestRef>::TIterator It(Requests); It; ++It)
 	{
-		TSharedRef<IHttpRequest, ESPMode::ThreadSafe> Request = *It;
+		FHttpRequestRef Request = *It;
 		Request->Tick(DeltaSeconds);
 	}
 
@@ -263,8 +263,9 @@ bool FHttpManager::Tick(float DeltaSeconds)
 		// Finish and remove any completed requests
 		for (IHttpThreadedRequest* CompletedRequest : CompletedThreadedRequests)
 		{
+			FHttpRequestRef CompletedRequestRef = CompletedRequest->AsShared();
+			Requests.Remove(CompletedRequestRef);
 			CompletedRequest->FinishRequest();
-			Requests.Remove(CompletedRequest->AsShared());
 		}
 	}
 	// keep ticking
@@ -276,14 +277,14 @@ void FHttpManager::FlushTick(float DeltaSeconds)
 	Tick(DeltaSeconds);
 }
 
-void FHttpManager::AddRequest(const TSharedRef<IHttpRequest, ESPMode::ThreadSafe>& Request)
+void FHttpManager::AddRequest(const FHttpRequestRef& Request)
 {
 	FScopeLock ScopeLock(&RequestLock);
 	check(!bFlushing);
 	Requests.Add(Request);
 }
 
-void FHttpManager::RemoveRequest(const TSharedRef<IHttpRequest, ESPMode::ThreadSafe>& Request)
+void FHttpManager::RemoveRequest(const FHttpRequestRef& Request)
 {
 	FScopeLock ScopeLock(&RequestLock);
 
@@ -312,7 +313,7 @@ bool FHttpManager::IsValidRequest(const IHttpRequest* RequestPtr) const
 	FScopeLock ScopeLock(&RequestLock);
 
 	bool bResult = false;
-	for (const TSharedRef<IHttpRequest, ESPMode::ThreadSafe>& Request : Requests)
+	for (const FHttpRequestRef& Request : Requests)
 	{
 		if (&Request.Get() == RequestPtr)
 		{
@@ -329,7 +330,7 @@ void FHttpManager::DumpRequests(FOutputDevice& Ar) const
 	FScopeLock ScopeLock(&RequestLock);
 
 	Ar.Logf(TEXT("------- (%d) Http Requests"), Requests.Num());
-	for (const TSharedRef<IHttpRequest, ESPMode::ThreadSafe>& Request : Requests)
+	for (const FHttpRequestRef& Request : Requests)
 	{
 		Ar.Logf(TEXT("	verb=[%s] url=[%s] status=%s"),
 			*Request->GetVerb(), *Request->GetURL(), EHttpRequestStatus::ToString(Request->GetStatus()));
