[UE-131486] HttpServer/ Build fix
+ include full type definition to satisfy usage of TUniquePtr<FHttpServerResponse>

diff --git a/Engine/Source/Runtime/Online/HTTPServer/Private/HttpConnectionResponseWriteContext.h b/Engine/Source/Runtime/Online/HTTPServer/Private/HttpConnectionResponseWriteContext.h
--- a/Engine/Source/Runtime/Online/HTTPServer/Private/HttpConnectionResponseWriteContext.h
+++ b/Engine/Source/Runtime/Online/HTTPServer/Private/HttpConnectionResponseWriteContext.h
@@ -3,9 +3,9 @@
 #include "CoreMinimal.h"
 #include "HttpConnectionContext.h"
 #include "HttpServerHttpVersion.h"
+#include "HttpServerResponse.h"
 
 
-struct FHttpServerResponse;
 class FSocket;
 
 DECLARE_LOG_CATEGORY_EXTERN(LogHttpConnectionResponseWriteContext, Log, All);
