[UE-134338] HttpServer / Build fix
+ include full type definition to satisfy usage of TUniquePtr<FHttpServerResponse>

diff --git a/Engine/Source/Runtime/Online/HTTPServer/Public/HttpResultCallback.h b/Engine/Source/Runtime/Online/HTTPServer/Public/HttpResultCallback.h
--- a/Engine/Source/Runtime/Online/HTTPServer/Public/HttpResultCallback.h
+++ b/Engine/Source/Runtime/Online/HTTPServer/Public/HttpResultCallback.h
@@ -2,7 +2,7 @@
 #pragma once
 
 #include "CoreMinimal.h"
-struct FHttpServerResponse;
+#include "HttpServerResponse.h"
 
 /**
 * FHttpResultCallback
