Make WebSockets protocols configurable via ini

diff --git a/Engine/Config/BaseEngine.ini b/Engine/Config/BaseEngine.ini
--- a/Engine/Config/BaseEngine.ini
+++ b/Engine/Config/BaseEngine.ini
@@ -65,6 +65,14 @@ DeleteTempFilesWithoutURLMappingEntries=true
 ; This could be stale data from a previous run that was either never actually used, wasn't saved after the temp file was moved, etc.
 RemoveURLMappingEntriesWithoutPhysicalTempFiles=true
 
+[WebSockets]
++WebSocketsProtocols="ws"
++WebSocketsProtocols="wss"
++WebSocketsProtocols="v10.stomp"
++WebSocketsProtocols="v11.stomp"
++WebSocketsProtocols="v12.stomp"
++WebSocketsProtocols="xmpp"
+
 [WebSockets.LibWebSockets]
 ThreadStackSize=131072
 ThreadTargetFrameTimeInSeconds=0.0333
diff --git a/Engine/Source/Runtime/Online/WebSockets/Private/WebSocketsModule.cpp b/Engine/Source/Runtime/Online/WebSockets/Private/WebSocketsModule.cpp
--- a/Engine/Source/Runtime/Online/WebSockets/Private/WebSocketsModule.cpp
+++ b/Engine/Source/Runtime/Online/WebSockets/Private/WebSocketsModule.cpp
@@ -5,6 +5,7 @@
 
 #if WITH_WEBSOCKETS
 #include "PlatformWebSocket.h"
+#include "Misc/ConfigCacheIni.h"
 #endif // #if WITH_WEBSOCKETS
 
 DEFINE_LOG_CATEGORY(LogWebSockets);
@@ -29,7 +30,9 @@ void FWebSocketsModule::StartupModule()
 	Singleton = this;
 
 #if WITH_WEBSOCKETS
-	const FString Protocols[] = {TEXT("ws"), TEXT("wss"), TEXT("v10.stomp"), TEXT("v11.stomp"), TEXT("v12.stomp"), TEXT("xmpp")};
+	// Default configuration values can be found in BaseEngine.ini
+	TArray<FString> Protocols;
+	GConfig->GetArray(TEXT("WebSockets"), TEXT("WebSocketsProtocols"), Protocols, GEngineIni);
 
 	WebSocketsManager = new FPlatformWebSocketsManager;
 	WebSocketsManager->InitWebSockets(Protocols);
