Make sure that any remaining Http requests (if any) are cancelled so there is no problem when shutting down.
Add warning if values SoftLimitSeconds and HardLimitSeconds read from ini file are not valid

diff --git a/Engine/Source/Runtime/Online/HTTP/Private/HttpManager.cpp b/Engine/Source/Runtime/Online/HTTP/Private/HttpManager.cpp
--- a/Engine/Source/Runtime/Online/HTTP/Private/HttpManager.cpp
+++ b/Engine/Source/Runtime/Online/HTTP/Private/HttpManager.cpp
@@ -70,6 +70,23 @@ void FHttpManager::InitializeFlushTimeLimits()
 		case EHttpFlushReason::Shutdown:
 			GConfig->GetDouble(TEXT("HTTP"), TEXT("FlushSoftTimeLimitShutdown"), SoftLimitSeconds, GEngineIni);
 			GConfig->GetDouble(TEXT("HTTP"), TEXT("FlushHardTimeLimitShutdown"), HardLimitSeconds, GEngineIni);
+			
+			if ((HardLimitSeconds >= 0) && ((SoftLimitSeconds < 0) || (SoftLimitSeconds >= HardLimitSeconds)))
+			{
+				UE_CLOG(!IsRunningCommandlet(), LogHttp, Warning, TEXT("Soft limit[%.02f] is higher than the hard limit set[%.02f] in file [%s]. Please change the soft limit to a value lower than the hard limit for Flush to work correctly. - 1 is unlimited and therefore the highest possible value."), static_cast<float>(SoftLimitSeconds), static_cast<float>(HardLimitSeconds), *GEngineIni);
+				// we need to be absolutely sure that SoftLimitSeconds is always strictly less than HardLimitSeconds so remaining requests (if any) can be canceled before exiting
+				if (HardLimitSeconds > 0.0)
+				{
+					SoftLimitSeconds = HardLimitSeconds / 2.0;	// clamping SoftLimitSeconds to a reasonable value
+				}
+				else
+				{
+					// HardLimitSeconds should never be 0.0 while shutting down otherwise we can't cancel the remaining requests
+					HardLimitSeconds = 0.05;	// using a non zero value 
+					SoftLimitSeconds = 0.0;		// cancelling request immediately
+				}
+			}
+
 			break;
 		case EHttpFlushReason::FullFlush:
 			SoftLimitSeconds = -1.0;
