Change to CurlMalloc to resolve XCurl errors.
Revert of fix to POST, not needed anymore.

diff --git a/Engine/Source/Runtime/Online/HTTP/Private/Curl/CurlHttp.cpp b/Engine/Source/Runtime/Online/HTTP/Private/Curl/CurlHttp.cpp
--- a/Engine/Source/Runtime/Online/HTTP/Private/Curl/CurlHttp.cpp
+++ b/Engine/Source/Runtime/Online/HTTP/Private/Curl/CurlHttp.cpp
@@ -832,7 +832,7 @@ bool FCurlHttpRequest::SetupRequestHttpThread()
 			// In the case of using a streamed file, you must explicitly set the Content-Type, because RequestPayload->IsURLEncoded returns false.
 			check(!GetHeader(TEXT("Content-Type")).IsEmpty() || RequestPayload->IsURLEncoded());
 			curl_easy_setopt(EasyHandle, CURLOPT_POST, 1L);
-			curl_easy_setopt(EasyHandle, CURLOPT_POSTFIELDS, (RequestPayload->GetContentLength() == 0) ? "" : NULL);
+			curl_easy_setopt(EasyHandle, CURLOPT_POSTFIELDS, NULL);
 #if WITH_CURL_XCURL
 			curl_easy_setopt(EasyHandle, CURLOPT_INFILESIZE, RequestPayload->GetContentLength());
 #else
diff --git a/Engine/Source/Runtime/Online/HTTP/Private/Curl/CurlHttp.h b/Engine/Source/Runtime/Online/HTTP/Private/Curl/CurlHttp.h
--- a/Engine/Source/Runtime/Online/HTTP/Private/Curl/CurlHttp.h
+++ b/Engine/Source/Runtime/Online/HTTP/Private/Curl/CurlHttp.h
@@ -54,14 +54,7 @@ namespace
 	*/
 	void* CurlMalloc(size_t Size)
 	{
-		void* Return = NULL;
-
-		if (Size)
-		{
-			Return = FMemory::Malloc(Size);
-		}
-		
-		return Return;
+		return FMemory::Malloc(Size);
 	}
 
 	/**
