Remove UE4 references
Some of the work was already done in CL 15872931

diff --git a/Engine/Source/Runtime/Online/HTTP/Private/IXML/HttpIXML.h b/Engine/Source/Runtime/Online/HTTP/Private/IXML/HttpIXML.h
--- a/Engine/Source/Runtime/Online/HTTP/Private/IXML/HttpIXML.h
+++ b/Engine/Source/Runtime/Online/HTTP/Private/IXML/HttpIXML.h
@@ -11,7 +11,7 @@
 #include "GenericPlatform/HttpRequestPayload.h"
 
 // Default user agent string
-static const WCHAR USER_AGENT[] = L"UE4HTTPIXML\r\n";
+static const WCHAR USER_AGENT[] = L"UEHTTPIXML\r\n";
 
 
 /**
