HttpResponse codes GetDescription

diff --git a/Engine/Source/Runtime/Online/HTTP/Public/Interfaces/IHttpResponse.h b/Engine/Source/Runtime/Online/HTTP/Public/Interfaces/IHttpResponse.h
--- a/Engine/Source/Runtime/Online/HTTP/Public/Interfaces/IHttpResponse.h
+++ b/Engine/Source/Runtime/Online/HTTP/Public/Interfaces/IHttpResponse.h
@@ -5,6 +5,8 @@
 #include "CoreMinimal.h"
 #include "Interfaces/IHttpBase.h"
 
+#define LOCTEXT_NAMESPACE "IHttpResponse"
+
 namespace EHttpResponseCodes
 {
 	/**
@@ -104,10 +106,68 @@ namespace EHttpResponseCodes
 	{
 		return StatusCode >= Ok && StatusCode <= PartialContent;
 	}
+
+	/**
+	 * @param StatusCode http response code to retrieve
+	 * @return The status code description
+	 */
+	inline FText GetDescription(EHttpResponseCodes::Type StatusCode)
+	{
+		switch (StatusCode)
+		{
+		case EHttpResponseCodes::Continue: return LOCTEXT("HttpResponseCode_100", "Continue");
+		case EHttpResponseCodes::SwitchProtocol: return LOCTEXT("HttpResponseCode_101", "Switching Protocols");
+		case EHttpResponseCodes::Ok: return LOCTEXT("HttpResponseCode_200", "OK");
+		case EHttpResponseCodes::Created: return LOCTEXT("HttpResponseCode_201", "Created");
+		case EHttpResponseCodes::Accepted: return LOCTEXT("HttpResponseCode_202", "Accepted");
+		case EHttpResponseCodes::Partial: return LOCTEXT("HttpResponseCode_203", "Non-Authoritative Information");
+		case EHttpResponseCodes::NoContent: return LOCTEXT("HttpResponseCode_204", "No Content");
+		case EHttpResponseCodes::ResetContent: return LOCTEXT("HttpResponseCode_205", "Reset Content");
+		case EHttpResponseCodes::PartialContent: return LOCTEXT("HttpResponseCode_206", "Partial Content");
+
+		case EHttpResponseCodes::Ambiguous: return LOCTEXT("HttpResponseCode_300", "Multiple Choices");
+		case EHttpResponseCodes::Moved: return LOCTEXT("HttpResponseCode_301", "Moved Permanently");
+		case EHttpResponseCodes::Redirect: return LOCTEXT("HttpResponseCode_302", "Found/Moved temporarily");
+		case EHttpResponseCodes::RedirectMethod: return LOCTEXT("HttpResponseCode_303", "See Other");
+		case EHttpResponseCodes::NotModified: return LOCTEXT("HttpResponseCode_304", "Not Modified");
+		case EHttpResponseCodes::UseProxy: return LOCTEXT("HttpResponseCode_305", "Use Proxy");
+		case EHttpResponseCodes::RedirectKeepVerb: return LOCTEXT("HttpResponseCode_307", "Temporary Redirect");
+
+		case EHttpResponseCodes::BadRequest: return LOCTEXT("HttpResponseCode_400", "Bad Request");
+		case EHttpResponseCodes::Denied: return LOCTEXT("HttpResponseCode_401", "Unauthorized");
+		case EHttpResponseCodes::PaymentReq: return LOCTEXT("HttpResponseCode_402", "Payment Required");
+		case EHttpResponseCodes::Forbidden: return LOCTEXT("HttpResponseCode_403", "Forbidden");
+		case EHttpResponseCodes::NotFound: return LOCTEXT("HttpResponseCode_404", "Not Found");
+		case EHttpResponseCodes::BadMethod: return LOCTEXT("HttpResponseCode_405", "Method Not Allowed");
+		case EHttpResponseCodes::NoneAcceptable: return LOCTEXT("HttpResponseCode_406", "Not Acceptable");
+		case EHttpResponseCodes::ProxyAuthReq: return LOCTEXT("HttpResponseCode_407", "Proxy Authentication Required");
+		case EHttpResponseCodes::RequestTimeout: return LOCTEXT("HttpResponseCode_408", "Request Timeout");
+		case EHttpResponseCodes::Conflict: return LOCTEXT("HttpResponseCode_409", "Conflict");
+		case EHttpResponseCodes::Gone: return LOCTEXT("HttpResponseCode_410", "Gone");
+		case EHttpResponseCodes::LengthRequired: return LOCTEXT("HttpResponseCode_411", "Length Required");
+		case EHttpResponseCodes::PrecondFailed: return LOCTEXT("HttpResponseCode_412", "Precondition Failed");
+		case EHttpResponseCodes::RequestTooLarge: return LOCTEXT("HttpResponseCode_413", "Payload Too Large");
+		case EHttpResponseCodes::UriTooLong: return LOCTEXT("HttpResponseCode_414", "URI Too Long");
+		case EHttpResponseCodes::UnsupportedMedia: return LOCTEXT("HttpResponseCode_415", "Unsupported Media Type");
+		case EHttpResponseCodes::TooManyRequests: return LOCTEXT("HttpResponseCode_429", "Too Many Requests");
+		case EHttpResponseCodes::RetryWith: return LOCTEXT("HttpResponseCode_449", "Retry With");
+
+		case EHttpResponseCodes::ServerError: return LOCTEXT("HttpResponseCode_500", "Internal Server Error");
+		case EHttpResponseCodes::NotSupported: return LOCTEXT("HttpResponseCode_501", "Not Implemented");
+		case EHttpResponseCodes::BadGateway: return LOCTEXT("HttpResponseCode_502", "Bad Gateway");
+		case EHttpResponseCodes::ServiceUnavail: return LOCTEXT("HttpResponseCode_503", "Service Unavailable");
+		case EHttpResponseCodes::GatewayTimeout: return LOCTEXT("HttpResponseCode_504", "Gateway Timeout");
+		case EHttpResponseCodes::VersionNotSup: LOCTEXT("HttpResponseCode_505", "HTTP Version Not Supported");
+
+		case EHttpResponseCodes::Unknown:
+		default:
+			return LOCTEXT("HttpResponseCode_0", "Unknown");
+		}
+	}
 }
 
 /**
- * Inteface for Http responses that come back after starting an Http request
+ * Interface for Http responses that come back after starting an Http request
  */
 class IHttpResponse : public IHttpBase
 {
@@ -133,3 +193,5 @@ public:
 	 */
 	virtual ~IHttpResponse() = default;
 };
+
+#undef LOCTEXT_NAMESPACE
