PR #8724
Moved variable set operation inside critical section.
(Contributed by [REDACTED_PERSON])

diff --git a/Engine/Source/Runtime/Online/HTTP/Private/HttpManager.cpp b/Engine/Source/Runtime/Online/HTTP/Private/HttpManager.cpp
--- a/Engine/Source/Runtime/Online/HTTP/Private/HttpManager.cpp
+++ b/Engine/Source/Runtime/Online/HTTP/Private/HttpManager.cpp
@@ -202,10 +202,11 @@ void FHttpManager::Flush(EHttpFlushReason FlushReason)
 {
 	QUICK_SCOPE_CYCLE_COUNTER(STAT_FHttpManager_Flush);
 
-	// This variable set to true prevents new Http requests from being launched
-	bFlushing = true;
-
 	FScopeLock ScopeLock(&RequestLock);
+	
+	// This variable is set to indicate that flush is happening.
+	// While flushing is in progress, the RequestLock is held and threads are blocked when trying to submit new requests.
+	bFlushing = true;
 
 	double FlushTimeSoftLimitSeconds = FlushTimeLimitsMap[FlushReason].SoftLimitSeconds;
 	double FlushTimeHardLimitSeconds = FlushTimeLimitsMap[FlushReason].HardLimitSeconds;
@@ -378,10 +379,10 @@ void FHttpManager::RemoveRequest(const FHttpRequestRef& Request)
 
 void FHttpManager::AddThreadedRequest(const TSharedRef<IHttpThreadedRequest, ESPMode::ThreadSafe>& Request)
 {
-	check(!bFlushing);
 	check(Thread);
 	{
 		FScopeLock ScopeLock(&RequestLock);
+		check(!bFlushing);
 		Requests.Add(Request);
 	}
 	Thread->AddRequest(&Request.Get());
