Fix for crash on launch on macOS Catalina 10.15.

#rnx

diff --git a/Engine/Source/Runtime/Online/HTTP/Private/Apple/AppleHTTP.cpp b/Engine/Source/Runtime/Online/HTTP/Private/Apple/AppleHTTP.cpp
--- a/Engine/Source/Runtime/Online/HTTP/Private/Apple/AppleHTTP.cpp
+++ b/Engine/Source/Runtime/Online/HTTP/Private/Apple/AppleHTTP.cpp
@@ -712,7 +712,7 @@ static const unsigned char ecdsaSecp384r1Asn1Header[] =
                 SecTrustCreateWithCertificates(Cert, TrustPolicy, CertTrust.GetForAssignment());
                 SecTrustEvaluateWithError(CertTrust, nil);
                 TCFRef<SecKeyRef> CertPubKey;
-                if (@available(iOS 14, tvOS 14, *))
+                if (@available(iOS 14, tvOS 14, macOS 11, *))
                 {
                     CertPubKey = SecTrustCopyKey(CertTrust);
 
