Moved HTTPManager Flush time limit config reload to UpdateConfigs to enable Hotfix

diff --git a/Engine/Source/Runtime/Online/HTTP/Private/HttpManager.cpp b/Engine/Source/Runtime/Online/HTTP/Private/HttpManager.cpp
--- a/Engine/Source/Runtime/Online/HTTP/Private/HttpManager.cpp
+++ b/Engine/Source/Runtime/Online/HTTP/Private/HttpManager.cpp
@@ -45,11 +45,13 @@ void FHttpManager::Initialize()
 		Thread->StartThread();
 	}
 
-	InitializeFlushTimeLimits();
+	UpdateConfigs();
 }
 
-void FHttpManager::InitializeFlushTimeLimits()
+void FHttpManager::ReloadFlushTimeLimits()
 {
+	FlushTimeLimitsMap.Reset();
+
 	//Save int values of Default and FullFlush?
 	for (EHttpFlushReason Reason : TEnumRange<EHttpFlushReason>())
 	{
@@ -174,6 +176,8 @@ void FHttpManager::OnEndFramePostFork()
 
 void FHttpManager::UpdateConfigs()
 {
+	ReloadFlushTimeLimits();
+
 	if (Thread)
 	{
 		Thread->UpdateConfigs();
diff --git a/Engine/Source/Runtime/Online/HTTP/Public/HttpManager.h b/Engine/Source/Runtime/Online/HTTP/Public/HttpManager.h
--- a/Engine/Source/Runtime/Online/HTTP/Public/HttpManager.h
+++ b/Engine/Source/Runtime/Online/HTTP/Public/HttpManager.h
@@ -204,7 +204,7 @@ protected:
 	 */
 	virtual FHttpThread* CreateHttpThread();
 
-	void InitializeFlushTimeLimits();
+	void ReloadFlushTimeLimits();
 
 protected:
 	/** List of Http requests that are actively being processed */
