Fix build regression on macOS Catalina 10.15.7

#rnx

diff --git a/Engine/Source/Runtime/Online/HTTP/Private/Apple/AppleHTTP.cpp b/Engine/Source/Runtime/Online/HTTP/Private/Apple/AppleHTTP.cpp
--- a/Engine/Source/Runtime/Online/HTTP/Private/Apple/AppleHTTP.cpp
+++ b/Engine/Source/Runtime/Online/HTTP/Private/Apple/AppleHTTP.cpp
@@ -712,7 +712,14 @@ static const unsigned char ecdsaSecp384r1Asn1Header[] =
                 SecTrustCreateWithCertificates(Cert, TrustPolicy, CertTrust.GetForAssignment());
                 SecTrustEvaluateWithError(CertTrust, nil);
                 TCFRef<SecKeyRef> CertPubKey;
-                CertPubKey = SecTrustCopyKey(CertTrust);
+				if (@available(macOS 11, *))
+				{
+					CertPubKey = SecTrustCopyKey(CertTrust);
+				}
+				else
+				{
+					CertPubKey = SecTrustCopyPublicKey(CertTrust);
+				}
 				TCFRef<CFDataRef> CertPubKeyData = SecKeyCopyExternalRepresentation(CertPubKey, NULL);
                 if (!CertPubKeyData)
                 {
