Fix iOS game build deprecation warning

#rnx

diff --git a/Engine/Source/Runtime/Online/HTTP/Private/Apple/AppleHTTP.cpp b/Engine/Source/Runtime/Online/HTTP/Private/Apple/AppleHTTP.cpp
--- a/Engine/Source/Runtime/Online/HTTP/Private/Apple/AppleHTTP.cpp
+++ b/Engine/Source/Runtime/Online/HTTP/Private/Apple/AppleHTTP.cpp
@@ -712,14 +712,21 @@ static const unsigned char ecdsaSecp384r1Asn1Header[] =
                 SecTrustCreateWithCertificates(Cert, TrustPolicy, CertTrust.GetForAssignment());
                 SecTrustEvaluateWithError(CertTrust, nil);
                 TCFRef<SecKeyRef> CertPubKey;
+
 				if (@available(macOS 11, *))
 				{
 					CertPubKey = SecTrustCopyKey(CertTrust);
 				}
 				else
 				{
+					PRAGMA_DISABLE_DEPRECATION_WARNINGS
+					// warning: 'SecTrustCopyPublicKey' is deprecated: first deprecated in iOS 14.0 [-Wdeprecated-declarations]
+
 					CertPubKey = SecTrustCopyPublicKey(CertTrust);
+
+					PRAGMA_ENABLE_DEPRECATION_WARNINGS
 				}
+
 				TCFRef<CFDataRef> CertPubKeyData = SecKeyCopyExternalRepresentation(CertPubKey, NULL);
                 if (!CertPubKeyData)
                 {
