Replace/remove uses of whitelist/blacklist in WebSockets

diff --git a/Engine/Source/Runtime/Online/WebSockets/Private/Lws/LwsWebSocket.cpp b/Engine/Source/Runtime/Online/WebSockets/Private/Lws/LwsWebSocket.cpp
--- a/Engine/Source/Runtime/Online/WebSockets/Private/Lws/LwsWebSocket.cpp
+++ b/Engine/Source/Runtime/Online/WebSockets/Private/Lws/LwsWebSocket.cpp
@@ -100,20 +100,20 @@ void FLwsWebSocket::Connect()
 		return;
 	}
 
-	if (!FLwsWebSocketsManager::Get().bDisableDomainWhitelist)
+	if (!FLwsWebSocketsManager::Get().bDisableDomainAllowlist)
 	{
 		FHttpManager& HttpManager = FHttpModule::Get().GetHttpManager();
 		if (!HttpManager.IsDomainAllowed(Url))
 		{
 			State = EState::Error;
-			UE_LOG(LogWebSockets, Warning, TEXT("FLwsWebSocket[%d]::Connect: %s is not whitelisted. Refusing to connect."), Identifier, *Url);
+			UE_LOG(LogWebSockets, Warning, TEXT("FLwsWebSocket[%d]::Connect: %s is not allowlisted. Refusing to connect."), Identifier, *Url);
 			OnConnectionError().Broadcast(TEXT("Invalid Domain"));
 			return;
 		}
 	}
 	else
 	{
-		UE_LOG(LogWebSockets, Log, TEXT("FLwsWebSocket[%d]::Connect: Domain whitelisting has been disabled by config."), Identifier);
+		UE_LOG(LogWebSockets, Log, TEXT("FLwsWebSocket[%d]::Connect: Domain allowlisting has been disabled by config."), Identifier);
 	}
 
 	// No lock, we are not being processed on the websockets thread yet
diff --git a/Engine/Source/Runtime/Online/WebSockets/Private/Lws/LwsWebSocketsManager.cpp b/Engine/Source/Runtime/Online/WebSockets/Private/Lws/LwsWebSocketsManager.cpp
--- a/Engine/Source/Runtime/Online/WebSockets/Private/Lws/LwsWebSocketsManager.cpp
+++ b/Engine/Source/Runtime/Online/WebSockets/Private/Lws/LwsWebSocketsManager.cpp
@@ -56,7 +56,7 @@ FLwsWebSocketsManager::FLwsWebSocketsManager()
 	ThreadMinimumSleepTimeInSeconds = 0.0f;
 	GConfig->GetDouble(TEXT("WebSockets.LibWebSockets"), TEXT("ThreadMinimumSleepTimeInSeconds"), ThreadMinimumSleepTimeInSeconds, GEngineIni);
 
-	GConfig->GetBool(TEXT("LwsWebSocket"), TEXT("bDisableDomainWhitelist"), bDisableDomainWhitelist, GEngineIni);
+	GConfig->GetBool(TEXT("LwsWebSocket"), TEXT("bDisableDomainWhitelist"), bDisableDomainAllowlist, GEngineIni);
 	GConfig->GetBool(TEXT("LwsWebSocket"), TEXT("bDisableCertValidation"), bDisableCertValidation, GEngineIni);
 }
 
diff --git a/Engine/Source/Runtime/Online/WebSockets/Private/Lws/LwsWebSocketsManager.h b/Engine/Source/Runtime/Online/WebSockets/Private/Lws/LwsWebSocketsManager.h
--- a/Engine/Source/Runtime/Online/WebSockets/Private/Lws/LwsWebSocketsManager.h
+++ b/Engine/Source/Runtime/Online/WebSockets/Private/Lws/LwsWebSocketsManager.h
@@ -109,7 +109,7 @@ private:
 	/** Minimum time to sleep in our thread's tick, even if the sleep makes us exceed our target frame time */
 	double ThreadMinimumSleepTimeInSeconds;
 
-	bool bDisableDomainWhitelist = false;
+	bool bDisableDomainAllowlist = false;
 	bool bDisableCertValidation = false;
 
 	friend class FLwsWebSocket;
diff --git a/Engine/Source/Runtime/Online/WebSockets/Private/WinHttp/WinHttpWebSocket.cpp b/Engine/Source/Runtime/Online/WebSockets/Private/WinHttp/WinHttpWebSocket.cpp
--- a/Engine/Source/Runtime/Online/WebSockets/Private/WinHttp/WinHttpWebSocket.cpp
+++ b/Engine/Source/Runtime/Online/WebSockets/Private/WinHttp/WinHttpWebSocket.cpp
@@ -74,22 +74,22 @@ void FWinHttpWebSocket::Connect()
 
 	State = EWebSocketConnectionState::Connecting;
 
-	// Check Domain whitelist if enabled
-	bool bDisableDomainWhitelist = false;
-	GConfig->GetBool(TEXT("WinHttpWebSocket"), TEXT("bDisableDomainWhitelist"), bDisableDomainWhitelist, GEngineIni);
-	if (!bDisableDomainWhitelist)
+	// Check Domain allowedlist if enabled
+	bool bDisableDomainAllowlist = false;
+	GConfig->GetBool(TEXT("WinHttpWebSocket"), TEXT("bDisableDomainWhitelist"), bDisableDomainAllowlist, GEngineIni);
+	if (!bDisableDomainAllowlist)
 	{
 		FHttpManager& HttpManager = FHttpModule::Get().GetHttpManager();
 		if (!HttpManager.IsDomainAllowed(Url))
 		{
-			UE_LOG(LogWebSockets, Warning, TEXT("WinHttp WebSocket[%p]: %s is not whitelisted, refusing to connect."), this, *Url);
+			UE_LOG(LogWebSockets, Warning, TEXT("WinHttp WebSocket[%p]: %s is not in the allowed list, refusing to connect."), this, *Url);
 			HandleCloseComplete(EWebSocketConnectionState::FailedToConnect, UE_WEBSOCKET_CLOSE_APP_FAILURE, FString(TEXT("Invalid Domain")));
 			return;
 		}
 	}
 	else
 	{
-		UE_LOG(LogWebSockets, Log, TEXT("WinHttp WebSocket[%p]: Domain whitelisting has been disabled by config."), this);
+		UE_LOG(LogWebSockets, Log, TEXT("WinHttp WebSocket[%p]: Domain allowed list has been disabled by config."), this);
 	}
 
 	FWinHttpHttpManager* Manager = FWinHttpHttpManager::GetManager();
