Add RGB format param to WebImage

diff --git a/Engine/Source/Runtime/Online/ImageDownload/Private/WebImage.cpp b/Engine/Source/Runtime/Online/ImageDownload/Private/WebImage.cpp
--- a/Engine/Source/Runtime/Online/ImageDownload/Private/WebImage.cpp
+++ b/Engine/Source/Runtime/Online/ImageDownload/Private/WebImage.cpp
@@ -2,7 +2,6 @@
 
 #include "WebImage.h"
 
-#include "IImageWrapper.h"
 #include "IImageWrapperModule.h"
 #include "HttpModule.h"
 #include "Modules/ModuleManager.h"
@@ -14,6 +13,7 @@
 FWebImage::FWebImage()
 : StandInBrush(FCoreStyle::Get().GetDefaultBrush())
 , bDownloadSuccess(false)
+, RGBFormat(ERGBFormat::RGBA)
 {
 }
 
@@ -139,7 +139,7 @@ bool FWebImage::ProcessHttpResponse(const FString& RequestUrl, FHttpResponsePtr
 
 	// get the raw image data
 	TArray<uint8> RawImageData;
-	if (!ImageWrapper->GetRaw(ERGBFormat::RGBA, 8, RawImageData))
+	if (!ImageWrapper->GetRaw(RGBFormat, 8, RawImageData))
 	{
 		UE_LOG(LogImageDownload, Error, TEXT("Image Download: Unable to convert image format %d to BGRA 8"), (int32)ImageFormat);
 		return false;
diff --git a/Engine/Source/Runtime/Online/ImageDownload/Public/WebImage.h b/Engine/Source/Runtime/Online/ImageDownload/Public/WebImage.h
--- a/Engine/Source/Runtime/Online/ImageDownload/Public/WebImage.h
+++ b/Engine/Source/Runtime/Online/ImageDownload/Public/WebImage.h
@@ -11,6 +11,7 @@
 #include "Interfaces/IHttpRequest.h"
 #include "Interfaces/IHttpResponse.h"
 #include "Brushes/SlateDynamicImageBrush.h"
+#include "IImageWrapper.h"
 
 class IHttpRequest;
 struct FSlateBrush;
@@ -54,6 +55,9 @@ public:
 	/** Set the brush that is currently being returned (this will be overridden when any async download completes) */
 	FORCEINLINE FWebImage& SetStandInBrush(TAttribute<const FSlateBrush*> StandInBrushIn) { StandInBrush = StandInBrushIn; DownloadedBrush.Reset(); return *this; }
 
+	/** Sets the RGB format used to read the downloaded image data (ERGBFormat::RGBA by default) */
+	void SetRGBFormat(const ERGBFormat InRGBFormat) { RGBFormat = InRGBFormat; }
+
 	/** Begin downloading an image. This will automatically set the current brush to the downloaded image when it completes (if successful) */
 	bool BeginDownload(const FString& InUrl, const TOptional<FString>& StandInETag = TOptional<FString>(), const FOnImageDownloaded& DownloadCallback = FOnImageDownloaded());
 
@@ -121,4 +125,7 @@ private:
 
 	/** The ETag of the downloaded image */
 	TOptional<FString> ETag;
+
+	/** RGB format used to read the downloaded image data */
+	ERGBFormat RGBFormat;
 };
