Addition of Verbose logging to track progress of HttpManager::Flush

diff --git a/Engine/Source/Runtime/Online/HTTP/Private/HttpManager.cpp b/Engine/Source/Runtime/Online/HTTP/Private/HttpManager.cpp
--- a/Engine/Source/Runtime/Online/HTTP/Private/HttpManager.cpp
+++ b/Engine/Source/Runtime/Online/HTTP/Private/HttpManager.cpp
@@ -20,6 +20,20 @@
 
 FCriticalSection FHttpManager::RequestLock;
 
+const TCHAR* LexToString(const EHttpFlushReason& FlushReason)
+{
+	switch (FlushReason)
+	{
+	case EHttpFlushReason::Default:		return TEXT("Default");
+	case EHttpFlushReason::Background:	return TEXT("Background");
+	case EHttpFlushReason::Shutdown:	return TEXT("Shutdown");
+	case EHttpFlushReason::FullFlush:	return TEXT("FullFlush");
+	}
+
+	checkNoEntry();
+	return TEXT("Invalid");
+}
+
 FHttpManager::FHttpManager()
 	: FTSTickerObjectBase(0.0f, FTSBackgroundableTicker::GetCoreTicker())
 	, Thread(nullptr)
@@ -220,14 +234,13 @@ void FHttpManager::Flush(EHttpFlushReason FlushReason)
 	float SecondsToSleepForOutstandingThreadedRequests = 0.5f;
 	GConfig->GetFloat(TEXT("HTTP"), TEXT("RequestCleanupDelaySec"), SecondsToSleepForOutstandingThreadedRequests, GEngineIni);
 
+	UE_CLOG(!IsRunningCommandlet(), LogHttp, Verbose, TEXT("[FHttpManager::Flush] FlushReason [%s] FlushTimeSoftLimitSeconds [%.3fs] FlushTimeHardLimitSeconds [%.3fs] SecondsToSleepForOutstandingThreadedRequests [%.3fs]"), LexToString(FlushReason), FlushTimeSoftLimitSeconds, FlushTimeHardLimitSeconds, SecondsToSleepForOutstandingThreadedRequests);
+
 	// Clear all delegates bound to ongoing Http requests
 	if (FlushReason == EHttpFlushReason::Shutdown)
 	{
-		if (Requests.Num())
-		{
-			// Don't emit these tracking logs in commandlet runs. Build system traps warnings during cook, and these are not truly fatal, but useful for tracking down shutdown issues.
-			UE_CLOG(!IsRunningCommandlet(), LogHttp, Warning, TEXT("Http module shutting down, but needs to wait on %d outstanding Http requests:"), Requests.Num());
-		}
+		// Don't emit these tracking logs in commandlet runs. Build system traps warnings during cook, and these are not truly fatal, but useful for tracking down shutdown issues.
+		UE_CLOG(!IsRunningCommandlet() && Requests.Num(), LogHttp, Warning, TEXT("[FHttpManager::Flush] FlushReason was Shutdown. Unbinding delegates for %d outstanding Http Requests:"), Requests.Num());
 
 		// Clear delegates since they may point to deleted instances
 		for (TArray<FHttpRequestRef>::TIterator It(Requests); It; ++It)
@@ -242,8 +255,7 @@ void FHttpManager::Flush(EHttpFlushReason FlushReason)
 		}
 	}
 
-	// Don't emit these tracking logs in commandlet runs. Build system traps warnings during cook, and these are not truly fatal, but useful for tracking down shutdown issues.
-	UE_CLOG(!IsRunningCommandlet(), LogHttp, Warning, TEXT("Cleaning up %d outstanding Http requests."), Requests.Num());
+	UE_CLOG(!IsRunningCommandlet() && Requests.Num(), LogHttp, Verbose, TEXT("[FHttpManager::Flush] Cleanup starts for %d outstanding Http Requests."), Requests.Num());
 
 	double BeginWaitTime = FPlatformTime::Seconds();
 	double LastFlushTickTime = BeginWaitTime;
@@ -259,7 +271,7 @@ void FHttpManager::Flush(EHttpFlushReason FlushReason)
 		if (FlushTimeSoftLimitSeconds >= 0 && (AppTime - BeginWaitTime >= FlushTimeSoftLimitSeconds))
 		{
 			// Don't emit these tracking logs in commandlet runs. Build system traps warnings during cook, and these are not truly fatal, but useful for tracking down shutdown issues.
-			UE_CLOG(!IsRunningCommandlet(), LogHttp, Warning, TEXT("Canceling remaining %d HTTP requests"), Requests.Num());
+			UE_CLOG(!IsRunningCommandlet(), LogHttp, Warning, TEXT("[FHttpManager::Flush] FlushTimeSoftLimitSeconds [%.3fs] exceeded. Cancelling %d outstanding HTTP requests:"), FlushTimeSoftLimitSeconds, Requests.Num());
 
 			for (TArray<FHttpRequestRef>::TIterator It(Requests); It; ++It)
 			{
@@ -288,7 +300,7 @@ void FHttpManager::Flush(EHttpFlushReason FlushReason)
 					if (AppTime >= StallWarnTime)
 					{
 						// Don't emit these tracking logs in commandlet runs. Build system traps warnings during cook, and these are not truly fatal, but useful for tracking down shutdown issues.
-						UE_CLOG(!IsRunningCommandlet(), LogHttp, Warning, TEXT("Ticking HTTPThread for %d outstanding Http requests."), Requests.Num());
+						UE_CLOG(!IsRunningCommandlet(), LogHttp, Warning, TEXT("	Ticking HTTPThread for %d outstanding Http requests."), Requests.Num());
 						StallWarnTime = AppTime + 0.5;
 					}
 					Thread->Tick();
@@ -296,7 +308,7 @@ void FHttpManager::Flush(EHttpFlushReason FlushReason)
 				else
 				{
 					// Don't emit these tracking logs in commandlet runs. Build system traps warnings during cook, and these are not truly fatal, but useful for tracking down shutdown issues.
-					UE_CLOG(!IsRunningCommandlet(), LogHttp, Warning, TEXT("Sleeping %.3fs to wait for %d outstanding Http requests."), SecondsToSleepForOutstandingThreadedRequests, Requests.Num());
+					UE_CLOG(!IsRunningCommandlet(), LogHttp, Warning, TEXT("	Sleeping %.3fs to wait for %d outstanding Http Requests."), SecondsToSleepForOutstandingThreadedRequests, Requests.Num());
 					FPlatformProcess::Sleep(SecondsToSleepForOutstandingThreadedRequests);
 				}
 			}
@@ -309,10 +321,13 @@ void FHttpManager::Flush(EHttpFlushReason FlushReason)
 		AppTime = FPlatformTime::Seconds();
 	}
 
+	UE_CLOG(!IsRunningCommandlet(), LogHttp, Verbose, TEXT("[FHttpManager::Flush] Cleanup ended after %.3fs. %d outstanding Http Requests."), AppTime - BeginWaitTime, Requests.Num());
+
 	// Don't emit these tracking logs in commandlet runs. Build system traps warnings during cook, and these are not truly fatal, but useful for tracking down shutdown issues.
 	if (Requests.Num() > 0 && (FlushTimeHardLimitSeconds > 0 && (AppTime - BeginWaitTime > FlushTimeHardLimitSeconds)) && !IsRunningCommandlet())
 	{
-		UE_LOG(LogHttp, Warning, TEXT("HTTTManager::Flush exceeded hard limit time %.3fs. Current time is %.3fs. These requests are being abandoned without being flushed:"), FlushTimeHardLimitSeconds, AppTime - BeginWaitTime);
+		UE_LOG(LogHttp, Warning, TEXT("[FHttpManager::Flush] FlushTimeHardLimitSeconds [%.3fs] exceeded. The following requests are being abandoned without being flushed:"), FlushTimeHardLimitSeconds);
+
 		for (TArray<FHttpRequestRef>::TIterator It(Requests); It; ++It)
 		{
 			FHttpRequestRef& Request = *It;
diff --git a/Engine/Source/Runtime/Online/HTTP/Public/HttpManager.h b/Engine/Source/Runtime/Online/HTTP/Public/HttpManager.h
--- a/Engine/Source/Runtime/Online/HTTP/Public/HttpManager.h
+++ b/Engine/Source/Runtime/Online/HTTP/Public/HttpManager.h
@@ -26,6 +26,7 @@ enum class EHttpFlushReason : uint8
 	Count
 };
 ENUM_RANGE_BY_COUNT(EHttpFlushReason, EHttpFlushReason::Count)
+const TCHAR* LexToString(const EHttpFlushReason& FlushReason);
 
 /**
  * Manages Http request that are currently being processed
