UE-145584 //UE5/Main - Compile UnrealGame HoloLens - error LNK2019: unresolved externals
-Define the ntohs, etc network integer conversion functions for hololens.  Something has started using them.
-Also explicitly setting icmp to disabled for hololens just to document that more clearly.
manual preflight 622ba52432749f363c916d7e of enginetest hololens

diff --git a/Engine/Source/Runtime/Online/ICMP/Private/HoloLens/IcmpHoloLens.cpp b/Engine/Source/Runtime/Online/ICMP/Private/HoloLens/IcmpHoloLens.cpp
new file mode 100644
--- /dev/null
+++ b/Engine/Source/Runtime/Online/ICMP/Private/HoloLens/IcmpHoloLens.cpp
@@ -0,0 +1,27 @@
+#include "Icmp.h"
+
+#include "Microsoft/AllowMicrosoftPlatformTypes.h"
+#define _WINSOCK_DEPRECATED_NO_WARNINGS
+#include "WinSock2.h"
+#undef _WINSOCK_DEPRECATED_NO_WARNINGS
+#include "Microsoft/HideMicrosoftPlatformTypes.h"
+
+uint16 NtoHS(uint16 val)
+{
+	return ntohs(val);
+}
+
+uint16 HtoNS(uint16 val)
+{
+	return htons(val);
+}
+
+uint32 NtoHL(uint32 val)
+{
+	return ntohl(val);
+}
+
+uint32 HtoNL(uint32 val)
+{
+	return htonl(val);
+}
diff --git a/Engine/Source/Runtime/Online/ICMP/Private/IcmpPrivate.h b/Engine/Source/Runtime/Online/ICMP/Private/IcmpPrivate.h
--- a/Engine/Source/Runtime/Online/ICMP/Private/IcmpPrivate.h
+++ b/Engine/Source/Runtime/Online/ICMP/Private/IcmpPrivate.h
@@ -14,6 +14,12 @@
 		#define PLATFORM_USES_POSIX_ICMP 0
 		#define PING_ALLOWS_CUSTOM_THREAD_SIZE 1
 
+	#elif PLATFORM_HOLOLENS
+
+		#define PLATFORM_SUPPORTS_ICMP 0
+		#define PLATFORM_USES_POSIX_ICMP 0
+		#define PING_ALLOWS_CUSTOM_THREAD_SIZE 0
+
 	#elif PLATFORM_MAC
 
 		#define PLATFORM_SUPPORTS_ICMP 1
