UE-145584 //UE5/Main - Compile UnrealGame HoloLens - error LNK2019: unresolved externals
-Fix copyright warning.

diff --git a/Engine/Source/Runtime/Online/ICMP/Private/HoloLens/IcmpHoloLens.cpp b/Engine/Source/Runtime/Online/ICMP/Private/HoloLens/IcmpHoloLens.cpp
--- a/Engine/Source/Runtime/Online/ICMP/Private/HoloLens/IcmpHoloLens.cpp
+++ b/Engine/Source/Runtime/Online/ICMP/Private/HoloLens/IcmpHoloLens.cpp
@@ -1,3 +1,4 @@
+// Copyright Epic Games, Inc. All Rights Reserved.
 #include "Icmp.h"
 
 #include "Microsoft/AllowMicrosoftPlatformTypes.h"
