Add option to increase HTTP log verbosity when a request is being retried
This must be explicitly opted into by setting [HTTP.Retry]RetryManagerVerbosityLevel to the desired verbosity level
#tests Running with increased HTTP failure rate to observe HTTP verbosity

diff --git a/Engine/Source/Runtime/Online/HTTP/Private/HttpRetrySystem.cpp b/Engine/Source/Runtime/Online/HTTP/Private/HttpRetrySystem.cpp
--- a/Engine/Source/Runtime/Online/HTTP/Private/HttpRetrySystem.cpp
+++ b/Engine/Source/Runtime/Online/HTTP/Private/HttpRetrySystem.cpp
@@ -151,6 +151,18 @@ FHttpRetrySystem::FManager::FManager(const FRetryLimitCountSetting& InRetryLimit
 	, RetryTimeoutRelativeSecondsDefault(InRetryTimeoutRelativeSecondsDefault)
 {}
 
+FHttpRetrySystem::FManager::~FManager()
+{
+	// Decrement retried request for log verbosity tracker
+	for (const FHttpRetryRequestEntry& Request : RequestList)
+	{
+		if (Request.CurrentRetryCount > 0)
+		{
+			FHttpLogVerbosityTracker::Get().DecrementRetriedRequests();
+		}
+	}
+}
+
 TSharedRef<FHttpRetrySystem::FRequest, ESPMode::ThreadSafe> FHttpRetrySystem::FManager::CreateRequest(
 	const FRetryLimitCountSetting& InRetryLimitCountOverride,
 	const FRetryTimeoutRelativeSecondsSetting& InRetryTimeoutRelativeSecondsOverride,
@@ -270,6 +282,23 @@ bool FHttpRetrySystem::FManager::HasTimedOut(const FHttpRetryRequestEntry& HttpR
     return bResult;
 }
 
+void FHttpRetrySystem::FManager::RetryHttpRequest(FHttpRetryRequestEntry& RequestEntry)
+{
+	// if this fails the HttpRequest's state will be failed which will cause the retry logic to kick(as expected)
+	const bool bProcessRequestSuccess = RequestEntry.Request->HttpRequest->ProcessRequest();
+	if (bProcessRequestSuccess)
+	{
+		UE_LOG(LogHttp, Warning, TEXT("Retry %d on %s"), RequestEntry.CurrentRetryCount + 1, *(RequestEntry.Request->GetURL()));
+
+		if (RequestEntry.CurrentRetryCount == 0)
+		{
+			FHttpLogVerbosityTracker::Get().IncrementRetriedRequests();
+		}
+		++RequestEntry.CurrentRetryCount;
+		RequestEntry.Request->Status = FRequest::EStatus::Processing;
+	}
+}
+
 float FHttpRetrySystem::FManager::GetLockoutPeriodSeconds(const FHttpRetryRequestEntry& HttpRetryRequestEntry)
 {
 	float LockoutPeriod = 0.0f;
@@ -355,24 +384,29 @@ bool FHttpRetrySystem::FManager::Update(uint32* FileCount, uint32* FailingCount,
 	//    endif
 	// endfor
 
-	int32 index = 0;
-	while (index < RequestList.Num())
+	int32 RequestIndex = 0;
+	while (RequestIndex < RequestList.Num())
 	{
 		QUICK_SCOPE_CYCLE_COUNTER(STAT_FHttpRetrySystem_FManager_Update_RequestListItem);
 
-		FHttpRetryRequestEntry& HttpRetryRequestEntry = RequestList[index];
-		TSharedRef<FHttpRetrySystem::FRequest, ESPMode::ThreadSafe>& HttpRetryRequest = HttpRetryRequestEntry.Request;
+		FHttpRetryRequestEntry* HttpRetryRequestEntry = &RequestList[RequestIndex];
+		TSharedRef<FHttpRetrySystem::FRequest, ESPMode::ThreadSafe> HttpRetryRequest = HttpRetryRequestEntry->Request;
+		// Delegates fired in this loop can resize the array if new requests are added, invalidating HttpRetryRequestEntry and HttpRetryRequest.  Call this when the array may have been modified.
+		auto ResetCurrentIterationVariables = [this, &HttpRetryRequestEntry, &HttpRetryRequest, RequestIndex]()
+		{
+			HttpRetryRequestEntry = &RequestList[RequestIndex];
+		};
 
 		const EHttpRequestStatus::Type RequestStatus = HttpRetryRequest->GetStatus();
 
-		if (HttpRetryRequestEntry.bShouldCancel)
+		if (HttpRetryRequestEntry->bShouldCancel)
 		{
 			UE_LOG(LogHttp, Warning, TEXT("Request cancelled on %s"), *(HttpRetryRequest->GetURL()));
 			HttpRetryRequest->Status = FHttpRetrySystem::FRequest::EStatus::Cancelled;
 		}
 		else
 		{
-			if (!HasTimedOut(HttpRetryRequestEntry, NowAbsoluteSeconds))
+			if (!HasTimedOut(*HttpRetryRequestEntry, NowAbsoluteSeconds))
 			{
 				if (HttpRetryRequest->Status == FHttpRetrySystem::FRequest::EStatus::NotStarted)
 				{
@@ -409,8 +443,8 @@ bool FHttpRetrySystem::FManager::Update(uint32* FileCount, uint32* FailingCount,
 					bool bCanRetry = false;
 					if (RequestStatus == EHttpRequestStatus::Failed || RequestStatus == EHttpRequestStatus::Failed_ConnectionError || RequestStatus == EHttpRequestStatus::Succeeded)
 					{
-						bShouldRetry = ShouldRetry(HttpRetryRequestEntry);
-						bCanRetry = CanRetry(HttpRetryRequestEntry);
+						bShouldRetry = ShouldRetry(*HttpRetryRequestEntry);
+						bCanRetry = CanRetry(*HttpRetryRequestEntry);
 					}
 
 					if (RequestStatus == EHttpRequestStatus::Failed || RequestStatus == EHttpRequestStatus::Failed_ConnectionError || forceFail || (bShouldRetry && bCanRetry))
@@ -419,18 +453,19 @@ bool FHttpRetrySystem::FManager::Update(uint32* FileCount, uint32* FailingCount,
 
 						if (forceFail || (bShouldRetry && bCanRetry))
 						{
-							float LockoutPeriod = GetLockoutPeriodSeconds(HttpRetryRequestEntry);
+							float LockoutPeriod = GetLockoutPeriodSeconds(*HttpRetryRequestEntry);
 
 							if (LockoutPeriod > 0.0f)
 							{
 								UE_LOG(LogHttp, Warning, TEXT("Lockout of %fs on %s"), LockoutPeriod, *(HttpRetryRequest->GetURL()));
 							}
 
-							HttpRetryRequestEntry.LockoutEndTimeAbsoluteSeconds = NowAbsoluteSeconds + LockoutPeriod;
+							HttpRetryRequestEntry->LockoutEndTimeAbsoluteSeconds = NowAbsoluteSeconds + LockoutPeriod;
 							HttpRetryRequest->Status = FHttpRetrySystem::FRequest::EStatus::ProcessingLockout;
 							
 							QUICK_SCOPE_CYCLE_COUNTER(STAT_FHttpRetrySystem_FManager_Update_OnRequestWillRetry);
 							HttpRetryRequest->OnRequestWillRetry().ExecuteIfBound(HttpRetryRequest, HttpRetryRequest->GetResponse(), LockoutPeriod);
+							ResetCurrentIterationVariables();
 						}
 						else
 						{
@@ -444,7 +479,7 @@ bool FHttpRetrySystem::FManager::Update(uint32* FileCount, uint32* FailingCount,
 					}
 					else if (RequestStatus == EHttpRequestStatus::Succeeded)
 					{
-						if (HttpRetryRequestEntry.CurrentRetryCount > 0)
+						if (HttpRetryRequestEntry->CurrentRetryCount > 0)
 						{
 							UE_LOG(LogHttp, Warning, TEXT("Success on %s"), *(HttpRetryRequest->GetURL()));
 						}
@@ -460,17 +495,10 @@ bool FHttpRetrySystem::FManager::Update(uint32* FileCount, uint32* FailingCount,
 
 				if (HttpRetryRequest->Status == FHttpRetrySystem::FRequest::EStatus::ProcessingLockout)
 				{
-					if (NowAbsoluteSeconds >= HttpRetryRequestEntry.LockoutEndTimeAbsoluteSeconds)
+					if (NowAbsoluteSeconds >= HttpRetryRequestEntry->LockoutEndTimeAbsoluteSeconds)
 					{
-						// if this fails the HttpRequest's state will be failed which will cause the retry logic to kick(as expected)
-						bool success = HttpRetryRequest->HttpRequest->ProcessRequest();
-						if (success)
-						{
-							UE_LOG(LogHttp, Warning, TEXT("Retry %d on %s"), HttpRetryRequestEntry.CurrentRetryCount + 1, *(HttpRetryRequest->GetURL()));
-
-							++HttpRetryRequestEntry.CurrentRetryCount;
-							HttpRetryRequest->Status = FRequest::EStatus::Processing;
-						}
+						RetryHttpRequest(*HttpRetryRequestEntry);
+						ResetCurrentIterationVariables();
 					}
 
 					if (FailingCount != nullptr)
@@ -481,7 +509,7 @@ bool FHttpRetrySystem::FManager::Update(uint32* FileCount, uint32* FailingCount,
 			}
 			else
 			{
-				UE_LOG(LogHttp, Warning, TEXT("Timeout on retry %d: %s"), HttpRetryRequestEntry.CurrentRetryCount + 1, *(HttpRetryRequest->GetURL()));
+				UE_LOG(LogHttp, Warning, TEXT("Timeout on retry %d: %s"), HttpRetryRequestEntry->CurrentRetryCount + 1, *(HttpRetryRequest->GetURL()));
 				bIsGreen = false;
 				HttpRetryRequest->Status = FHttpRetrySystem::FRequest::EStatus::FailedTimeout;
 				if (FailedCount != nullptr)
@@ -508,15 +536,17 @@ bool FHttpRetrySystem::FManager::Update(uint32* FileCount, uint32* FailingCount,
 			if (bWasSuccessful)
 			{
 				HttpRetryRequest->BroadcastResponseHeadersReceived();
+				ResetCurrentIterationVariables();
 			}
 			
 			QUICK_SCOPE_CYCLE_COUNTER(STAT_FHttpRetrySystem_FManager_Update_OnProcessRequestComplete);
 			HttpRetryRequest->OnProcessRequestComplete().ExecuteIfBound(HttpRetryRequest, HttpRetryRequest->GetResponse(), bWasSuccessful);
+			ResetCurrentIterationVariables();
 		}
 
-        if(bWasSuccessful)
+        if (bWasSuccessful)
         {
-            if(CompletedCount != nullptr)
+            if (CompletedCount != nullptr)
             {
                 ++(*CompletedCount);
             }
@@ -524,11 +554,15 @@ bool FHttpRetrySystem::FManager::Update(uint32* FileCount, uint32* FailingCount,
 
 		if (bWasCompleted)
 		{
-			RequestList.RemoveAtSwap(index);
+			if (RequestList[RequestIndex].CurrentRetryCount > 0)
+			{
+				FHttpLogVerbosityTracker::Get().DecrementRetriedRequests();
+			}
+			RequestList.RemoveAtSwap(RequestIndex);
 		}
 		else
 		{
-			++index;
+			++RequestIndex;
 		}
 	}
 
@@ -597,4 +631,75 @@ void FHttpRetrySystem::FManager::BlockUntilFlushed(float InTimeoutSec)
 		FPlatformProcess::Sleep(SleepInterval);
 		TimeElapsed += SleepInterval;
 	}
-}
\ No newline at end of file
+}
+
+FHttpRetrySystem::FManager::FHttpLogVerbosityTracker& FHttpRetrySystem::FManager::FHttpLogVerbosityTracker::Get()
+{
+	static FHttpLogVerbosityTracker Tracker;
+	return Tracker;
+}
+
+FHttpRetrySystem::FManager::FHttpLogVerbosityTracker::FHttpLogVerbosityTracker()
+{
+	UpdateSettingsFromConfig();
+	FCoreDelegates::OnConfigSectionsChanged.AddRaw(this, &FHttpLogVerbosityTracker::OnConfigSectionsChanged);
+}
+
+FHttpRetrySystem::FManager::FHttpLogVerbosityTracker::~FHttpLogVerbosityTracker()
+{
+	FCoreDelegates::OnConfigSectionsChanged.RemoveAll(this);
+}
+
+void FHttpRetrySystem::FManager::FHttpLogVerbosityTracker::IncrementRetriedRequests()
+{
+	check(IsInGameThread());
+	++NumRetriedRequests;
+	if (NumRetriedRequests == 1)
+	{
+		OriginalVerbosity = UE_GET_LOG_VERBOSITY(LogHttp);
+		if (TargetVerbosity != ELogVerbosity::NoLogging)
+		{
+			UE_LOG(LogHttp, Warning, TEXT("HttpRetry: Increasing log verbosity from %s to %s due to requests being retried"), ToString(OriginalVerbosity), ToString(TargetVerbosity));
+			//UE_SET_LOG_VERBOSITY(LogHttp, TargetVerbosity); // Macro requires the value to be a ELogVerbosity constant
+#if !NO_LOGGING
+			LogHttp.SetVerbosity(TargetVerbosity);
+#endif
+		}
+	}
+}
+
+void FHttpRetrySystem::FManager::FHttpLogVerbosityTracker::DecrementRetriedRequests()
+{
+	check(IsInGameThread());
+	--NumRetriedRequests;
+	check(NumRetriedRequests >= 0);
+	if (NumRetriedRequests == 0)
+	{
+		UE_LOG(LogHttp, Warning, TEXT("HttpRetry: Resetting log verbosity to %s due to requests being retried"), ToString(OriginalVerbosity));
+		//UE_SET_LOG_VERBOSITY(LogHttp, OriginalVerbosity); // Macro requires the value to be a ELogVerbosity constant
+#if !NO_LOGGING
+		LogHttp.SetVerbosity(OriginalVerbosity);
+#endif
+}
+}
+
+void FHttpRetrySystem::FManager::FHttpLogVerbosityTracker::UpdateSettingsFromConfig()
+{
+	FString TargetVerbosityAsString;
+	if (GConfig->GetString(TEXT("HTTP.Retry"), TEXT("RetryManagerVerbosityLevel"), TargetVerbosityAsString, GEngineIni))
+	{
+		TargetVerbosity = ParseLogVerbosityFromString(TargetVerbosityAsString);
+	}
+	else
+	{
+		TargetVerbosity = ELogVerbosity::NoLogging;
+	}
+}
+
+void FHttpRetrySystem::FManager::FHttpLogVerbosityTracker::OnConfigSectionsChanged(const FString& IniFilename, const TSet<FString>& SectionName)
+{
+	if (IniFilename == GEngineIni && SectionName.Contains(TEXT("HTTP.Retry")))
+	{
+		UpdateSettingsFromConfig();
+	}
+}
diff --git a/Engine/Source/Runtime/Online/HTTP/Public/HttpRetrySystem.h b/Engine/Source/Runtime/Online/HTTP/Public/HttpRetrySystem.h
--- a/Engine/Source/Runtime/Online/HTTP/Public/HttpRetrySystem.h
+++ b/Engine/Source/Runtime/Online/HTTP/Public/HttpRetrySystem.h
@@ -143,7 +143,7 @@ namespace FHttpRetrySystem
 			const FRetryDomainsPtr& InRetryDomains = FRetryDomainsPtr()
 			);
 
-		HTTP_API virtual ~FManager() = default;
+		HTTP_API virtual ~FManager();
 
         /**
          * Updates the entries in the list of retry requests. Optional parameters are for future connection health assessment
@@ -177,6 +177,30 @@ namespace FHttpRetrySystem
             TSharedRef<FRequest, ESPMode::ThreadSafe>	Request;
         };
 
+		class FHttpLogVerbosityTracker
+		{
+		public:
+			FHttpLogVerbosityTracker();
+			~FHttpLogVerbosityTracker();
+
+			/** Mark that a request is being retried */
+			void IncrementRetriedRequests();
+			/** Mark that a retried request is no longer being retried */
+			void DecrementRetriedRequests();
+			static FHttpLogVerbosityTracker& Get();
+		protected:
+			/** Update settings from config */
+			void UpdateSettingsFromConfig();
+			void OnConfigSectionsChanged(const FString& IniFilename, const TSet<FString>& SectionName);
+
+			/** Number of requests that are in a retried state.  When this is non-zero, verbosity will be adjusted. */
+			int32 NumRetriedRequests = 0;
+			/** Verbosity to restore to when there are no requests being retried */
+			ELogVerbosity::Type OriginalVerbosity = ELogVerbosity::Error;
+			/** Config driven target verbosity to set to when requests are being retried.  NoLogging means the verbosity will not be modified. */
+			ELogVerbosity::Type TargetVerbosity = ELogVerbosity::NoLogging;
+		};
+
 		bool ProcessRequest(TSharedRef<FRequest, ESPMode::ThreadSafe>& HttpRequest);
 		void CancelRequest(TSharedRef<FRequest, ESPMode::ThreadSafe>& HttpRequest);
 
@@ -190,6 +214,12 @@ namespace FHttpRetrySystem
         // @return true if the retry request has timed out
         bool HasTimedOut(const FHttpRetryRequestEntry& HttpRetryRequestEntry, const double NowAbsoluteSeconds);
 
+		/**
+		 * Retry an HTTP request
+		 * @param RequestEntry request to retry
+		 */
+		void RetryHttpRequest(FHttpRetryRequestEntry& RequestEntry);
+
         // @return number of seconds to lockout for
         float GetLockoutPeriodSeconds(const FHttpRetryRequestEntry& HttpRetryRequestEntry);
 
@@ -198,6 +228,6 @@ namespace FHttpRetrySystem
         FRetryLimitCountSetting              RetryLimitCountDefault;
         FRetryTimeoutRelativeSecondsSetting  RetryTimeoutRelativeSecondsDefault;
 
-        TArray<FHttpRetryRequestEntry>        RequestList;
+        TArray<FHttpRetryRequestEntry>       RequestList;
     };
 }
