Build fix
#tests compile

diff --git a/Engine/Source/Runtime/Online/HTTP/Private/HttpRetrySystem.cpp b/Engine/Source/Runtime/Online/HTTP/Private/HttpRetrySystem.cpp
--- a/Engine/Source/Runtime/Online/HTTP/Private/HttpRetrySystem.cpp
+++ b/Engine/Source/Runtime/Online/HTTP/Private/HttpRetrySystem.cpp
@@ -9,6 +9,8 @@
 #include "Http.h"
 #include "HttpManager.h"
 #include "Stats/Stats.h"
+#include "Misc/ConfigCacheIni.h"
+#include "Misc/CoreDelegates.h"
 
 LLM_DEFINE_TAG(HTTP);
 
@@ -680,7 +682,7 @@ void FHttpRetrySystem::FManager::FHttpLogVerbosityTracker::DecrementRetriedReque
 #if !NO_LOGGING
 		LogHttp.SetVerbosity(OriginalVerbosity);
 #endif
-}
+	}
 }
 
 void FHttpRetrySystem::FManager::FHttpLogVerbosityTracker::UpdateSettingsFromConfig()
