Http Get request with body contents - warning and cancel

diff --git a/Engine/Source/Runtime/Online/HTTP/Private/Curl/CurlHttp.cpp b/Engine/Source/Runtime/Online/HTTP/Private/Curl/CurlHttp.cpp
--- a/Engine/Source/Runtime/Online/HTTP/Private/Curl/CurlHttp.cpp
+++ b/Engine/Source/Runtime/Online/HTTP/Private/Curl/CurlHttp.cpp
@@ -758,6 +758,12 @@ bool FCurlHttpRequest::SetupRequest()
 		UE_LOG(LogHttp, Log, TEXT("Cannot process HTTP request: URL is empty"));
 		return false;
 	}
+	if ((GetVerb().IsEmpty() || GetVerb().Equals(TEXT("GET"), ESearchCase::IgnoreCase))
+		&& (RequestPayload.IsValid() && RequestPayload->GetContentLength() > 0))
+	{
+		UE_LOG(LogHttp, Warning, TEXT("An HTTP Get request cannot contain a payload."));
+		return false;
+	}
 
 	// set up request
 
