- Rolling back some Lws optimizations that may have broken some platforms

diff --git a/Engine/Source/Runtime/Online/WebSockets/Private/Lws/LwsWebSocket.cpp b/Engine/Source/Runtime/Online/WebSockets/Private/Lws/LwsWebSocket.cpp
--- a/Engine/Source/Runtime/Online/WebSockets/Private/Lws/LwsWebSocket.cpp
+++ b/Engine/Source/Runtime/Online/WebSockets/Private/Lws/LwsWebSocket.cpp
@@ -158,8 +158,6 @@ void FLwsWebSocket::Close(int32 Code, const FString& Reason)
 void FLwsWebSocket::Send(const void* Data, SIZE_T Size, bool bIsBinary)
 {
 	SendQueue.Enqueue(new FLwsSendBuffer(static_cast<const uint8*>(Data), Size, bIsBinary));
-	FLwsWebSocketsManager& WebSocketsManager = FLwsWebSocketsManager::Get();
-	WebSocketsManager.WakeService();
 }
 
 void FLwsWebSocket::Send(const FString& Data)
diff --git a/Engine/Source/Runtime/Online/WebSockets/Private/Lws/LwsWebSocketsManager.cpp b/Engine/Source/Runtime/Online/WebSockets/Private/Lws/LwsWebSocketsManager.cpp
--- a/Engine/Source/Runtime/Online/WebSockets/Private/Lws/LwsWebSocketsManager.cpp
+++ b/Engine/Source/Runtime/Online/WebSockets/Private/Lws/LwsWebSocketsManager.cpp
@@ -227,8 +227,13 @@ uint32 FLwsWebSocketsManager::Run()
 {
 	while (!ExitRequest.GetValue())
 	{
-		// tick blocks in the lws_service() call until needed
+		double BeginTime = FPlatformTime::Seconds();
 		Tick();
+		double EndTime = FPlatformTime::Seconds();
+
+		double TotalTime = EndTime - BeginTime;
+		double SleepTime = FMath::Max(ThreadTargetFrameTimeInSeconds - TotalTime, ThreadMinimumSleepTimeInSeconds);
+		FPlatformProcess::SleepNoStats(SleepTime);
 	}
 
 	return 0;
@@ -380,7 +385,7 @@ void FLwsWebSocketsManager::Tick()
 	}
 	if (LwsContext)
 	{
-		lws_service(LwsContext, MAX_int32);
+		lws_service(LwsContext, 0);
 	}
 	for (FLwsWebSocket* Socket : SocketsDestroyedDuringService)
 	{
@@ -406,13 +411,6 @@ void FLwsWebSocketsManager::StartProcessingWebSocket(FLwsWebSocket* Socket)
 {
 	Sockets.Emplace(Socket->AsShared());
 	SocketsToStart.Enqueue(Socket);
-	WakeService();
-}
-
-void FLwsWebSocketsManager::WakeService()
-{
-	// Safe to call from other threads
-	lws_cancel_service(LwsContext);
 }
 
 bool FLwsWebSocketsManager::GameThreadTick(float DeltaTime)
diff --git a/Engine/Source/Runtime/Online/WebSockets/Private/Lws/LwsWebSocketsManager.h b/Engine/Source/Runtime/Online/WebSockets/Private/Lws/LwsWebSocketsManager.h
--- a/Engine/Source/Runtime/Online/WebSockets/Private/Lws/LwsWebSocketsManager.h
+++ b/Engine/Source/Runtime/Online/WebSockets/Private/Lws/LwsWebSocketsManager.h
@@ -43,11 +43,6 @@ public:
 	 */
 	void StartProcessingWebSocket(FLwsWebSocket* Socket);
 
-	/**
-	 * Trigger/Wake up the websocket thread to process new requests
-	 */
-	void WakeService();
-
 	// IWebSocketsManager
 	virtual void InitWebSockets(TArrayView<const FString> Protocols) override;
 	virtual void ShutdownWebSockets() override;
