Move WebSockets platform references to platform extensions

diff --git a/Engine/Platforms/Hololens/Source/Runtime/Online/WebSockets/WebSockets_HoloLens.build.cs b/Engine/Platforms/Hololens/Source/Runtime/Online/WebSockets/WebSockets_HoloLens.build.cs
--- a/Engine/Platforms/Hololens/Source/Runtime/Online/WebSockets/WebSockets_HoloLens.build.cs
+++ b/Engine/Platforms/Hololens/Source/Runtime/Online/WebSockets/WebSockets_HoloLens.build.cs
@@ -12,5 +12,21 @@ namespace UnrealBuildTool.Rules
 		public WebSockets_HoloLens(ReadOnlyTargetRules Target) : base(Target)
 		{
 		}
+
+		protected override string WebSocketsManagerPlatformInclude
+		{
+			get
+			{
+				return "HoloLensWebSocketsManager.h";
+			}
+		}
+
+		protected override string WebSocketsManagerPlatformClass
+		{
+			get
+			{
+				return "FHoloLensWebSocketsManager";
+			}
+		}
 	}
 }
diff --git a/Engine/Source/Runtime/Online/WebSockets/Private/PlatformWebSocket.h b/Engine/Source/Runtime/Online/WebSockets/Private/PlatformWebSocket.h
--- a/Engine/Source/Runtime/Online/WebSockets/Private/PlatformWebSocket.h
+++ b/Engine/Source/Runtime/Online/WebSockets/Private/PlatformWebSocket.h
@@ -4,28 +4,9 @@
 
 #include "HAL/Platform.h"
 
-#if WITH_LIBWEBSOCKETS
-	#if PLATFORM_SWITCH
-		#include "Lws/LwsSwitchWebSocketsManager.h"
-	#else
-		#include "Lws/LwsWebSocketsManager.h"
-	#endif //PLATFORM_SWITCH
-#elif WITH_WINHTTPWEBSOCKETS
-	#include "WinHttp/WinHttpWebSocketsManager.h"
-#elif PLATFORM_HOLOLENS
-	#include "HoloLensWebSocketsManager.h"
+#if defined(WEBSOCKETS_MANAGER_PLATFORM_INCLUDE)
+	#include WEBSOCKETS_MANAGER_PLATFORM_INCLUDE
+	typedef WEBSOCKETS_MANAGER_PLATFORM_CLASS FPlatformWebSocketsManager;
 #else
 	#error "Web Sockets not implemented on this platform yet"
-#endif
-
-#if WITH_LIBWEBSOCKETS
-	#if PLATFORM_SWITCH
-		typedef FLwsSwitchWebSocketsManager FPlatformWebSocketsManager;
-	#else
-		typedef FLwsWebSocketsManager FPlatformWebSocketsManager;
-	#endif // !PLATFORM_SWITCH
-#elif WITH_WINHTTPWEBSOCKETS
-	typedef FWinHttpWebSocketsManager FPlatformWebSocketsManager;
-#elif PLATFORM_HOLOLENS
-	typedef FHoloLensWebSocketsManager FPlatformWebSocketsManager;
-#endif
+#endif
\ No newline at end of file
diff --git a/Engine/Source/Runtime/Online/WebSockets/WebSockets.Build.cs b/Engine/Source/Runtime/Online/WebSockets/WebSockets.Build.cs
--- a/Engine/Source/Runtime/Online/WebSockets/WebSockets.Build.cs
+++ b/Engine/Source/Runtime/Online/WebSockets/WebSockets.Build.cs
@@ -44,6 +44,44 @@ public class WebSockets : ModuleRules
 		}
 	}
 
+	protected virtual string WebSocketsManagerPlatformInclude
+	{
+		get
+		{
+			if (PlatformSupportsLibWebsockets)
+			{
+				return "Lws/LwsWebSocketsManager.h";
+			}
+			else if (bPlatformSupportsWinHttpWebSockets)
+			{
+				return "WinHttp/WinHttpWebSocketsManager.h";
+			}
+			else
+			{
+				return "";
+			}
+		}
+	}
+
+	protected virtual string WebSocketsManagerPlatformClass
+	{
+		get
+		{
+			if (PlatformSupportsLibWebsockets)
+			{
+				return "FLwsWebSocketsManager";
+			}
+			else if (bPlatformSupportsWinHttpWebSockets)
+			{
+				return "FWinHttpWebSocketsManager";
+			}
+			else
+			{
+				return "";
+			}
+		}
+	}
+
 	public WebSockets(ReadOnlyTargetRules Target) : base(Target)
 	{
 		PrivateDependencyModuleNames.AddRange(
@@ -82,7 +120,7 @@ public class WebSockets : ModuleRules
 					PrivateDependencyModuleNames.Add("SSL");
 				}
 			}
-			if (bPlatformSupportsWinHttpWebSockets)
+			else if (bPlatformSupportsWinHttpWebSockets)
 			{
 				// Enable WinHttp Support
 				bWithWinHttpWebSockets = true;
@@ -102,5 +140,11 @@ public class WebSockets : ModuleRules
 		PublicDefinitions.Add("WITH_WEBSOCKETS=" + (bWithWebSockets ? "1" : "0"));
 		PublicDefinitions.Add("WITH_LIBWEBSOCKETS=" + (bWithLibWebSockets ? "1" : "0"));
 		PublicDefinitions.Add("WITH_WINHTTPWEBSOCKETS=" + (bWithWinHttpWebSockets ? "1" : "0"));
+		string PlatformInclude = WebSocketsManagerPlatformInclude;
+		if (PlatformInclude.Length > 0)
+		{
+			PublicDefinitions.Add("WEBSOCKETS_MANAGER_PLATFORM_INCLUDE=\"" + WebSocketsManagerPlatformInclude + "\"");
+			PublicDefinitions.Add("WEBSOCKETS_MANAGER_PLATFORM_CLASS=" + WebSocketsManagerPlatformClass);
+		}
 	}
 }
