Removed BackgroundHTTPFileHash private pch.
Timing when compiling module with PCH(including compiling PCH): 2.08 secs
Timing using shared PCH: 0.85s secs

diff --git a/Engine/Source/Runtime/Online/BackgroundHTTPFileHash/BackgroundHTTPFileHash.Build.cs b/Engine/Source/Runtime/Online/BackgroundHTTPFileHash/BackgroundHTTPFileHash.Build.cs
--- a/Engine/Source/Runtime/Online/BackgroundHTTPFileHash/BackgroundHTTPFileHash.Build.cs
+++ b/Engine/Source/Runtime/Online/BackgroundHTTPFileHash/BackgroundHTTPFileHash.Build.cs
@@ -12,7 +12,5 @@ public class BackgroundHTTPFileHash : ModuleRules
                 "Core",
             }
         );
-        
-        PrivatePCHHeaderFile = "Private/BackgroundHttpFileHashPrivate.h";
     }
 }
