Silence http manager warnings during automation tests
#rnx

diff --git a/Engine/Source/Runtime/Online/HTTP/Private/HttpManager.cpp b/Engine/Source/Runtime/Online/HTTP/Private/HttpManager.cpp
--- a/Engine/Source/Runtime/Online/HTTP/Private/HttpManager.cpp
+++ b/Engine/Source/Runtime/Online/HTTP/Private/HttpManager.cpp
@@ -216,6 +216,14 @@ void FHttpManager::Flush(bool bShutdown)
 	Flush(bShutdown ? EHttpFlushReason::Shutdown : EHttpFlushReason::Default);
 }
 
+namespace
+{
+	bool ShouldOutputHttpWarnings()
+	{
+		return !IsRunningCommandlet() && !GIsAutomationTesting;
+	}
+}
+
 void FHttpManager::Flush(EHttpFlushReason FlushReason)
 {
 	QUICK_SCOPE_CYCLE_COUNTER(STAT_FHttpManager_Flush);
@@ -240,7 +248,7 @@ void FHttpManager::Flush(EHttpFlushReason FlushReason)
 	if (FlushReason == EHttpFlushReason::Shutdown)
 	{
 		// Don't emit these tracking logs in commandlet runs. Build system traps warnings during cook, and these are not truly fatal, but useful for tracking down shutdown issues.
-		UE_CLOG(!IsRunningCommandlet() && Requests.Num(), LogHttp, Warning, TEXT("[FHttpManager::Flush] FlushReason was Shutdown. Unbinding delegates for %d outstanding Http Requests:"), Requests.Num());
+		UE_CLOG(ShouldOutputHttpWarnings() && Requests.Num(), LogHttp, Warning, TEXT("[FHttpManager::Flush] FlushReason was Shutdown. Unbinding delegates for %d outstanding Http Requests:"), Requests.Num());
 
 		// Clear delegates since they may point to deleted instances
 		for (TArray<FHttpRequestRef>::TIterator It(Requests); It; ++It)
@@ -251,7 +259,7 @@ void FHttpManager::Flush(EHttpFlushReason FlushReason)
 			Request->OnHeaderReceived().Unbind();
 
 			// Don't emit these tracking logs in commandlet runs. Build system traps warnings during cook, and these are not truly fatal, but useful for tracking down shutdown issues.
-			UE_CLOG(!IsRunningCommandlet(), LogHttp, Warning, TEXT("	verb=[%s] url=[%s] refs=[%d] status=%s"), *Request->GetVerb(), *Request->GetURL(), Request.GetSharedReferenceCount(), EHttpRequestStatus::ToString(Request->GetStatus()));
+			UE_CLOG(ShouldOutputHttpWarnings(), LogHttp, Warning, TEXT("	verb=[%s] url=[%s] refs=[%d] status=%s"), *Request->GetVerb(), *Request->GetURL(), Request.GetSharedReferenceCount(), EHttpRequestStatus::ToString(Request->GetStatus()));
 		}
 	}
 
@@ -271,14 +279,14 @@ void FHttpManager::Flush(EHttpFlushReason FlushReason)
 		if (FlushTimeSoftLimitSeconds >= 0 && (AppTime - BeginWaitTime >= FlushTimeSoftLimitSeconds))
 		{
 			// Don't emit these tracking logs in commandlet runs. Build system traps warnings during cook, and these are not truly fatal, but useful for tracking down shutdown issues.
-			UE_CLOG(!IsRunningCommandlet(), LogHttp, Warning, TEXT("[FHttpManager::Flush] FlushTimeSoftLimitSeconds [%.3fs] exceeded. Cancelling %d outstanding HTTP requests:"), FlushTimeSoftLimitSeconds, Requests.Num());
+			UE_CLOG(ShouldOutputHttpWarnings(), LogHttp, Warning, TEXT("[FHttpManager::Flush] FlushTimeSoftLimitSeconds [%.3fs] exceeded. Cancelling %d outstanding HTTP requests:"), FlushTimeSoftLimitSeconds, Requests.Num());
 
 			for (TArray<FHttpRequestRef>::TIterator It(Requests); It; ++It)
 			{
 				FHttpRequestRef& Request = *It;
 
 				// Don't emit these tracking logs in commandlet runs. Build system traps warnings during cook, and these are not truly fatal, but useful for tracking down shutdown issues.
-				UE_CLOG(!IsRunningCommandlet(), LogHttp, Warning, TEXT("	verb=[%s] url=[%s] refs=[%d] status=%s"), *Request->GetVerb(), *Request->GetURL(), Request.GetSharedReferenceCount(), EHttpRequestStatus::ToString(Request->GetStatus()));
+				UE_CLOG(ShouldOutputHttpWarnings(), LogHttp, Warning, TEXT("	verb=[%s] url=[%s] refs=[%d] status=%s"), *Request->GetVerb(), *Request->GetURL(), Request.GetSharedReferenceCount(), EHttpRequestStatus::ToString(Request->GetStatus()));
 
 				FScopedEnterBackgroundEvent(*Request->GetURL());
 
@@ -300,7 +308,7 @@ void FHttpManager::Flush(EHttpFlushReason FlushReason)
 					if (AppTime >= StallWarnTime)
 					{
 						// Don't emit these tracking logs in commandlet runs. Build system traps warnings during cook, and these are not truly fatal, but useful for tracking down shutdown issues.
-						UE_CLOG(!IsRunningCommandlet(), LogHttp, Warning, TEXT("	Ticking HTTPThread for %d outstanding Http requests."), Requests.Num());
+						UE_CLOG(ShouldOutputHttpWarnings(), LogHttp, Warning, TEXT("	Ticking HTTPThread for %d outstanding Http requests."), Requests.Num());
 						StallWarnTime = AppTime + 0.5;
 					}
 					Thread->Tick();
@@ -308,7 +316,7 @@ void FHttpManager::Flush(EHttpFlushReason FlushReason)
 				else
 				{
 					// Don't emit these tracking logs in commandlet runs. Build system traps warnings during cook, and these are not truly fatal, but useful for tracking down shutdown issues.
-					UE_CLOG(!IsRunningCommandlet(), LogHttp, Warning, TEXT("	Sleeping %.3fs to wait for %d outstanding Http Requests."), SecondsToSleepForOutstandingThreadedRequests, Requests.Num());
+					UE_CLOG(ShouldOutputHttpWarnings(), LogHttp, Warning, TEXT("	Sleeping %.3fs to wait for %d outstanding Http Requests."), SecondsToSleepForOutstandingThreadedRequests, Requests.Num());
 					FPlatformProcess::Sleep(SecondsToSleepForOutstandingThreadedRequests);
 				}
 			}
@@ -324,7 +332,7 @@ void FHttpManager::Flush(EHttpFlushReason FlushReason)
 	UE_CLOG(!IsRunningCommandlet(), LogHttp, Verbose, TEXT("[FHttpManager::Flush] Cleanup ended after %.3fs. %d outstanding Http Requests."), AppTime - BeginWaitTime, Requests.Num());
 
 	// Don't emit these tracking logs in commandlet runs. Build system traps warnings during cook, and these are not truly fatal, but useful for tracking down shutdown issues.
-	if (Requests.Num() > 0 && (FlushTimeHardLimitSeconds > 0 && (AppTime - BeginWaitTime > FlushTimeHardLimitSeconds)) && !IsRunningCommandlet())
+	if (Requests.Num() > 0 && (FlushTimeHardLimitSeconds > 0 && (AppTime - BeginWaitTime > FlushTimeHardLimitSeconds)) && ShouldOutputHttpWarnings())
 	{
 		UE_LOG(LogHttp, Warning, TEXT("[FHttpManager::Flush] FlushTimeHardLimitSeconds [%.3fs] exceeded. The following requests are being abandoned without being flushed:"), FlushTimeHardLimitSeconds);
 
