Expand silencing of http manager warnings to all unattended jobs
#rnx

diff --git a/Engine/Source/Runtime/Online/HTTP/Private/HttpManager.cpp b/Engine/Source/Runtime/Online/HTTP/Private/HttpManager.cpp
--- a/Engine/Source/Runtime/Online/HTTP/Private/HttpManager.cpp
+++ b/Engine/Source/Runtime/Online/HTTP/Private/HttpManager.cpp
@@ -6,6 +6,7 @@
 #include "HAL/PlatformProcess.h"
 #include "Misc/ScopeLock.h"
 #include "Http.h"
+#include "Misc/App.h"
 #include "Misc/Guid.h"
 #include "Misc/Fork.h"
 
@@ -220,7 +221,7 @@ namespace
 {
 	bool ShouldOutputHttpWarnings()
 	{
-		return !IsRunningCommandlet() && !GIsAutomationTesting;
+		return !IsRunningCommandlet() && !FApp::IsUnattended();
 	}
 }
 
