optionally fail to get a valid HttpRouter if we fail to bind/start listening on the given port, modify callers to request this behavior if it's clear they expected it to behave this way, and allow listener sockets to reuse addresses/ports when binding

#rnx

diff --git a/Engine/Plugins/VirtualProduction/RemoteControl/Source/WebRemoteControl/Private/WebRemoteControl.cpp b/Engine/Plugins/VirtualProduction/RemoteControl/Source/WebRemoteControl/Private/WebRemoteControl.cpp
--- a/Engine/Plugins/VirtualProduction/RemoteControl/Source/WebRemoteControl/Private/WebRemoteControl.cpp
+++ b/Engine/Plugins/VirtualProduction/RemoteControl/Source/WebRemoteControl/Private/WebRemoteControl.cpp
@@ -416,7 +416,7 @@ void FWebRemoteControlModule::StartHttpServer()
 {
 	if (!HttpRouter)
 	{
-		HttpRouter = FHttpServerModule::Get().GetHttpRouter(HttpServerPort);
+		HttpRouter = FHttpServerModule::Get().GetHttpRouter(HttpServerPort, /* bFailOnBindFailure = */ true);
 		if (!HttpRouter)
 		{
 			UE_LOG(LogRemoteControl, Error, TEXT("Web Remote Call server couldn't be started on port %d"), HttpServerPort);
diff --git a/Engine/Source/Runtime/Online/HTTPServer/Private/HttpListener.cpp b/Engine/Source/Runtime/Online/HTTPServer/Private/HttpListener.cpp
--- a/Engine/Source/Runtime/Online/HTTPServer/Private/HttpListener.cpp
+++ b/Engine/Source/Runtime/Online/HTTPServer/Private/HttpListener.cpp
@@ -78,6 +78,12 @@ bool FHttpListener::StartListening()
 	}
 
 	BindAddress->SetPort(ListenPort);
+
+	if (Config.bReuseAddressAndPort)
+	{
+		NewSocket->SetReuseAddr(true);
+	}
+
 	if (!NewSocket->Bind(*BindAddress))
 	{
 		UE_LOG(LogHttpListener, Error, 
diff --git a/Engine/Source/Runtime/Online/HTTPServer/Private/HttpServerConfig.cpp b/Engine/Source/Runtime/Online/HTTPServer/Private/HttpServerConfig.cpp
--- a/Engine/Source/Runtime/Online/HTTPServer/Private/HttpServerConfig.cpp
+++ b/Engine/Source/Runtime/Online/HTTPServer/Private/HttpServerConfig.cpp
@@ -18,6 +18,7 @@ const FHttpServerListenerConfig FHttpServerConfig::GetListenerConfig(uint32 Port
 	GConfig->GetInt(*IniSectionName, TEXT("DefaultBufferSize"), Config.BufferSize, GEngineIni);
 	GConfig->GetInt(*IniSectionName, TEXT("DefaultConnectionsBacklogSize"), Config.ConnectionsBacklogSize, GEngineIni);
 	GConfig->GetInt(*IniSectionName, TEXT("DefaultMaxConnectionsAcceptPerFrame"), Config.MaxConnectionsAcceptPerFrame, GEngineIni);
+	GConfig->GetBool(*IniSectionName, TEXT("DefaultReuseAddressAndPort"), Config.bReuseAddressAndPort, GEngineIni);
 
 	// Apply per-port ini overrides
 	TArray<FString> ListenerConfigs;
@@ -46,6 +47,7 @@ const FHttpServerListenerConfig FHttpServerConfig::GetListenerConfig(uint32 Port
 				FParse::Value(*ListenerConfigStr, TEXT("BufferSize="), Config.BufferSize);
 				FParse::Value(*ListenerConfigStr, TEXT("ConnectionsBacklogSize="), Config.ConnectionsBacklogSize);
 				FParse::Value(*ListenerConfigStr, TEXT("MaxConnectionsAcceptPerFrame="), Config.MaxConnectionsAcceptPerFrame);
+				FParse::Bool(*ListenerConfigStr, TEXT("ReuseAddressAndPort="), Config.bReuseAddressAndPort);
 				break;
 			}
 		}
diff --git a/Engine/Source/Runtime/Online/HTTPServer/Private/HttpServerConfig.h b/Engine/Source/Runtime/Online/HTTPServer/Private/HttpServerConfig.h
--- a/Engine/Source/Runtime/Online/HTTPServer/Private/HttpServerConfig.h
+++ b/Engine/Source/Runtime/Online/HTTPServer/Private/HttpServerConfig.h
@@ -17,6 +17,8 @@ struct FHttpServerListenerConfig
 	int32 ConnectionsBacklogSize = 16;
 	/** Max Number of connections to accept per frame */
 	int32 MaxConnectionsAcceptPerFrame = 1;
+	/** If true, call FSocket::SetReuseAddr when binding to allow the use of an already bound address/port */
+	bool bReuseAddressAndPort = false;
 };
 
 struct FHttpServerConfig
diff --git a/Engine/Source/Runtime/Online/HTTPServer/Private/HttpServerModule.cpp b/Engine/Source/Runtime/Online/HTTPServer/Private/HttpServerModule.cpp
--- a/Engine/Source/Runtime/Online/HTTPServer/Private/HttpServerModule.cpp
+++ b/Engine/Source/Runtime/Online/HTTPServer/Private/HttpServerModule.cpp
@@ -96,27 +96,66 @@ FHttpServerModule& FHttpServerModule::Get()
 	return *Singleton;
 }
 
-TSharedPtr<IHttpRouter> FHttpServerModule::GetHttpRouter(uint32 Port)
+TSharedPtr<IHttpRouter> FHttpServerModule::GetHttpRouter(uint32 Port, bool bFailOnBindFailure /* = false */)
 {
 	check(Singleton == this);
 
-	// We may already be listening on this port
+	bool bFailedToListen = false;
+
+	// We may already have a listener for this port
 	TUniquePtr<FHttpListener>* ExistingListener = Listeners.Find(Port);
 	if (ExistingListener)
 	{
-		return ExistingListener->Get()->GetRouter();
+		if (bHttpListenersEnabled)
+		{
+			// if listeners are enabled, the existing listener for this port
+			// should always be listening (IsListening() will only be true
+			// if it fully initialized/successfully bound and actually started listening)
+			if (ensureMsgf(ExistingListener->Get()->IsListening(), TEXT("[%s] the existing listener for port %d is not listening/bound and listeners are still enabled"), ANSI_TO_TCHAR(__FUNCTION__), Port))
+			{
+				UE_LOG(LogHttpServerModule, Verbose, TEXT("[%s] found an existing, active listener for port %d"), ANSI_TO_TCHAR(__FUNCTION__), Port);
+				return ExistingListener->Get()->GetRouter();
+			}
+			else
+			{
+				// get rid of it and create a new one for now
+				Listeners.Remove(Port);
+			}
+		}
+		else
+		{
+			UE_LOG(LogHttpServerModule, Verbose, TEXT("[%s] found an existing listener for port %d but listeners are currently disabled"), ANSI_TO_TCHAR(__FUNCTION__), Port);
+			return ExistingListener->Get()->GetRouter();
+		}
 	}
 
 	// Otherwise create a new one
+	UE_LOG(LogHttpServerModule, VeryVerbose, TEXT("[%s] creating a new listener for port %d"), ANSI_TO_TCHAR(__FUNCTION__), Port);
 	TUniquePtr<FHttpListener> NewListener = MakeUnique<FHttpListener>(Port);
 
 	// Try to start this listener now
 	if (bHttpListenersEnabled)
 	{
-		NewListener->StartListening();
+		if (!NewListener->StartListening())
+		{
+			UE_LOG(LogHttpServerModule, Warning, TEXT("[%s] failed to start listening on port %d! (bFailOnBindFailure? %s)"), ANSI_TO_TCHAR(__FUNCTION__), Port, *LexToString(bFailOnBindFailure));
+			bFailedToListen = true;
+		}
+	}
+
+	if (bFailedToListen && bFailOnBindFailure)
+	{
+		// if we actually failed to listen on a port for whatever reason
+		// (i.e. listeners are enabled and we attempted to do so) then
+		// let the one we created fall out of scope and try again next time
+		return nullptr;
+	}
+	else
+	{
+		// the legacy behavior returns the router regardless of listener success
+		const auto& NewListenerRef = Listeners.Add(Port, MoveTemp(NewListener));
+		return NewListenerRef->GetRouter();
 	}
-	const auto& NewListenerRef = Listeners.Add(Port, MoveTemp(NewListener));
-	return NewListenerRef->GetRouter();
 }
 
 bool FHttpServerModule::Tick(float DeltaTime)
diff --git a/Engine/Source/Runtime/Online/HTTPServer/Private/Tests/HttpServerTests.cpp b/Engine/Source/Runtime/Online/HTTPServer/Private/Tests/HttpServerTests.cpp
--- a/Engine/Source/Runtime/Online/HTTPServer/Private/Tests/HttpServerTests.cpp
+++ b/Engine/Source/Runtime/Online/HTTPServer/Private/Tests/HttpServerTests.cpp
@@ -2,6 +2,7 @@
 
 #include "CoreMinimal.h"
 #include "Misc/AutomationTest.h"
+#include "Math/NumericLimits.h"
 #include "HttpServerModule.h"
 #include "IHttpRouter.h"
 #include "HttpRouteHandle.h"
@@ -17,6 +18,7 @@ IMPLEMENT_SIMPLE_AUTOMATION_TEST(FHttpServerIntegrationTest, "System.Online.Http
 bool FHttpServerIntegrationTest::RunTest(const FString& Parameters)
 {
 	const uint32 HttpRouterPort = 8888;
+	const uint32 InvalidHttpRouterPort = TNumericLimits<uint16>::Max() + 1; // 65536
 	const FHttpPath HttpPath(TEXT("/TestHttpServer"));
 
 	// Ensure router creation
@@ -27,6 +29,14 @@ bool FHttpServerIntegrationTest::RunTest(const FString& Parameters)
 	TSharedPtr<IHttpRouter> DuplicateHttpRouter = FHttpServerModule::Get().GetHttpRouter(HttpRouterPort);
 	TestEqual(TEXT("HttpRouter Duplicates"), HttpRouter, DuplicateHttpRouter);
 
+	// Ensure failed port binds result in a null router instance if requested
+	TSharedPtr<IHttpRouter> InvalidHttpRouterOnFail = FHttpServerModule::Get().GetHttpRouter(InvalidHttpRouterPort, /* bFailOnBindFailure = */ true);
+	TestFalse(TEXT("HttpRouter is null on bind failure if requested"), InvalidHttpRouterOnFail.IsValid());
+
+	// Ensure failed port binds still return a valid router if not explicitly requested to fail (and by default)
+	TSharedPtr<IHttpRouter> ValidHttpRouterOnFail = FHttpServerModule::Get().GetHttpRouter(InvalidHttpRouterPort /*, bFailOnBindFailure = false */);
+	TestTrue(TEXT("HttpRouter is NOT null on bind failure by default"), ValidHttpRouterOnFail.IsValid());
+
 	// Ensure we can create route bindings
 	const FHttpRequestHandler RequestHandler = [this]
 	(const FHttpServerRequest& Request, const FHttpResultCallback& OnComplete)
diff --git a/Engine/Source/Runtime/Online/HTTPServer/Public/HttpServerModule.h b/Engine/Source/Runtime/Online/HTTPServer/Public/HttpServerModule.h
--- a/Engine/Source/Runtime/Online/HTTPServer/Public/HttpServerModule.h
+++ b/Engine/Source/Runtime/Online/HTTPServer/Public/HttpServerModule.h
@@ -40,10 +40,11 @@ public:
 	/**
 	 * Per-port-binding access to an http router
 	 *
-	 * @param  Port The listener's bound port 
+	 * @param  Port The listener's bound port
+	 * @param  bFailOnBindFailure if true, return nullptr if we fail to bind/listen on the given port
 	 * @return An IHttpRouter instance that can be leveraged to respond to HTTP requests
 	 */
-	HTTPSERVER_API TSharedPtr<IHttpRouter> GetHttpRouter(uint32 Port);
+	HTTPSERVER_API TSharedPtr<IHttpRouter> GetHttpRouter(uint32 Port, bool bFailOnBindFailure = false);
 
 	/**
 	 * FTSTicker callback
diff --git a/Engine/Source/Runtime/PerfCounters/Private/PerfCounters.cpp b/Engine/Source/Runtime/PerfCounters/Private/PerfCounters.cpp
--- a/Engine/Source/Runtime/PerfCounters/Private/PerfCounters.cpp
+++ b/Engine/Source/Runtime/PerfCounters/Private/PerfCounters.cpp
@@ -62,7 +62,7 @@ bool FPerfCounters::Initialize()
 	}
 
 	// Get an IHttpRouter on the command-line designated port
-	HttpRouter = FHttpServerModule::Get().GetHttpRouter(StatsPort);
+	HttpRouter = FHttpServerModule::Get().GetHttpRouter(StatsPort, /* bFailOnBindFailure = */ true);
 	if (!HttpRouter)
 	{
 		UE_LOG(LogPerfCounters, Error, 
