Fix TSan warnings by converting TQueue use to the improved TSpscQueue

diff --git a/Engine/Source/Runtime/Core/Public/Containers/SpscQueue.h b/Engine/Source/Runtime/Core/Public/Containers/SpscQueue.h
--- a/Engine/Source/Runtime/Core/Public/Containers/SpscQueue.h
+++ b/Engine/Source/Runtime/Core/Public/Containers/SpscQueue.h
@@ -80,6 +80,18 @@ public:
 		return Value;
 	}
 
+	bool Dequeue(ElementType& OutElem)
+	{
+		TOptional<ElementType> LocalElement = Dequeue();
+		if (LocalElement.IsSet())
+		{
+			OutElem = LocalElement.GetValue();
+			return true;
+		}
+		
+		return false;
+	}
+
 private:
 	struct FNode
 	{
diff --git a/Engine/Source/Runtime/Online/HTTP/Private/Curl/CurlHttp.h b/Engine/Source/Runtime/Online/HTTP/Private/Curl/CurlHttp.h
--- a/Engine/Source/Runtime/Online/HTTP/Private/Curl/CurlHttp.h
+++ b/Engine/Source/Runtime/Online/HTTP/Private/Curl/CurlHttp.h
@@ -6,7 +6,7 @@
 #include "HAL/ThreadSafeCounter.h"
 #include "Interfaces/IHttpResponse.h"
 #include "IHttpThreadedRequest.h"
-#include "Containers/Queue.h"
+#include "Containers/SpscQueue.h"
 #include "GenericPlatform/HttpRequestPayload.h"
 #include "HAL/ThreadSafeBool.h"
 
@@ -467,7 +467,7 @@ private:
 	/** Cached key/value header pairs. Parsed once request completes. Only accessible on the game thread. */
 	TMap<FString, FString> Headers;
 	/** Newly received headers we need to inform listeners about */
-	TQueue<TPair<FString, FString>> NewlyReceivedHeaders;
+	TSpscQueue<TPair<FString, FString>> NewlyReceivedHeaders;
 	/** Cached code from completed response */
 	int32 HttpCode;
 	/** Cached content length from completed response */
diff --git a/Engine/Source/Runtime/Online/HTTP/Private/HttpThread.h b/Engine/Source/Runtime/Online/HTTP/Private/HttpThread.h
--- a/Engine/Source/Runtime/Online/HTTP/Private/HttpThread.h
+++ b/Engine/Source/Runtime/Online/HTTP/Private/HttpThread.h
@@ -8,6 +8,7 @@
 #include "HttpPackage.h"
 #include "Misc/SingleThreadRunnable.h"
 #include "Containers/Queue.h"
+#include "Containers/SpscQueue.h"
 
 class IHttpThreadedRequest;
 
@@ -153,7 +154,7 @@ protected:
 	 * Threaded requests that have completed and are waiting for the game thread to process.
 	 * Added to on HTTP thread, processed then cleared on game thread (Single producer, single consumer)
 	 */
-	TQueue<IHttpThreadedRequest*, EQueueMode::Spsc> CompletedThreadedRequests;
+	TSpscQueue<IHttpThreadedRequest*> CompletedThreadedRequests;
 
 	/** Pointer to Runnable Thread */
 	FRunnableThread* Thread;
