fix HTTP server integration tests

#rnx

diff --git a/Engine/Source/Runtime/Online/HTTPServer/Private/Tests/HttpServerTests.cpp b/Engine/Source/Runtime/Online/HTTPServer/Private/Tests/HttpServerTests.cpp
--- a/Engine/Source/Runtime/Online/HTTPServer/Private/Tests/HttpServerTests.cpp
+++ b/Engine/Source/Runtime/Online/HTTPServer/Private/Tests/HttpServerTests.cpp
@@ -29,10 +29,6 @@ bool FHttpServerIntegrationTest::RunTest(const FString& Parameters)
 	TSharedPtr<IHttpRouter> DuplicateHttpRouter = FHttpServerModule::Get().GetHttpRouter(HttpRouterPort);
 	TestEqual(TEXT("HttpRouter Duplicates"), HttpRouter, DuplicateHttpRouter);
 
-	// Ensure failed port binds result in a null router instance if requested
-	TSharedPtr<IHttpRouter> InvalidHttpRouterOnFail = FHttpServerModule::Get().GetHttpRouter(InvalidHttpRouterPort, /* bFailOnBindFailure = */ true);
-	TestFalse(TEXT("HttpRouter is null on bind failure if requested"), InvalidHttpRouterOnFail.IsValid());
-
 	// Ensure failed port binds still return a valid router if not explicitly requested to fail (and by default)
 	TSharedPtr<IHttpRouter> ValidHttpRouterOnFail = FHttpServerModule::Get().GetHttpRouter(InvalidHttpRouterPort /*, bFailOnBindFailure = false */);
 	TestTrue(TEXT("HttpRouter is NOT null on bind failure by default"), ValidHttpRouterOnFail.IsValid());
@@ -52,6 +48,10 @@ bool FHttpServerIntegrationTest::RunTest(const FString& Parameters)
 
 	FHttpServerModule::Get().StartAllListeners();
 
+	// Ensure failed port binds result in a null router instance if requested (and listeners are enabled)
+	TSharedPtr<IHttpRouter> InvalidHttpRouterOnFail = FHttpServerModule::Get().GetHttpRouter(InvalidHttpRouterPort, /* bFailOnBindFailure = */ true);
+	TestFalse(TEXT("HttpRouter is null on bind failure if requested"), InvalidHttpRouterOnFail.IsValid());
+
 	// Make a request
 	/*
 	ADD_LATENT_AUTOMATION_COMMAND(FExecStringLatentCommand(TEXT("HTTP TEST 1 http://localhost:8888/TestHttpServer")));
