Fixed TSAN warning by making `RunningThreadedRequestLimit` atomic

diff --git a/Engine/Source/Runtime/Online/HTTP/Private/HttpThread.cpp b/Engine/Source/Runtime/Online/HTTP/Private/HttpThread.cpp
--- a/Engine/Source/Runtime/Online/HTTP/Private/HttpThread.cpp
+++ b/Engine/Source/Runtime/Online/HTTP/Private/HttpThread.cpp
@@ -174,11 +174,17 @@ bool FHttpThread::NeedsSingleThreadTick() const
 
 void FHttpThread::UpdateConfigs()
 {
-	GConfig->GetInt(TEXT("HTTP.HttpThread"), TEXT("RunningThreadedRequestLimit"), RunningThreadedRequestLimit, GEngineIni);
-	if (RunningThreadedRequestLimit < 1)
+	int32 LocalRunningThreadedRequestLimit = -1;
+	if (GConfig->GetInt(TEXT("HTTP.HttpThread"), TEXT("RunningThreadedRequestLimit"), LocalRunningThreadedRequestLimit, GEngineIni))
 	{
-		UE_LOG(LogHttp, Warning, TEXT("RunningThreadedRequestLimit must be configured as a number greater than 0. Current value is %d."), RunningThreadedRequestLimit);
-		RunningThreadedRequestLimit = INT_MAX;
+		if (LocalRunningThreadedRequestLimit < 1)
+		{
+			UE_LOG(LogHttp, Warning, TEXT("RunningThreadedRequestLimit must be configured as a number greater than 0. The configured value is %d. Ignored. The current value is still %d"), LocalRunningThreadedRequestLimit, RunningThreadedRequestLimit.load());
+		}
+		else
+		{
+			RunningThreadedRequestLimit = LocalRunningThreadedRequestLimit;
+		}
 	}
 }
 
diff --git a/Engine/Source/Runtime/Online/HTTP/Private/HttpThread.h b/Engine/Source/Runtime/Online/HTTP/Private/HttpThread.h
--- a/Engine/Source/Runtime/Online/HTTP/Private/HttpThread.h
+++ b/Engine/Source/Runtime/Online/HTTP/Private/HttpThread.h
@@ -10,6 +10,8 @@
 #include "Containers/Queue.h"
 #include "Containers/SpscQueue.h"
 
+#include <atomic>
+
 class IHttpThreadedRequest;
 
 /**
@@ -168,5 +170,5 @@ private:
 	bool bIsStopped;
 
 	/** Limit for threaded http requests running at the same time. If not specified through configuration values, there will be no limit */
-	int32 RunningThreadedRequestLimit = INT_MAX;
+	std::atomic<int32> RunningThreadedRequestLimit = INT_MAX;
 };
