Direct FTCHARToUTF8_Convert usage replaced with StringCast or FPlatformString in HTTP module.

diff --git a/Engine/Source/Runtime/Online/HTTP/Private/Curl/CurlHttp.cpp b/Engine/Source/Runtime/Online/HTTP/Private/Curl/CurlHttp.cpp
--- a/Engine/Source/Runtime/Online/HTTP/Private/Curl/CurlHttp.cpp
+++ b/Engine/Source/Runtime/Online/HTTP/Private/Curl/CurlHttp.cpp
@@ -331,10 +331,10 @@ void FCurlHttpRequest::SetContentAsString(const FString& ContentString)
 		return;
 	}
 
-	int32 Utf8Length = FTCHARToUTF8_Convert::ConvertedLength(*ContentString, ContentString.Len());
+	int32 Utf8Length = FPlatformString::ConvertedLength<UTF8CHAR>(*ContentString, ContentString.Len());
 	TArray<uint8> Buffer;
 	Buffer.SetNumUninitialized(Utf8Length);
-	FTCHARToUTF8_Convert::Convert((UTF8CHAR*)Buffer.GetData(), Buffer.Num(), *ContentString, ContentString.Len());
+	FPlatformString::Convert((UTF8CHAR*)Buffer.GetData(), Buffer.Num(), *ContentString, ContentString.Len());
 	RequestPayload = MakeUnique<FRequestPayloadInMemory>(MoveTemp(Buffer));
 	bIsRequestPayloadSeekable = true;
 }
diff --git a/Engine/Source/Runtime/Online/HTTP/Private/IXML/HttpIXml.cpp b/Engine/Source/Runtime/Online/HTTP/Private/IXML/HttpIXml.cpp
--- a/Engine/Source/Runtime/Online/HTTP/Private/IXML/HttpIXml.cpp
+++ b/Engine/Source/Runtime/Online/HTTP/Private/IXML/HttpIXml.cpp
@@ -137,10 +137,10 @@ void FHttpRequestIXML::SetContentAsString(const FString& ContentString)
 {
 	if ( ContentString.Len() )
 	{
-		int32 Utf8Length = FTCHARToUTF8_Convert::ConvertedLength(*ContentString, ContentString.Len());
+		int32 Utf8Length = FPlatformString::ConvertedLength<UTF8CHAR>(*ContentString, ContentString.Len());
 		TArray<uint8> Buffer;
 		Buffer.SetNumUninitialized(Utf8Length);
-		FTCHARToUTF8_Convert::Convert((ANSICHAR*)Buffer.GetData(), Buffer.Num(), *ContentString, ContentString.Len());
+		FPlatformString::Convert((UTF8CHAR*)Buffer.GetData(), Buffer.Num(), *ContentString, ContentString.Len());
 		Payload = MakeUnique<FRequestPayloadInMemory>(MoveTemp(Buffer));
 	}
 }
diff --git a/Engine/Source/Runtime/Online/HTTP/Private/NullHttp.cpp b/Engine/Source/Runtime/Online/HTTP/Private/NullHttp.cpp
--- a/Engine/Source/Runtime/Online/HTTP/Private/NullHttp.cpp
+++ b/Engine/Source/Runtime/Online/HTTP/Private/NullHttp.cpp
@@ -79,9 +79,9 @@ void FNullHttpRequest::SetContent(TArray<uint8>&& ContentPayload)
 
 void FNullHttpRequest::SetContentAsString(const FString& ContentString)
 {
-	int32 Utf8Length = FTCHARToUTF8_Convert::ConvertedLength(*ContentString, ContentString.Len());
+	int32 Utf8Length = FPlatformString::ConvertedLength<UTF8CHAR>(*ContentString, ContentString.Len());
 	Payload.SetNumUninitialized(Utf8Length);
-	FTCHARToUTF8_Convert::Convert((UTF8CHAR*)Payload.GetData(), Payload.Num(), *ContentString, ContentString.Len());
+	FPlatformString::Convert((UTF8CHAR*)Payload.GetData(), Payload.Num(), *ContentString, ContentString.Len());
 }
 
 bool FNullHttpRequest::SetContentAsStreamedFile(const FString& Filename)
