Decrease priority of WebSockets thread from normal to below normal. This helps performance on platforms with a low core count.

diff --git a/Engine/Source/Runtime/Online/WebSockets/Private/Lws/LwsWebSocketsManager.cpp b/Engine/Source/Runtime/Online/WebSockets/Private/Lws/LwsWebSocketsManager.cpp
--- a/Engine/Source/Runtime/Online/WebSockets/Private/Lws/LwsWebSocketsManager.cpp
+++ b/Engine/Source/Runtime/Online/WebSockets/Private/Lws/LwsWebSocketsManager.cpp
@@ -157,7 +157,7 @@ void FLwsWebSocketsManager::InitWebSockets(TArrayView<const FString> Protocols)
 	
 	int32 ThreadStackSize = 128 * 1024;
 	GConfig->GetInt(TEXT("WebSockets.LibWebSockets"), TEXT("ThreadStackSize"), ThreadStackSize, GEngineIni);
-	Thread = FForkProcessHelper::CreateForkableThread(this, TEXT("LibwebsocketsThread"), 128 * 1024, TPri_Normal, FPlatformAffinity::GetNoAffinityMask());
+	Thread = FForkProcessHelper::CreateForkableThread(this, TEXT("LibwebsocketsThread"), 128 * 1024, TPri_BelowNormal, FPlatformAffinity::GetNoAffinityMask());
 	if (!Thread)
 	{
 		UE_LOG(LogWebSockets, Error, TEXT("FLwsWebSocketsManager failed to initialize thread!"));
