Added scopes when HTTP thread is sleeping

diff --git a/Engine/Source/Runtime/Online/HTTP/Private/HttpThread.cpp b/Engine/Source/Runtime/Online/HTTP/Private/HttpThread.cpp
--- a/Engine/Source/Runtime/Online/HTTP/Private/HttpThread.cpp
+++ b/Engine/Source/Runtime/Online/HTTP/Private/HttpThread.cpp
@@ -20,6 +20,8 @@ DECLARE_CYCLE_STAT(TEXT("StartThreadedRequest"), STAT_HTTPThread_StartThreadedRe
 DECLARE_CYCLE_STAT(TEXT("HttpThreadTick"), STAT_HTTPThread_HttpThreadTick, STATGROUP_HTTPThread);
 DECLARE_CYCLE_STAT(TEXT("IsThreadedRequestComplete"), STAT_HTTPThread_IsThreadedRequestComplete, STATGROUP_HTTPThread);
 DECLARE_CYCLE_STAT(TEXT("CompleteThreadedRequest"), STAT_HTTPThread_CompleteThreadedRequest, STATGROUP_HTTPThread);
+DECLARE_CYCLE_STAT(TEXT("ActiveSleep"), STAT_HTTPThread_ActiveSleep, STATGROUP_HTTPThread);
+DECLARE_CYCLE_STAT(TEXT("IdleSleep"), STAT_HTTPThread_IdleSleep, STATGROUP_HTTPThread);
 
 // FHttpThread
 
@@ -135,6 +137,7 @@ uint32 FHttpThread::Run()
 				const double InnerLoopEnd = FPlatformTime::Seconds();
 				if (bKeepProcessing)
 				{
+					SCOPE_CYCLE_COUNTER(STAT_HTTPThread_ActiveSleep);
 					double InnerLoopTime = InnerLoopEnd - InnerLoopBegin;
 					double InnerSleep = FMath::Max(HttpThreadActiveFrameTimeInSeconds - InnerLoopTime, HttpThreadActiveMinimumSleepTimeInSeconds);
 					FPlatformProcess::SleepNoStats(InnerSleep);
@@ -144,6 +147,7 @@ uint32 FHttpThread::Run()
 					OuterLoopEnd = InnerLoopEnd;
 				}
 			}
+			SCOPE_CYCLE_COUNTER(STAT_HTTPThread_IdleSleep)
 			double OuterLoopTime = OuterLoopEnd - OuterLoopBegin;
 			double OuterSleep = FMath::Max(HttpThreadIdleFrameTimeInSeconds - OuterLoopTime, HttpThreadIdleMinimumSleepTimeInSeconds);
 			FPlatformProcess::SleepNoStats(OuterSleep);
