CurlHttp
* Slightly optimized DebugCallback function
* Remove extra FString allocations by directly converting from ANSICHAR to TCHAR in the final FString.
* Don't create FStrings in CURLINFO_HEADER_OUT events unless LogHTTP is set to very verbose

diff --git a/Engine/Source/Runtime/Online/HTTP/Private/Curl/CurlHttp.cpp b/Engine/Source/Runtime/Online/HTTP/Private/Curl/CurlHttp.cpp
--- a/Engine/Source/Runtime/Online/HTTP/Private/Curl/CurlHttp.cpp
+++ b/Engine/Source/Runtime/Online/HTTP/Private/Curl/CurlHttp.cpp
@@ -629,8 +629,8 @@ size_t FCurlHttpRequest::DebugCallback(CURL * Handle, curl_infotype DebugInfoTyp
 				char* FoundNulPtr = (char*)memchr(DebugInfo, 0, DebugInfoSize);
 				int CalculatedSize = FoundNulPtr != nullptr ? FoundNulPtr - DebugInfo : DebugInfoSize;
 
-				auto ConvertedString = StringCast<TCHAR>(static_cast<const ANSICHAR*>(DebugInfo), CalculatedSize);
-				FString DebugText(ConvertedString.Length(), ConvertedString.Get());
+				FString DebugText(CalculatedSize, static_cast<const ANSICHAR*>(DebugInfo));
+
 				DebugText.ReplaceInline(TEXT("\n"), TEXT(""), ESearchCase::CaseSensitive);
 				DebugText.ReplaceInline(TEXT("\r"), TEXT(""), ESearchCase::CaseSensitive);
 				UE_LOG(LogHttp, VeryVerbose, TEXT("%p: '%s'"), this, *DebugText);
@@ -649,49 +649,52 @@ size_t FCurlHttpRequest::DebugCallback(CURL * Handle, curl_infotype DebugInfoTyp
 
 		case CURLINFO_HEADER_OUT:
 			{
-				// C string is not null terminated:  https://curl.haxx.se/libcurl/c/CURLOPT_DEBUGFUNCTION.html
-
-				// Scan for \r\n\r\n.  According to some code in tool_cb_dbg.c, special processing is needed for
-				// CURLINFO_HEADER_OUT blocks when containing both headers and data (which may be binary).
-				//
-				// Truncate at 1023 characters. This is just an arbitrary number based on a buffer size seen in
-				// the libcurl code.
-				int RecalculatedSize = FMath::Min(DebugInfoSize, (size_t)1023);
-				for (int Index = 0; Index <= RecalculatedSize - 4; ++Index)
+				if (UE_LOG_ACTIVE(LogHttp, VeryVerbose))
 				{
-					if (DebugInfo[Index] == '\r' && DebugInfo[Index + 1] == '\n'
-							&& DebugInfo[Index + 2] == '\r' && DebugInfo[Index + 3] == '\n')
+					// C string is not null terminated:  https://curl.haxx.se/libcurl/c/CURLOPT_DEBUGFUNCTION.html
+
+					// Scan for \r\n\r\n.  According to some code in tool_cb_dbg.c, special processing is needed for
+					// CURLINFO_HEADER_OUT blocks when containing both headers and data (which may be binary).
+					//
+					// Truncate at 1023 characters. This is just an arbitrary number based on a buffer size seen in
+					// the libcurl code.
+					int RecalculatedSize = FMath::Min(DebugInfoSize, (size_t)1023);
+					for (int Index = 0; Index <= RecalculatedSize - 4; ++Index)
 					{
-						RecalculatedSize = Index;
-						break;
+						if (DebugInfo[Index] == '\r' && DebugInfo[Index + 1] == '\n'
+								&& DebugInfo[Index + 2] == '\r' && DebugInfo[Index + 3] == '\n')
+						{
+							RecalculatedSize = Index;
+							break;
+						}
 					}
-				}
 
-				// As lib/http.c states that CURLINFO_HEADER_OUT may contain binary data, only print it if
-				// the header data is readable.
-				bool bIsPrintable = true;
-				for (int Index = 0; Index < RecalculatedSize; ++Index)
-				{
-					unsigned char Ch = DebugInfo[Index];
-					if (!isprint(Ch) && !isspace(Ch))
+					// As lib/http.c states that CURLINFO_HEADER_OUT may contain binary data, only print it if
+					// the header data is readable.
+					bool bIsPrintable = true;
+					for (int Index = 0; Index < RecalculatedSize; ++Index)
 					{
-						bIsPrintable = false;
-						break;
+						unsigned char Ch = DebugInfo[Index];
+						if (!isprint(Ch) && !isspace(Ch))
+						{
+							bIsPrintable = false;
+							break;
+						}
 					}
-				}
 
-				if (bIsPrintable)
-				{
-					auto ConvertedString = StringCast<TCHAR>(static_cast<const ANSICHAR*>(DebugInfo), RecalculatedSize);
-					FString DebugText(ConvertedString.Length(), ConvertedString.Get());
-					DebugText.ReplaceInline(TEXT("\n"), TEXT(""), ESearchCase::CaseSensitive);
-					DebugText.ReplaceInline(TEXT("\r"), TEXT(""), ESearchCase::CaseSensitive);
-					UE_LOG(LogHttp, VeryVerbose, TEXT("%p: Sent header (%d bytes) - %s"), this, RecalculatedSize, *DebugText);
-				}
-				else
-				{
-					UE_LOG(LogHttp, VeryVerbose, TEXT("%p: Sent header (%d bytes) - contains binary data"), this, RecalculatedSize);
-				}
+					if (bIsPrintable)
+					{
+						FString DebugText(RecalculatedSize, static_cast<const ANSICHAR*>(DebugInfo));
+
+						DebugText.ReplaceInline(TEXT("\n"), TEXT(""), ESearchCase::CaseSensitive);
+						DebugText.ReplaceInline(TEXT("\r"), TEXT(""), ESearchCase::CaseSensitive);
+						UE_LOG(LogHttp, VeryVerbose, TEXT("%p: Sent header (%d bytes) - %s"), this, RecalculatedSize, *DebugText);
+					}
+					else
+					{
+						UE_LOG(LogHttp, VeryVerbose, TEXT("%p: Sent header (%d bytes) - contains binary data"), this, RecalculatedSize);
+					}
+					}
 			}
 			break;
 
