Set UnsafeTypeCastWarningLevel to Error for a few widely-used modules

diff --git a/Engine/Source/Developer/DesktopPlatform/DesktopPlatform.Build.cs b/Engine/Source/Developer/DesktopPlatform/DesktopPlatform.Build.cs
--- a/Engine/Source/Developer/DesktopPlatform/DesktopPlatform.Build.cs
+++ b/Engine/Source/Developer/DesktopPlatform/DesktopPlatform.Build.cs
@@ -37,5 +37,7 @@ public class DesktopPlatform : ModuleRules
 
 			AddEngineThirdPartyPrivateStaticDependencies(Target, "SDL2");
 		}
+
+		UnsafeTypeCastWarningLevel = WarningLevel.Error;
 	}
 }
diff --git a/Engine/Source/Developer/Zen/Zen.Build.cs b/Engine/Source/Developer/Zen/Zen.Build.cs
--- a/Engine/Source/Developer/Zen/Zen.Build.cs
+++ b/Engine/Source/Developer/Zen/Zen.Build.cs
@@ -11,5 +11,7 @@ public class Zen : ModuleRules
 		PrivateIncludePathModuleNames.Add("DesktopPlatform");
 		AddEngineThirdPartyPrivateStaticDependencies(Target, "libcurl");
 		AddEngineThirdPartyPrivateStaticDependencies(Target, "OpenSSL");
+
+		UnsafeTypeCastWarningLevel = WarningLevel.Error;
 	}
 }
diff --git a/Engine/Source/Programs/UnrealHeaderTool/UnrealHeaderTool.Build.cs b/Engine/Source/Programs/UnrealHeaderTool/UnrealHeaderTool.Build.cs
--- a/Engine/Source/Programs/UnrealHeaderTool/UnrealHeaderTool.Build.cs
+++ b/Engine/Source/Programs/UnrealHeaderTool/UnrealHeaderTool.Build.cs
@@ -31,6 +31,6 @@ public class UnrealHeaderTool : ModuleRules
 		
 		bEnableExceptions = true;
 
-		//UnsafeTypeCastWarningLevel = WarningLevel.Warning;
+		UnsafeTypeCastWarningLevel = WarningLevel.Error;
 	}
 }
diff --git a/Engine/Source/Runtime/ApplicationCore/ApplicationCore.Build.cs b/Engine/Source/Runtime/ApplicationCore/ApplicationCore.Build.cs
--- a/Engine/Source/Runtime/ApplicationCore/ApplicationCore.Build.cs
+++ b/Engine/Source/Runtime/ApplicationCore/ApplicationCore.Build.cs
@@ -1,4 +1,4 @@
-﻿// Copyright Epic Games, Inc. All Rights Reserved.
+// Copyright Epic Games, Inc. All Rights Reserved.
 
 using UnrealBuildTool;
 
@@ -96,5 +96,7 @@ public class ApplicationCore : ModuleRules
 		{
 			throw new System.Exception("ApplicationCore cannot be used when Target.bCompileAgainstApplicationCore = false.");
 		}
+
+		UnsafeTypeCastWarningLevel = WarningLevel.Error;
 	}
 }
diff --git a/Engine/Source/Runtime/BuildSettings/BuildSettings.Build.cs b/Engine/Source/Runtime/BuildSettings/BuildSettings.Build.cs
--- a/Engine/Source/Runtime/BuildSettings/BuildSettings.Build.cs
+++ b/Engine/Source/Runtime/BuildSettings/BuildSettings.Build.cs
@@ -23,5 +23,7 @@ public class BuildSettings : ModuleRules
 		PrivateDefinitions.Add(string.Format("BRANCH_NAME=\"{0}\"", Target.Version.BranchName));
 
 		PrivateDefinitions.Add(string.Format("BUILD_VERSION=\"{0}\"", Target.BuildVersion));
+
+		UnsafeTypeCastWarningLevel = WarningLevel.Error;
 	}
 }
diff --git a/Engine/Source/Runtime/Core/Core.Build.cs b/Engine/Source/Runtime/Core/Core.Build.cs
--- a/Engine/Source/Runtime/Core/Core.Build.cs
+++ b/Engine/Source/Runtime/Core/Core.Build.cs
@@ -316,7 +316,7 @@ public class Core : ModuleRules
 
 		PrivateDefinitions.Add("UE_DEFINE_LEGACY_MATH_CONSTANT_MACRO_NAMES=0");
 
-		UnsafeTypeCastWarningLevel = WarningLevel.Warning;
+		UnsafeTypeCastWarningLevel = WarningLevel.Error;
 	}
 
 	protected virtual bool SupportsBinaryConfig(ReadOnlyTargetRules Target)
diff --git a/Engine/Source/Runtime/CoreUObject/CoreUObject.Build.cs b/Engine/Source/Runtime/CoreUObject/CoreUObject.Build.cs
--- a/Engine/Source/Runtime/CoreUObject/CoreUObject.Build.cs
+++ b/Engine/Source/Runtime/CoreUObject/CoreUObject.Build.cs
@@ -26,7 +26,7 @@ public class CoreUObject : ModuleRules
         PrivateDependencyModuleNames.Add("Json");
 
 		//@TODO: UE-127233
-		// UnsafeTypeCastWarningLevel = WarningLevel.Warning;
+		// UnsafeTypeCastWarningLevel = WarningLevel.Error;
 
 		if (Target.bBuildWithEditorOnlyData)
 		{
diff --git a/Engine/Source/Runtime/Net/Common/NetCommon.Build.cs b/Engine/Source/Runtime/Net/Common/NetCommon.Build.cs
--- a/Engine/Source/Runtime/Net/Common/NetCommon.Build.cs
+++ b/Engine/Source/Runtime/Net/Common/NetCommon.Build.cs
@@ -14,5 +14,7 @@ public class NetCommon : ModuleRules
 				"Core"
 			}
 		);
+
+		UnsafeTypeCastWarningLevel = WarningLevel.Error;
 	}
 }
diff --git a/Engine/Source/Runtime/Online/BackgroundHTTPFileHash/BackgroundHTTPFileHash.Build.cs b/Engine/Source/Runtime/Online/BackgroundHTTPFileHash/BackgroundHTTPFileHash.Build.cs
--- a/Engine/Source/Runtime/Online/BackgroundHTTPFileHash/BackgroundHTTPFileHash.Build.cs
+++ b/Engine/Source/Runtime/Online/BackgroundHTTPFileHash/BackgroundHTTPFileHash.Build.cs
@@ -12,5 +12,7 @@ public class BackgroundHTTPFileHash : ModuleRules
                 "Core",
             }
         );
+
+        UnsafeTypeCastWarningLevel = WarningLevel.Error;
     }
 }
diff --git a/Engine/Source/Runtime/Online/SSL/SSL.Build.cs b/Engine/Source/Runtime/Online/SSL/SSL.Build.cs
--- a/Engine/Source/Runtime/Online/SSL/SSL.Build.cs
+++ b/Engine/Source/Runtime/Online/SSL/SSL.Build.cs
@@ -12,10 +12,11 @@ public class SSL : ModuleRules
 				Target.Platform == UnrealTargetPlatform.Mac ||
 				Target.Platform == UnrealTargetPlatform.Win64 ||
 				Target.IsInPlatformGroup(UnrealPlatformGroup.Unix) ||
-	            Target.Platform == UnrealTargetPlatform.IOS ||
-	            Target.Platform == UnrealTargetPlatform.Android;
+				Target.Platform == UnrealTargetPlatform.IOS ||
+				Target.Platform == UnrealTargetPlatform.Android;
 		}
 	}
+
 	protected virtual bool bUseDefaultSSLCert
 	{
 		get
@@ -26,10 +27,9 @@ public class SSL : ModuleRules
 		}
 	}
 
-
 	public SSL(ReadOnlyTargetRules Target) : base(Target)
-    {
-        PublicDefinitions.Add("SSL_PACKAGE=1");
+	{
+		PublicDefinitions.Add("SSL_PACKAGE=1");
 
 		PrivateDependencyModuleNames.AddRange(
 			new string[] {
@@ -59,5 +59,7 @@ public class SSL : ModuleRules
 		{
 			PublicDefinitions.Add("WITH_SSL=0");
 		}
-    }
+
+		UnsafeTypeCastWarningLevel = WarningLevel.Error;
+	}
 }
diff --git a/Engine/Source/Runtime/Projects/Projects.Build.cs b/Engine/Source/Runtime/Projects/Projects.Build.cs
--- a/Engine/Source/Runtime/Projects/Projects.Build.cs
+++ b/Engine/Source/Runtime/Projects/Projects.Build.cs
@@ -64,6 +64,8 @@ namespace UnrealBuildTool.Rules
 			{
 				PublicDefinitions.Add("LOAD_PLUGINS_FOR_TARGET_PLATFORMS=0");
 			}
+
+			UnsafeTypeCastWarningLevel = WarningLevel.Error;
 		}
 	}
 }
diff --git a/Engine/Source/Runtime/Sockets/Sockets.Build.cs b/Engine/Source/Runtime/Sockets/Sockets.Build.cs
--- a/Engine/Source/Runtime/Sockets/Sockets.Build.cs
+++ b/Engine/Source/Runtime/Sockets/Sockets.Build.cs
@@ -17,5 +17,7 @@ public class Sockets : ModuleRules
 			});
 
 		PublicDefinitions.Add("SOCKETS_PACKAGE=1");
+
+		UnsafeTypeCastWarningLevel = WarningLevel.Error;
 	}
 }
diff --git a/Engine/Source/Runtime/TraceLog/TraceLog.Build.cs b/Engine/Source/Runtime/TraceLog/TraceLog.Build.cs
--- a/Engine/Source/Runtime/TraceLog/TraceLog.Build.cs
+++ b/Engine/Source/Runtime/TraceLog/TraceLog.Build.cs
@@ -10,7 +10,9 @@ public class TraceLog : ModuleRules
 	{
 		bRequiresImplementModule = false;
 		PublicIncludePathModuleNames.Add("Core");
-    }
+
+		UnsafeTypeCastWarningLevel = WarningLevel.Error;
+	}
 
 	// used by platform extension derived classes. probably should become a project setting at some point!
 	protected void EnableTraceByDefault(ReadOnlyTargetRules Target)
