Fix warnings in build health related to iOS 15

diff --git a/Engine/Plugins/Media/ElectraUtil/Source/ElectraHTTPStream/Private/Apple/AppleElectraHTTPStream.cpp b/Engine/Plugins/Media/ElectraUtil/Source/ElectraHTTPStream/Private/Apple/AppleElectraHTTPStream.cpp
--- a/Engine/Plugins/Media/ElectraUtil/Source/ElectraHTTPStream/Private/Apple/AppleElectraHTTPStream.cpp
+++ b/Engine/Plugins/Media/ElectraUtil/Source/ElectraHTTPStream/Private/Apple/AppleElectraHTTPStream.cpp
@@ -1307,8 +1307,9 @@ FElectraHTTPStreamApple::FChallengeResponse FElectraHTTPStreamApple::ReceivedCha
 			// the chain starts with the server's cert itself, so walk backwards to optimize for roots first
 			TArray<TArray<uint8, TFixedAllocator<ISslCertificateManager::PUBLIC_KEY_DIGEST_SIZE>>> CertDigests;
 
-			CFIndex NumCerts = SecTrustGetCertificateCount(RemoteTrust);
-			for(int i = static_cast<int>(NumCerts) - 1; i >= 0; i--)
+                        CFArrayRef Certificates = SecTrustCopyCertificateChain(RemoteTrust);
+                        CFIndex CertificateCount = CFArrayGetCount(Certificates);
+			for (int i = 0; CertificateCount; ++i)
 			{
 				SecCertificateRef Cert = SecTrustGetCertificateAtIndex(RemoteTrust, i);
 
diff --git a/Engine/Plugins/Online/IOS/OnlineSubsystemIOS/Source/Private/OnlineLeaderboardsInterfaceIOS.cpp b/Engine/Plugins/Online/IOS/OnlineSubsystemIOS/Source/Private/OnlineLeaderboardsInterfaceIOS.cpp
--- a/Engine/Plugins/Online/IOS/OnlineSubsystemIOS/Source/Private/OnlineLeaderboardsInterfaceIOS.cpp
+++ b/Engine/Plugins/Online/IOS/OnlineSubsystemIOS/Source/Private/OnlineLeaderboardsInterfaceIOS.cpp
@@ -7,6 +7,7 @@
 // GameCenter includes
 #include <GameKit/GKLeaderboard.h>
 #include <GameKit/GKScore.h>
+#include <GameKit/GKLocalPlayer.h>
 
 FOnlineLeaderboardsIOS::FOnlineLeaderboardsIOS(FOnlineSubsystemIOS* InSubsystem)
 {
@@ -161,7 +162,8 @@ bool FOnlineLeaderboardsIOS::ReadLeaderboards(const TArray< FUniqueNetIdRef >& P
 		}
 
 		// Kick off a game center read request for the list of users
-		[GKPlayer loadPlayersForIdentifiers:FriendIds withCompletionHandler:^(NSArray *players, NSError *Error)
+                 GKLocalPlayer* GKLocalUser = [GKLocalPlayer localPlayer];
+                 [GKLocalUser loadFriendsWithIdentifiers:(NSArray<NSString *> *)FriendIds completionHandler:^(NSArray<GKPlayer *> *players, NSError *Error)
 		 {
 			bool bWasSuccessful = (Error == nil) && [players count] > 0;
 			if (bWasSuccessful)
diff --git a/Engine/Source/Runtime/ApplicationCore/Private/IOS/IOSWindow.cpp b/Engine/Source/Runtime/ApplicationCore/Private/IOS/IOSWindow.cpp
--- a/Engine/Source/Runtime/ApplicationCore/Private/IOS/IOSWindow.cpp
+++ b/Engine/Source/Runtime/ApplicationCore/Private/IOS/IOSWindow.cpp
@@ -23,16 +23,8 @@ void FIOSWindow::Initialize( class FIOSApplication* const Application, const TSh
 {
 	OwningApplication = Application;
 	Definition = InDefinition;
-    
-    NSArray         *allWindows = [[UIApplication sharedApplication]windows];
-    for (UIWindow   *currentWindow in allWindows)
-    {
-        if (currentWindow.isKeyWindow)
-        {
-			Window = currentWindow;
-            break;
-        }
-    }
+        
+        Window = (UIWindow *)[[[[UIApplication sharedApplication] connectedScenes] allObjects] firstObject];
 
 #if !PLATFORM_TVOS
 	if(InParent.Get() != NULL)
diff --git a/Engine/Source/Runtime/Online/HTTP/Private/Apple/AppleHTTP.cpp b/Engine/Source/Runtime/Online/HTTP/Private/Apple/AppleHTTP.cpp
--- a/Engine/Source/Runtime/Online/HTTP/Private/Apple/AppleHTTP.cpp
+++ b/Engine/Source/Runtime/Online/HTTP/Private/Apple/AppleHTTP.cpp
@@ -699,10 +699,12 @@ static const unsigned char ecdsaSecp384r1Asn1Header[] =
             // the chain starts with the server's cert itself, so walk backwards to optimize for roots first
             TArray<TArray<uint8, TFixedAllocator<ISslCertificateManager::PUBLIC_KEY_DIGEST_SIZE>>> CertDigests;
             
-            CFIndex NumCerts = SecTrustGetCertificateCount(RemoteTrust);
-            for (int i = static_cast<int>(NumCerts) - 1; i >= 0; i--)
+            CFArrayRef Certificates = SecTrustCopyCertificateChain(RemoteTrust);
+            CFIndex CertificateCount = CFArrayGetCount(Certificates);
+
+            for (int i = 0; CertificateCount; ++i)
             {
-                SecCertificateRef Cert = SecTrustGetCertificateAtIndex(RemoteTrust, i);
+                SecCertificateRef Cert = (SecCertificateRef)CFArrayGetValueAtIndex(Certificates, i);
                 
                 // this is not great, but the only way to extract a public key from a SecCertificateRef
                 // is to create an individual SecTrustRef for each cert that only contains itself and then
