Added additional libcurl HTTP configuration options to set the size of the connection cache and forbidding connection re-use

diff --git a/Engine/Source/Runtime/Online/HTTP/Private/Curl/CurlHttpManager.cpp b/Engine/Source/Runtime/Online/HTTP/Private/Curl/CurlHttpManager.cpp
--- a/Engine/Source/Runtime/Online/HTTP/Private/Curl/CurlHttpManager.cpp
+++ b/Engine/Source/Runtime/Online/HTTP/Private/Curl/CurlHttpManager.cpp
@@ -198,6 +198,17 @@ void FCurlHttpManager::InitCurl()
 			}
 		}
 
+		int32 MaxConnects = 0;
+		if (GConfig->GetInt(TEXT("HTTP.Curl"), TEXT("MaxConnects"), MaxConnects, GEngineIni) && MaxConnects >= 0)
+		{
+			const CURLMcode SetOptResult = curl_multi_setopt(GMultiHandle, CURLMOPT_MAXCONNECTS, static_cast<long>(MaxConnects));
+			if (SetOptResult != CURLM_OK)
+			{
+				UE_LOG(LogInit, Warning, TEXT("Failed to set libcurl max connects options (%d), error %d ('%s')"),
+					MaxConnects, static_cast<int32>(SetOptResult), StringCast<TCHAR>(curl_multi_strerror(SetOptResult)).Get());
+			}
+		}
+
 #if !WITH_CURL_XCURL
 		GShareHandle = curl_share_init();
 		if (NULL != GShareHandle)
@@ -218,6 +229,12 @@ void FCurlHttpManager::InitCurl()
 	}
 
 	// Init curl request options
+	bool bForbidReuse = false;
+	if (GConfig->GetBool(TEXT("HTTP.Curl"), TEXT("bForbidReuse"), bForbidReuse, GEngineIni))
+	{
+		CurlRequestOptions.bDontReuseConnections = bForbidReuse;
+	}
+	// If set on the command line for debugging this overrides the setting from the .ini file.
 	if (FParse::Param(FCommandLine::Get(), TEXT("noreuseconn")))
 	{
 		CurlRequestOptions.bDontReuseConnections = true;
