Add handling for setting the default/loopback address when using ipv6

diff --git a/Engine/Source/Runtime/Online/HTTPServer/Private/HttpListener.cpp b/Engine/Source/Runtime/Online/HTTPServer/Private/HttpListener.cpp
--- a/Engine/Source/Runtime/Online/HTTPServer/Private/HttpListener.cpp
+++ b/Engine/Source/Runtime/Online/HTTPServer/Private/HttpListener.cpp
@@ -64,6 +64,10 @@ bool FHttpListener::StartListening()
 	{
 		BindAddress->SetAnyAddress();
 	}
+	else if (0 == Config.BindAddress.Compare(TEXT("localhost"), ESearchCase::IgnoreCase))
+	{
+		BindAddress->SetLoopbackAddress();
+	}
 	else
 	{
 		bool bIsValidAddress = false;
diff --git a/Engine/Source/Runtime/Online/HTTPServer/Private/HttpServerConfig.h b/Engine/Source/Runtime/Online/HTTPServer/Private/HttpServerConfig.h
--- a/Engine/Source/Runtime/Online/HTTPServer/Private/HttpServerConfig.h
+++ b/Engine/Source/Runtime/Online/HTTPServer/Private/HttpServerConfig.h
@@ -10,7 +10,7 @@ DECLARE_LOG_CATEGORY_EXTERN(LogHttpServerConfig, Log, All);
 struct FHttpServerListenerConfig
 {
 	/** The address to bind */
-	FString BindAddress = FString(TEXT("127.0.0.1"));
+	FString BindAddress = FString(TEXT("localhost"));
 	/** Maximum send buffer size */
 	int32 BufferSize = 512 * 1024;
 	/** Number of pending connections to queue */
