Add http RequestAdded delegate.

diff --git a/Engine/Source/Runtime/Online/HTTP/Private/HttpManager.cpp b/Engine/Source/Runtime/Online/HTTP/Private/HttpManager.cpp
--- a/Engine/Source/Runtime/Online/HTTP/Private/HttpManager.cpp
+++ b/Engine/Source/Runtime/Online/HTTP/Private/HttpManager.cpp
@@ -396,6 +396,7 @@ void FHttpManager::AddRequest(const FHttpRequestRef& Request)
 	FScopeLock ScopeLock(&RequestLock);
 	check(!bFlushing);
 	Requests.Add(Request);
+	RequestAddedDelegate.ExecuteIfBound(Request);
 }
 
 void FHttpManager::RemoveRequest(const FHttpRequestRef& Request)
@@ -409,9 +410,7 @@ void FHttpManager::AddThreadedRequest(const TSharedRef<IHttpThreadedRequest, ESP
 {
 	check(Thread);
 	{
-		FScopeLock ScopeLock(&RequestLock);
-		check(!bFlushing);
-		Requests.Add(Request);
+		AddRequest(Request);
 	}
 	Thread->AddRequest(&Request.Get());
 }
diff --git a/Engine/Source/Runtime/Online/HTTP/Public/HttpManager.h b/Engine/Source/Runtime/Online/HTTP/Public/HttpManager.h
--- a/Engine/Source/Runtime/Online/HTTP/Public/HttpManager.h
+++ b/Engine/Source/Runtime/Online/HTTP/Public/HttpManager.h
@@ -36,6 +36,13 @@ enum class EHttpFlushReason : uint8
 ENUM_RANGE_BY_COUNT(EHttpFlushReason, EHttpFlushReason::Count)
 const TCHAR* LexToString(const EHttpFlushReason& FlushReason);
 
+/**
+ * Delegate called when an Http request added
+ *
+ * @param Request Http request that start things
+ */
+DECLARE_DELEGATE_OneParam(FHttpManagerRequestAddedDelegate, const FHttpRequestRef& /*Request*/);
+
 /**
  * Manages Http request that are currently being processed
  */
@@ -69,6 +76,9 @@ public:
 	 */
 	void AddRequest(const FHttpRequestRef& Request);
 
+	/** Delegate that will get called once request added */
+	FHttpManagerRequestAddedDelegate RequestAddedDelegate;
+
 	/**
 	 * Removes an Http request instance from the manager
 	 * Presumably it is done being processed
