List some ssl error codes which need to be treated as connection error.

diff --git a/Engine/Source/Runtime/Online/HTTP/Private/Curl/CurlHttp.cpp b/Engine/Source/Runtime/Online/HTTP/Private/Curl/CurlHttp.cpp
--- a/Engine/Source/Runtime/Online/HTTP/Private/Curl/CurlHttp.cpp
+++ b/Engine/Source/Runtime/Online/HTTP/Private/Curl/CurlHttp.cpp
@@ -1325,6 +1325,7 @@ void FCurlHttpRequest::FinishedRequest()
 			case CURLE_COULDNT_CONNECT:
 			case CURLE_COULDNT_RESOLVE_PROXY:
 			case CURLE_COULDNT_RESOLVE_HOST:
+			case CURLE_SSL_CONNECT_ERROR:
 				// report these as connection errors (safe to retry)
 				CompletionStatus = EHttpRequestStatus::Failed_ConnectionError;
 				break;
