Fix WebSockets unacceptable words

diff --git a/Engine/Config/ConfigRedirects.ini b/Engine/Config/ConfigRedirects.ini
--- a/Engine/Config/ConfigRedirects.ini
+++ b/Engine/Config/ConfigRedirects.ini
@@ -31,3 +31,9 @@ IniSectionBlacklist=IniSectionDenylist
 
 [/Script/Engine.RPCDoSDetectionConfig]
 RPCBlockWhitelist=RPCBlockAllowlist
+
+[LwsWebSocket]
+bDisableDomainWhitelist=bDisableDomainAllowlist
+
+[WinHttpWebSocket]
+bDisableDomainWhitelist=bDisableDomainAllowlist
\ No newline at end of file
diff --git a/Engine/Source/Runtime/Online/WebSockets/Private/Lws/LwsWebSocketsManager.cpp b/Engine/Source/Runtime/Online/WebSockets/Private/Lws/LwsWebSocketsManager.cpp
--- a/Engine/Source/Runtime/Online/WebSockets/Private/Lws/LwsWebSocketsManager.cpp
+++ b/Engine/Source/Runtime/Online/WebSockets/Private/Lws/LwsWebSocketsManager.cpp
@@ -56,7 +56,7 @@ FLwsWebSocketsManager::FLwsWebSocketsManager()
 	ThreadMinimumSleepTimeInSeconds = 0.0f;
 	GConfig->GetDouble(TEXT("WebSockets.LibWebSockets"), TEXT("ThreadMinimumSleepTimeInSeconds"), ThreadMinimumSleepTimeInSeconds, GEngineIni);
 
-	GConfig->GetBool(TEXT("LwsWebSocket"), TEXT("bDisableDomainWhitelist"), bDisableDomainAllowlist, GEngineIni);
+	GConfig->GetBool(TEXT("LwsWebSocket"), TEXT("bDisableDomainAllowlist"), bDisableDomainAllowlist, GEngineIni);
 	GConfig->GetBool(TEXT("LwsWebSocket"), TEXT("bDisableCertValidation"), bDisableCertValidation, GEngineIni);
 }
 
diff --git a/Engine/Source/Runtime/Online/WebSockets/Private/WinHttp/WinHttpWebSocket.cpp b/Engine/Source/Runtime/Online/WebSockets/Private/WinHttp/WinHttpWebSocket.cpp
--- a/Engine/Source/Runtime/Online/WebSockets/Private/WinHttp/WinHttpWebSocket.cpp
+++ b/Engine/Source/Runtime/Online/WebSockets/Private/WinHttp/WinHttpWebSocket.cpp
@@ -76,7 +76,7 @@ void FWinHttpWebSocket::Connect()
 
 	// Check Domain allowedlist if enabled
 	bool bDisableDomainAllowlist = false;
-	GConfig->GetBool(TEXT("WinHttpWebSocket"), TEXT("bDisableDomainWhitelist"), bDisableDomainAllowlist, GEngineIni);
+	GConfig->GetBool(TEXT("WinHttpWebSocket"), TEXT("bDisableDomainAllowlist"), bDisableDomainAllowlist, GEngineIni);
 	if (!bDisableDomainAllowlist)
 	{
 		FHttpManager& HttpManager = FHttpModule::Get().GetHttpManager();
