FHttpServer and FInternetAddr: Fix failed check statements and ensures in FHttpServerIntegrationTest.
Fix FHttpListener::StartListening to return false rather than fail a check statement when the ListenPort is invalid.
Downgrade ensure in FHttpServerModule::GetHttpRouter to an error when the listening fails.
#rnx

diff --git a/Engine/Source/Runtime/Online/HTTPServer/Private/HttpListener.cpp b/Engine/Source/Runtime/Online/HTTPServer/Private/HttpListener.cpp
--- a/Engine/Source/Runtime/Online/HTTPServer/Private/HttpListener.cpp
+++ b/Engine/Source/Runtime/Online/HTTPServer/Private/HttpListener.cpp
@@ -81,6 +81,13 @@ bool FHttpListener::StartListening()
 		}
 	}
 
+	if (!BindAddress->IsPortValid(ListenPort))
+	{
+		UE_LOG(LogHttpListener, Error,
+			TEXT("HttpListener detected invalid port %u"),
+			ListenPort, *Config.BindAddress, ListenPort);
+		return false;
+	}
 	BindAddress->SetPort(ListenPort);
 
 	if (Config.bReuseAddressAndPort)
diff --git a/Engine/Source/Runtime/Online/HTTPServer/Private/HttpServerModule.cpp b/Engine/Source/Runtime/Online/HTTPServer/Private/HttpServerModule.cpp
--- a/Engine/Source/Runtime/Online/HTTPServer/Private/HttpServerModule.cpp
+++ b/Engine/Source/Runtime/Online/HTTPServer/Private/HttpServerModule.cpp
@@ -111,7 +111,7 @@ TSharedPtr<IHttpRouter> FHttpServerModule::GetHttpRouter(uint32 Port, bool bFail
 			// if listeners are enabled, the existing listener for this port
 			// should always be listening (IsListening() will only be true
 			// if it fully initialized/successfully bound and actually started listening)
-			if (ensureMsgf(ExistingListener->Get()->IsListening(), TEXT("[%s] the existing listener for port %d is not listening/bound and listeners are still enabled"), ANSI_TO_TCHAR(__FUNCTION__), Port))
+			if (ExistingListener->Get()->IsListening())
 			{
 				UE_LOG(LogHttpServerModule, Verbose, TEXT("[%s] found an existing, active listener for port %d"), ANSI_TO_TCHAR(__FUNCTION__), Port);
 				return ExistingListener->Get()->GetRouter();
@@ -119,6 +119,8 @@ TSharedPtr<IHttpRouter> FHttpServerModule::GetHttpRouter(uint32 Port, bool bFail
 			else
 			{
 				// get rid of it and create a new one for now
+				UE_LOG(LogHttpServerModule, Error, TEXT("[%s] the existing listener for port %d is not listening/bound and listeners are still enabled"),
+					ANSI_TO_TCHAR(__FUNCTION__), Port);
 				Listeners.Remove(Port);
 			}
 		}
diff --git a/Engine/Source/Runtime/Sockets/Private/BSDSockets/IPAddressBSD.cpp b/Engine/Source/Runtime/Sockets/Private/BSDSockets/IPAddressBSD.cpp
--- a/Engine/Source/Runtime/Sockets/Private/BSDSockets/IPAddressBSD.cpp
+++ b/Engine/Source/Runtime/Sockets/Private/BSDSockets/IPAddressBSD.cpp
@@ -1,6 +1,8 @@
 // Copyright Epic Games, Inc. All Rights Reserved.
 
 #include "IPAddressBSD.h"
+
+#include "Math/NumericLimits.h"
 #include "SocketSubsystemBSD.h"
 
 #if PLATFORM_HAS_BSD_SOCKETS || PLATFORM_HAS_BSD_IPV6_SOCKETS
@@ -324,6 +326,13 @@ void FInternetAddrBSD::SetPort(int32 InPort)
 	((sockaddr_in*)&Addr)->sin_port = htons(IntCastChecked<uint16>(InPort));
 }
 
+/** Report whether the port is in a valid range for SetPort. */
+bool FInternetAddrBSD::IsPortValid(int32 InPort) const
+{
+	return 0 <= InPort && InPort <= MAX_uint16;
+}
+
+
 int32 FInternetAddrBSD::GetPort() const
 {
 #if PLATFORM_HAS_BSD_IPV6_SOCKETS
diff --git a/Engine/Source/Runtime/Sockets/Private/BSDSockets/IPAddressBSD.h b/Engine/Source/Runtime/Sockets/Private/BSDSockets/IPAddressBSD.h
--- a/Engine/Source/Runtime/Sockets/Private/BSDSockets/IPAddressBSD.h
+++ b/Engine/Source/Runtime/Sockets/Private/BSDSockets/IPAddressBSD.h
@@ -182,6 +182,9 @@ public:
 	 */
 	virtual void SetPort(int32 InPort) override;
 
+	/** Report whether the port is in a valid range for SetPort. */
+	virtual bool IsPortValid(int32 InPort) const override;
+
 	/** Returns the port number from this address in host byte order */
 	virtual int32 GetPort() const override;
 
diff --git a/Engine/Source/Runtime/Sockets/Public/IPAddress.h b/Engine/Source/Runtime/Sockets/Public/IPAddress.h
--- a/Engine/Source/Runtime/Sockets/Public/IPAddress.h
+++ b/Engine/Source/Runtime/Sockets/Public/IPAddress.h
@@ -78,6 +78,9 @@ public:
 	 */
 	virtual void SetPort(int32 InPort) = 0;
 
+	/** Report whether the port is in a valid range for SetPort. */
+	virtual bool IsPortValid(int32 InPort) const { return true; }
+
 	/**
 	 * Copies the port number from this address and places it into a host byte order int
 	 *
