BPS - Use correct print format for sizes when doing patch data enumeration.

diff --git a/Engine/Source/Runtime/Online/BuildPatchServices/Private/Enumeration/PatchDataEnumeration.cpp b/Engine/Source/Runtime/Online/BuildPatchServices/Private/Enumeration/PatchDataEnumeration.cpp
--- a/Engine/Source/Runtime/Online/BuildPatchServices/Private/Enumeration/PatchDataEnumeration.cpp
+++ b/Engine/Source/Runtime/Online/BuildPatchServices/Private/Enumeration/PatchDataEnumeration.cpp
@@ -61,7 +61,7 @@ namespace EnumerationHelpers
 				if (bIncludeSizes)
 				{
 					uint64 FileSize = Manifest.GetDataSize(DataGuid);
-					OutputLine += FString::Printf(TEXT("\t%u"), FileSize);
+					OutputLine += FString::Printf(TEXT("\t%" UINT64_FMT), FileSize);
 				}
 				UE_LOG(LogDataEnumeration, Verbose, TEXT("%s"), *OutputLine);
 				OutFiles.Add(MoveTemp(OutputLine));
@@ -177,7 +177,7 @@ namespace BuildPatchServices
 						if (Configuration.bIncludeSizes)
 						{
 							uint64 FileSize = File->TotalSize();
-							OutputLine += FString::Printf(TEXT("\t%u"), FileSize);
+							OutputLine += FString::Printf(TEXT("\t%" UINT64_FMT), FileSize);
 						}
 						UE_LOG(LogDataEnumeration, Verbose, TEXT("%s"), *OutputLine);
 						OutFiles.Add(MoveTemp(OutputLine));
