Add a delegate fired when an HTTP request completes to allow games to add HTTP analytics

diff --git a/Engine/Source/Runtime/Online/HTTP/Private/Apple/AppleHTTP.cpp b/Engine/Source/Runtime/Online/HTTP/Private/Apple/AppleHTTP.cpp
--- a/Engine/Source/Runtime/Online/HTTP/Private/Apple/AppleHTTP.cpp
+++ b/Engine/Source/Runtime/Online/HTTP/Private/Apple/AppleHTTP.cpp
@@ -431,13 +431,15 @@ void FAppleHttpRequest::FinishedRequest()
 {
 	UE_LOG(LogHttp, Verbose, TEXT("FAppleHttpRequest::FinishedRequest()"));
 	ElapsedTime = (float)(FPlatformTime::Seconds() - StartRequestTime);
-	if( Response.IsValid() && Response->IsReady() && !Response->HadError())
+	FHttpManager& HttpManager = FHttpModule::Get().GetHttpManager();
+	if (Response.IsValid() && Response->IsReady() && !Response->HadError())
 	{
 		UE_LOG(LogHttp, Verbose, TEXT("Request succeeded"));
 		CompletionStatus = EHttpRequestStatus::Succeeded;
 
 		// TODO: Try to broadcast OnHeaderReceived when we receive headers instead of here at the end
 		BroadcastResponseHeadersReceived();
+		HttpManager.BroadcastHttpRequestCompleted(Request);
 		OnProcessRequestComplete().ExecuteIfBound(SharedThis(this), Response, true);
 	}
 	else
@@ -451,6 +453,7 @@ void FAppleHttpRequest::FinishedRequest()
 		}
 
 		Response = nullptr;
+		HttpManager.BroadcastHttpRequestCompleted(Request);
 		OnProcessRequestComplete().ExecuteIfBound(SharedThis(this), nullptr, false);
 	}
 
@@ -458,9 +461,9 @@ void FAppleHttpRequest::FinishedRequest()
 	CleanupRequest();
 
 	// Remove from global list since processing is now complete
-	if (FHttpModule::Get().GetHttpManager().IsValidRequest(this))
+	if (HttpManager.IsValidRequest(this))
 	{
-		FHttpModule::Get().GetHttpManager().RemoveRequest(SharedThis(this));
+		HttpManager.RemoveRequest(SharedThis(this));
 	}
 }
 
diff --git a/Engine/Source/Runtime/Online/HTTP/Private/HttpManager.cpp b/Engine/Source/Runtime/Online/HTTP/Private/HttpManager.cpp
--- a/Engine/Source/Runtime/Online/HTTP/Private/HttpManager.cpp
+++ b/Engine/Source/Runtime/Online/HTTP/Private/HttpManager.cpp
@@ -380,6 +380,7 @@ bool FHttpManager::Tick(float DeltaSeconds)
 			FHttpRequestRef CompletedRequestRef = CompletedRequest->AsShared();
 			Requests.Remove(CompletedRequestRef);
 			CompletedRequest->FinishRequest();
+			BroadcastHttpRequestCompleted(CompletedRequestRef);
 		}
 	}
 	// keep ticking
@@ -396,7 +397,9 @@ void FHttpManager::AddRequest(const FHttpRequestRef& Request)
 	FScopeLock ScopeLock(&RequestLock);
 	check(!bFlushing);
 	Requests.Add(Request);
+PRAGMA_DISABLE_DEPRECATION_WARNINGS // Valid direct access of RequestAddedDelegate
 	RequestAddedDelegate.ExecuteIfBound(Request);
+PRAGMA_ENABLE_DEPRECATION_WARNINGS
 }
 
 void FHttpManager::RemoveRequest(const FHttpRequestRef& Request)
@@ -438,6 +441,18 @@ bool FHttpManager::IsValidRequest(const IHttpRequest* RequestPtr) const
 	return bResult;
 }
 
+void FHttpManager::SetRequestAddedDelegate(const FHttpManagerRequestAddedDelegate& Delegate)
+{
+PRAGMA_DISABLE_DEPRECATION_WARNINGS // Valid direct access of RequestAddedDelegate
+	RequestAddedDelegate = Delegate;
+PRAGMA_ENABLE_DEPRECATION_WARNINGS
+}
+
+void FHttpManager::SetRequestCompletedDelegate(const FHttpManagerRequestCompletedDelegate& Delegate)
+{
+	RequestCompletedDelegate = Delegate;
+}
+
 void FHttpManager::DumpRequests(FOutputDevice& Ar) const
 {
 	FScopeLock ScopeLock(&RequestLock);
@@ -454,3 +469,8 @@ bool FHttpManager::SupportsDynamicProxy() const
 {
 	return false;
 }
+
+void FHttpManager::BroadcastHttpRequestCompleted(const FHttpRequestRef& Request)
+{
+	RequestCompletedDelegate.ExecuteIfBound(Request);
+}
diff --git a/Engine/Source/Runtime/Online/HTTP/Private/NullHttp.cpp b/Engine/Source/Runtime/Online/HTTP/Private/NullHttp.cpp
--- a/Engine/Source/Runtime/Online/HTTP/Private/NullHttp.cpp
+++ b/Engine/Source/Runtime/Online/HTTP/Private/NullHttp.cpp
@@ -178,8 +178,10 @@ float FNullHttpRequest::GetElapsedTime() const
 void FNullHttpRequest::FinishedRequest()
 {
 	CompletionStatus = EHttpRequestStatus::Failed;
-	TSharedRef<IHttpRequest, ESPMode::ThreadSafe> Request = SharedThis(this);
-	FHttpModule::Get().GetHttpManager().RemoveRequest(Request);
+	FHttpRequestRef Request = SharedThis(this);
+	FHttpManager& HttpManager = FHttpModule::Get().GetHttpManager();
+	HttpManager.BroadcastHttpRequestCompleted(Request);
+	HttpManager.RemoveRequest(Request);
 
 	UE_LOG(LogHttp, Log, TEXT("Finished request %p. no response %s url=%s elapsed=%.3f"),
 		this, *GetVerb(), *GetURL(), ElapsedTime);
diff --git a/Engine/Source/Runtime/Online/HTTP/Public/HttpManager.h b/Engine/Source/Runtime/Online/HTTP/Public/HttpManager.h
--- a/Engine/Source/Runtime/Online/HTTP/Public/HttpManager.h
+++ b/Engine/Source/Runtime/Online/HTTP/Public/HttpManager.h
@@ -37,12 +37,17 @@ ENUM_RANGE_BY_COUNT(EHttpFlushReason, EHttpFlushReason::Count)
 const TCHAR* LexToString(const EHttpFlushReason& FlushReason);
 
 /**
- * Delegate called when an Http request added
- *
- * @param Request Http request that start things
+ * Delegate called when an Http request added. Intended to be used for analytics. Called from the thread that adds the HTTP request.
+ * @param Request Http request that was added
  */
 DECLARE_DELEGATE_OneParam(FHttpManagerRequestAddedDelegate, const FHttpRequestRef& /*Request*/);
 
+/**
+ * Delegate called when an Http request completes. Intended to be used for analytics. Called from the game thread.
+ * @param Request Http request that completed
+ */
+DECLARE_DELEGATE_OneParam(FHttpManagerRequestCompletedDelegate, const FHttpRequestRef& /*Request*/);
+
 /**
  * Manages Http request that are currently being processed
  */
@@ -77,8 +82,12 @@ public:
 	void AddRequest(const FHttpRequestRef& Request);
 
 	/** Delegate that will get called once request added */
+	UE_DEPRECATED(5.2, "Direct access to RequestAddedDelegate is deprecated. Please use SetRequestAddedDelegate")
 	FHttpManagerRequestAddedDelegate RequestAddedDelegate;
 
+	void SetRequestAddedDelegate(const FHttpManagerRequestAddedDelegate& Delegate);
+	void SetRequestCompletedDelegate(const FHttpManagerRequestCompletedDelegate& Delegate);
+
 	/**
 	 * Removes an Http request instance from the manager
 	 * Presumably it is done being processed
@@ -240,6 +249,9 @@ protected:
 	// This variable is set to true in Flush(EHttpFlushReason), and prevents new Http requests from being launched
 	bool bFlushing;
 
+	/** Delegate that will get called when a request completes */
+	FHttpManagerRequestCompletedDelegate RequestCompletedDelegate;
+
 	struct FHttpFlushTimeLimit
 	{
 		/**
@@ -270,4 +282,10 @@ PACKAGE_SCOPE:
 
 	/** Used to lock access to add/remove/find requests */
 	static FCriticalSection RequestLock;
+	/**
+	 * Broadcast that a non-threaded HTTP request is complete.
+	 * Called automatically internally for threaded requests.
+	 * Called explicitly by non-threaded requests
+	 */
+	void BroadcastHttpRequestCompleted(const FHttpRequestRef& Request);
 };
