Undo changes to AppleHTTP to fix compiler error

diff --git a/Engine/Source/Runtime/Online/HTTP/Private/Apple/AppleHTTP.cpp b/Engine/Source/Runtime/Online/HTTP/Private/Apple/AppleHTTP.cpp
--- a/Engine/Source/Runtime/Online/HTTP/Private/Apple/AppleHTTP.cpp
+++ b/Engine/Source/Runtime/Online/HTTP/Private/Apple/AppleHTTP.cpp
@@ -431,15 +431,13 @@ void FAppleHttpRequest::FinishedRequest()
 {
 	UE_LOG(LogHttp, Verbose, TEXT("FAppleHttpRequest::FinishedRequest()"));
 	ElapsedTime = (float)(FPlatformTime::Seconds() - StartRequestTime);
-	FHttpManager& HttpManager = FHttpModule::Get().GetHttpManager();
-	if (Response.IsValid() && Response->IsReady() && !Response->HadError())
+	if( Response.IsValid() && Response->IsReady() && !Response->HadError())
 	{
 		UE_LOG(LogHttp, Verbose, TEXT("Request succeeded"));
 		CompletionStatus = EHttpRequestStatus::Succeeded;
 
 		// TODO: Try to broadcast OnHeaderReceived when we receive headers instead of here at the end
 		BroadcastResponseHeadersReceived();
-		HttpManager.BroadcastHttpRequestCompleted(Request);
 		OnProcessRequestComplete().ExecuteIfBound(SharedThis(this), Response, true);
 	}
 	else
@@ -453,7 +451,6 @@ void FAppleHttpRequest::FinishedRequest()
 		}
 
 		Response = nullptr;
-		HttpManager.BroadcastHttpRequestCompleted(Request);
 		OnProcessRequestComplete().ExecuteIfBound(SharedThis(this), nullptr, false);
 	}
 
@@ -461,9 +458,9 @@ void FAppleHttpRequest::FinishedRequest()
 	CleanupRequest();
 
 	// Remove from global list since processing is now complete
-	if (HttpManager.IsValidRequest(this))
+	if (FHttpModule::Get().GetHttpManager().IsValidRequest(this))
 	{
-		HttpManager.RemoveRequest(SharedThis(this));
+		FHttpModule::Get().GetHttpManager().RemoveRequest(SharedThis(this));
 	}
 }
 
