Add deprecation suppression around FHttpManager constructor and destructor to address deprecation warnings

diff --git a/Engine/Source/Runtime/Online/HTTP/Private/HttpManager.cpp b/Engine/Source/Runtime/Online/HTTP/Private/HttpManager.cpp
--- a/Engine/Source/Runtime/Online/HTTP/Private/HttpManager.cpp
+++ b/Engine/Source/Runtime/Online/HTTP/Private/HttpManager.cpp
@@ -35,6 +35,7 @@ const TCHAR* LexToString(const EHttpFlushReason& FlushReason)
 	return TEXT("Invalid");
 }
 
+PRAGMA_DISABLE_DEPRECATION_WARNINGS
 FHttpManager::FHttpManager()
 	: FTSTickerObjectBase(0.0f, FTSBackgroundableTicker::GetCoreTicker())
 	, Thread(nullptr)
@@ -51,6 +52,7 @@ FHttpManager::~FHttpManager()
 		delete Thread;
 	}
 }
+PRAGMA_ENABLE_DEPRECATION_WARNINGS
 
 void FHttpManager::Initialize()
 {
