- Fix LastReportedBytesSent/Read not being reset on a new request

diff --git a/Engine/Source/Runtime/Online/HTTP/Private/Curl/CurlHttp.cpp b/Engine/Source/Runtime/Online/HTTP/Private/Curl/CurlHttp.cpp
--- a/Engine/Source/Runtime/Online/HTTP/Private/Curl/CurlHttp.cpp
+++ b/Engine/Source/Runtime/Online/HTTP/Private/Curl/CurlHttp.cpp
@@ -786,6 +786,7 @@ bool FCurlHttpRequest::SetupRequest()
 	bCurlRequestCompleted = false;
 	bCanceled = false;
 	CurlAddToMultiResult = CURLM_OK;
+	LastReportedBytesSent = 0;
 
 	// default no verb to a GET
 	if (Verb.IsEmpty())
@@ -992,6 +993,7 @@ bool FCurlHttpRequest::ProcessRequest()
 
 	// Clear out response. If this is a re-used request, Response could point to a stale response until SetupRequestHttpThread is called
 	Response = nullptr;
+	LastReportedBytesRead = 0;
 
 	bool bStarted = false;
 	if (!FHttpModule::Get().GetHttpManager().IsDomainAllowed(URL))
