BuldPatchServices: Integrating fixes for DownloadService delegates crash if http requests complete during installer shutdown.
Fix for client Crash in BuildPatchServices::FCloudChunkSource::FDownloadDelegates::OnDownloadComplete
- Add cancelling (abandoning) of in flight downloads if should abort

p4 integrate //Portal/Dev-Main-Staging/...[REDACTED_HANDLE],19486880
p4 integrate //Portal/Dev-Main-Staging/...[REDACTED_HANDLE],20280996

diff --git a/Engine/Source/Runtime/Online/BuildPatchServices/Private/Installer/CloudChunkSource.cpp b/Engine/Source/Runtime/Online/BuildPatchServices/Private/Installer/CloudChunkSource.cpp
--- a/Engine/Source/Runtime/Online/BuildPatchServices/Private/Installer/CloudChunkSource.cpp
+++ b/Engine/Source/Runtime/Online/BuildPatchServices/Private/Installer/CloudChunkSource.cpp
@@ -705,6 +705,15 @@ namespace BuildPatchServices
 			Platform->Sleep(0.01f);
 		}
 
+		// Abandon in flight downloads if should abort.
+		if (bShouldAbort)
+		{
+			for (const TPair<int32, FGuid>& InFlightDownload : InFlightDownloads)
+			{
+				DownloadService->RequestAbandon(InFlightDownload.Key);
+			}
+		}
+
 		// Provide final stat values.
 		CloudChunkSourceStat->OnDownloadHealthUpdated(TrackedDownloadHealth);
 		CloudChunkSourceStat->OnSuccessRateUpdated(ChunkSuccessRate.GetOverall());
diff --git a/Engine/Source/Runtime/Online/BuildPatchServices/Private/Installer/DownloadService.cpp b/Engine/Source/Runtime/Online/BuildPatchServices/Private/Installer/DownloadService.cpp
--- a/Engine/Source/Runtime/Online/BuildPatchServices/Private/Installer/DownloadService.cpp
+++ b/Engine/Source/Runtime/Online/BuildPatchServices/Private/Installer/DownloadService.cpp
@@ -190,6 +190,11 @@ namespace BuildPatchServices
 			ExecuteCancelled(RequestId, DownloadDelegates);
 			DownloadDelegates = FDownloadDelegates();
 		}
+
+		void Reset(FDownloadDelegates& DownloadDelegates)
+		{
+			DownloadDelegates = FDownloadDelegates();
+		}
 	}
 
 	class FDownloadService
@@ -219,15 +224,20 @@ namespace BuildPatchServices
 		// IDownloadService interface begin.
 		virtual int32 RequestFile(const FString& FileUri, const FDownloadCompleteDelegate& OnCompleteDelegate, const FDownloadProgressDelegate& OnProgressDelegate) override;
 		virtual void RequestCancel(int32 RequestId) override;
+		virtual void RequestAbandon(int32 RequestId) override;
 		// IDownloadService interface end.
 
 	private:
 		bool Tick(float DeltaTime);
 		void ProcessCancelRequests();
+		void ProcessAbandonRequests();
 		void ProcessNewRequests();
 		void ProcessProgressUpdates();
 		void ProcessCompletedRequests();
 
+		TSet<int32> CancelNewRequests(const TArray<int32>& RequestsIDs);
+		TSet<int32> CancelOngoingRequests(const TArray<int32>& RequestsIDs);
+
 		int32 MakeRequestId();
 		IDownloadServiceStat::FDownloadRecord MakeDownloadRecord(int32 RequestId, FString Uri);
 		TFunction<void()> MakeFileLoadTask(const int32& RequestId, const FString& FileUri, FFileRequest* FileRequest);
@@ -260,6 +270,9 @@ namespace BuildPatchServices
 		FCriticalSection CancelRequestsCS;
 		TArray<int32> CancelRequests;
 
+		FCriticalSection AbandonRequestsCS;
+		TArray<int32> AbandonRequests;
+
 		FCriticalSection ActiveRequestsCS;
 		TMap<int32, TSharedRef<IHttpRequest, ESPMode::ThreadSafe>> ActiveHttpRequests;
 		TMap<int32, TUniquePtr<FFileRequest>> ActiveFileRequests;
@@ -303,6 +316,8 @@ namespace BuildPatchServices
 		, NewRequests()
 		, CancelRequestsCS()
 		, CancelRequests()
+		, AbandonRequestsCS()
+		, AbandonRequests()
 		, ActiveRequestsCS()
 		, ActiveHttpRequests()
 		, ActiveFileRequests()
@@ -373,12 +388,21 @@ namespace BuildPatchServices
 		CancelRequests.Add(RequestId);
 	}
 
+	void FDownloadService::RequestAbandon(int32 RequestId)
+	{
+		// Add the request.
+		FScopeLock ScopeLock(&AbandonRequestsCS);
+		AbandonRequests.Add(RequestId);
+	}
+
 	bool FDownloadService::Tick(float DeltaTime)
 	{
 		QUICK_SCOPE_CYCLE_COUNTER(STAT_FDownloadService_Tick);
 
 		ProcessCancelRequests();
 
+		ProcessAbandonRequests();
+
 		ProcessNewRequests();
 
 		ProcessProgressUpdates();
@@ -397,42 +421,48 @@ namespace BuildPatchServices
 		CancelRequestsCS.Unlock();
 
 		// Cancel new requests that were not processed yet.
-		TSet<int32> UnstartedRequests;
-		NewRequestsCS.Lock();
-		for (const int32& RequestId : FrameCancelRequests)
-		{
-			if (NewRequests.Remove(RequestId) > 0)
-			{
-				UnstartedRequests.Add(RequestId);
-			}
-		}
-		NewRequestsCS.Unlock();
+		TSet<int32> UnstartedRequests = CancelNewRequests(FrameCancelRequests);
 
 		// Cancel ongoing requests.
-		TSet<int32> CancelledRequests;
-		ActiveRequestsCS.Lock();
-		for (const int32& RequestId : FrameCancelRequests)
+		TSet<int32> CancelledRequests = CancelOngoingRequests(FrameCancelRequests);
+
+		// Call and remove delegates for new requests that were canceled.
+		RequestDelegatesCS.Lock();
+		for (const int32& RequestId : UnstartedRequests.Union(CancelledRequests))
 		{
-			if (ActiveHttpRequests.Contains(RequestId))
-			{
-				ActiveHttpRequests[RequestId]->CancelRequest();
-				CancelledRequests.Add(RequestId);
-			}
-			if (ActiveFileRequests.Contains(RequestId))
+			if (ensureMsgf(RequestDelegates.Contains(RequestId), TEXT("Missing request delegates for %d"), RequestId))
 			{
-				ActiveFileRequests[RequestId]->ShouldCancel = true;
-				CancelledRequests.Add(RequestId);
+				DownloadServiceHelpers::ExecuteCancelledAndReset(RequestId, RequestDelegates[RequestId]);
+				if (UnstartedRequests.Contains(RequestId))
+				{
+					RequestDelegates.Remove(RequestId);
+				}
 			}
 		}
-		ActiveRequestsCS.Unlock();
+		RequestDelegatesCS.Unlock();
+	}
 
-		// Call and remove delegates for new requests that were canceled.
+	void FDownloadService::ProcessAbandonRequests()
+	{
+		// Grab new abandon requests for this frame.
+		TArray<int32> FrameAbandonRequests;
+		AbandonRequestsCS.Lock();
+		FrameAbandonRequests = MoveTemp(AbandonRequests);
+		AbandonRequestsCS.Unlock();
+
+		// Cancel new requests that were not processed yet.
+		TSet<int32> UnstartedRequests = CancelNewRequests(FrameAbandonRequests);
+
+		// Cancel ongoing requests.
+		TSet<int32> CancelledRequests = CancelOngoingRequests(FrameAbandonRequests);
+
+		// Remove delegates for new requests that were canceled.
 		RequestDelegatesCS.Lock();
 		for (const int32& RequestId : UnstartedRequests.Union(CancelledRequests))
 		{
 			if (ensureMsgf(RequestDelegates.Contains(RequestId), TEXT("Missing request delegates for %d"), RequestId))
 			{
-				DownloadServiceHelpers::ExecuteCancelledAndReset(RequestId, RequestDelegates[RequestId]);
+				DownloadServiceHelpers::Reset(RequestDelegates[RequestId]);
 				if (UnstartedRequests.Contains(RequestId))
 				{
 					RequestDelegates.Remove(RequestId);
@@ -527,6 +557,44 @@ namespace BuildPatchServices
 		}
 	}
 
+	TSet<int32> FDownloadService::CancelNewRequests(const TArray<int32>& RequestsIDs)
+	{
+		TSet<int32> UnstartedRequests;
+		NewRequestsCS.Lock();
+		for (const int32& RequestId : RequestsIDs)
+		{
+			if (NewRequests.Remove(RequestId) > 0)
+			{
+				UnstartedRequests.Add(RequestId);
+			}
+		}
+		NewRequestsCS.Unlock();
+
+		return UnstartedRequests;
+	}
+
+	TSet<int32> FDownloadService::CancelOngoingRequests(const TArray<int32>& RequestsIDs)
+	{
+		TSet<int32> CancelledRequests;
+		ActiveRequestsCS.Lock();
+		for (const int32& RequestId : RequestsIDs)
+		{
+			if (ActiveHttpRequests.Contains(RequestId))
+			{
+				ActiveHttpRequests[RequestId]->CancelRequest();
+				CancelledRequests.Add(RequestId);
+			}
+			if (ActiveFileRequests.Contains(RequestId))
+			{
+				ActiveFileRequests[RequestId]->ShouldCancel = true;
+				CancelledRequests.Add(RequestId);
+			}
+		}
+		ActiveRequestsCS.Unlock();
+
+		return CancelledRequests;
+	}
+
 	int32 FDownloadService::MakeRequestId()
 	{
 		return RequestIdCounter.Increment();
diff --git a/Engine/Source/Runtime/Online/BuildPatchServices/Private/Installer/DownloadService.h b/Engine/Source/Runtime/Online/BuildPatchServices/Private/Installer/DownloadService.h
--- a/Engine/Source/Runtime/Online/BuildPatchServices/Private/Installer/DownloadService.h
+++ b/Engine/Source/Runtime/Online/BuildPatchServices/Private/Installer/DownloadService.h
@@ -82,6 +82,12 @@ namespace BuildPatchServices
 		 * @param RequestId             The id that was returned by the call to RequestFile which should be canceled.
 		 */
 		virtual void RequestCancel(int32 RequestId) = 0;
+
+		/**
+		 *  Requests the abandoning of a requested file (cancelling without calling complete/progress handlers)
+		 * @param RequestId             The id that was returned by the call to RequestFile which should be canceled.
+		 */
+		virtual void RequestAbandon(int32 RequestId) = 0;
 	};
 
 	/**
diff --git a/Engine/Source/Runtime/Online/BuildPatchServices/Private/Tests/Mock/DownloadService.mock.h b/Engine/Source/Runtime/Online/BuildPatchServices/Private/Tests/Mock/DownloadService.mock.h
--- a/Engine/Source/Runtime/Online/BuildPatchServices/Private/Tests/Mock/DownloadService.mock.h
+++ b/Engine/Source/Runtime/Online/BuildPatchServices/Private/Tests/Mock/DownloadService.mock.h
@@ -47,6 +47,16 @@ namespace BuildPatchServices
 			RxRequestCancel.Emplace(FStatsCollector::GetSeconds(), RequestId);
 		}
 
+		virtual void RequestAbandon(int32 RequestId) override
+		{
+			FScopeLock ScopeLock(&ThreadLock);
+			if (RequestCancelFunc)
+			{
+				RequestCancelFunc(RequestId);
+			}
+			RxRequestCancel.Emplace(FStatsCollector::GetSeconds(), RequestId);
+		}
+
 	public:
 		FCriticalSection ThreadLock;
 		int32 Count;
