FWebSocketModule::Get now will load the WebSockets module if it is not loaded, to match the documentation.

#9376

diff --git a/Engine/Source/Runtime/Online/WebSockets/Private/WebSocketsModule.cpp b/Engine/Source/Runtime/Online/WebSockets/Private/WebSocketsModule.cpp
--- a/Engine/Source/Runtime/Online/WebSockets/Private/WebSocketsModule.cpp
+++ b/Engine/Source/Runtime/Online/WebSockets/Private/WebSocketsModule.cpp
@@ -55,6 +55,12 @@ void FWebSocketsModule::ShutdownModule()
 
 FWebSocketsModule& FWebSocketsModule::Get()
 {
+	if (nullptr == Singleton)
+	{
+		check(IsInGameThread());
+		FModuleManager::LoadModuleChecked<FWebSocketsModule>("WebSockets");
+	}
+	check(Singleton);
 	return *Singleton;
 }
 
