Update comment to be in line with inclusive terminology standards

diff --git a/Engine/Source/Runtime/Online/HTTP/Private/Curl/CurlHttp.cpp b/Engine/Source/Runtime/Online/HTTP/Private/Curl/CurlHttp.cpp
--- a/Engine/Source/Runtime/Online/HTTP/Private/Curl/CurlHttp.cpp
+++ b/Engine/Source/Runtime/Online/HTTP/Private/Curl/CurlHttp.cpp
@@ -1177,7 +1177,7 @@ void FCurlHttpRequest::BroadcastNewlyReceivedHeaders()
 	check(IsInGameThread());
 	if (Response.IsValid())
 	{
-		// Process the headers received on the HTTP thread and merge them into our master list and then broadcast the new headers
+		// Process the headers received on the HTTP thread and merge them into the response's list of headers and then broadcast the new headers
 		TPair<FString, FString> NewHeader;
 		while (Response->NewlyReceivedHeaders.Dequeue(NewHeader))
 		{
