WebSocket - Packet fragmentation support and upper memory limit for receiving utf-8 messages. Differenciate between binary and text frame for events.

diff --git a/Engine/Config/BaseEngine.ini b/Engine/Config/BaseEngine.ini
--- a/Engine/Config/BaseEngine.ini
+++ b/Engine/Config/BaseEngine.ini
@@ -72,6 +72,7 @@ RemoveURLMappingEntriesWithoutPhysicalTempFiles=true
 +WebSocketsProtocols="v11.stomp"
 +WebSocketsProtocols="v12.stomp"
 +WebSocketsProtocols="xmpp"
+TextMessageMemoryLimit=1048576
 
 [WebSockets.LibWebSockets]
 ThreadStackSize=131072
diff --git a/Engine/Plugins/Media/PixelStreaming/Source/PixelStreaming/Private/Tests/SignallingServerConnectionTests.cpp b/Engine/Plugins/Media/PixelStreaming/Source/PixelStreaming/Private/Tests/SignallingServerConnectionTests.cpp
--- a/Engine/Plugins/Media/PixelStreaming/Source/PixelStreaming/Private/Tests/SignallingServerConnectionTests.cpp
+++ b/Engine/Plugins/Media/PixelStreaming/Source/PixelStreaming/Private/Tests/SignallingServerConnectionTests.cpp
@@ -88,10 +88,12 @@ namespace
 		virtual bool IsConnected() override { return bConnected; }
 		virtual void Send(const FString& Data) override { OnMockResponseEvent.Broadcast(Data); }
 		virtual void Send(const void* Data, SIZE_T Size, bool bIsBinary = false) override {}
+		virtual void SetTextMessageMemoryLimit(uint64 TextMessageMemoryLimit) override {}
 		virtual FWebSocketConnectedEvent& OnConnected() override { return OnConnectedEvent; }
 		virtual FWebSocketConnectionErrorEvent& OnConnectionError() override { return OnErrorEvent; }
 		virtual FWebSocketClosedEvent& OnClosed() override { return OnClosedEvent; }
 		virtual FWebSocketMessageEvent& OnMessage() override { return OnMessageEvent; }
+		virtual FWebSocketBinaryMessageEvent& OnBinaryMessage() override { return OnBinaryMessageEvent; }
 		virtual FWebSocketRawMessageEvent& OnRawMessage() override { return OnRawMessageEvent; }
 		virtual FWebSocketMessageSentEvent& OnMessageSent() override { return OnMessageSentEvent; }
 
@@ -106,6 +108,7 @@ namespace
 		FWebSocketConnectionErrorEvent OnErrorEvent;
 		FWebSocketClosedEvent OnClosedEvent;
 		FWebSocketMessageEvent OnMessageEvent;
+		FWebSocketBinaryMessageEvent OnBinaryMessageEvent;
 		FWebSocketRawMessageEvent OnRawMessageEvent;
 		FWebSocketMessageSentEvent OnMessageSentEvent;
 
diff --git a/Engine/Source/Runtime/Online/WebSockets/Private/Lws/LwsWebSocket.cpp b/Engine/Source/Runtime/Online/WebSockets/Private/Lws/LwsWebSocket.cpp
--- a/Engine/Source/Runtime/Online/WebSockets/Private/Lws/LwsWebSocket.cpp
+++ b/Engine/Source/Runtime/Online/WebSockets/Private/Lws/LwsWebSocket.cpp
@@ -45,6 +45,15 @@ FLwsReceiveBufferBinary::FLwsReceiveBufferBinary(const uint8* Data, const int32
 	Payload.Append(Data, Size);
 }
 
+// FLwsReceiveBufferBinaryFragment
+FLwsReceiveBufferBinaryFragment::FLwsReceiveBufferBinaryFragment(const uint8* Data, const int32 Size, const bool bInIsLastFragment)
+	: bIsLastFragment(bInIsLastFragment)
+{
+	check(Data);
+	check(Size >= 0);
+	Payload.Append(Data, Size);
+}
+
 // FLwsReceiveBufferText
 FLwsReceiveBufferText::FLwsReceiveBufferText(FString&& InText)
 	: Text(MoveTemp(InText))
@@ -72,7 +81,7 @@ const TCHAR* FLwsWebSocket::ToString(const EState InState)
 	return TEXT("Unknown");
 }
 
-FLwsWebSocket::FLwsWebSocket(FPrivateToken, const FString& InUrl, const TArray<FString>& InProtocols, const FString& InUpgradeHeader)
+FLwsWebSocket::FLwsWebSocket(FPrivateToken, const FString& InUrl, const TArray<FString>& InProtocols, const FString& InUpgradeHeader, uint64 TextMessageMemoryLimit)
 	: State(EState::None)
 	, LastGameThreadState(EState::None)
 	, bWasSendQueueEmpty(true)
@@ -80,6 +89,7 @@ FLwsWebSocket::FLwsWebSocket(FPrivateToken, const FString& InUrl, const TArray<F
 	, Url(InUrl)
 	, Protocols(InProtocols)
 	, UpgradeHeader(InUpgradeHeader)
+	, MaxTextMessageBufferSize(TextMessageMemoryLimit)
 	, Identifier(++IncrementingIdentifier)
 {
 	UE_LOG(LogWebSockets, VeryVerbose, TEXT("FLwsWebSocket[%d]: Constructed url=%s protocols=%s"), Identifier, *InUrl, *FString::Join(Protocols, TEXT(",")));
@@ -122,6 +132,7 @@ void FLwsWebSocket::Connect()
 
 	bWantsMessageEvents = OnMessage().IsBound();
 	bWantsRawMessageEvents = OnRawMessage().IsBound();
+	bWantsBinaryMessageEvents = OnBinaryMessage().IsBound();
 	
 	UE_LOG(LogWebSockets, Verbose, TEXT("FLwsWebSocket[%d]::Connect: setting State=%s url=%s bWantsMessageEvents=%d bWantsRawMessageEvents=%d"), Identifier, ToString(State), *Url, (int32)bWantsMessageEvents, (int32)bWantsRawMessageEvents);
 
@@ -167,6 +178,12 @@ void FLwsWebSocket::Send(const FString& Data)
 	OnMessageSent().Broadcast(Data);
 }
 
+
+void FLwsWebSocket::SetTextMessageMemoryLimit(uint64 TextMessageMemoryLimit)
+{
+	MaxTextMessageBufferSize = TextMessageMemoryLimit;
+}
+
 void FLwsWebSocket::SendFromQueue()
 {
 	check(LwsConnection);
@@ -210,6 +227,7 @@ void FLwsWebSocket::ClearData()
 {
 	check(State != EState::Connected);
 	ReceiveBinaryQueue.Empty();
+	ReceiveBinaryFragmentQueue.Empty();
 	ReceiveTextQueue.Empty();
 	FLwsSendBuffer* SendBuffer;
 	while (SendQueue.Dequeue(SendBuffer))
@@ -281,24 +299,44 @@ int FLwsWebSocket::LwsCallback(lws* Instance, lws_callback_reasons Reason, void*
 	case LWS_CALLBACK_CLIENT_RECEIVE:
 	{
 		const SIZE_T BytesLeft = lws_remaining_packet_payload(Instance);
-		UE_LOG(LogWebSockets, VeryVerbose, TEXT("FLwsWebSocket[%d]::LwsCallback: Received LWS_CALLBACK_CLIENT_RECEIVE Length=%d BytesLeft=%d"), Identifier, Length, BytesLeft);
+		bool bIsFinalFragment = (bool)lws_is_final_fragment(Instance);
+		bool bIsBinary = (bool)lws_frame_is_binary(Instance);
+		UE_LOG(LogWebSockets, VeryVerbose, TEXT("FLwsWebSocket[%d]::LwsCallback: Received LWS_CALLBACK_CLIENT_RECEIVE Length=%d BytesLeftInFragment=%d IsFinalFragment=%s"), Identifier, Length, BytesLeft, bIsFinalFragment ? "true" : "false");
 		bool bWakeGameThread = false;
-		if (bWantsMessageEvents)
+		if (bWantsMessageEvents && !bIsBinary)
 		{
 			FUTF8ToTCHAR Convert((const ANSICHAR*)Data, Length);
-			ReceiveBuffer.Append(Convert.Get(), Convert.Length());
-			if (BytesLeft == 0)
+
+			if (ReceiveBuffer.Len() + Length > MaxTextMessageBufferSize)
 			{
-				bWakeGameThread = true;
-				ReceiveTextQueue.Enqueue(MakeUnique<FLwsReceiveBufferText>(MoveTemp(ReceiveBuffer)));
-				ReceiveBuffer.Empty();
+				FString ReasonString = FString::Printf(TEXT("Received text message exceeded memory limit of %lu bytes"), MaxTextMessageBufferSize);
+				Close(LWS_CLOSE_STATUS_MESSAGE_TOO_LARGE, ReasonString);
+
+				UE_LOG(LogWebSockets, Verbose, TEXT("Received text message too large - use SetTextMessageMemoryLimit() to increase buffer size. Current Size=%lu"), MaxTextMessageBufferSize);
 			}
+			else
+			{
+				ReceiveBuffer.Append(Convert.Get(), Convert.Length());
+				if (BytesLeft == 0 && bIsFinalFragment)
+				{
+					bWakeGameThread = true;
+					ReceiveTextQueue.Enqueue(MakeUnique<FLwsReceiveBufferText>(MoveTemp(ReceiveBuffer)));
+					ReceiveBuffer.Empty();
+				}
+			}
+		}
+		if (bWantsBinaryMessageEvents && bIsBinary)
+		{
+			bWakeGameThread = true;
+			bool bIsLast = (bool)bIsFinalFragment && BytesLeft == 0;
+			ReceiveBinaryFragmentQueue.Enqueue(MakeUnique<FLwsReceiveBufferBinaryFragment>(static_cast<const uint8*>(Data), Length, bIsLast));
 		}
 		if (bWantsRawMessageEvents)
 		{
 			bWakeGameThread = true;
 			ReceiveBinaryQueue.Enqueue(MakeUnique<FLwsReceiveBufferBinary>(static_cast<const uint8*>(Data), Length, BytesLeft));
 		}
+
 		if (bWakeGameThread)
 		{
 			FEmbeddedCommunication::WakeGameThread();
@@ -375,12 +413,26 @@ int FLwsWebSocket::LwsCallback(lws* Instance, lws_callback_reasons Reason, void*
 				FEmbeddedCommunication::WakeGameThread();
 			}
 			else if (State == EState::ClosingByRequest)
-			{
+			{		
+				FString ClosedReasonString;
+				int32 ClosedStatus;
+
+				if (CloseRequest.Code == LWS_CLOSE_STATUS_MESSAGE_TOO_LARGE)
+				{
+					ClosedReasonString = TEXT("Received text message exceeded memory limit");
+					ClosedStatus = LWS_CLOSE_STATUS_MESSAGE_TOO_LARGE;
+				}
+				else
+				{
+					ClosedReasonString = TEXT("Successfully closed connection to our peer");
+					ClosedStatus = LWS_CLOSE_STATUS_NORMAL;
+				}
+
 				{
 					FScopeLock ScopeLock(&StateLock);
 					State = EState::Closed;
-					ClosedReason.Reason = TEXT("Successfully closed connection to our peer");
-					ClosedReason.CloseStatus = LWS_CLOSE_STATUS_NORMAL;
+					ClosedReason.Reason = MoveTemp(ClosedReasonString);
+					ClosedReason.CloseStatus = ClosedStatus;
 					ClosedReason.bWasClean = true;
 				}
 				FEmbeddedCommunication::WakeGameThread();
@@ -556,6 +608,12 @@ void FLwsWebSocket::GameThreadTick()
 		{
 			OnRawMessage().Broadcast(BufferBinary->Payload.GetData(), BufferBinary->Payload.Num(), BufferBinary->BytesRemaining);
 		}
+
+		FLwsReceiveBufferBinaryFragmentPtr BufferBinaryFragment;
+		while (ReceiveBinaryFragmentQueue.Dequeue(BufferBinaryFragment))
+		{
+			OnBinaryMessage().Broadcast(BufferBinaryFragment->Payload.GetData(), BufferBinaryFragment->Payload.Num(), BufferBinaryFragment->bIsLastFragment);
+		}
 	}
 }
 
diff --git a/Engine/Source/Runtime/Online/WebSockets/Private/Lws/LwsWebSocket.h b/Engine/Source/Runtime/Online/WebSockets/Private/Lws/LwsWebSocket.h
--- a/Engine/Source/Runtime/Online/WebSockets/Private/Lws/LwsWebSocket.h
+++ b/Engine/Source/Runtime/Online/WebSockets/Private/Lws/LwsWebSocket.h
@@ -63,7 +63,7 @@ struct FLwsSendBuffer
 	bool bHasError;
 };
 
-/** Buffer for one incoming binary packet fragment */
+/** Buffer for one incoming binary packet */
 struct FLwsReceiveBufferBinary
 {
 	/**
@@ -82,6 +82,25 @@ struct FLwsReceiveBufferBinary
 
 typedef TUniquePtr<FLwsReceiveBufferBinary> FLwsReceiveBufferBinaryPtr;
 
+/** Buffer for one incoming binary packet fragment */
+struct FLwsReceiveBufferBinaryFragment
+{
+	/**
+	 * Constructor
+	 * @param Data pointer to data to fill our buffer with
+	 * @param Size size of Data
+	 * @param bInIsLastFragment Whether or not this is the last fragment
+	 */
+	FLwsReceiveBufferBinaryFragment(const uint8* Data, const int32 Size, const bool bInIsLastFragment);
+
+	/** Payload received */
+	TArray<uint8> Payload;
+	/** Whether or not this is the last fragment */
+	const bool bIsLastFragment;
+};
+
+typedef TUniquePtr<FLwsReceiveBufferBinaryFragment> FLwsReceiveBufferBinaryFragmentPtr;
+
 /** Buffer for one incoming text packet, fully received */
 struct FLwsReceiveBufferText
 {
@@ -119,6 +138,7 @@ public:
 
 	virtual void Send(const FString& Data);
 	virtual void Send(const void* Data, SIZE_T Size, bool bIsBinary) override;
+	virtual void SetTextMessageMemoryLimit(uint64 TextMessageMemoryLimit) override;
 
 	/** Delegate called when a web socket connection has been established */
 	DECLARE_DERIVED_EVENT(FLwsWebSocket, IWebSocket::FWebSocketConnectedEvent, FWebSocketConnectedEvent);
@@ -141,20 +161,28 @@ public:
 		return ClosedEvent;
 	}
 
-	/** Delegate called when a web socket message has been received */
+	/** Delegate called when a web socket text message has been received */
 	DECLARE_DERIVED_EVENT(FLwsWebSocket, IWebSocket::FWebSocketMessageEvent, FWebSocketMessageEvent);
 	virtual FWebSocketMessageEvent& OnMessage() override
 	{
 		return MessageEvent;
 	}
 
-	/** Delegate called when a raw web socket message has been received */
+	/** Delegate called when a web socket binary message has been received */
+	DECLARE_DERIVED_EVENT(FLwsWebSocket, IWebSocket::FWebSocketBinaryMessageEvent, FWebSocketBinaryMessageEvent);
+	virtual FWebSocketBinaryMessageEvent& OnBinaryMessage() override
+	{
+		return BinaryMessageEvent;
+	}
+
+	/** Delegate called when a any web socket message has been received */
 	DECLARE_DERIVED_EVENT(FLwsWebSocket, IWebSocket::FWebSocketRawMessageEvent, FWebSocketRawMessageEvent);
 	virtual FWebSocketRawMessageEvent& OnRawMessage() override
 	{
 		return RawMessageEvent;
 	}
 
+
 	DECLARE_DERIVED_EVENT(FLwsWebSocket, IWebSocket::FWebSocketMessageSentEvent, FWebSocketMessageSentEvent);
 	virtual FWebSocketMessageSentEvent& OnMessageSent() override
 	{
@@ -194,7 +222,7 @@ private:
 
 public:
 	/** Constructor */
-	FLwsWebSocket(FPrivateToken, const FString& Url, const TArray<FString>& Protocols, const FString& UpgradeHeader);
+	FLwsWebSocket(FPrivateToken, const FString& Url, const TArray<FString>& Protocols, const FString& UpgradeHeader, uint64 TextMessageMemoryLimit);
 
 private:
 	/**
@@ -283,6 +311,7 @@ private:
 	FWebSocketConnectionErrorEvent ConnectionErrorEvent;
 	FWebSocketClosedEvent ClosedEvent;
 	FWebSocketMessageEvent MessageEvent;
+	FWebSocketBinaryMessageEvent BinaryMessageEvent;
 	FWebSocketRawMessageEvent RawMessageEvent;
 	FWebSocketMessageSentEvent OnMessageSentEvent;
 
@@ -305,16 +334,26 @@ private:
 	 * For performance reasons if nothing was bound at Connect() time, we will never trigger OnRawMessage
 	 */
 	bool bWantsRawMessageEvents;
+	/**
+	 * Whether or not OnBinaryMessage was bound to when Connect() was called.
+	 * For performance reasons if nothing was bound at Connect() time, we will never trigger OnBinaryMessage
+	 */
+	bool bWantsBinaryMessageEvents;
+
 
 	/** Buffer of an incomplete packet received */
 	FString ReceiveBuffer;
-	/** Received binary fragments, waiting for delegates to be triggered on the game thread */
+	/** Received binary packets, waiting for delegates to be triggered on the game thread */
 	TQueue<FLwsReceiveBufferBinaryPtr, EQueueMode::Spsc> ReceiveBinaryQueue;
+	/** Received fragments of binary packets, waiting for delegates to be triggered on the game thread */
+	TQueue<FLwsReceiveBufferBinaryFragmentPtr, EQueueMode::Spsc> ReceiveBinaryFragmentQueue;
 	/** Received text packets, waiting for delegates to be triggered on the game thread */
 	TQueue<FLwsReceiveBufferTextPtr, EQueueMode::Spsc> ReceiveTextQueue;
 	/** Pending outgoing packets, populated by the game thread and processed on the libwebsockets thread */
 	TQueue<FLwsSendBuffer*, EQueueMode::Spsc> SendQueue;
 
+	uint64 MaxTextMessageBufferSize;
+
 	// Unique identifier for logging
 	/** Incrementing identifier to give each web socket a unique identifier */
 	static int32 IncrementingIdentifier;
diff --git a/Engine/Source/Runtime/Online/WebSockets/Private/Lws/LwsWebSocketsManager.cpp b/Engine/Source/Runtime/Online/WebSockets/Private/Lws/LwsWebSocketsManager.cpp
--- a/Engine/Source/Runtime/Online/WebSockets/Private/Lws/LwsWebSocketsManager.cpp
+++ b/Engine/Source/Runtime/Online/WebSockets/Private/Lws/LwsWebSocketsManager.cpp
@@ -403,7 +403,11 @@ TSharedRef<IWebSocket> FLwsWebSocketsManager::CreateWebSocket(const FString& Url
 		UpgradeHeaderString += FString::Printf(TEXT("%s: %s\r\n"), *OneHeader.Key, *OneHeader.Value);
 	}
 
-	FLwsWebSocketRef Socket = MakeShared<FLwsWebSocket>(FLwsWebSocket::FPrivateToken{}, Url, Protocols, UpgradeHeaderString);
+	// default memory limit for IWebSocket text messages
+	int TextMessageMemoryLimit = 1024 * 1024;
+	GConfig->GetInt(TEXT("WebSockets"), TEXT("TextMessageMemoryLimit"), TextMessageMemoryLimit, GEngineIni);
+
+	FLwsWebSocketRef Socket = MakeShared<FLwsWebSocket>(FLwsWebSocket::FPrivateToken{}, Url, Protocols, UpgradeHeaderString, TextMessageMemoryLimit);
 	return Socket;
 }
 
diff --git a/Engine/Source/Runtime/Online/WebSockets/Public/IWebSocket.h b/Engine/Source/Runtime/Online/WebSockets/Public/IWebSocket.h
--- a/Engine/Source/Runtime/Online/WebSockets/Public/IWebSocket.h
+++ b/Engine/Source/Runtime/Online/WebSockets/Public/IWebSocket.h
@@ -42,6 +42,13 @@ public:
 	 */
 	virtual void Send(const void* Data, SIZE_T Size, bool bIsBinary = false) = 0;
 
+	/**
+	 * Optionally change memory limit for receiving UTF-8 text on this socket. 
+	 * Default from config TextMessageMemoryLimit under [WebSockets] or 1MB.
+	 * @param TextMessageMemoryLimit new buffer size upper limit in bytes.
+	 */
+	virtual void SetTextMessageMemoryLimit(uint64 TextMessageMemoryLimit) = 0;
+
 	/**
 	 * Delegate called when a web socket connection has been established successfully.
 	 *
@@ -70,12 +77,19 @@ public:
 	 */
 	DECLARE_EVENT_OneParam(IWebSocket, FWebSocketMessageEvent, const FString& /* MessageString */);
 	virtual FWebSocketMessageEvent& OnMessage() = 0;
-
-	/**
-	 * Delegate called when a web socket data has been received.
+	
+      /**
+	 * Delegate called when web socket binary data has been received.
+	 * bIsLastFragment will be true if it is the last fragment of the current packet.
+	 */
+	DECLARE_EVENT_ThreeParams(IWebSocket, FWebSocketBinaryMessageEvent, const void* /* Data */, SIZE_T /* Size */, bool /* bIsLastFragment */);
+	virtual FWebSocketBinaryMessageEvent& OnBinaryMessage() = 0;
+	
+      /**
+	 * Delegate called when any web socket data has been received.
 	 * May be called multiple times for a message if the message was split into multiple frames. 
-	 * The last parameter will be 0 on the last frame in the packet.
-	 *
+	 * This is called for every web socket fragment.
+	 * BytesRemaining is the number of bytes left in the current fragment.
 	 */
 	DECLARE_EVENT_ThreeParams(IWebSocket, FWebSocketRawMessageEvent, const void* /* Data */, SIZE_T /* Size */, SIZE_T /* BytesRemaining */);
 	virtual FWebSocketRawMessageEvent& OnRawMessage() = 0;
