added missing function implementation for WinHttpWebsocket for submission 23727568

diff --git a/Engine/Source/Runtime/Online/WebSockets/Private/WinHttp/WinHttpWebSocket.cpp b/Engine/Source/Runtime/Online/WebSockets/Private/WinHttp/WinHttpWebSocket.cpp
--- a/Engine/Source/Runtime/Online/WebSockets/Private/WinHttp/WinHttpWebSocket.cpp
+++ b/Engine/Source/Runtime/Online/WebSockets/Private/WinHttp/WinHttpWebSocket.cpp
@@ -187,6 +187,11 @@ void FWinHttpWebSocket::Send(const void* Data, SIZE_T Size, bool bIsBinary)
 	WebSocket->SendMessage(MessageType, MoveTemp(Message));
 }
 
+void FWinHttpWebSocket::SetTextMessageMemoryLimit(uint64 TextMessageMemoryLimit)
+{
+	UE_LOG(LogWebSockets, Verbose, TEXT("SetTextMessageMemoryLimit not implemented for WinHttpWebSocket."));
+}
+
 FWinHttpWebSocket::FWebSocketConnectedEvent& FWinHttpWebSocket::OnConnected()
 {
 	return OnConnectedHandler;
@@ -207,6 +212,11 @@ FWinHttpWebSocket::FWebSocketMessageEvent& FWinHttpWebSocket::OnMessage()
 	return OnMessageHandler;
 }
 
+FWebSocketBinaryMessageEvent& FWinHttpWebSocket::OnBinaryMessage() 
+{ 
+	return BinaryMessageHandler; 
+}
+
 FWinHttpWebSocket::FWebSocketRawMessageEvent& FWinHttpWebSocket::OnRawMessage()
 {
 	return OnRawMessageHandler;
diff --git a/Engine/Source/Runtime/Online/WebSockets/Private/WinHttp/WinHttpWebSocket.h b/Engine/Source/Runtime/Online/WebSockets/Private/WinHttp/WinHttpWebSocket.h
--- a/Engine/Source/Runtime/Online/WebSockets/Private/WinHttp/WinHttpWebSocket.h
+++ b/Engine/Source/Runtime/Online/WebSockets/Private/WinHttp/WinHttpWebSocket.h
@@ -43,6 +43,7 @@ public:
 	virtual bool IsConnected() override final;
 	virtual void Send(const FString& Data) override final;
 	virtual void Send(const void* Data, SIZE_T Size, bool bIsBinary) override final;
+	virtual void SetTextMessageMemoryLimit(uint64 TextMessageMemoryLimit) override final;
 	DECLARE_DERIVED_EVENT(FWinHttpWebSocket, IWebSocket::FWebSocketConnectedEvent, FWebSocketConnectedEvent);
 	virtual FWebSocketConnectedEvent& OnConnected() override final;
 	DECLARE_DERIVED_EVENT(FWinHttpWebSocket, IWebSocket::FWebSocketConnectionErrorEvent, FWebSocketConnectionErrorEvent);
@@ -52,6 +53,8 @@ public:
 	DECLARE_DERIVED_EVENT(FWinHttpWebSocket, IWebSocket::FWebSocketMessageEvent, FWebSocketMessageEvent);
 	virtual FWebSocketMessageEvent& OnMessage() override final;
 	DECLARE_DERIVED_EVENT(FWinHttpWebSocket, IWebSocket::FWebSocketRawMessageEvent, FWebSocketRawMessageEvent);
+	DECLARE_DERIVED_EVENT(FLwsWebSocket, IWebSocket::FWebSocketBinaryMessageEvent, FWebSocketBinaryMessageEvent);
+	virtual FWebSocketBinaryMessageEvent& OnBinaryMessage() override final;
 	virtual FWebSocketRawMessageEvent& OnRawMessage() override final;
 	DECLARE_DERIVED_EVENT(FWinHttpWebSocket, IWebSocket::FWebSocketMessageSentEvent, FWebSocketMessageSentEvent);
 	virtual FWebSocketMessageSentEvent& OnMessageSent() override final;
@@ -86,6 +89,7 @@ protected:
 	FWebSocketConnectionErrorEvent OnErrorHandler;
 	FWebSocketClosedEvent OnClosedHandler;
 	FWebSocketMessageEvent OnMessageHandler;
+	FWebSocketBinaryMessageEvent BinaryMessageHandler;
 	FWebSocketRawMessageEvent OnRawMessageHandler;
 	FWebSocketMessageSentEvent OnMessageSentHandler;
 
