correcting return value fo submission fix 23728075

diff --git a/Engine/Source/Runtime/Online/WebSockets/Private/WinHttp/WinHttpWebSocket.cpp b/Engine/Source/Runtime/Online/WebSockets/Private/WinHttp/WinHttpWebSocket.cpp
--- a/Engine/Source/Runtime/Online/WebSockets/Private/WinHttp/WinHttpWebSocket.cpp
+++ b/Engine/Source/Runtime/Online/WebSockets/Private/WinHttp/WinHttpWebSocket.cpp
@@ -212,7 +212,7 @@ FWinHttpWebSocket::FWebSocketMessageEvent& FWinHttpWebSocket::OnMessage()
 	return OnMessageHandler;
 }
 
-FWebSocketBinaryMessageEvent& FWinHttpWebSocket::OnBinaryMessage() 
+FWinHttpWebSocket::FWebSocketBinaryMessageEvent& FWinHttpWebSocket::OnBinaryMessage()
 { 
 	return BinaryMessageHandler; 
 }
diff --git a/Engine/Source/Runtime/Online/WebSockets/Private/WinHttp/WinHttpWebSocket.h b/Engine/Source/Runtime/Online/WebSockets/Private/WinHttp/WinHttpWebSocket.h
--- a/Engine/Source/Runtime/Online/WebSockets/Private/WinHttp/WinHttpWebSocket.h
+++ b/Engine/Source/Runtime/Online/WebSockets/Private/WinHttp/WinHttpWebSocket.h
@@ -53,7 +53,7 @@ public:
 	DECLARE_DERIVED_EVENT(FWinHttpWebSocket, IWebSocket::FWebSocketMessageEvent, FWebSocketMessageEvent);
 	virtual FWebSocketMessageEvent& OnMessage() override final;
 	DECLARE_DERIVED_EVENT(FWinHttpWebSocket, IWebSocket::FWebSocketRawMessageEvent, FWebSocketRawMessageEvent);
-	DECLARE_DERIVED_EVENT(FLwsWebSocket, IWebSocket::FWebSocketBinaryMessageEvent, FWebSocketBinaryMessageEvent);
+	DECLARE_DERIVED_EVENT(FWinHttpWebSocket, IWebSocket::FWebSocketBinaryMessageEvent, FWebSocketBinaryMessageEvent);
 	virtual FWebSocketBinaryMessageEvent& OnBinaryMessage() override final;
 	virtual FWebSocketRawMessageEvent& OnRawMessage() override final;
 	DECLARE_DERIVED_EVENT(FWinHttpWebSocket, IWebSocket::FWebSocketMessageSentEvent, FWebSocketMessageSentEvent);
