Moved status deserializing outside the GT.
Added deserialized status caching

diff --git a/Engine/Source/Runtime/Online/XMPP/Public/XmppConnection.h b/Engine/Source/Runtime/Online/XMPP/Public/XmppConnection.h
--- a/Engine/Source/Runtime/Online/XMPP/Public/XmppConnection.h
+++ b/Engine/Source/Runtime/Online/XMPP/Public/XmppConnection.h
@@ -175,6 +175,11 @@ public:
 		return !operator==(Other);
 	}
 
+	friend uint32 GetTypeHash(const FXmppUserJid& A)
+	{
+		return HashCombine(HashCombine(GetTypeHash(A.Id), GetTypeHash(A.Domain)), GetTypeHash(A.Resource));
+	}
+
 	/** full jid path <id@domain/resource> */
 	FString GetFullPath() const
 	{
