Engine FExec cleanup: removing Exec overrides that aren't used in online systems as Exec overrides will soon need to be wrapped with ifdefs.

diff --git a/Engine/Plugins/Online/OnlineFramework/Source/Lobby/Private/LobbyModule.cpp b/Engine/Plugins/Online/OnlineFramework/Source/Lobby/Private/LobbyModule.cpp
--- a/Engine/Plugins/Online/OnlineFramework/Source/Lobby/Private/LobbyModule.cpp
+++ b/Engine/Plugins/Online/OnlineFramework/Source/Lobby/Private/LobbyModule.cpp
@@ -18,15 +18,3 @@ void FLobbyModule::ShutdownModule()
 {
 }
 
-bool FLobbyModule::Exec(UWorld* InWorld, const TCHAR* Cmd, FOutputDevice& Ar)
-{
-	// Ignore any execs that don't start with Lobby
-	if (FParse::Command(&Cmd, TEXT("Lobby")))
-	{
-		return false;
-	}
-
-	return false;
-}
-
-
diff --git a/Engine/Plugins/Online/OnlineFramework/Source/Lobby/Public/LobbyModule.h b/Engine/Plugins/Online/OnlineFramework/Source/Lobby/Public/LobbyModule.h
--- a/Engine/Plugins/Online/OnlineFramework/Source/Lobby/Public/LobbyModule.h
+++ b/Engine/Plugins/Online/OnlineFramework/Source/Lobby/Public/LobbyModule.h
@@ -4,7 +4,6 @@
 
 #include "CoreMinimal.h"
 #include "Stats/Stats.h"
-#include "Misc/CoreMisc.h"
 #include "Modules/ModuleInterface.h"
 #include "Modules/ModuleManager.h"
 
@@ -20,14 +19,11 @@ DECLARE_STATS_GROUP(TEXT("Lobby"), STATGROUP_Lobby, STATCAT_Advanced);
  * Module for lobbies via online beacon
  */
 class FLobbyModule : 
-	public IModuleInterface, public FSelfRegisteringExec
+	public IModuleInterface
 {
 
 public:
 
-	// FSelfRegisteringExec
-	virtual bool Exec(UWorld* InWorld, const TCHAR* Cmd, FOutputDevice& Ar) override;
-
 	/**
 	 * Singleton-like access to this module's interface.  This is just for convenience!
 	 * Beware of calling this during the shutdown phase, though.  Your module might have been unloaded already.
diff --git a/Engine/Plugins/Online/OnlineFramework/Source/Party/Private/PartyModule.cpp b/Engine/Plugins/Online/OnlineFramework/Source/Party/Private/PartyModule.cpp
--- a/Engine/Plugins/Online/OnlineFramework/Source/Party/Private/PartyModule.cpp
+++ b/Engine/Plugins/Online/OnlineFramework/Source/Party/Private/PartyModule.cpp
@@ -18,16 +18,3 @@ void FPartyModule::ShutdownModule()
 {
 }
 
-bool FPartyModule::Exec(UWorld* InWorld, const TCHAR* Cmd, FOutputDevice& Ar)
-{
-	// Ignore any execs that don't start with Party
-	if (FParse::Command(&Cmd, TEXT("Party")))
-	{
-		// this is where we can implement "Party XYZABC" commands
-		return false;
-	}
-
-	return false;
-}
-
-
diff --git a/Engine/Plugins/Online/OnlineFramework/Source/Party/Private/SocialToolkit.cpp b/Engine/Plugins/Online/OnlineFramework/Source/Party/Private/SocialToolkit.cpp
--- a/Engine/Plugins/Online/OnlineFramework/Source/Party/Private/SocialToolkit.cpp
+++ b/Engine/Plugins/Online/OnlineFramework/Source/Party/Private/SocialToolkit.cpp
@@ -1239,11 +1239,6 @@ void USocialToolkit::HandlePartyRequestToJoinRemoved(const FUniqueNetId& LocalUs
 	}
 }
 
-bool USocialToolkit::Exec(class UWorld* InWorld, const TCHAR* Cmd, FOutputDevice& Out)
-{
-	return false;
-}
-
 #if WITH_EDITOR
 void USocialToolkit::Debug_OnStartRandomizeUserPresence(uint8 NumRandomUser, float TickerTimer)
 {
diff --git a/Engine/Plugins/Online/OnlineFramework/Source/Party/Public/PartyModule.h b/Engine/Plugins/Online/OnlineFramework/Source/Party/Public/PartyModule.h
--- a/Engine/Plugins/Online/OnlineFramework/Source/Party/Public/PartyModule.h
+++ b/Engine/Plugins/Online/OnlineFramework/Source/Party/Public/PartyModule.h
@@ -20,14 +20,11 @@ DECLARE_CYCLE_STAT_EXTERN(TEXT("PartyStat1"), STAT_PartyStat1, STATGROUP_Party,
  * Module for the party framework that builds on top of IOnlineParty
  */
 class FPartyModule : 
-	public IModuleInterface, public FSelfRegisteringExec
+	public IModuleInterface
 {
 
 public:
 
-	// FSelfRegisteringExec
-	virtual bool Exec(UWorld* InWorld, const TCHAR* Cmd, FOutputDevice& Ar) override;
-
 	/**
 	 * Singleton-like access to this module's interface.  This is just for convenience!
 	 * Beware of calling this during the shutdown phase, though.  Your module might have been unloaded already.
diff --git a/Engine/Plugins/Online/OnlineFramework/Source/Party/Public/SocialToolkit.h b/Engine/Plugins/Online/OnlineFramework/Source/Party/Public/SocialToolkit.h
--- a/Engine/Plugins/Online/OnlineFramework/Source/Party/Public/SocialToolkit.h
+++ b/Engine/Plugins/Online/OnlineFramework/Source/Party/Public/SocialToolkit.h
@@ -79,9 +79,6 @@ public:
 	ULocalPlayer& GetOwningLocalPlayer() const;
 	const TArray<USocialUser*>& GetAllUsers() const { return AllUsers; }
 
-	// FExec
-	virtual bool Exec(class UWorld* InWorld, const TCHAR* Cmd, FOutputDevice& Out) override;
-
 	/** Finds a SocialUser given a unique net ID from any OSS */
 	USocialUser* FindUser(const FUniqueNetIdRepl& UserId) const;
 
diff --git a/Engine/Source/Runtime/Online/ICMP/Private/IcmpModule.cpp b/Engine/Source/Runtime/Online/ICMP/Private/IcmpModule.cpp
--- a/Engine/Source/Runtime/Online/ICMP/Private/IcmpModule.cpp
+++ b/Engine/Source/Runtime/Online/ICMP/Private/IcmpModule.cpp
@@ -14,15 +14,3 @@ void FIcmpModule::ShutdownModule()
 {
 }
 
-bool FIcmpModule::Exec(UWorld* InWorld, const TCHAR* Cmd, FOutputDevice& Ar)
-{
-	// Ignore any execs that don't start with Icmp
-	if (FParse::Command(&Cmd, TEXT("Icmp")))
-	{
-		return false;
-	}
-
-	return false;
-}
-
-
diff --git a/Engine/Source/Runtime/Online/ICMP/Public/IcmpModule.h b/Engine/Source/Runtime/Online/ICMP/Public/IcmpModule.h
--- a/Engine/Source/Runtime/Online/ICMP/Public/IcmpModule.h
+++ b/Engine/Source/Runtime/Online/ICMP/Public/IcmpModule.h
@@ -5,7 +5,6 @@
 #include "CoreMinimal.h"
 #include "HAL/Platform.h"
 #include "Logging/LogMacros.h"
-#include "Misc/CoreMisc.h"
 #include "Modules/ModuleInterface.h"
 #include "Modules/ModuleManager.h"
 
@@ -19,14 +18,10 @@ ICMP_API DECLARE_LOG_CATEGORY_EXTERN(LogIcmp, Display, All);
  * Module for Icmp service utilities
  */
 class FIcmpModule : 
-	public IModuleInterface, public FSelfRegisteringExec
+	public IModuleInterface
 {
 
 public:
-
-	// FSelfRegisteringExec
-	virtual bool Exec(UWorld* InWorld, const TCHAR* Cmd, FOutputDevice& Ar) override;
-
 	/**
 	 * Singleton-like access to this module's interface.  This is just for convenience!
 	 * Beware of calling this during the shutdown phase, though.  Your module might have been unloaded already.
diff --git a/Engine/Source/Runtime/Online/SSL/Private/SslModule.cpp b/Engine/Source/Runtime/Online/SSL/Private/SslModule.cpp
--- a/Engine/Source/Runtime/Online/SSL/Private/SslModule.cpp
+++ b/Engine/Source/Runtime/Online/SSL/Private/SslModule.cpp
@@ -21,19 +21,6 @@ FSslModule::FSslModule()
 {
 }
 
-bool FSslModule::Exec(UWorld* InWorld, const TCHAR* Cmd, FOutputDevice& Ar)
-{
-	bool bResult = false;
-
-	// Ignore any execs that don't start with HTTP
-	if (FParse::Command(&Cmd, TEXT("SSL")))
-	{
-		bResult = false;
-	}
-
-	return bResult;
-}
-
 void FSslModule::StartupModule()
 {	
 	Singleton = this;
diff --git a/Engine/Source/Runtime/Online/SSL/Public/SslModule.h b/Engine/Source/Runtime/Online/SSL/Public/SslModule.h
--- a/Engine/Source/Runtime/Online/SSL/Public/SslModule.h
+++ b/Engine/Source/Runtime/Online/SSL/Public/SslModule.h
@@ -4,7 +4,6 @@
 
 #include "HAL/Platform.h"
 #include "Misc/AssertionMacros.h"
-#include "Misc/CoreMisc.h"
 #include "Modules/ModuleInterface.h"
 #include "Modules/ModuleManager.h"
 
@@ -17,24 +16,11 @@ class UWorld;
  * Module for SSL/TLS certificate management
  */
 class FSslModule : 
-	public IModuleInterface, public FSelfRegisteringExec
+	public IModuleInterface
 {
 
 public:
 
-	// FSelfRegisteringExec
-
-	/**
-	 * Handle exec commands starting with "SSL"
-	 *
-	 * @param InWorld	the world context
-	 * @param Cmd		the exec command being executed
-	 * @param Ar		the archive to log results to
-	 *
-	 * @return true if the handler consumed the input, false to continue searching handlers
-	 */
-	virtual bool Exec(UWorld* InWorld, const TCHAR* Cmd, FOutputDevice& Ar) override;
-
 	// FSslModule
 
 	/** Constructor */
diff --git a/Engine/Source/Runtime/Online/Voice/Private/VoiceModule.cpp b/Engine/Source/Runtime/Online/Voice/Private/VoiceModule.cpp
--- a/Engine/Source/Runtime/Online/Voice/Private/VoiceModule.cpp
+++ b/Engine/Source/Runtime/Online/Voice/Private/VoiceModule.cpp
@@ -61,17 +61,6 @@ void FVoiceModule::ShutdownModule()
 #endif
 }
 
-bool FVoiceModule::Exec(UWorld* InWorld, const TCHAR* Cmd, FOutputDevice& Ar)
-{
-	// Ignore any execs that don't start with Voice
-	if (FParse::Command(&Cmd, TEXT("Voice")))
-	{
-		return false;
-	}
-
-	return false;
-}
-
 TSharedPtr<IVoiceCapture> FVoiceModule::CreateVoiceCapture(const FString& DeviceName, int32 SampleRate, int32 NumChannels)
 {
 #if PLATFORM_SUPPORTS_VOICE_CAPTURE
diff --git a/Engine/Source/Runtime/Online/Voice/Public/VoiceModule.h b/Engine/Source/Runtime/Online/Voice/Public/VoiceModule.h
--- a/Engine/Source/Runtime/Online/Voice/Public/VoiceModule.h
+++ b/Engine/Source/Runtime/Online/Voice/Public/VoiceModule.h
@@ -3,7 +3,6 @@
 #pragma once
 
 #include "CoreMinimal.h"
-#include "Misc/CoreMisc.h"
 #include "Modules/ModuleInterface.h"
 #include "Modules/ModuleManager.h"
 #include "Interfaces/VoiceCapture.h"
@@ -33,14 +32,11 @@ DECLARE_LOG_CATEGORY_EXTERN(LogVoiceCapture, Warning, All);
  * Module for Voice capture/compression/decompression implementations
  */
 class FVoiceModule : 
-	public IModuleInterface, public FSelfRegisteringExec
+	public IModuleInterface
 {
 
 public:
 
-	// FSelfRegisteringExec
-	virtual bool Exec(UWorld* InWorld, const TCHAR* Cmd, FOutputDevice& Ar) override;
-
 	/**
 	 * Singleton-like access to this module's interface.  This is just for convenience!
 	 * Beware of calling this during the shutdown phase, though.  Your module might have been unloaded already.
