Update AppleHTTP to use NSURLSession instead of the deprecated NSURLConnection

Also, added some VeryVerbose logs to make easier to debug certificate pinning

diff --git a/Engine/Source/Runtime/Online/HTTP/Private/Apple/AppleHTTP.cpp b/Engine/Source/Runtime/Online/HTTP/Private/Apple/AppleHTTP.cpp
deleted file mode 100644
diff --git a/Engine/Source/Runtime/Online/HTTP/Private/Apple/AppleHTTP.h b/Engine/Source/Runtime/Online/HTTP/Private/Apple/AppleHTTP.h
deleted file mode 100644
diff --git a/Engine/Source/Runtime/Online/HTTP/Private/Apple/AppleHTTPNSUrlConnection.cpp b/Engine/Source/Runtime/Online/HTTP/Private/Apple/AppleHTTPNSUrlConnection.cpp
new file mode 100644
--- /dev/null
+++ b/Engine/Source/Runtime/Online/HTTP/Private/Apple/AppleHTTPNSUrlConnection.cpp
@@ -0,0 +1,1022 @@
+// Copyright Epic Games, Inc. All Rights Reserved.
+
+
+#include "AppleHTTPNSUrlConnection.h"
+#include "Misc/EngineVersion.h"
+#include "Security/Security.h"
+#include "CommonCrypto/CommonDigest.h"
+#include "Foundation/Foundation.h"
+#include "Misc/App.h"
+#include "Misc/Base64.h"
+#include "HAL/PlatformTime.h"
+#include "Http.h"
+#include "HttpModule.h"
+#include "Apple/CFRef.h"
+
+#if WITH_SSL
+#include "Ssl.h"
+#endif
+
+/****************************************************************************
+ * FAppleHttpNSUrlConnectionRequest implementation
+ ***************************************************************************/
+
+
+FAppleHttpNSUrlConnectionRequest::FAppleHttpNSUrlConnectionRequest()
+:	Connection(nullptr)
+,	bIsPayloadFile(false)
+,	RequestPayloadByteLength(0)
+,	CompletionStatus(EHttpRequestStatus::NotStarted)
+,	ProgressBytesSent(0)
+,	StartRequestTime(0.0)
+,	ElapsedTime(0.0f)
+{
+	UE_LOG(LogHttp, Verbose, TEXT("FAppleHttpNSUrlConnectionRequest::FAppleHttpNSUrlConnectionRequest()"));
+	Request = [[NSMutableURLRequest alloc] init];
+	Request.timeoutInterval = FHttpModule::Get().GetHttpTimeout();
+
+	// Disable cache to mimic WinInet behavior
+	Request.cachePolicy = NSURLRequestReloadIgnoringLocalCacheData;
+
+	// Add default headers
+	const TMap<FString, FString>& DefaultHeaders = FHttpModule::Get().GetDefaultHeaders();
+	for (TMap<FString, FString>::TConstIterator It(DefaultHeaders); It; ++It)
+	{
+		SetHeader(It.Key(), It.Value());
+	}
+
+#if WITH_SSL
+	// Make sure the module is loaded on the game thread before being used by FHttpResponseAppleNSUrlConnectionWrapper, which will be called on the main thread
+	FSslModule::Get();
+#endif
+}
+
+
+FAppleHttpNSUrlConnectionRequest::~FAppleHttpNSUrlConnectionRequest()
+{
+	UE_LOG(LogHttp, Verbose, TEXT("FAppleHttpNSUrlConnectionRequest::~FAppleHttpNSUrlConnectionRequest()"));
+	check(Connection == nullptr);
+	[Request release];
+}
+
+
+FString FAppleHttpNSUrlConnectionRequest::GetURL() const
+{
+	SCOPED_AUTORELEASE_POOL;
+	NSURL* URL = Request.URL;
+	if (URL != nullptr)
+	{
+		FString ConvertedURL(URL.absoluteString);
+		UE_LOG(LogHttp, Verbose, TEXT("FAppleHttpNSUrlConnectionRequest::GetURL() - %s"), *ConvertedURL);
+		return ConvertedURL;
+	}
+	else
+	{
+		UE_LOG(LogHttp, Verbose, TEXT("FAppleHttpNSUrlConnectionRequest::GetURL() - NULL"));
+		return FString();
+	}
+}
+
+
+void FAppleHttpNSUrlConnectionRequest::SetURL(const FString& URL)
+{
+	SCOPED_AUTORELEASE_POOL;
+	UE_LOG(LogHttp, Verbose, TEXT("FAppleHttpNSUrlConnectionRequest::SetURL() - %s"), *URL);
+	Request.URL = [NSURL URLWithString: URL.GetNSString()];
+}
+
+
+FString FAppleHttpNSUrlConnectionRequest::GetURLParameter(const FString& ParameterName) const
+{
+	SCOPED_AUTORELEASE_POOL;
+	UE_LOG(LogHttp, Verbose, TEXT("FAppleHttpNSUrlConnectionRequest::GetURLParameter() - %s"), *ParameterName);
+
+	NSRange ParametersStart = [Request.URL.query rangeOfString:@"?"];
+	if (ParametersStart.location != NSNotFound && ParametersStart.length > 0)
+	{
+		NSString* ParametersStr = [Request.URL.query substringFromIndex:ParametersStart.location + 1];
+		NSString* ParameterNameStr = ParameterName.GetNSString();
+		NSArray* Parameters = [ParametersStr componentsSeparatedByString:@"&"];
+		for (NSString* Parameter in Parameters)
+		{
+			NSArray* KeyValue = [Parameter componentsSeparatedByString:@"="];
+			NSString* Key = KeyValue[0];
+			if ([Key compare:ParameterNameStr] == NSOrderedSame)
+			{
+				return FString(KeyValue[1]);
+			}
+		}
+	}
+
+	return FString();
+}
+
+
+FString FAppleHttpNSUrlConnectionRequest::GetHeader(const FString& HeaderName) const
+{
+	SCOPED_AUTORELEASE_POOL;
+	FString Header([Request valueForHTTPHeaderField:HeaderName.GetNSString()]);
+	UE_LOG(LogHttp, Verbose, TEXT("FAppleHttpNSUrlConnectionRequest::GetHeader() - %s"), *Header);
+	return Header;
+}
+
+
+void FAppleHttpNSUrlConnectionRequest::SetHeader(const FString& HeaderName, const FString& HeaderValue)
+{
+	SCOPED_AUTORELEASE_POOL;
+	UE_LOG(LogHttp, Verbose, TEXT("FAppleHttpNSUrlConnectionRequest::SetHeader() - %s / %s"), *HeaderName, *HeaderValue );
+	[Request setValue: HeaderValue.GetNSString() forHTTPHeaderField: HeaderName.GetNSString()];
+}
+
+void FAppleHttpNSUrlConnectionRequest::AppendToHeader(const FString& HeaderName, const FString& AdditionalHeaderValue)
+{
+    if (!HeaderName.IsEmpty() && !AdditionalHeaderValue.IsEmpty())
+    {
+        NSDictionary* Headers = [Request allHTTPHeaderFields];
+        NSString* PreviousHeaderValuePtr = [Headers objectForKey: HeaderName.GetNSString()];
+        FString PreviousValue(PreviousHeaderValuePtr);
+		FString NewValue;
+		if (!PreviousValue.IsEmpty())
+		{
+			NewValue = PreviousValue + TEXT(", ");
+		}
+		NewValue += AdditionalHeaderValue;
+
+        SetHeader(HeaderName, NewValue);
+	}
+}
+
+TArray<FString> FAppleHttpNSUrlConnectionRequest::GetAllHeaders() const
+{
+	SCOPED_AUTORELEASE_POOL;
+	UE_LOG(LogHttp, Verbose, TEXT("FAppleHttpNSUrlConnectionRequest::GetAllHeaders()"));
+	NSDictionary* Headers = Request.allHTTPHeaderFields;
+	TArray<FString> Result;
+	Result.Reserve(Headers.count);
+	for (NSString* Key in Headers.allKeys)
+	{
+		FString ConvertedValue(Headers[Key]);
+		FString ConvertedKey(Key);
+		UE_LOG(LogHttp, Verbose, TEXT("Header= %s, Key= %s"), *ConvertedValue, *ConvertedKey);
+
+		Result.Add( FString::Printf( TEXT("%s: %s"), *ConvertedKey, *ConvertedValue ) );
+	}
+	return Result;
+}
+
+
+const TArray<uint8>& FAppleHttpNSUrlConnectionRequest::GetContent() const
+{
+	UE_LOG(LogHttp, Verbose, TEXT("FAppleHttpNSUrlConnectionRequest::GetContent()"));
+	if (bIsPayloadFile)
+	{
+		UE_LOG(LogHttp, Warning, TEXT("FAppleHttpNSUrlConnectionRequest::GetContent() called on a request that is set up for streaming a file. Return value is an empty buffer"));
+		RequestPayload.Empty();
+		return RequestPayload;
+	}
+	else
+	{
+		SCOPED_AUTORELEASE_POOL;
+		NSData* Body = Request.HTTPBody; // accessing HTTPBody will call copy on the value, increasing its retain count
+		RequestPayload.Empty();
+		RequestPayload.Append((const uint8*)Body.bytes, Body.length);
+		return RequestPayload;
+	}
+}
+
+
+void FAppleHttpNSUrlConnectionRequest::SetContent(const TArray<uint8>& ContentPayload)
+{
+	if (CompletionStatus == EHttpRequestStatus::Processing)
+	{
+		UE_LOG(LogHttp, Warning, TEXT("FAppleHttpNSUrlConnectionRequest::SetContent() - attempted to set content on a request that is inflight"));
+		return;
+	}
+
+	UE_LOG(LogHttp, Verbose, TEXT("FAppleHttpNSUrlConnectionRequest::SetContent()"));
+	Request.HTTPBody = [NSData dataWithBytes:ContentPayload.GetData() length:ContentPayload.Num()];
+	RequestPayloadByteLength = ContentPayload.Num();
+	bIsPayloadFile = false;
+
+	ContentData.Empty();
+}
+
+
+void FAppleHttpNSUrlConnectionRequest::SetContent(TArray<uint8>&& ContentPayload)
+{
+	if (CompletionStatus == EHttpRequestStatus::Processing)
+	{
+		UE_LOG(LogHttp, Warning, TEXT("FAppleHttpNSUrlConnectionRequest::SetContent() - attempted to set content on a request that is inflight"));
+		return;
+	}
+
+	UE_LOG(LogHttp, Verbose, TEXT("FAppleHttpNSUrlConnectionRequest::SetContent()"));
+	ContentData = MoveTemp(ContentPayload);
+
+	Request.HTTPBody = [NSData dataWithBytesNoCopy:ContentData.GetData() length:ContentData.Num() freeWhenDone:false];
+	RequestPayloadByteLength = ContentData.Num();
+	bIsPayloadFile = false;
+}
+
+
+FString FAppleHttpNSUrlConnectionRequest::GetContentType() const
+{
+	FString ContentType = GetHeader(TEXT("Content-Type"));
+	UE_LOG(LogHttp, Verbose, TEXT("FAppleHttpNSUrlConnectionRequest::GetContentType() - %s"), *ContentType);
+	return ContentType;
+}
+
+
+int32 FAppleHttpNSUrlConnectionRequest::GetContentLength() const
+{
+	UE_LOG(LogHttp, Verbose, TEXT("FAppleHttpNSUrlConnectionRequest::GetContentLength() - %i"), RequestPayloadByteLength);
+	return RequestPayloadByteLength;
+}
+
+
+void FAppleHttpNSUrlConnectionRequest::SetContentAsString(const FString& ContentString)
+{
+	if (CompletionStatus == EHttpRequestStatus::Processing)
+	{
+		UE_LOG(LogHttp, Warning, TEXT("FAppleHttpNSUrlConnectionRequest::SetContentAsString() - attempted to set content on a request that is inflight"));
+		return;
+	}
+
+	UE_LOG(LogHttp, Verbose, TEXT("FAppleHttpNSUrlConnectionRequest::SetContentAsString() - %s"), *ContentString);
+	FTCHARToUTF8 Converter(*ContentString);
+
+	// The extra length computation here is unfortunate, but it's technically not safe to assume the length is the same.
+	Request.HTTPBody = [NSData dataWithBytes:(ANSICHAR*)Converter.Get() length:Converter.Length()];
+	RequestPayloadByteLength = Converter.Length();
+	bIsPayloadFile = false;
+
+	ContentData.Empty();
+}
+
+bool FAppleHttpNSUrlConnectionRequest::SetContentAsStreamedFile(const FString& Filename)
+{
+	SCOPED_AUTORELEASE_POOL;
+    UE_LOG(LogHttp, Verbose, TEXT("FAppleHttpNSUrlConnectionRequest::SetContentAsStreamedFile() - %s"), *Filename);
+
+	if (CompletionStatus == EHttpRequestStatus::Processing)
+	{
+		UE_LOG(LogHttp, Warning, TEXT("FAppleHttpNSUrlConnectionRequest::SetContentAsStreamedFile() - attempted to set content on a request that is inflight"));
+		return false;
+	}
+
+	NSString* PlatformFilename = Filename.GetNSString();
+
+	Request.HTTPBody = nil;
+	ContentData.Empty();
+
+	struct stat FileAttrs = { 0 };
+	if (stat(PlatformFilename.fileSystemRepresentation, &FileAttrs) == 0)
+	{
+		UE_LOG(LogHttp, VeryVerbose, TEXT("FAppleHttpNSUrlConnectionRequest::SetContentAsStreamedFile succeeded in getting the file size - %d"), FileAttrs.st_size);
+		// Under the hood, the Foundation framework unsets HTTPBody, and takes over as the stream delegate.
+		// The stream itself should be unopened when passed to setHTTPBodyStream.
+		Request.HTTPBodyStream = [NSInputStream inputStreamWithFileAtPath: PlatformFilename];
+		RequestPayloadByteLength = FileAttrs.st_size;
+		bIsPayloadFile = true;
+	}
+	else
+	{
+		UE_LOG(LogHttp, VeryVerbose, TEXT("FAppleHttpNSUrlConnectionRequest::SetContentAsStreamedFile failed to get file size"));
+		Request.HTTPBodyStream = nil;
+		RequestPayloadByteLength = 0;
+		bIsPayloadFile = false;
+	}
+
+	return bIsPayloadFile;
+}
+
+
+bool FAppleHttpNSUrlConnectionRequest::SetContentFromStream(TSharedRef<FArchive, ESPMode::ThreadSafe> Stream)
+{
+	UE_LOG(LogHttp, Warning, TEXT("FAppleHttpNSUrlConnectionRequest::SetContentFromStream is not implemented"));
+	return false;
+}
+
+
+FString FAppleHttpNSUrlConnectionRequest::GetVerb() const
+{
+	FString ConvertedVerb(Request.HTTPMethod);
+	UE_LOG(LogHttp, Verbose, TEXT("FAppleHttpNSUrlConnectionRequest::GetVerb() - %s"), *ConvertedVerb);
+	return ConvertedVerb;
+}
+
+
+void FAppleHttpNSUrlConnectionRequest::SetVerb(const FString& Verb)
+{
+	SCOPED_AUTORELEASE_POOL;
+	UE_LOG(LogHttp, Verbose, TEXT("FAppleHttpNSUrlConnectionRequest::SetVerb() - %s"), *Verb);
+	Request.HTTPMethod = Verb.GetNSString();
+}
+
+void FAppleHttpNSUrlConnectionRequest::SetTimeout(float InTimeoutSecs)
+{
+	Request.timeoutInterval = InTimeoutSecs;
+}
+
+void FAppleHttpNSUrlConnectionRequest::ClearTimeout()
+{
+	Request.timeoutInterval = FHttpModule::Get().GetHttpTimeout();
+}
+
+TOptional<float> FAppleHttpNSUrlConnectionRequest::GetTimeout() const
+{
+	return TOptional<float>(Request.timeoutInterval);
+}
+
+bool FAppleHttpNSUrlConnectionRequest::ProcessRequest()
+{
+	SCOPED_AUTORELEASE_POOL;
+	UE_LOG(LogHttp, Verbose, TEXT("FAppleHttpNSUrlConnectionRequest::ProcessRequest()"));
+	bool bStarted = false;
+
+	FString Scheme(Request.URL.scheme);
+	Scheme = Scheme.ToLower();
+
+	// Prevent overlapped requests using the same instance
+	if (CompletionStatus == EHttpRequestStatus::Processing)
+	{
+		UE_LOG(LogHttp, Warning, TEXT("ProcessRequest failed. Still processing last request."));
+	}
+	else if(GetURL().Len() == 0)
+	{
+		UE_LOG(LogHttp, Warning, TEXT("ProcessRequest failed. No URL was specified."));
+	}
+	else if( Scheme != TEXT("http") && Scheme != TEXT("https"))
+	{
+		UE_LOG(LogHttp, Warning, TEXT("ProcessRequest failed. URL '%s' is not a valid HTTP request. %p"), *GetURL(), this);
+	}
+	else if (!FHttpModule::Get().GetHttpManager().IsDomainAllowed(GetURL()))
+	{
+		UE_LOG(LogHttp, Warning, TEXT("ProcessRequest failed. URL '%s' is not using an allowed domain. %p"), *GetURL(), this);
+	}
+	else
+	{
+		bStarted = StartRequest();
+	}
+
+	if( !bStarted )
+	{
+		// Ensure we run on game thread
+		if (!IsInGameThread())
+		{
+			FHttpModule::Get().GetHttpManager().AddGameThreadTask([StrongThis = StaticCastSharedRef<FAppleHttpNSUrlConnectionRequest>(AsShared())]()
+			{
+				StrongThis->FinishedRequest();
+			});
+		}
+		else
+		{
+			FinishedRequest();
+		}
+	}
+
+	return bStarted;
+}
+
+bool FAppleHttpNSUrlConnectionRequest::StartRequest()
+{
+	SCOPED_AUTORELEASE_POOL;
+	UE_LOG(LogHttp, Verbose, TEXT("FAppleHttpNSUrlConnectionRequest::StartRequest()"));
+	bool bStarted = false;
+
+	// set the content-length and user-agent
+	if(GetContentLength() > 0)
+	{
+		[Request setValue:[NSString stringWithFormat:@"%d", GetContentLength()] forHTTPHeaderField:@"Content-Length"];
+	}
+
+	const FString UserAgent = GetHeader("User-Agent");
+	if(UserAgent.IsEmpty())
+	{
+		NSString* Tag = FPlatformHttp::GetDefaultUserAgent().GetNSString();
+		[Request setValue:Tag forHTTPHeaderField:@"User-Agent"];
+	}
+
+	Response = MakeShareable( new FAppleHttpNSUrlConnectionResponse( *this ) );
+
+#pragma clang diagnostic push
+#pragma clang diagnostic ignored "-Wdeprecated-declarations"
+	// Create the connection, tell it to run in the main run loop, and kick it off.
+	Connection = [[NSURLConnection alloc] initWithRequest:Request delegate:Response->ResponseWrapper startImmediately:NO];
+#pragma clang diagnostic pop
+	if (Connection != nullptr && Response->ResponseWrapper != nullptr)
+	{
+		CompletionStatus = EHttpRequestStatus::Processing;
+		[Connection scheduleInRunLoop:[NSRunLoop mainRunLoop] forMode:NSDefaultRunLoopMode];
+		[Connection start];
+		UE_LOG(LogHttp, Verbose, TEXT("[Connection start]"));
+
+		bStarted = true;
+		// Add to global list while being processed so that the ref counted request does not get deleted
+		FHttpModule::Get().GetHttpManager().AddRequest(SharedThis(this));
+	}
+	else
+	{
+		UE_LOG(LogHttp, Warning, TEXT("ProcessRequest failed. Could not initialize Internet connection."));
+		CompletionStatus = EHttpRequestStatus::Failed_ConnectionError;
+	}
+	StartRequestTime = FPlatformTime::Seconds();
+	// reset the elapsed time.
+	ElapsedTime = 0.0f;
+
+	return bStarted;
+}
+
+void FAppleHttpNSUrlConnectionRequest::FinishedRequest()
+{
+	UE_LOG(LogHttp, Verbose, TEXT("FAppleHttpNSUrlConnectionRequest::FinishedRequest()"));
+	ElapsedTime = (float)(FPlatformTime::Seconds() - StartRequestTime);
+	if( Response.IsValid() && Response->IsReady() && !Response->HadError())
+	{
+		UE_LOG(LogHttp, Verbose, TEXT("Request succeeded"));
+		CompletionStatus = EHttpRequestStatus::Succeeded;
+
+		// TODO: Try to broadcast OnHeaderReceived when we receive headers instead of here at the end
+		BroadcastResponseHeadersReceived();
+		OnProcessRequestComplete().ExecuteIfBound(SharedThis(this), Response, true);
+	}
+	else
+	{
+		UE_LOG(LogHttp, Verbose, TEXT("Request failed"));
+		FString URL([[Request URL] absoluteString]);
+		CompletionStatus = EHttpRequestStatus::Failed;
+		if (Response.IsValid() && [Response->ResponseWrapper bIsHostConnectionFailure])
+		{
+			CompletionStatus = EHttpRequestStatus::Failed_ConnectionError;
+		}
+
+		Response = nullptr;
+		OnProcessRequestComplete().ExecuteIfBound(SharedThis(this), nullptr, false);
+	}
+
+	// Clean up session/request handles that may have been created
+	CleanupRequest();
+
+	// Remove from global list since processing is now complete
+	if (FHttpModule::Get().GetHttpManager().IsValidRequest(this))
+	{
+		FHttpModule::Get().GetHttpManager().RemoveRequest(SharedThis(this));
+	}
+}
+
+
+void FAppleHttpNSUrlConnectionRequest::CleanupRequest()
+{
+	UE_LOG(LogHttp, Verbose, TEXT("FAppleHttpNSUrlConnectionRequest::CleanupRequest()"));
+
+	if(CompletionStatus == EHttpRequestStatus::Processing)
+	{
+		CancelRequest();
+	}
+
+	if(Connection != nullptr)
+	{
+		[Connection release];
+		Connection = nullptr;
+	}
+}
+
+
+void FAppleHttpNSUrlConnectionRequest::CancelRequest()
+{
+	
+	UE_LOG(LogHttp, Verbose, TEXT("FAppleHttpNSUrlConnectionRequest::CancelRequest()"));
+	if(Connection != nullptr)
+	{
+		[Connection cancel];
+	}
+
+	// Ensure we run on game thread
+	if (!IsInGameThread())
+	{
+		FHttpModule::Get().GetHttpManager().AddGameThreadTask([StrongThis = StaticCastSharedRef<FAppleHttpNSUrlConnectionRequest>(AsShared())]()
+		{
+			StrongThis->FinishedRequest();
+		});
+	}
+	else
+	{
+		FinishedRequest();
+	}
+}
+
+
+EHttpRequestStatus::Type FAppleHttpNSUrlConnectionRequest::GetStatus() const
+{
+	UE_LOG(LogHttp, Verbose, TEXT("FAppleHttpNSUrlConnectionRequest::GetStatus()"));
+	return CompletionStatus;
+}
+
+
+const FHttpResponsePtr FAppleHttpNSUrlConnectionRequest::GetResponse() const
+{
+	return Response;
+}
+
+void FAppleHttpNSUrlConnectionRequest::Tick(float DeltaSeconds)
+{
+	if (Response.IsValid() && (CompletionStatus == EHttpRequestStatus::Processing || Response->HadError()))
+	{
+		if (OnRequestProgress().IsBound())
+		{
+			const int32 BytesWritten = Response->GetNumBytesWritten();
+			const int32 BytesRead = Response->GetNumBytesReceived();
+			if (BytesWritten > 0 || BytesRead > 0)
+			{
+				OnRequestProgress().ExecuteIfBound(SharedThis(this), BytesWritten, BytesRead);
+			}
+		}
+		if (Response->IsReady())
+		{
+			FinishedRequest();
+		}
+	}
+}
+
+float FAppleHttpNSUrlConnectionRequest::GetElapsedTime() const
+{
+	return ElapsedTime;
+}
+
+
+/****************************************************************************
+ * FHttpResponseAppleNSUrlConnectionWrapper implementation
+ ***************************************************************************/
+
+@implementation FHttpResponseAppleNSUrlConnectionWrapper
+@synthesize Response;
+@synthesize bIsReady;
+@synthesize bHadError;
+@synthesize bIsHostConnectionFailure;
+@synthesize BytesWritten;
+
+
+-(FHttpResponseAppleNSUrlConnectionWrapper*) init
+{
+	UE_LOG(LogHttp, Verbose, TEXT("-(FHttpResponseAppleNSUrlConnectionWrapper*) init"));
+	self = [super init];
+	bIsReady = false;
+	bHadError = false;
+	bIsHostConnectionFailure = false;
+	
+	return self;
+}
+
+- (void)dealloc
+{
+	[Response release];
+	[super dealloc];
+}
+
+-(void) connection:(NSURLConnection *)connection didSendBodyData:(NSInteger)bytesWritten totalBytesWritten:(NSInteger)totalBytesWritten totalBytesExpectedToWrite:(NSInteger)totalBytesExpectedToWrite
+{
+	UE_LOG(LogHttp, Verbose, TEXT("didSendBodyData:(NSInteger)bytesWritten totalBytes:Written:(NSInteger)totalBytesWritten totalBytesExpectedToWrite:(NSInteger)totalBytesExpectedToWrite"));
+	self.BytesWritten = totalBytesWritten;
+	UE_LOG(LogHttp, Verbose, TEXT("didSendBodyData: totalBytesWritten = %d, totalBytesExpectedToWrite = %d: %p"), totalBytesWritten, totalBytesExpectedToWrite, self);
+}
+
+-(void) connection:(NSURLConnection *)connection didReceiveResponse:(NSURLResponse *)response
+{
+	UE_LOG(LogHttp, Verbose, TEXT("didReceiveResponse:(NSURLResponse *)response"));
+	self.Response = (NSHTTPURLResponse*)response;
+	
+	// presize the payload container if possible
+	Payload.Empty([response expectedContentLength] != NSURLResponseUnknownLength ? [response expectedContentLength] : 0);
+	UE_LOG(LogHttp, Verbose, TEXT("didReceiveResponse: expectedContentLength = %d. Length = %d: %p"), [response expectedContentLength], Payload.Max(), self);
+}
+
+
+-(void) connection:(NSURLConnection *)connection didReceiveData:(NSData *)data
+{
+	Payload.Append((const uint8*)[data bytes], [data length]);
+	UE_LOG(LogHttp, Verbose, TEXT("didReceiveData with %d bytes. After Append, Payload Length = %d: %p"), [data length], Payload.Num(), self);
+}
+
+
+-(void) connection:(NSURLConnection *)connection didFailWithError:(NSError *)error
+{
+	self.bIsReady = YES;
+	self.bHadError = YES;
+	UE_LOG(LogHttp, Warning, TEXT("didFailWithError. Http request failed - %s %s: %p"), 
+		*FString([error localizedDescription]),
+		*FString([[error userInfo] objectForKey:NSURLErrorFailingURLStringErrorKey]),
+		self);
+	// Determine if the specific error was failing to connect to the host.
+	switch ([error code])
+	{
+		case NSURLErrorCannotFindHost:
+		case NSURLErrorCannotConnectToHost:
+		case NSURLErrorDNSLookupFailed:
+			self.bIsHostConnectionFailure = YES;
+	}
+	// Log more details if verbose logging is enabled and this is an SSL error
+	if (UE_LOG_ACTIVE(LogHttp, Verbose))
+	{
+		SecTrustRef PeerTrustInfo = reinterpret_cast<SecTrustRef>([[error userInfo] objectForKey:NSURLErrorFailingURLPeerTrustErrorKey]);
+		if (PeerTrustInfo != nullptr)
+		{
+			SecTrustResultType TrustResult = kSecTrustResultInvalid;
+			SecTrustGetTrustResult(PeerTrustInfo, &TrustResult);
+
+			FString TrustResultString;
+			switch (TrustResult)
+			{
+#define MAP_TO_RESULTSTRING(Constant) case Constant: TrustResultString = TEXT(#Constant); break;
+			MAP_TO_RESULTSTRING(kSecTrustResultInvalid)
+			MAP_TO_RESULTSTRING(kSecTrustResultProceed)
+			MAP_TO_RESULTSTRING(kSecTrustResultDeny)
+			MAP_TO_RESULTSTRING(kSecTrustResultUnspecified)
+			MAP_TO_RESULTSTRING(kSecTrustResultRecoverableTrustFailure)
+			MAP_TO_RESULTSTRING(kSecTrustResultFatalTrustFailure)
+			MAP_TO_RESULTSTRING(kSecTrustResultOtherError)
+#undef MAP_TO_RESULTSTRING
+			default:
+				TrustResultString = TEXT("unknown");
+				break;
+			}
+			UE_LOG(LogHttp, Verbose, TEXT("didFailWithError. SSL trust result: %s (%d)"), *TrustResultString, TrustResult);
+		}
+	}
+}
+
+#if WITH_SSL
+-(void) connection:(NSURLConnection *)connection willSendRequestForAuthenticationChallenge: (NSURLAuthenticationChallenge *)challenge
+{
+	// CC gives the actual key, but strips the ASN.1 header... which means
+	// we can't calulate a proper SPKI hash without reconstructing it. sigh.
+	static const unsigned char rsa2048Asn1Header[] =
+	{
+		0x30, 0x82, 0x01, 0x22, 0x30, 0x0d, 0x06, 0x09,
+		0x2a, 0x86, 0x48, 0x86, 0xf7, 0x0d, 0x01, 0x01,
+		0x01, 0x05, 0x00, 0x03, 0x82, 0x01, 0x0f, 0x00
+	};
+	static const unsigned char rsa4096Asn1Header[] =
+	{
+		0x30, 0x82, 0x02, 0x22, 0x30, 0x0d, 0x06, 0x09,
+		0x2a, 0x86, 0x48, 0x86, 0xf7, 0x0d, 0x01, 0x01,
+		0x01, 0x05, 0x00, 0x03, 0x82, 0x02, 0x0f, 0x00
+	};
+	static const unsigned char ecdsaSecp256r1Asn1Header[] =
+	{
+		0x30, 0x59, 0x30, 0x13, 0x06, 0x07, 0x2a, 0x86,
+		0x48, 0xce, 0x3d, 0x02, 0x01, 0x06, 0x08, 0x2a,
+		0x86, 0x48, 0xce, 0x3d, 0x03, 0x01, 0x07, 0x03,
+		0x42, 0x00
+	};
+	static const unsigned char ecdsaSecp384r1Asn1Header[] =
+	{
+		0x30, 0x76, 0x30, 0x10, 0x06, 0x07, 0x2a, 0x86,
+		0x48, 0xce, 0x3d, 0x02, 0x01, 0x06, 0x05, 0x2b,
+		0x81, 0x04, 0x00, 0x22, 0x03, 0x62, 0x00
+	};
+
+    if (ensure(ISslCertificateManager::PUBLIC_KEY_DIGEST_SIZE == CC_SHA256_DIGEST_LENGTH))
+    {
+        // we only care about challenges to the received certificate chain
+        if ([challenge.protectionSpace.authenticationMethod isEqualToString:NSURLAuthenticationMethodServerTrust])
+        {
+            SecTrustRef RemoteTrust = challenge.protectionSpace.serverTrust;
+            FString RemoteHost = FString(UTF8_TO_TCHAR([challenge.protectionSpace.host UTF8String]));
+            if ((RemoteTrust == NULL) || (RemoteHost.IsEmpty()))
+            {
+                UE_LOG(LogHttp, Error, TEXT("failed certificate pinning validation: could not parse parameters during certificate pinning evaluation"));
+                [challenge.sender cancelAuthenticationChallenge: challenge];
+                return;
+            }
+
+			if (!SecTrustEvaluateWithError(RemoteTrust, nil))
+			{
+				UE_LOG(LogHttp, Error, TEXT("failed certificate pinning validation: default certificate trust evaluation failed for domain '%s'"), *RemoteHost);
+				[challenge.sender cancelAuthenticationChallenge: challenge];
+				return;
+			}
+            // look at all certs in the remote chain and calculate the SHA256 hash of their DER-encoded SPKI
+            // the chain starts with the server's cert itself, so walk backwards to optimize for roots first
+            TArray<TArray<uint8, TFixedAllocator<ISslCertificateManager::PUBLIC_KEY_DIGEST_SIZE>>> CertDigests;
+            
+            CFArrayRef Certificates = SecTrustCopyCertificateChain(RemoteTrust);
+            if (Certificates == nil)
+            {
+                UE_LOG(LogHttp, Error, TEXT("No certificate could be copied in the certificate chain used to evaluate trust."));
+                return;
+            }
+            
+            CFIndex CertificateCount = CFArrayGetCount(Certificates);
+            for (int i = 0; i < CertificateCount; ++i)
+            {
+                SecCertificateRef Cert = (SecCertificateRef)CFArrayGetValueAtIndex(Certificates, i);
+                
+                // this is not great, but the only way to extract a public key from a SecCertificateRef
+                // is to create an individual SecTrustRef for each cert that only contains itself and then
+                // evaluate that against an empty X509 policy.
+                TCFRef<SecTrustRef> CertTrust;
+                TCFRef<SecPolicyRef> TrustPolicy = SecPolicyCreateBasicX509();
+                SecTrustCreateWithCertificates(Cert, TrustPolicy, CertTrust.GetForAssignment());
+                SecTrustEvaluateWithError(CertTrust, nil);
+                TCFRef<SecKeyRef> CertPubKey;
+
+				CertPubKey = SecTrustCopyKey(CertTrust);
+
+				TCFRef<CFDataRef> CertPubKeyData = SecKeyCopyExternalRepresentation(CertPubKey, NULL);
+                if (!CertPubKeyData)
+                {
+                    UE_LOG(LogHttp, Warning, TEXT("could not extract public key from certificate %i for domain '%s'; skipping!"), i, *RemoteHost);
+                    continue;
+                }
+                
+				// we got the key. now we have to figure out what type of key it is; thanks, CommonCrypto.
+                TCFRef<CFDictionaryRef> CertPubKeyAttr = SecKeyCopyAttributes(CertPubKey);
+                NSString *CertPubKeyType = static_cast<NSString *>(CFDictionaryGetValue(CertPubKeyAttr, kSecAttrKeyType));
+                NSNumber *CertPubKeySize = static_cast<NSNumber *>(CFDictionaryGetValue(CertPubKeyAttr, kSecAttrKeySizeInBits));
+                char *CertPubKeyASN1Header;
+                uint8_t CertPubKeyASN1HeaderSize = 0;
+                if ([CertPubKeyType isEqualToString: (NSString *)kSecAttrKeyTypeRSA])
+                {
+                    switch ([CertPubKeySize integerValue])
+                    {
+                        case 2048:
+                            UE_LOG(LogHttp, VeryVerbose, TEXT("found 2048 bit RSA pubkey"));
+                            CertPubKeyASN1Header = (char *)rsa2048Asn1Header;
+                            CertPubKeyASN1HeaderSize = sizeof(rsa2048Asn1Header);
+                            break;
+                        case 4096:
+                            UE_LOG(LogHttp, VeryVerbose, TEXT("found 4096 bit RSA pubkey"));
+                            CertPubKeyASN1Header = (char *)rsa4096Asn1Header;
+                            CertPubKeyASN1HeaderSize = sizeof(rsa4096Asn1Header);
+                            break;
+                        default:
+                            UE_LOG(LogHttp, Log, TEXT("unsupported RSA key length %i for certificate %i for domain '%s'; skipping!"), [CertPubKeySize integerValue], i, *RemoteHost);
+                            continue;
+                    }
+                }
+                else if ([CertPubKeyType isEqualToString: (NSString *)kSecAttrKeyTypeECSECPrimeRandom])
+                {
+                    switch ([CertPubKeySize integerValue])
+                    {
+                        case 256:
+                            UE_LOG(LogHttp, VeryVerbose, TEXT("found 256 bit ECDSA pubkey"));
+                            CertPubKeyASN1Header = (char *)ecdsaSecp256r1Asn1Header;
+                            CertPubKeyASN1HeaderSize = sizeof(ecdsaSecp256r1Asn1Header);
+                            break;
+                        case 384:
+                            UE_LOG(LogHttp, VeryVerbose, TEXT("found 384 bit ECDSA pubkey"));
+                            CertPubKeyASN1Header = (char *)ecdsaSecp384r1Asn1Header;
+                            CertPubKeyASN1HeaderSize = sizeof(ecdsaSecp384r1Asn1Header);
+                            break;
+                        default:
+                            UE_LOG(LogHttp, Log, TEXT("unsupported ECDSA key length %i for certificate %i for domain '%s'; skipping!"), [CertPubKeySize integerValue], i, *RemoteHost);
+                            continue;
+                    }
+                }
+                else {
+                    UE_LOG(LogHttp, Log, TEXT("unsupported key type (not RSA or ECDSA) for certificate %i for domain '%s'; skipping!"), i, *RemoteHost);
+                    continue;
+                }
+				
+                UE_LOG(LogHttp, VeryVerbose, TEXT("constructed key header: [%d] %s"), CertPubKeyASN1HeaderSize, UTF8_TO_TCHAR([[[NSData dataWithBytes:CertPubKeyASN1Header length:CertPubKeyASN1HeaderSize] description] UTF8String]));
+                UE_LOG(LogHttp, VeryVerbose, TEXT("current pubkey: [%d] %s"), [(NSData*)CertPubKeyData length], UTF8_TO_TCHAR([[[NSData dataWithBytes:[(NSData*)CertPubKeyData bytes] length:[(NSData*)CertPubKeyData length]] description] UTF8String]));
+                
+                // smash 'em together to get a proper key with an ASN.1 header
+                NSMutableData *ReconstructedPubKey = [NSMutableData data];
+                [ReconstructedPubKey appendBytes:CertPubKeyASN1Header length:CertPubKeyASN1HeaderSize];
+                [ReconstructedPubKey appendData:CertPubKeyData];
+                UE_LOG(LogHttp, VeryVerbose, TEXT("reconstructed key: [%d] %s"), [ReconstructedPubKey length], UTF8_TO_TCHAR([[ReconstructedPubKey description] UTF8String]));
+                
+                TArray<uint8, TFixedAllocator<ISslCertificateManager::PUBLIC_KEY_DIGEST_SIZE>> CertCalcDigest;
+                CertCalcDigest.AddUninitialized(CC_SHA256_DIGEST_LENGTH);
+                if (!CC_SHA256([ReconstructedPubKey bytes], (CC_LONG)[ReconstructedPubKey length], CertCalcDigest.GetData()))
+                {
+                    UE_LOG(LogHttp, Warning, TEXT("could not calculate SHA256 digest of public key %d for domain '%s'; skipping!"), i, *RemoteHost);
+                }
+                else
+                {
+                    CertDigests.Add(CertCalcDigest);
+                    UE_LOG(LogHttp, Verbose, TEXT("added SHA256 digest to list for evaluation: domain: '%s' digest: [%d] %s"), *RemoteHost, CertCalcDigest.Num(), UTF8_TO_TCHAR([[[NSData dataWithBytes:CertCalcDigest.GetData() length:CertCalcDigest.Num()] description] UTF8String]));
+                }
+            }
+            
+            //finally, see if any of the pubkeys in the chain match any of our pinned pubkey hashes
+            if (CertDigests.Num() <= 0 || !FSslModule::Get().GetCertificateManager().VerifySslCertificates(CertDigests, RemoteHost))
+            {
+                // we could not validate any of the provided certs in chain with the pinned hashes for this host
+                // so we tell the sender to cancel (which cancels the pending connection)
+                UE_LOG(LogHttp, Error, TEXT("failed certificate pinning validation: no SPKI hashes in request matched pinned hashes for domain '%s' (was provided %d certificates in request)"), *RemoteHost, CertDigests.Num());
+                [challenge.sender cancelAuthenticationChallenge:challenge];
+                return;
+            }
+        }
+    }
+    else
+    {
+        UE_LOG(LogHttp, Error, TEXT("failed certificate pinning validation: SslCertificateManager is using non-SHA256 SPKI hashes [expected %d bytes, got %d bytes]"), CC_SHA256_DIGEST_LENGTH, ISslCertificateManager::PUBLIC_KEY_DIGEST_SIZE);
+        [challenge.sender cancelAuthenticationChallenge:challenge];
+        return;
+    }
+    
+    // if we got this far, pinning validation either succeeded or was disabled (or this was checking for client auth, etc.)
+    // so tell the connection to keep going with whatever else it was trying to validate
+    UE_LOG(LogHttp, Verbose, TEXT("certificate public key pinning either succeeded, is disabled, or challenge was not a server trust; continuing with auth"));
+    [challenge.sender performDefaultHandlingForAuthenticationChallenge: challenge];
+}
+#endif
+
+
+-(void) connectionDidFinishLoading:(NSURLConnection *)connection
+{
+	UE_LOG(LogHttp, Verbose, TEXT("connectionDidFinishLoading: %p"), self);
+	self.bIsReady = YES;
+}
+
+- (TArray<uint8>&)getPayload
+{
+	return Payload;
+}
+
+-(int32)getBytesWritten
+{
+	return self.BytesWritten;
+}
+
+@end
+
+
+
+
+/****************************************************************************
+ * FAppleHttpNSUrlConnectionResponse implementation
+ **************************************************************************/
+
+FAppleHttpNSUrlConnectionResponse::FAppleHttpNSUrlConnectionResponse(const FAppleHttpNSUrlConnectionRequest& InRequest)
+	: Request( InRequest )
+{
+	UE_LOG(LogHttp, Verbose, TEXT("FAppleHttpNSUrlConnectionResponse::FAppleHttpNSUrlConnectionResponse()"));
+	ResponseWrapper = [[FHttpResponseAppleNSUrlConnectionWrapper alloc] init];
+}
+
+
+FAppleHttpNSUrlConnectionResponse::~FAppleHttpNSUrlConnectionResponse()
+{
+	UE_LOG(LogHttp, Verbose, TEXT("FAppleHttpNSUrlConnectionResponse::~FAppleHttpNSUrlConnectionResponse()"));
+	[ResponseWrapper getPayload].Empty();
+
+	[ResponseWrapper release];
+	ResponseWrapper = nil;
+}
+
+
+FString FAppleHttpNSUrlConnectionResponse::GetURL() const
+{
+	UE_LOG(LogHttp, Verbose, TEXT("FAppleHttpNSUrlConnectionResponse::GetURL()"));
+	return FString(Request.Request.URL.query);
+}
+
+
+FString FAppleHttpNSUrlConnectionResponse::GetURLParameter(const FString& ParameterName) const
+{
+	SCOPED_AUTORELEASE_POOL;
+	UE_LOG(LogHttp, Verbose, TEXT("FAppleHttpNSUrlConnectionResponse::GetURLParameter()"));
+
+	NSString* ParameterNameStr = ParameterName.GetNSString();
+	NSArray* Parameters = [[[Request.Request URL] query] componentsSeparatedByString:@"&"];
+	for (NSString* Parameter in Parameters)
+	{
+		NSArray* KeyValue = [Parameter componentsSeparatedByString:@"="];
+		NSString* Key = [KeyValue objectAtIndex:0];
+		if ([Key compare:ParameterNameStr] == NSOrderedSame)
+		{
+			return FString([[KeyValue objectAtIndex:1] stringByRemovingPercentEncoding]);
+		}
+	}
+	return FString();
+}
+
+
+FString FAppleHttpNSUrlConnectionResponse::GetHeader(const FString& HeaderName) const
+{
+	SCOPED_AUTORELEASE_POOL;
+	UE_LOG(LogHttp, Verbose, TEXT("FAppleHttpNSUrlConnectionResponse::GetHeader()"));
+
+	NSString* ConvertedHeaderName = HeaderName.GetNSString();
+	return FString([[[ResponseWrapper Response] allHeaderFields] objectForKey:ConvertedHeaderName]);
+}
+
+
+TArray<FString> FAppleHttpNSUrlConnectionResponse::GetAllHeaders() const
+{
+	UE_LOG(LogHttp, Verbose, TEXT("FAppleHttpNSUrlConnectionResponse::GetAllHeaders()"));
+
+	NSDictionary* Headers = [GetResponseObj() allHeaderFields];
+	TArray<FString> Result;
+	Result.Reserve([Headers count]);
+	for (NSString* Key in [Headers allKeys])
+	{
+		FString ConvertedValue([Headers objectForKey:Key]);
+		FString ConvertedKey(Key);
+		Result.Add( FString::Printf( TEXT("%s: %s"), *ConvertedKey, *ConvertedValue ) );
+	}
+	return Result;
+}
+
+
+FString FAppleHttpNSUrlConnectionResponse::GetContentType() const
+{
+	UE_LOG(LogHttp, Verbose, TEXT("FAppleHttpNSUrlConnectionResponse::GetContentType()"));
+
+	return GetHeader( TEXT( "Content-Type" ) );
+}
+
+
+int32 FAppleHttpNSUrlConnectionResponse::GetContentLength() const
+{
+	UE_LOG(LogHttp, Verbose, TEXT("FAppleHttpNSUrlConnectionResponse::GetContentLength()"));
+
+	return ResponseWrapper.Response.expectedContentLength;
+}
+
+
+const TArray<uint8>& FAppleHttpNSUrlConnectionResponse::GetContent() const
+{
+	if( !IsReady() )
+	{
+		UE_LOG(LogHttp, Warning, TEXT("Payload is incomplete. Response still processing. %p"), &Request);
+	}
+	else
+	{
+		Payload = [ResponseWrapper getPayload];
+		UE_LOG(LogHttp, Verbose, TEXT("FAppleHttpNSUrlConnectionResponse::GetContent() - Num: %i"), [ResponseWrapper getPayload].Num());
+	}
+
+	return Payload;
+}
+
+
+FString FAppleHttpNSUrlConnectionResponse::GetContentAsString() const
+{
+	UE_LOG(LogHttp, Verbose, TEXT("FAppleHttpNSUrlConnectionResponse::GetContentAsString()"));
+
+	// Fill in our data.
+	GetContent();
+
+	TArray<uint8> ZeroTerminatedPayload;
+	ZeroTerminatedPayload.AddZeroed( Payload.Num() + 1 );
+	FMemory::Memcpy( ZeroTerminatedPayload.GetData(), Payload.GetData(), Payload.Num() );
+
+	return UTF8_TO_TCHAR( ZeroTerminatedPayload.GetData() );
+}
+
+
+int32 FAppleHttpNSUrlConnectionResponse::GetResponseCode() const
+{
+	UE_LOG(LogHttp, Verbose, TEXT("FAppleHttpNSUrlConnectionResponse::GetResponseCode()"));
+
+	return [GetResponseObj() statusCode];
+}
+
+
+NSHTTPURLResponse* FAppleHttpNSUrlConnectionResponse::GetResponseObj() const
+{
+	UE_LOG(LogHttp, Verbose, TEXT("FAppleHttpNSUrlConnectionResponse::GetResponseObj()"));
+
+	return [ResponseWrapper Response];
+}
+
+
+bool FAppleHttpNSUrlConnectionResponse::IsReady() const
+{
+	bool Ready = [ResponseWrapper bIsReady];
+
+	if( Ready )
+	{
+		UE_LOG(LogHttp, Verbose, TEXT("FAppleHttpNSUrlConnectionResponse::IsReady()"));
+	}
+
+	return Ready;
+}
+
+bool FAppleHttpNSUrlConnectionResponse::HadError() const
+{
+	bool bHadError = [ResponseWrapper bHadError];
+	
+	if( bHadError )
+	{
+		UE_LOG(LogHttp, Verbose, TEXT("FAppleHttpNSUrlConnectionResponse::HadError()"));
+	}
+	
+	return bHadError;
+}
+
+const int32 FAppleHttpNSUrlConnectionResponse::GetNumBytesReceived() const
+{
+	return [ResponseWrapper getPayload].Num();
+}
+
+const int32 FAppleHttpNSUrlConnectionResponse::GetNumBytesWritten() const
+{
+    int32 NumBytesWritten = [ResponseWrapper getBytesWritten];
+    return NumBytesWritten;
+}
diff --git a/Engine/Source/Runtime/Online/HTTP/Private/Apple/AppleHTTPNSUrlConnection.h b/Engine/Source/Runtime/Online/HTTP/Private/Apple/AppleHTTPNSUrlConnection.h
new file mode 100644
--- /dev/null
+++ b/Engine/Source/Runtime/Online/HTTP/Private/Apple/AppleHTTPNSUrlConnection.h
@@ -0,0 +1,238 @@
+// Copyright Epic Games, Inc. All Rights Reserved.
+
+#pragma once
+
+#include "CoreMinimal.h"
+#include "Containers/UnrealString.h"
+#include "GenericPlatform/HttpRequestImpl.h"
+#include "Interfaces/IHttpResponse.h"
+#include "HttpManager.h"
+#include "PlatformHttp.h"
+
+class FHttpManager;
+class FString;
+class IHttpRequest;
+
+/**
+ * Apple implementation of an Http request
+ */
+class FAppleHttpNSUrlConnectionRequest : public FHttpRequestImpl
+{
+public:
+	// implementation friends
+	friend class FAppleHttpNSUrlConnectionResponse;
+
+
+	//~ Begin IHttpBase Interface
+	virtual FString GetURL() const override;
+	virtual FString GetURLParameter(const FString& ParameterName) const override;
+	virtual FString GetHeader(const FString& HeaderName) const override;
+	virtual TArray<FString> GetAllHeaders() const override;	
+	virtual FString GetContentType() const override;
+	virtual int32 GetContentLength() const override;
+	virtual const TArray<uint8>& GetContent() const override;
+	//~ End IHttpBase Interface
+
+	//~ Begin IHttpRequest Interface
+	virtual FString GetVerb() const override;
+	virtual void SetVerb(const FString& Verb) override;
+	virtual void SetURL(const FString& URL) override;
+	virtual void SetContent(const TArray<uint8>& ContentPayload) override;
+	virtual void SetContent(TArray<uint8>&& ContentPayload) override;
+	virtual void SetContentAsString(const FString& ContentString) override;
+    virtual bool SetContentAsStreamedFile(const FString& Filename) override;
+	virtual bool SetContentFromStream(TSharedRef<FArchive, ESPMode::ThreadSafe> Stream) override;
+	virtual void SetHeader(const FString& HeaderName, const FString& HeaderValue) override;
+	virtual void AppendToHeader(const FString& HeaderName, const FString& AdditionalHeaderValue) override;
+	virtual void SetTimeout(float InTimeoutSecs) override;
+	virtual void ClearTimeout() override;
+	virtual TOptional<float> GetTimeout() const override;
+	virtual bool ProcessRequest() override;
+	virtual void CancelRequest() override;
+	virtual EHttpRequestStatus::Type GetStatus() const override;
+	virtual const FHttpResponsePtr GetResponse() const override;
+	virtual void Tick(float DeltaSeconds) override;
+	virtual float GetElapsedTime() const override;
+	//~ End IHttpRequest Interface
+
+	/**
+	 * Constructor
+	 */
+	FAppleHttpNSUrlConnectionRequest();
+
+	/**
+	 * Destructor. Clean up any connection/request handles
+	 */
+	virtual ~FAppleHttpNSUrlConnectionRequest();
+
+
+private:
+
+	/**
+	 * Create the session connection and initiate the web request
+	 *
+	 * @return true if the request was started
+	 */
+	bool StartRequest();
+
+	/**
+	 * Process state for a finished request that no longer needs to be ticked
+	 * Calls the completion delegate
+	 */
+	void FinishedRequest();
+
+	/**
+	 * Close session/request handles and unregister callbacks
+	 */
+	void CleanupRequest();
+
+
+private:
+	/** This is the NSMutableURLRequest, all our Apple functionality will deal with this. */
+	NSMutableURLRequest* Request;
+
+	/** This is the connection our request is sent along. */
+	NSURLConnection* Connection;
+
+	/** Flag whether the request payload source is a file */
+	bool bIsPayloadFile;
+
+	/** The request payload length in bytes. This must be tracked separately for a file stream */
+	int32 RequestPayloadByteLength;
+
+	/** The response object which we will use to pair with this request */
+	TSharedPtr<class FAppleHttpNSUrlConnectionResponse,ESPMode::ThreadSafe> Response;
+
+	/** BYTE array payload to use with the request. Typically for a POST */
+	mutable TArray<uint8> RequestPayload;
+
+	/** BYTE array for content which we now own */
+	TArray<uint8> ContentData;
+
+	/** Current status of request being processed */
+	EHttpRequestStatus::Type CompletionStatus;
+
+	/** Number of bytes sent to progress update */
+	int32 ProgressBytesSent;
+
+	/** Start of the request */
+	double StartRequestTime;
+
+	/** Time taken to complete/cancel the request. */
+	float ElapsedTime;
+};
+
+
+/**
+ * Apple Response Wrapper which will be used for it's delegates to receive responses.
+ */
+@interface FHttpResponseAppleNSUrlConnectionWrapper : NSObject
+{
+	/** Holds the payload as we receive it. */
+	TArray<uint8> Payload;
+}
+/** A handle for the response */
+@property(retain) NSHTTPURLResponse* Response;
+/** Flag whether the response is ready */
+@property BOOL bIsReady;
+/** When the response is complete, indicates whether the response was received without error. */
+@property BOOL bHadError;
+/** When the response is complete, indicates whether the response failed with an error specific to connecting to the host. */
+@property BOOL bIsHostConnectionFailure;
+/** The total number of bytes written out during the request/response */
+@property int32 BytesWritten;
+
+/** Delegate called when we send data. See Apple docs for when/how this should be used. */
+-(void) connection:(NSURLConnection *)connection didSendBodyData:(NSInteger)bytesWritten totalBytesWritten:(NSInteger)totalBytesWritten totalBytesExpectedToWrite:(NSInteger)totalBytesExpectedToWrite;
+/** Delegate called with we receive a response. See Apple docs for when/how this should be used. */
+-(void) connection:(NSURLConnection *)connection didReceiveResponse:(NSURLResponse *)response;
+/** Delegate called with we receive data. See Apple docs for when/how this should be used. */
+-(void) connection:(NSURLConnection *)connection didReceiveData:(NSData *)data;
+/** Delegate called with we complete with an error. See Apple docs for when/how this should be used. */
+-(void) connection:(NSURLConnection *)connection didFailWithError:(NSError *)error;
+/** Delegate called with we complete successfully. See Apple docs for when/how this should be used. */
+-(void) connectionDidFinishLoading:(NSURLConnection *)connection;
+
+#if WITH_SSL
+/** Delegate called when the connection is about to validate an auth challenge. We only care about server trust. See Apple docs for when/how this should be used. */
+-(void)connection:(NSURLConnection *)connection willSendRequestForAuthenticationChallenge: (NSURLAuthenticationChallenge *)challenge;
+#endif
+
+- (TArray<uint8>&)getPayload;
+- (int32)getBytesWritten;
+@end
+
+
+/**
+ * Apple implementation of an Http response
+ */
+class FAppleHttpNSUrlConnectionResponse : public IHttpResponse
+{
+private:
+	// This is the NSHTTPURLResponse, all our functionality will deal with.
+	FHttpResponseAppleNSUrlConnectionWrapper* ResponseWrapper;
+
+	/** Request that owns this response */
+	const FAppleHttpNSUrlConnectionRequest& Request;
+
+
+public:
+	// implementation friends
+	friend class FAppleHttpNSUrlConnectionRequest;
+
+
+	//~ Begin IHttpBase Interface
+	virtual FString GetURL() const override;
+	virtual FString GetURLParameter(const FString& ParameterName) const override;
+	virtual FString GetHeader(const FString& HeaderName) const override;
+	virtual TArray<FString> GetAllHeaders() const override;	
+	virtual FString GetContentType() const override;
+	virtual int32 GetContentLength() const override;
+	virtual const TArray<uint8>& GetContent() const override;
+	//~ End IHttpBase Interface
+
+	//~ Begin IHttpResponse Interface
+	virtual int32 GetResponseCode() const override;
+	virtual FString GetContentAsString() const override;
+	//~ End IHttpResponse Interface
+
+	NSHTTPURLResponse* GetResponseObj() const;
+
+	/**
+	 * Check whether a response is ready or not.
+	 */
+	bool IsReady() const;
+	
+	/**
+	 * Check whether a response had an error.
+	 */
+	bool HadError() const;
+
+	/**
+	 * Get the number of bytes received so far
+	 */
+	const int32 GetNumBytesReceived() const;
+
+	/**
+	* Get the number of bytes sent so far
+	*/
+	const int32 GetNumBytesWritten() const;
+
+	/**
+	 * Constructor
+	 *
+	 * @param InRequest - original request that created this response
+	 */
+	FAppleHttpNSUrlConnectionResponse(const FAppleHttpNSUrlConnectionRequest& InRequest);
+
+	/**
+	 * Destructor
+	 */
+	virtual ~FAppleHttpNSUrlConnectionResponse();
+
+
+private:
+
+	/** BYTE array to fill in as the response is read via didReceiveData */
+	mutable TArray<uint8> Payload;
+};
diff --git a/Engine/Source/Runtime/Online/HTTP/Private/Apple/AppleHTTPNSUrlSession.cpp b/Engine/Source/Runtime/Online/HTTP/Private/Apple/AppleHTTPNSUrlSession.cpp
new file mode 100644
--- /dev/null
+++ b/Engine/Source/Runtime/Online/HTTP/Private/Apple/AppleHTTPNSUrlSession.cpp
@@ -0,0 +1,815 @@
+// Copyright Epic Games, Inc. All Rights Reserved.
+
+
+#include "AppleHTTPNSUrlSession.h"
+#include "HttpManager.h"
+#include "Misc/EngineVersion.h"
+#include "Misc/App.h"
+#include "Misc/Base64.h"
+#include "HAL/PlatformTime.h"
+#include "Http.h"
+#include "HttpModule.h"
+
+/**
+ * State of a request's response
+ */
+enum class EAppleHttpRequestResponseState: uint8
+{
+	NotReady,
+	Success,
+	Error,
+	ConnectionError
+};
+
+/**
+ * Class to hold data from delegate implementation notifications.
+ */
+
+@interface FAppleHttpNSUrlSessionResponseDelegate : NSObject<NSURLSessionDataDelegate>
+{
+	/** Holds the payload as we receive it. */
+	@public TArray<uint8> Payload;
+}
+/** A handle for the response */
+@property(retain) NSHTTPURLResponse* Response;
+/** The total number of bytes written out during the request/response */
+@property int32 BytesWritten;
+/** The total number of bytes received out during the request/response */
+@property int32 BytesReceived;
+/** Response state */
+@property EAppleHttpRequestResponseState ResponseState;
+
+/** NSURLSessionDataDelegate delegate methods. Those are called from a thread controlled by the NSURLSession */
+
+/** Sent periodically to notify the delegate of upload progress. */
+- (void)URLSession:(NSURLSession *)session task:(NSURLSessionTask *)task didSendBodyData:(int64_t)bytesSent totalBytesSent:(int64_t)totalBytesSent totalBytesExpectedToSend:(int64_t)totalBytesExpectedToSend;
+/** The task has received a response and no further messages will be received until the completion block is called. */
+- (void)URLSession:(NSURLSession *)session dataTask:(NSURLSessionDataTask *)dataTask didReceiveResponse:(NSURLResponse *)response completionHandler:(void (^)(NSURLSessionResponseDisposition disposition))completionHandler;
+/** Sent when data is available for the delegate to consume. Data may be discontiguous */
+- (void)URLSession:(NSURLSession *)session dataTask:(NSURLSessionDataTask *)dataTask didReceiveData:(NSData *)data;
+/** Sent as the last message related to a specific task.  A nil Error implies that no error occurred and this task is complete. */
+- (void)URLSession:(NSURLSession *)session task:(NSURLSessionTask *)task didCompleteWithError:(nullable NSError *)error;
+
+@end
+
+@implementation FAppleHttpNSUrlSessionResponseDelegate
+@synthesize Response;
+@synthesize ResponseState;
+@synthesize BytesWritten;
+@synthesize BytesReceived;
+
+-(FAppleHttpNSUrlSessionResponseDelegate*) init
+{
+	self = [super init];
+	
+	BytesWritten = 0;
+	BytesReceived = 0;
+	ResponseState = EAppleHttpRequestResponseState::NotReady;
+	
+	return self;
+}
+
+- (void) dealloc
+{
+	[Response release];
+	[super dealloc];
+}
+
+- (void)URLSession:(NSURLSession *)session task:(NSURLSessionTask *)task didSendBodyData:(int64_t)bytesSent totalBytesSent:(int64_t)totalBytesSent totalBytesExpectedToSend:(int64_t)totalBytesExpectedToSend
+{
+	UE_LOG(LogHttp, Verbose, TEXT("URLSession:task:didSendBodyData:totalBytesSent:totalBytesExpectedToSend: totalBytesSent = %lld, totalBytesSent = %lld: %p"), totalBytesSent, totalBytesExpectedToSend, self);
+	self.BytesWritten = (int32)totalBytesSent;
+}
+
+- (void)URLSession:(NSURLSession *)session dataTask:(NSURLSessionDataTask *)dataTask didReceiveResponse:(NSURLResponse *)response completionHandler:(void (^)(NSURLSessionResponseDisposition disposition))completionHandler
+{
+	UE_LOG(LogHttp, Verbose, TEXT("URLSession:dataTask:didReceiveResponse:completionHandler"));
+	
+	self.Response = (NSHTTPURLResponse*)response;
+	int32 ExpectedResponseLength = (int32)response.expectedContentLength;
+	if(ExpectedResponseLength != NSURLResponseUnknownLength)
+	{
+		Payload.Empty(ExpectedResponseLength);
+	}
+	UE_LOG(LogHttp, Verbose, TEXT("URLSession:dataTask:didReceiveResponse:completionHandler: expectedContentLength = %lld. Length = %d: %p"), response.expectedContentLength, Payload.Max(), self);
+	completionHandler(NSURLSessionResponseAllow);
+}
+
+- (void)URLSession:(NSURLSession *)session dataTask:(NSURLSessionDataTask *)dataTask didReceiveData:(NSData *)data
+{
+	[data enumerateByteRangesUsingBlock:^(const void *bytes, NSRange byteRange, BOOL *stop) {
+		Payload.Append((const uint8*)bytes, byteRange.length);
+	}];
+	// Keep BytesReceived as a separated value to avoid concurrent accesses to Payload
+	self.BytesReceived = Payload.Num();
+	
+	UE_LOG(LogHttp, Verbose, TEXT("URLSession:dataTask:didReceiveData with %u bytes. After Append, Payload Length = %d: %p"), [data length], Payload.Num(), self);
+}
+
+- (void)URLSession:(NSURLSession *)session task:(NSURLSessionTask *)task didCompleteWithError:(nullable NSError *)error;
+{
+	if (error == nil)
+	{
+		UE_LOG(LogHttp, Verbose, TEXT("URLSession:task:didCompleteWithError. Http request succeeded: %p"), self);
+		self.ResponseState = EAppleHttpRequestResponseState::Success;
+	}
+	else
+	{
+		UE_LOG(LogHttp, Warning, TEXT("URLSession:task:didCompleteWithError. Http request failed - %s %s: %p"),
+			   *FString([error localizedDescription]),
+			   *FString([[error userInfo] objectForKey:NSURLErrorFailingURLStringErrorKey]),
+			   self);
+		
+		// Determine if the specific error was failing to connect to the host.
+		switch ([error code])
+		{
+			case NSURLErrorCannotFindHost:
+			case NSURLErrorCannotConnectToHost:
+			case NSURLErrorDNSLookupFailed:
+				self.ResponseState = EAppleHttpRequestResponseState::ConnectionError;
+				break;
+			default:
+				self.ResponseState = EAppleHttpRequestResponseState::Error;
+		}
+		// Log more details if verbose logging is enabled and this is an SSL error
+		if (UE_LOG_ACTIVE(LogHttp, Verbose))
+		{
+			SecTrustRef PeerTrustInfo = reinterpret_cast<SecTrustRef>([[error userInfo] objectForKey:NSURLErrorFailingURLPeerTrustErrorKey]);
+			if (PeerTrustInfo != nullptr)
+			{
+				SecTrustResultType TrustResult = kSecTrustResultInvalid;
+				SecTrustGetTrustResult(PeerTrustInfo, &TrustResult);
+				
+				FString TrustResultString;
+				switch (TrustResult)
+				{
+#define MAP_TO_RESULTSTRING(Constant) case Constant: TrustResultString = TEXT(#Constant); break;
+						MAP_TO_RESULTSTRING(kSecTrustResultInvalid)
+						MAP_TO_RESULTSTRING(kSecTrustResultProceed)
+						MAP_TO_RESULTSTRING(kSecTrustResultDeny)
+						MAP_TO_RESULTSTRING(kSecTrustResultUnspecified)
+						MAP_TO_RESULTSTRING(kSecTrustResultRecoverableTrustFailure)
+						MAP_TO_RESULTSTRING(kSecTrustResultFatalTrustFailure)
+						MAP_TO_RESULTSTRING(kSecTrustResultOtherError)
+#undef MAP_TO_RESULTSTRING
+					default:
+						TrustResultString = TEXT("unknown");
+						break;
+				}
+				UE_LOG(LogHttp, Verbose, TEXT("URLSession:task:didCompleteWithError. SSL trust result: %s (%d)"), *TrustResultString, TrustResult);
+			}
+		}
+	}
+}
+@end
+
+/****************************************************************************
+ * FAppleHttpNSUrlSessionRequest implementation
+ ***************************************************************************/
+
+FAppleHttpNSUrlSessionRequest::FAppleHttpNSUrlSessionRequest(NSURLSession* InSession)
+:   Session([InSession retain])
+,   Task(nil)
+,	bIsPayloadFile(false)
+,	ContentBytesLength(0)
+,	CompletionStatus(EHttpRequestStatus::NotStarted)
+,	StartRequestTime(0.0)
+,	ElapsedTime(0.0f)
+{
+	UE_LOG(LogHttp, Verbose, TEXT("FAppleHttpNSUrlSessionRequest::FAppleHttpNSUrlSessionRequest()"));
+	Request = [[NSMutableURLRequest alloc] init];
+	Request.timeoutInterval = FHttpModule::Get().GetHttpTimeout();
+
+	// Disable cache to mimic WinInet behavior
+	Request.cachePolicy = NSURLRequestReloadIgnoringLocalCacheData;
+
+	// Add default headers
+	const TMap<FString, FString>& DefaultHeaders = FHttpModule::Get().GetDefaultHeaders();
+	for (TMap<FString, FString>::TConstIterator It(DefaultHeaders); It; ++It)
+	{
+		SetHeader(It.Key(), It.Value());
+	}
+}
+
+FAppleHttpNSUrlSessionRequest::~FAppleHttpNSUrlSessionRequest()
+{
+	UE_LOG(LogHttp, Verbose, TEXT("FAppleHttpNSUrlSessionRequest::~FAppleHttpNSUrlSessionRequest()"));
+	check(Task == nil);
+	[Request release];
+    [Session release];
+}
+
+FString FAppleHttpNSUrlSessionRequest::GetURL() const
+{
+	SCOPED_AUTORELEASE_POOL;
+	NSURL* URL = Request.URL;
+	if (URL != nullptr)
+	{
+		FString ConvertedURL(URL.absoluteString);
+		UE_LOG(LogHttp, Verbose, TEXT("FAppleHttpNSUrlSessionRequest::GetURL() - %s"), *ConvertedURL);
+		return ConvertedURL;
+	}
+	else
+	{
+		UE_LOG(LogHttp, Verbose, TEXT("FAppleHttpNSUrlSessionRequest::GetURL() - NULL"));
+		return FString();
+	}
+}
+
+void FAppleHttpNSUrlSessionRequest::SetURL(const FString& URL)
+{
+	SCOPED_AUTORELEASE_POOL;
+	UE_LOG(LogHttp, Verbose, TEXT("FAppleHttpNSUrlSessionRequest::SetURL() - %s"), *URL);
+	Request.URL = [NSURL URLWithString: URL.GetNSString()];
+}
+
+FString FAppleHttpNSUrlSessionRequest::GetURLParameter(const FString& ParameterName) const
+{
+	SCOPED_AUTORELEASE_POOL;
+	UE_LOG(LogHttp, Verbose, TEXT("FAppleHttpNSUrlSessionRequest::GetURLParameter() - %s"), *ParameterName);
+
+	NSRange ParametersStart = [Request.URL.query rangeOfString:@"?"];
+	if (ParametersStart.location != NSNotFound && ParametersStart.length > 0)
+	{
+		NSString* ParametersStr = [Request.URL.query substringFromIndex:ParametersStart.location + 1];
+		NSString* ParameterNameStr = ParameterName.GetNSString();
+		NSArray* Parameters = [ParametersStr componentsSeparatedByString:@"&"];
+		for (NSString* Parameter in Parameters)
+		{
+			NSArray* KeyValue = [Parameter componentsSeparatedByString:@"="];
+			NSString* Key = KeyValue[0];
+			if ([Key compare:ParameterNameStr] == NSOrderedSame)
+			{
+				return FString(KeyValue[1]);
+			}
+		}
+	}
+
+	return FString();
+}
+
+FString FAppleHttpNSUrlSessionRequest::GetHeader(const FString& HeaderName) const
+{
+	SCOPED_AUTORELEASE_POOL;
+	FString Header([Request valueForHTTPHeaderField:HeaderName.GetNSString()]);
+	UE_LOG(LogHttp, Verbose, TEXT("FAppleHttpNSUrlSessionRequest::GetHeader() - %s"), *Header);
+	return Header;
+}
+
+
+void FAppleHttpNSUrlSessionRequest::SetHeader(const FString& HeaderName, const FString& HeaderValue)
+{
+	SCOPED_AUTORELEASE_POOL;
+	UE_LOG(LogHttp, Verbose, TEXT("FAppleHttpNSUrlSessionRequest::SetHeader() - %s / %s"), *HeaderName, *HeaderValue );
+	[Request setValue: HeaderValue.GetNSString() forHTTPHeaderField: HeaderName.GetNSString()];
+}
+
+void FAppleHttpNSUrlSessionRequest::AppendToHeader(const FString& HeaderName, const FString& AdditionalHeaderValue)
+{
+    if (!HeaderName.IsEmpty() && !AdditionalHeaderValue.IsEmpty())
+    {
+        NSDictionary* Headers = [Request allHTTPHeaderFields];
+        NSString* PreviousHeaderValuePtr = [Headers objectForKey: HeaderName.GetNSString()];
+        FString PreviousValue(PreviousHeaderValuePtr);
+		FString NewValue;
+		if (!PreviousValue.IsEmpty())
+		{
+			NewValue = PreviousValue + TEXT(", ");
+		}
+		NewValue += AdditionalHeaderValue;
+
+        SetHeader(HeaderName, NewValue);
+	}
+}
+
+TArray<FString> FAppleHttpNSUrlSessionRequest::GetAllHeaders() const
+{
+	SCOPED_AUTORELEASE_POOL;
+	UE_LOG(LogHttp, Verbose, TEXT("FAppleHttpNSUrlSessionRequest::GetAllHeaders()"));
+	NSDictionary* Headers = Request.allHTTPHeaderFields;
+	TArray<FString> Result;
+	Result.Reserve(Headers.count);
+	for (NSString* Key in Headers.allKeys)
+	{
+		FString ConvertedValue(Headers[Key]);
+		FString ConvertedKey(Key);
+		UE_LOG(LogHttp, Verbose, TEXT("Header= %s, Key= %s"), *ConvertedValue, *ConvertedKey);
+
+		Result.Add( FString::Printf( TEXT("%s: %s"), *ConvertedKey, *ConvertedValue ) );
+	}
+	return Result;
+}
+
+const TArray<uint8>& FAppleHttpNSUrlSessionRequest::GetContent() const
+{
+	UE_LOG(LogHttp, Verbose, TEXT("FAppleHttpNSUrlSessionRequest::GetContent()"));
+	StorageForGetContent.Empty();
+	if (bIsPayloadFile)
+	{
+		UE_LOG(LogHttp, Warning, TEXT("FAppleHttpNSUrlSessionRequest::GetContent() called on a request that is set up for streaming a file. Return value is an empty buffer"));
+	}
+	else
+	{
+		SCOPED_AUTORELEASE_POOL;
+		NSData* Body = Request.HTTPBody; // accessing HTTPBody will call retain autorelease on the value, increasing its retain count
+		StorageForGetContent.Append((const uint8*)Body.bytes, Body.length);
+	}
+	return StorageForGetContent;
+}
+
+void FAppleHttpNSUrlSessionRequest::SetContent(const TArray<uint8>& ContentPayload)
+{
+	if (CompletionStatus == EHttpRequestStatus::Processing)
+	{
+		UE_LOG(LogHttp, Warning, TEXT("FAppleHttpNSUrlSessionRequest::SetContent() - attempted to set content on a request that is inflight"));
+		return;
+	}
+
+	UE_LOG(LogHttp, Verbose, TEXT("FAppleHttpNSUrlSessionRequest::SetContent()"));
+	Request.HTTPBody = [NSData dataWithBytes:ContentPayload.GetData() length:ContentPayload.Num()];
+	ContentBytesLength = ContentPayload.Num();
+	bIsPayloadFile = false;
+}
+
+void FAppleHttpNSUrlSessionRequest::SetContent(TArray<uint8>&& ContentPayload)
+{
+	if (CompletionStatus == EHttpRequestStatus::Processing)
+	{
+		UE_LOG(LogHttp, Warning, TEXT("FAppleHttpNSUrlSessionRequest::SetContent() - attempted to set content on a request that is inflight"));
+		return;
+	}
+
+	UE_LOG(LogHttp, Verbose, TEXT("FAppleHttpNSUrlSessionRequest::SetContent()"));
+
+	// We cannot use NSData dataWithBytesNoCopy:length:freeWhenDone: and keep the data in this instance because we don't have control
+	// over the lifetime of the request copy that NSURLSessionTask keeps
+	Request.HTTPBody = [NSData dataWithBytes:ContentPayload.GetData() length:ContentPayload.Num()];
+	ContentBytesLength = ContentPayload.Num();
+	bIsPayloadFile = false;
+
+	// Clear argument content since client code probably expects that
+	ContentPayload.Empty();
+}
+
+FString FAppleHttpNSUrlSessionRequest::GetContentType() const
+{
+	FString ContentType = GetHeader(TEXT("Content-Type"));
+	UE_LOG(LogHttp, Verbose, TEXT("FAppleHttpNSUrlSessionRequest::GetContentType() - %s"), *ContentType);
+	return ContentType;
+}
+
+int32 FAppleHttpNSUrlSessionRequest::GetContentLength() const
+{
+	UE_LOG(LogHttp, Verbose, TEXT("FAppleHttpNSUrlSessionRequest::GetContentLength() - %i"), ContentBytesLength);
+	return ContentBytesLength;
+}
+
+void FAppleHttpNSUrlSessionRequest::SetContentAsString(const FString& ContentString)
+{
+	if (CompletionStatus == EHttpRequestStatus::Processing)
+	{
+		UE_LOG(LogHttp, Warning, TEXT("FAppleHttpNSUrlSessionRequest::SetContentAsString() - attempted to set content on a request that is inflight"));
+		return;
+	}
+
+	UE_LOG(LogHttp, Verbose, TEXT("FAppleHttpNSUrlSessionRequest::SetContentAsString() - %s"), *ContentString);
+	FTCHARToUTF8 Converter(*ContentString);
+
+	// The extra length computation here is unfortunate, but it's technically not safe to assume the length is the same.
+	Request.HTTPBody = [NSData dataWithBytes:(ANSICHAR*)Converter.Get() length:Converter.Length()];
+	ContentBytesLength = Converter.Length();
+	bIsPayloadFile = false;
+}
+
+bool FAppleHttpNSUrlSessionRequest::SetContentAsStreamedFile(const FString& Filename)
+{
+	SCOPED_AUTORELEASE_POOL;
+    UE_LOG(LogHttp, Verbose, TEXT("FAppleHttpNSUrlSessionRequest::SetContentAsStreamedFile() - %s"), *Filename);
+
+	if (CompletionStatus == EHttpRequestStatus::Processing)
+	{
+		UE_LOG(LogHttp, Warning, TEXT("FAppleHttpNSUrlSessionRequest::SetContentAsStreamedFile() - attempted to set content on a request that is inflight"));
+		return false;
+	}
+
+	NSString* PlatformFilename = Filename.GetNSString();
+
+	Request.HTTPBody = nil;
+
+	struct stat FileAttrs = { 0 };
+	if (stat(PlatformFilename.fileSystemRepresentation, &FileAttrs) == 0)
+	{
+		UE_LOG(LogHttp, VeryVerbose, TEXT("FAppleHttpNSUrlSessionRequest::SetContentAsStreamedFile succeeded in getting the file size - %lld"), FileAttrs.st_size);
+		// Under the hood, the Foundation framework unsets HTTPBody, and takes over as the stream delegate.
+		// The stream itself should be unopened when passed to setHTTPBodyStream.
+		Request.HTTPBodyStream = [NSInputStream inputStreamWithFileAtPath: PlatformFilename];
+		ContentBytesLength = FileAttrs.st_size;
+		bIsPayloadFile = true;
+	}
+	else
+	{
+		UE_LOG(LogHttp, VeryVerbose, TEXT("FAppleHttpNSUrlSessionRequest::SetContentAsStreamedFile failed to get file size"));
+		Request.HTTPBodyStream = nil;
+		ContentBytesLength = 0;
+		bIsPayloadFile = false;
+	}
+
+	return bIsPayloadFile;
+}
+
+bool FAppleHttpNSUrlSessionRequest::SetContentFromStream(TSharedRef<FArchive, ESPMode::ThreadSafe> Stream)
+{
+	UE_LOG(LogHttp, Warning, TEXT("FAppleHttpNSUrlSessionRequest::SetContentFromStream is not implemented"));
+	return false;
+}
+
+FString FAppleHttpNSUrlSessionRequest::GetVerb() const
+{
+	FString ConvertedVerb(Request.HTTPMethod);
+	UE_LOG(LogHttp, Verbose, TEXT("FAppleHttpNSUrlSessionRequest::GetVerb() - %s"), *ConvertedVerb);
+	return ConvertedVerb;
+}
+
+void FAppleHttpNSUrlSessionRequest::SetVerb(const FString& Verb)
+{
+	SCOPED_AUTORELEASE_POOL;
+	UE_LOG(LogHttp, Verbose, TEXT("FAppleHttpNSUrlSessionRequest::SetVerb() - %s"), *Verb);
+	Request.HTTPMethod = Verb.GetNSString();
+}
+
+void FAppleHttpNSUrlSessionRequest::SetTimeout(float InTimeoutSecs)
+{
+	Request.timeoutInterval = InTimeoutSecs;
+}
+
+void FAppleHttpNSUrlSessionRequest::ClearTimeout()
+{
+	Request.timeoutInterval = FHttpModule::Get().GetHttpTimeout();
+}
+
+TOptional<float> FAppleHttpNSUrlSessionRequest::GetTimeout() const
+{
+	return TOptional<float>(Request.timeoutInterval);
+}
+
+bool FAppleHttpNSUrlSessionRequest::ProcessRequest()
+{
+	SCOPED_AUTORELEASE_POOL;
+	UE_LOG(LogHttp, Verbose, TEXT("FAppleHttpNSUrlSessionRequest::ProcessRequest()"));
+	bool bStarted = false;
+
+	FString Scheme(Request.URL.scheme);
+	Scheme = Scheme.ToLower();
+
+	// Prevent overlapped requests using the same instance
+	if (CompletionStatus == EHttpRequestStatus::Processing)
+	{
+		UE_LOG(LogHttp, Warning, TEXT("ProcessRequest failed. Still processing last request."));
+	}
+	else if(GetURL().Len() == 0)
+	{
+		UE_LOG(LogHttp, Warning, TEXT("ProcessRequest failed. No URL was specified."));
+	}
+	else if( Scheme != TEXT("http") && Scheme != TEXT("https"))
+	{
+		UE_LOG(LogHttp, Warning, TEXT("ProcessRequest failed. URL '%s' is not a valid HTTP request. %p"), *GetURL(), this);
+	}
+	else if (!FHttpModule::Get().GetHttpManager().IsDomainAllowed(GetURL()))
+	{
+		UE_LOG(LogHttp, Warning, TEXT("ProcessRequest failed. URL '%s' is not using an allowed domain. %p"), *GetURL(), this);
+	}
+	else
+	{
+		bStarted = StartRequest();
+	}
+
+	if( !bStarted )
+	{
+		// Ensure we run on game thread
+		if (!IsInGameThread())
+		{
+			FHttpModule::Get().GetHttpManager().AddGameThreadTask([StrongThis = StaticCastSharedRef<FAppleHttpNSUrlSessionRequest>(AsShared())]()
+			{
+				StrongThis->FinishedRequest();
+			});
+		}
+		else
+		{
+			FinishedRequest();
+		}
+	}
+
+	return bStarted;
+}
+
+bool FAppleHttpNSUrlSessionRequest::StartRequest()
+{
+	SCOPED_AUTORELEASE_POOL;
+	UE_LOG(LogHttp, Verbose, TEXT("FAppleHttpNSUrlSessionRequest::StartRequest()"));
+	bool bStarted = false;
+
+	// set the content-length and user-agent
+	if(GetContentLength() > 0)
+	{
+		[Request setValue:[NSString stringWithFormat:@"%d", GetContentLength()] forHTTPHeaderField:@"Content-Length"];
+	}
+
+	const FString UserAgent = GetHeader("User-Agent");
+	if(UserAgent.IsEmpty())
+	{
+		NSString* Tag = FPlatformHttp::GetDefaultUserAgent().GetNSString();
+		[Request setValue:Tag forHTTPHeaderField:@"User-Agent"];
+	}
+
+	Task = [Session dataTaskWithRequest: Request];
+	
+	if (Task != nil)
+	{
+		bStarted = true;
+
+		CompletionStatus = EHttpRequestStatus::Processing;
+
+		Response = MakeShared<FAppleHttpNSUrlSessionResponse>(*this);
+
+		// Both Task and Response keep a strong reference to the delegate
+		Task.delegate = Response->ResponseDelegate;
+
+		[[Task retain] resume];
+		UE_LOG(LogHttp, Verbose, TEXT("[NSURLSessionTask resume]"));
+
+		FHttpModule::Get().GetHttpManager().AddRequest(SharedThis(this));
+	}
+	else
+	{
+		UE_LOG(LogHttp, Warning, TEXT("ProcessRequest failed. Could not initialize Internet connection."));
+		CompletionStatus = EHttpRequestStatus::Failed_ConnectionError;
+	}
+
+	StartRequestTime = FPlatformTime::Seconds();
+	// reset the elapsed time.
+	ElapsedTime = 0.0f;
+
+	return bStarted;
+}
+
+void FAppleHttpNSUrlSessionRequest::FinishedRequest()
+{
+	UE_LOG(LogHttp, Verbose, TEXT("FAppleHttpNSUrlSessionRequest::FinishedRequest()"));
+	ElapsedTime = (float)(FPlatformTime::Seconds() - StartRequestTime);
+	if( Response.IsValid() && Response->IsReady() && !Response->HadError())
+	{
+		UE_LOG(LogHttp, Verbose, TEXT("Request succeeded"));
+		CompletionStatus = EHttpRequestStatus::Succeeded;
+
+		// TODO: Try to broadcast OnHeaderReceived when we receive headers instead of here at the end
+		BroadcastResponseHeadersReceived();
+		OnProcessRequestComplete().ExecuteIfBound(SharedThis(this), Response, true);
+	}
+	else
+	{
+		UE_LOG(LogHttp, Verbose, TEXT("Request failed"));
+		FString URL([[Request URL] absoluteString]);
+		CompletionStatus = EHttpRequestStatus::Failed;
+		if (Response.IsValid() && Response->HadConnectionError())
+		{
+			CompletionStatus = EHttpRequestStatus::Failed_ConnectionError;
+		}
+
+		Response = nullptr;
+		OnProcessRequestComplete().ExecuteIfBound(SharedThis(this), nullptr, false);
+	}
+
+	// Clean up session/request handles that may have been created
+	CleanupRequest();
+
+	// Remove from global list since processing is now complete
+	if (FHttpModule::Get().GetHttpManager().IsValidRequest(this))
+	{
+		FHttpModule::Get().GetHttpManager().RemoveRequest(SharedThis(this));
+	}
+}
+
+void FAppleHttpNSUrlSessionRequest::CleanupRequest()
+{
+	UE_LOG(LogHttp, Verbose, TEXT("FAppleHttpNSUrlSessionRequest::CleanupRequest()"));
+	
+	if(CompletionStatus == EHttpRequestStatus::Processing)
+	{
+		CancelRequest();
+	}
+	
+	if(Task != nil)
+	{
+		[Task release];
+		Task = nil;
+	}
+}
+
+void FAppleHttpNSUrlSessionRequest::CancelRequest()
+{
+	UE_LOG(LogHttp, Verbose, TEXT("FAppleHttpNSUrlSessionRequest::CancelRequest()"));
+
+	if( Task != nil )
+	{
+		[Task cancel];
+	}
+
+	// Ensure we run on game thread
+	if (!IsInGameThread())
+	{
+		FHttpModule::Get().GetHttpManager().AddGameThreadTask([StrongThis = StaticCastSharedRef<FAppleHttpNSUrlSessionRequest>(AsShared())]()
+		{
+			StrongThis->FinishedRequest();
+		});
+	}
+	else
+	{
+		FinishedRequest();
+	}
+}
+
+EHttpRequestStatus::Type FAppleHttpNSUrlSessionRequest::GetStatus() const
+{
+	UE_LOG(LogHttp, Verbose, TEXT("FAppleHttpNSUrlSessionRequest::GetStatus()"));
+	return CompletionStatus;
+}
+
+const FHttpResponsePtr FAppleHttpNSUrlSessionRequest::GetResponse() const
+{
+	return Response;
+}
+
+void FAppleHttpNSUrlSessionRequest::Tick(float DeltaSeconds)
+{
+	if (Response.IsValid() && (CompletionStatus == EHttpRequestStatus::Processing || Response->HadError()))
+	{
+		if (OnRequestProgress().IsBound())
+		{
+			const int32 BytesWritten = Response->GetNumBytesWritten();
+			const int32 BytesRead = Response->GetNumBytesReceived();
+			if (BytesWritten > 0 || BytesRead > 0)
+			{
+				OnRequestProgress().ExecuteIfBound(SharedThis(this), BytesWritten, BytesRead);
+			}
+		}
+		if (Response->IsReady())
+		{
+			FinishedRequest();
+		}
+	}
+}
+
+float FAppleHttpNSUrlSessionRequest::GetElapsedTime() const
+{
+	return ElapsedTime;
+}
+
+/****************************************************************************
+ * FAppleHttpNSUrlSessionResponse implementation
+ **************************************************************************/
+
+FAppleHttpNSUrlSessionResponse::FAppleHttpNSUrlSessionResponse(const FAppleHttpNSUrlSessionRequest& InRequest)
+	: Request( InRequest )
+{
+	UE_LOG(LogHttp, Verbose, TEXT("FAppleHttpNSUrlSessionResponse::FAppleHttpNSUrlSessionResponse()"));
+	ResponseDelegate = [[FAppleHttpNSUrlSessionResponseDelegate alloc] init];
+}
+
+FAppleHttpNSUrlSessionResponse::~FAppleHttpNSUrlSessionResponse()
+{
+	UE_LOG(LogHttp, Verbose, TEXT("FAppleHttpNSUrlSessionResponse::~FAppleHttpNSUrlSessionResponse()"));
+	
+	[ResponseDelegate release];
+	ResponseDelegate = nil;
+}
+
+FString FAppleHttpNSUrlSessionResponse::GetURL() const
+{
+	UE_LOG(LogHttp, Verbose, TEXT("FAppleHttpNSUrlSessionResponse::GetURL()"));
+	return FString(Request.Request.URL.query);
+}
+
+FString FAppleHttpNSUrlSessionResponse::GetURLParameter(const FString& ParameterName) const
+{
+	SCOPED_AUTORELEASE_POOL;
+	UE_LOG(LogHttp, Verbose, TEXT("FAppleHttpNSUrlSessionResponse::GetURLParameter()"));
+
+	NSString* ParameterNameStr = ParameterName.GetNSString();
+	NSArray* Parameters = [[[Request.Request URL] query] componentsSeparatedByString:@"&"];
+	for (NSString* Parameter in Parameters)
+	{
+		NSArray* KeyValue = [Parameter componentsSeparatedByString:@"="];
+		NSString* Key = [KeyValue objectAtIndex:0];
+		if ([Key compare:ParameterNameStr] == NSOrderedSame)
+		{
+			return FString([[KeyValue objectAtIndex:1] stringByRemovingPercentEncoding]);
+		}
+	}
+	return FString();
+}
+
+FString FAppleHttpNSUrlSessionResponse::GetHeader(const FString& HeaderName) const
+{
+	SCOPED_AUTORELEASE_POOL;
+	UE_LOG(LogHttp, Verbose, TEXT("FAppleHttpNSUrlSessionResponse::GetHeader()"));
+	NSString* ConvertedHeaderName = HeaderName.GetNSString();
+	return FString([ResponseDelegate.Response.allHeaderFields objectForKey:ConvertedHeaderName]);
+}
+
+TArray<FString> FAppleHttpNSUrlSessionResponse::GetAllHeaders() const
+{
+	UE_LOG(LogHttp, Verbose, TEXT("FAppleHttpNSUrlSessionResponse::GetAllHeaders()"));
+
+	NSDictionary* Headers = ResponseDelegate.Response.allHeaderFields;
+	TArray<FString> Result;
+	Result.Reserve([Headers count]);
+	for (NSString* Key in [Headers allKeys])
+	{
+		FString ConvertedValue([Headers objectForKey:Key]);
+		FString ConvertedKey(Key);
+		Result.Add( FString::Printf( TEXT("%s: %s"), *ConvertedKey, *ConvertedValue ) );
+	}
+	return Result;
+}
+
+FString FAppleHttpNSUrlSessionResponse::GetContentType() const
+{
+	UE_LOG(LogHttp, Verbose, TEXT("FAppleHttpNSUrlSessionResponse::GetContentType()"));
+
+	return GetHeader( TEXT( "Content-Type" ) );
+}
+
+int32 FAppleHttpNSUrlSessionResponse::GetContentLength() const
+{
+	UE_LOG(LogHttp, Verbose, TEXT("FAppleHttpNSUrlSessionResponse::GetContentLength()"));
+	
+	return ResponseDelegate.Response.expectedContentLength;
+}
+
+const TArray<uint8>& FAppleHttpNSUrlSessionResponse::GetContent() const
+{
+	if( !IsReady() )
+	{
+		const static TArray<uint8> EmptyPayload;
+		UE_LOG(LogHttp, Warning, TEXT("Payload is incomplete. Response still processing. %p"), &Request);
+		return EmptyPayload;
+	}
+	else
+	{
+		UE_LOG(LogHttp, Verbose, TEXT("FAppleHttpNSUrlSessionResponse::GetContent() - Num: %i"), ResponseDelegate->Payload.Num());
+		return ResponseDelegate->Payload;
+	}
+}
+
+FString FAppleHttpNSUrlSessionResponse::GetContentAsString() const
+{
+	UE_LOG(LogHttp, Verbose, TEXT("FAppleHttpNSUrlSessionResponse::GetContentAsString()"));
+
+	// Fill in our data.
+	const TArray<uint8>& Payload = GetContent();
+
+	TArray<uint8> ZeroTerminatedPayload;
+	ZeroTerminatedPayload.AddZeroed( Payload.Num() + 1 );
+	FMemory::Memcpy( ZeroTerminatedPayload.GetData(), Payload.GetData(), Payload.Num() );
+
+	return UTF8_TO_TCHAR( ZeroTerminatedPayload.GetData() );
+}
+
+int32 FAppleHttpNSUrlSessionResponse::GetResponseCode() const
+{
+	UE_LOG(LogHttp, Verbose, TEXT("FAppleHttpNSUrlSessionResponse::GetResponseCode()"));
+
+	return ResponseDelegate.Response.statusCode;
+}
+
+bool FAppleHttpNSUrlSessionResponse::IsReady() const
+{
+	return (ResponseDelegate.ResponseState != EAppleHttpRequestResponseState::NotReady);
+}
+
+bool FAppleHttpNSUrlSessionResponse::HadError() const
+{
+	switch(ResponseDelegate.ResponseState)
+	{
+		case EAppleHttpRequestResponseState::Error:
+		case EAppleHttpRequestResponseState::ConnectionError:
+			return true;
+		default:
+			return false;
+	}
+}
+
+bool FAppleHttpNSUrlSessionResponse::HadConnectionError() const
+{
+	return (ResponseDelegate.ResponseState == EAppleHttpRequestResponseState::ConnectionError);
+}
+
+const int32 FAppleHttpNSUrlSessionResponse::GetNumBytesReceived() const
+{
+	return ResponseDelegate.BytesReceived;
+}
+
+const int32 FAppleHttpNSUrlSessionResponse::GetNumBytesWritten() const
+{
+	return ResponseDelegate.BytesWritten;
+}
diff --git a/Engine/Source/Runtime/Online/HTTP/Private/Apple/AppleHTTPNSUrlSession.h b/Engine/Source/Runtime/Online/HTTP/Private/Apple/AppleHTTPNSUrlSession.h
new file mode 100644
--- /dev/null
+++ b/Engine/Source/Runtime/Online/HTTP/Private/Apple/AppleHTTPNSUrlSession.h
@@ -0,0 +1,194 @@
+// Copyright Epic Games, Inc. All Rights Reserved.
+
+#pragma once
+
+#include "CoreMinimal.h"
+#include "GenericPlatform/HttpRequestImpl.h"
+#include "Interfaces/IHttpResponse.h"
+#include "PlatformHttp.h"
+
+/**
+ * Apple implementation of an Http request
+ */
+class FAppleHttpNSUrlSessionRequest : public FHttpRequestImpl
+{
+public:
+	// implementation friends
+	friend class FAppleHttpNSUrlSessionResponse;
+
+
+	//~ Begin IHttpBase Interface
+	virtual FString GetURL() const override;
+	virtual FString GetURLParameter(const FString& ParameterName) const override;
+	virtual FString GetHeader(const FString& HeaderName) const override;
+	virtual TArray<FString> GetAllHeaders() const override;	
+	virtual FString GetContentType() const override;
+	virtual int32 GetContentLength() const override;
+	virtual const TArray<uint8>& GetContent() const override;
+	//~ End IHttpBase Interface
+
+	//~ Begin IHttpRequest Interface
+	virtual FString GetVerb() const override;
+	virtual void SetVerb(const FString& Verb) override;
+	virtual void SetURL(const FString& URL) override;
+	virtual void SetContent(const TArray<uint8>& ContentPayload) override;
+	virtual void SetContent(TArray<uint8>&& ContentPayload) override;
+	virtual void SetContentAsString(const FString& ContentString) override;
+    virtual bool SetContentAsStreamedFile(const FString& Filename) override;
+	virtual bool SetContentFromStream(TSharedRef<FArchive, ESPMode::ThreadSafe> Stream) override;
+	virtual void SetHeader(const FString& HeaderName, const FString& HeaderValue) override;
+	virtual void AppendToHeader(const FString& HeaderName, const FString& AdditionalHeaderValue) override;
+	virtual void SetTimeout(float InTimeoutSecs) override;
+	virtual void ClearTimeout() override;
+	virtual TOptional<float> GetTimeout() const override;
+	virtual bool ProcessRequest() override;
+	virtual void CancelRequest() override;
+	virtual EHttpRequestStatus::Type GetStatus() const override;
+	virtual const FHttpResponsePtr GetResponse() const override;
+	virtual void Tick(float DeltaSeconds) override;
+	virtual float GetElapsedTime() const override;
+	//~ End IHttpRequest Interface
+
+	/**
+	 * Constructor
+	 *
+	 * @param InSession - NSURLSession session used to create NSURLSessionTask to retrieve the response
+	 */
+	explicit FAppleHttpNSUrlSessionRequest(NSURLSession* InSession);
+
+	/**
+	 * Destructor. Clean up any connection/request handles
+	 */
+	virtual ~FAppleHttpNSUrlSessionRequest();
+
+
+private:
+
+	/**
+	 * Create the session connection and initiate the web request
+	 *
+	 * @return true if the request was started
+	 */
+	bool StartRequest();
+
+	/**
+	 * Process state for a finished request that no longer needs to be ticked
+	 * Calls the completion delegate
+	 */
+	void FinishedRequest();
+
+	/**
+	 * Close session/request handles and unregister callbacks
+	 */
+	void CleanupRequest();
+
+
+private:
+	/** This is the NSMutableURLRequest, all our Apple functionality will deal with this. */
+	NSMutableURLRequest* Request;
+
+    /** This is the session our request belongs to */
+    NSURLSession* Session;
+	
+	/** This is the Task associated to the sessionin charge of our request */
+	NSURLSessionTask* Task;
+    
+	/** Flag whether the request payload source is a file */
+	bool bIsPayloadFile;
+
+	/** The request payload length in bytes. This must be tracked separately for a file stream */
+	int32 ContentBytesLength;
+
+	/** The response object which we will use to pair with this request */
+	TSharedPtr<class FAppleHttpNSUrlSessionResponse,ESPMode::ThreadSafe> Response;
+
+	/** Array used to retrieve back content set on the ObjC request when calling GetContent*/
+	mutable TArray<uint8> StorageForGetContent;
+
+	/** Current status of request being processed */
+	EHttpRequestStatus::Type CompletionStatus;
+
+	/** Start of the request */
+	double StartRequestTime;
+
+	/** Time taken to complete/cancel the request. */
+	float ElapsedTime;
+};
+
+@class FAppleHttpNSUrlSessionResponseDelegate;
+
+/**
+ * Apple implementation of an Http response
+ */
+class FAppleHttpNSUrlSessionResponse : public IHttpResponse
+{
+private:
+	// Delegate implementation. Keeps the response state and data
+	FAppleHttpNSUrlSessionResponseDelegate* ResponseDelegate;
+
+	/** Request that owns this response */
+	const FAppleHttpNSUrlSessionRequest& Request;
+
+
+public:
+	// implementation friends
+	friend class FAppleHttpNSUrlSessionRequest;
+
+
+	//~ Begin IHttpBase Interface
+	virtual FString GetURL() const override;
+	virtual FString GetURLParameter(const FString& ParameterName) const override;
+	virtual FString GetHeader(const FString& HeaderName) const override;
+	virtual TArray<FString> GetAllHeaders() const override;	
+	virtual FString GetContentType() const override;
+	virtual int32 GetContentLength() const override;
+	virtual const TArray<uint8>& GetContent() const override;
+	//~ End IHttpBase Interface
+
+	//~ Begin IHttpResponse Interface
+	virtual int32 GetResponseCode() const override;
+	virtual FString GetContentAsString() const override;
+	//~ End IHttpResponse Interface
+
+	/**
+	 * Check whether headers are available.
+	 */
+	bool AreHeadersAvailable() const;
+
+	/**
+	 * Check whether a response is ready or not.
+	 */
+	bool IsReady() const;
+	
+	/**
+	 * Check whether a response had an error.
+	 */
+	bool HadError() const;
+
+	/**
+	 * Check whether a response had a connection error.
+	 */
+	bool HadConnectionError() const;
+
+	/**
+	 * Get the number of bytes received so far
+	 */
+	const int32 GetNumBytesReceived() const;
+
+	/**
+	* Get the number of bytes sent so far
+	*/
+	const int32 GetNumBytesWritten() const;
+
+	/**
+	 * Constructor
+	 *
+	 * @param InRequest - original request that created this response
+	 */
+	FAppleHttpNSUrlSessionResponse(const FAppleHttpNSUrlSessionRequest& InRequest);
+
+	/**
+	 * Destructor
+	 */
+	virtual ~FAppleHttpNSUrlSessionResponse();
+};
diff --git a/Engine/Source/Runtime/Online/HTTP/Private/Apple/ApplePlatformHttp.cpp b/Engine/Source/Runtime/Online/HTTP/Private/Apple/ApplePlatformHttp.cpp
--- a/Engine/Source/Runtime/Online/HTTP/Private/Apple/ApplePlatformHttp.cpp
+++ b/Engine/Source/Runtime/Online/HTTP/Private/Apple/ApplePlatformHttp.cpp
@@ -1,10 +1,217 @@
 // Copyright Epic Games, Inc. All Rights Reserved.
 
 #include "Apple/ApplePlatformHttp.h"
-#include "AppleHTTP.h"
+#include "Http.h"
+#include "AppleHTTPNSUrlConnection.h"
+#include "AppleHTTPNSUrlSession.h"
+#include "Apple/CFRef.h"
+#include "Misc/Base64.h"
+#include "Misc/CommandLine.h"
+
+#include "CommonCrypto/CommonDigest.h"
 
 #if WITH_SSL
 #include "Ssl.h"
+
+static FString AppleGetCertificateSummary(SecCertificateRef Cert)
+{
+	TCFRef<CFStringRef> Summary = SecCertificateCopySubjectSummary(Cert);
+	return FString(Summary);
+}
+
+static bool ValidateCertificatePublicKeyPinningApple(NSURLAuthenticationChallenge *challenge)
+{
+	// CC gives the actual key, but strips the ASN.1 header... which means
+	// we can't calulate a proper SPKI hash without reconstructing it. sigh.
+	static const unsigned char rsa2048Asn1Header[] =
+	{
+		0x30, 0x82, 0x01, 0x22, 0x30, 0x0d, 0x06, 0x09,
+		0x2a, 0x86, 0x48, 0x86, 0xf7, 0x0d, 0x01, 0x01,
+		0x01, 0x05, 0x00, 0x03, 0x82, 0x01, 0x0f, 0x00
+	};
+	static const unsigned char rsa4096Asn1Header[] =
+	{
+		0x30, 0x82, 0x02, 0x22, 0x30, 0x0d, 0x06, 0x09,
+		0x2a, 0x86, 0x48, 0x86, 0xf7, 0x0d, 0x01, 0x01,
+		0x01, 0x05, 0x00, 0x03, 0x82, 0x02, 0x0f, 0x00
+	};
+	static const unsigned char ecdsaSecp256r1Asn1Header[] =
+	{
+		0x30, 0x59, 0x30, 0x13, 0x06, 0x07, 0x2a, 0x86,
+		0x48, 0xce, 0x3d, 0x02, 0x01, 0x06, 0x08, 0x2a,
+		0x86, 0x48, 0xce, 0x3d, 0x03, 0x01, 0x07, 0x03,
+		0x42, 0x00
+	};
+	static const unsigned char ecdsaSecp384r1Asn1Header[] =
+	{
+		0x30, 0x76, 0x30, 0x10, 0x06, 0x07, 0x2a, 0x86,
+		0x48, 0xce, 0x3d, 0x02, 0x01, 0x06, 0x05, 0x2b,
+		0x81, 0x04, 0x00, 0x22, 0x03, 0x62, 0x00
+	};
+
+	SecTrustRef RemoteTrust = challenge.protectionSpace.serverTrust;
+	FString RemoteHost = FString(UTF8_TO_TCHAR([challenge.protectionSpace.host UTF8String]));
+	if ((RemoteTrust == NULL) || (RemoteHost.IsEmpty()))
+	{
+		UE_LOG(LogHttp, Error, TEXT("failed certificate pinning validation: could not parse parameters during certificate pinning evaluation"));
+		return false;
+	}
+
+	if (!SecTrustEvaluateWithError(RemoteTrust, nil))
+	{
+		UE_LOG(LogHttp, Error, TEXT("failed certificate pinning validation: default certificate trust evaluation failed for domain '%s'"), *RemoteHost);
+		return false;
+	}
+	// look at all certs in the remote chain and calculate the SHA256 hash of their DER-encoded SPKI
+	// the chain starts with the server's cert itself, so walk backwards to optimize for roots first
+	TArray<TArray<uint8, TFixedAllocator<ISslCertificateManager::PUBLIC_KEY_DIGEST_SIZE>>> CertDigests;
+	
+	TCFRef<CFArrayRef> Certificates = SecTrustCopyCertificateChain(RemoteTrust);
+	if (Certificates == nil)
+	{
+		UE_LOG(LogHttp, Error, TEXT("No certificate could be copied in the certificate chain used to evaluate trust."));
+		return false;
+	}
+	
+	CFIndex CertificateCount = CFArrayGetCount(Certificates);
+	for (int i = 0; i < CertificateCount; ++i)
+	{
+		SecCertificateRef Cert = (SecCertificateRef)CFArrayGetValueAtIndex(Certificates, i);
+
+		TCFRef<CFStringRef> Summary = SecCertificateCopySubjectSummary(Cert);
+		UE_LOG(LogHttp, VeryVerbose, TEXT("inspecting certificate. Summary: %s "), *AppleGetCertificateSummary(Cert));
+
+		TCFRef<SecKeyRef> CertPubKey = SecCertificateCopyKey(Cert);
+
+		TCFRef<CFDataRef> CertPubKeyData = SecKeyCopyExternalRepresentation(CertPubKey, NULL);
+		if (!CertPubKeyData)
+		{
+			UE_LOG(LogHttp, Warning, TEXT("could not extract public key from certificate %i for domain '%s'; skipping!"), i, *RemoteHost);
+			continue;
+		}
+		
+		// we got the key. now we have to figure out what type of key it is; thanks, CommonCrypto.
+		TCFRef<CFDictionaryRef> CertPubKeyAttr = SecKeyCopyAttributes(CertPubKey);
+		NSString *CertPubKeyType = static_cast<NSString *>(CFDictionaryGetValue(CertPubKeyAttr, kSecAttrKeyType));
+		NSNumber *CertPubKeySize = static_cast<NSNumber *>(CFDictionaryGetValue(CertPubKeyAttr, kSecAttrKeySizeInBits));
+		char *CertPubKeyASN1Header;
+		uint8_t CertPubKeyASN1HeaderSize = 0;
+		if ([CertPubKeyType isEqualToString: (NSString *)kSecAttrKeyTypeRSA])
+		{
+			switch ([CertPubKeySize integerValue])
+			{
+				case 2048:
+					UE_LOG(LogHttp, VeryVerbose, TEXT("found 2048 bit RSA pubkey"));
+					CertPubKeyASN1Header = (char *)rsa2048Asn1Header;
+					CertPubKeyASN1HeaderSize = sizeof(rsa2048Asn1Header);
+					break;
+				case 4096:
+					UE_LOG(LogHttp, VeryVerbose, TEXT("found 4096 bit RSA pubkey"));
+					CertPubKeyASN1Header = (char *)rsa4096Asn1Header;
+					CertPubKeyASN1HeaderSize = sizeof(rsa4096Asn1Header);
+					break;
+				default:
+					UE_LOG(LogHttp, Log, TEXT("unsupported RSA key length %i for certificate %i for domain '%s'; skipping!"), [CertPubKeySize integerValue], i, *RemoteHost);
+					continue;
+			}
+		}
+		else if ([CertPubKeyType isEqualToString: (NSString *)kSecAttrKeyTypeECSECPrimeRandom])
+		{
+			switch ([CertPubKeySize integerValue])
+			{
+				case 256:
+					UE_LOG(LogHttp, VeryVerbose, TEXT("found 256 bit ECDSA pubkey"));
+					CertPubKeyASN1Header = (char *)ecdsaSecp256r1Asn1Header;
+					CertPubKeyASN1HeaderSize = sizeof(ecdsaSecp256r1Asn1Header);
+					break;
+				case 384:
+					UE_LOG(LogHttp, VeryVerbose, TEXT("found 384 bit ECDSA pubkey"));
+					CertPubKeyASN1Header = (char *)ecdsaSecp384r1Asn1Header;
+					CertPubKeyASN1HeaderSize = sizeof(ecdsaSecp384r1Asn1Header);
+					break;
+				default:
+					UE_LOG(LogHttp, Log, TEXT("unsupported ECDSA key length %i for certificate %i for domain '%s'; skipping!"), [CertPubKeySize integerValue], i, *RemoteHost);
+					continue;
+			}
+		}
+		else {
+			UE_LOG(LogHttp, Log, TEXT("unsupported key type (not RSA or ECDSA) for certificate %i for domain '%s'; skipping!"), i, *RemoteHost);
+			continue;
+		}
+		
+		UE_LOG(LogHttp, VeryVerbose, TEXT("constructed key header: [%d] %s"), CertPubKeyASN1HeaderSize, *FString::FromHexBlob((const uint8*)CertPubKeyASN1Header, CertPubKeyASN1HeaderSize));
+		UE_LOG(LogHttp, VeryVerbose, TEXT("current pubkey: [%d] %s"), [(NSData*)CertPubKeyData length], *FString::FromHexBlob((const uint8*)[(NSData*)CertPubKeyData bytes], [(NSData*)CertPubKeyData length]));
+		
+		// smash 'em together to get a proper key with an ASN.1 header
+		NSMutableData *ReconstructedPubKey = [NSMutableData data];
+		[ReconstructedPubKey appendBytes:CertPubKeyASN1Header length:CertPubKeyASN1HeaderSize];
+		[ReconstructedPubKey appendData:CertPubKeyData];
+		UE_LOG(LogHttp, VeryVerbose, TEXT("reconstructed key: [%d] %s"), [ReconstructedPubKey length], *FString::FromHexBlob((const uint8*)[ReconstructedPubKey bytes], [ReconstructedPubKey length]));
+		
+		TArray<uint8, TFixedAllocator<ISslCertificateManager::PUBLIC_KEY_DIGEST_SIZE>> CertCalcDigest;
+		CertCalcDigest.AddUninitialized(CC_SHA256_DIGEST_LENGTH);
+		if (!CC_SHA256([ReconstructedPubKey bytes], (CC_LONG)[ReconstructedPubKey length], CertCalcDigest.GetData()))
+		{
+			UE_LOG(LogHttp, Warning, TEXT("could not calculate SHA256 digest of public key %d for domain '%s'; skipping!"), i, *RemoteHost);
+		}
+		else
+		{
+			CertDigests.Add(CertCalcDigest);
+			UE_LOG(LogHttp, Verbose, TEXT("added SHA256 digest to list for evaluation: domain: '%s' Base64 digest: %s"), *RemoteHost, *FBase64::Encode(CertCalcDigest.GetData(), CertCalcDigest.Num()));
+			UE_LOG(LogHttp, VeryVerbose, TEXT("Certificate digest binary content: [%d] '%s'"), CertCalcDigest.Num(), *FString::FromHexBlob(CertCalcDigest.GetData(), CertCalcDigest.Num()));
+		}
+	}
+	
+	//finally, see if any of the pubkeys in the chain match any of our pinned pubkey hashes
+	if (CertDigests.IsEmpty() || !FSslModule::Get().GetCertificateManager().VerifySslCertificates(CertDigests, RemoteHost))
+	{
+		// we could not validate any of the provided certs in chain with the pinned hashes for this host
+		UE_LOG(LogHttp, Error, TEXT("failed certificate pinning validation: no SPKI hashes in request matched pinned hashes for domain '%s' (was provided %d certificates in request)"), *RemoteHost, CertDigests.Num());
+		return false;
+	}
+	else
+	{
+		UE_LOG(LogHttp, Verbose, TEXT("certificate public key pinning either succeeded or is disabled; continuing with auth"));
+		return true;
+	}
+}
+
+/** NSURLSessionDelegate implementation in charge of validating certificate pinning */
+@interface FApplePlatformHttpSessionDelegate: NSObject<NSURLSessionDelegate>
+- (void)URLSession:(NSURLSession *)session didReceiveChallenge:(NSURLAuthenticationChallenge *)challenge completionHandler:(void (^)(NSURLSessionAuthChallengeDisposition disposition, NSURLCredential *credential))completionHandler;
+@end
+
+@implementation FApplePlatformHttpSessionDelegate
+
+- (void)URLSession:(NSURLSession *)session didReceiveChallenge:(NSURLAuthenticationChallenge *)challenge completionHandler:(void (^)(NSURLSessionAuthChallengeDisposition disposition, NSURLCredential *credential))completionHandler
+{
+	if (!ensure(ISslCertificateManager::PUBLIC_KEY_DIGEST_SIZE == CC_SHA256_DIGEST_LENGTH))
+	{
+		UE_LOG(LogHttp, Error, TEXT("failed certificate pinning validation: SslCertificateManager is using non-SHA256 SPKI hashes [expected %d bytes, got %d bytes]"), CC_SHA256_DIGEST_LENGTH, ISslCertificateManager::PUBLIC_KEY_DIGEST_SIZE);
+		completionHandler(NSURLSessionAuthChallengeCancelAuthenticationChallenge, nil);
+		return;
+	}
+	
+	// we only care about challenges to the received certificate chain
+	if ([challenge.protectionSpace.authenticationMethod isEqualToString:NSURLAuthenticationMethodServerTrust] == NO)
+	{
+		UE_LOG(LogHttp, Verbose, TEXT("challenge was not a server trust; continuing with auth"));
+		completionHandler(NSURLSessionAuthChallengePerformDefaultHandling, nil);
+		return;
+	}
+	
+	if (ValidateCertificatePublicKeyPinningApple(challenge))
+	{
+		completionHandler(NSURLSessionAuthChallengeUseCredential, [NSURLCredential credentialForTrust: challenge.protectionSpace.serverTrust]);
+	}
+	else
+	{
+		completionHandler(NSURLSessionAuthChallengeCancelAuthenticationChallenge, nil);
+	}
+}
+
+@end
+
 #endif
 
 void FApplePlatformHttp::Init()
@@ -14,15 +221,62 @@ void FApplePlatformHttp::Init()
 	// and doesn't unload SSL before unloading HTTP module at exit
 	FSslModule::Get();
 #endif
-}
 
+	bool bUseNSUrlConnection = FParse::Param(FCommandLine::Get(), TEXT("UseNSUrlConnection"));
+	bUseNSUrlSession = !bUseNSUrlConnection;
+
+	if (bUseNSUrlSession)
+	{
+		InitWithNSUrlSession();
+	}
+}
 
 void FApplePlatformHttp::Shutdown()
 {
+	if (bUseNSUrlSession)
+	{
+		ShutdownWithNSUrlSession();
+	}
 }
 
+void FApplePlatformHttp::InitWithNSUrlSession()
+{
+	NSURLSessionConfiguration* Config = [NSURLSessionConfiguration defaultSessionConfiguration];
+
+	// Disable cache to mimic WinInet behavior
+	Config.requestCachePolicy = NSURLRequestReloadIgnoringLocalCacheData;
+	
+#if WITH_SSL
+	// Load SSL module during HTTP module's StatupModule() to make sure module manager figures out the dependencies correctly
+	// and doesn't unload SSL before unloading HTTP module at exit
+	FSslModule::Get();
+
+	FApplePlatformHttpSessionDelegate* Delegate = [[FApplePlatformHttpSessionDelegate alloc] init];
+	
+	Session = [NSURLSession sessionWithConfiguration: Config delegate: Delegate delegateQueue: nil];
+#else
+	Session = [NSURLSession sessionWithConfiguration: Config];
+#endif
+	
+    [Session retain];
+
+}
+
+void FApplePlatformHttp::ShutdownWithNSUrlSession()
+{
+	[Session finishTasksAndInvalidate];
+	[Session release];
+	Session = nil;
+}
 
 IHttpRequest* FApplePlatformHttp::ConstructRequest()
 {
-	return new FAppleHttpRequest();
+	if(bUseNSUrlSession)
+	{
+		return new FAppleHttpNSUrlSessionRequest(Session);
+	}
+	else
+	{
+		return new FAppleHttpNSUrlConnectionRequest();
+	}
 }
diff --git a/Engine/Source/Runtime/Online/HTTP/Public/Apple/ApplePlatformHttp.h b/Engine/Source/Runtime/Online/HTTP/Public/Apple/ApplePlatformHttp.h
--- a/Engine/Source/Runtime/Online/HTTP/Public/Apple/ApplePlatformHttp.h
+++ b/Engine/Source/Runtime/Online/HTTP/Public/Apple/ApplePlatformHttp.h
@@ -36,7 +36,17 @@ public:
 	 * @return request object
 	 */
 	static IHttpRequest* ConstructRequest();
+
+private:
+	/** Flag to allow fall back to use NSUrlConnection instead of NSUrlSession. Assigned from commandline */
+	static inline bool bUseNSUrlSession = false;
+
+    /** Session used to create Apple based requests */
+    static inline NSURLSession* Session = nil;
+
+	static void InitWithNSUrlSession();
+	static void ShutdownWithNSUrlSession();
 };
 
 
-typedef FApplePlatformHttp FPlatformHttp;
\ No newline at end of file
+typedef FApplePlatformHttp FPlatformHttp;
diff --git a/Engine/Source/Runtime/Online/SSL/Private/SslCertificateManager.cpp b/Engine/Source/Runtime/Online/SSL/Private/SslCertificateManager.cpp
--- a/Engine/Source/Runtime/Online/SSL/Private/SslCertificateManager.cpp
+++ b/Engine/Source/Runtime/Online/SSL/Private/SslCertificateManager.cpp
@@ -299,9 +299,11 @@ bool FSslCertificateManager::VerifySslCertificates(TArray<TArray<uint8, TFixedAl
 		return true;
 	}
 	bool bFoundMatch = false;
-	for (int32 CertIndex = 0; CertIndex < Digests.Num(); ++CertIndex)
+	for (const TArray<uint8, TFixedAllocator<PUBLIC_KEY_DIGEST_SIZE>>& CurrentDigest: Digests)
 	{
-		if (PinnedKeys->Contains(Digests[CertIndex]))
+		UE_LOG(LogSsl, VeryVerbose, TEXT("checking digest. Base64: '%s'"), *FBase64::Encode(CurrentDigest.GetData(), CurrentDigest.Num()));
+		
+		if (PinnedKeys->Contains(CurrentDigest))
 		{
 			UE_LOG(LogSsl, Verbose, TEXT("found public key digest in request that matches a pinned key for '%s'"), *Domain);
 			bFoundMatch = true;
