LwsWebSocket ping pong interval from ini.

diff --git a/Engine/Config/BaseEngine.ini b/Engine/Config/BaseEngine.ini
--- a/Engine/Config/BaseEngine.ini
+++ b/Engine/Config/BaseEngine.ini
@@ -79,6 +79,7 @@ ThreadStackSize=131072
 ThreadTargetFrameTimeInSeconds=0.0333
 ThreadMinimumSleepTimeInSeconds=0.0
 MaxHttpHeaderData=32768
+PingPongInterval=0
 
 [Ping]
 StackSize=1048576
diff --git a/Engine/Source/Runtime/Online/WebSockets/Private/Lws/LwsWebSocketsManager.cpp b/Engine/Source/Runtime/Online/WebSockets/Private/Lws/LwsWebSocketsManager.cpp
--- a/Engine/Source/Runtime/Online/WebSockets/Private/Lws/LwsWebSocketsManager.cpp
+++ b/Engine/Source/Runtime/Online/WebSockets/Private/Lws/LwsWebSocketsManager.cpp
@@ -113,6 +113,10 @@ void FLwsWebSocketsManager::InitWebSockets(TArrayView<const FString> Protocols)
 	ContextInfo.max_http_header_data2 = MaxHttpHeaderData;
 	ContextInfo.pt_serv_buf_size = MaxHttpHeaderData;
 	
+	int32 PingPongInterval = 0;
+	GConfig->GetInt(TEXT("WebSockets.LibWebSockets"), TEXT("PingPongInterval"), PingPongInterval, GEngineIni);
+	ContextInfo.ws_ping_pong_interval = PingPongInterval;
+
 	// HTTP proxy
 	const FString& ProxyAddress = FHttpModule::Get().GetProxyAddress();
 	TOptional<FTCHARToUTF8> Converter;
