Fix issues with URL parsing:
- add missing bounds check
- adjust existing bounds checks to avoid potentially overflowing an int32

diff --git a/Engine/Source/Runtime/Online/HTTP/Private/GenericPlatform/GenericPlatformHttp.cpp b/Engine/Source/Runtime/Online/HTTP/Private/GenericPlatform/GenericPlatformHttp.cpp
--- a/Engine/Source/Runtime/Online/HTTP/Private/GenericPlatform/GenericPlatformHttp.cpp
+++ b/Engine/Source/Runtime/Online/HTTP/Private/GenericPlatform/GenericPlatformHttp.cpp
@@ -189,9 +189,9 @@ FString FGenericPlatformHttp::UrlDecode(const FStringView EncodedString)
 		if (UTF8Data[CharIdx] == '%')
 		{
 			int32 Value = 0;
-			if (UTF8Data[CharIdx + 1] == 'u')
+			if (CharIdx < Converter.Length() - 1 && UTF8Data[CharIdx + 1] == 'u')
 			{
-				if (CharIdx + 6 <= Converter.Length())
+				if (CharIdx <= Converter.Length() - 6)
 				{
 					// Treat all %uXXXX as code point
 					Value = FParse::HexDigit(UTF8Data[CharIdx + 2]) << 12;
@@ -215,7 +215,7 @@ FString FGenericPlatformHttp::UrlDecode(const FStringView EncodedString)
 					continue;
 				}
 			}
-			else if(CharIdx + 3 <= Converter.Length())
+			else if(CharIdx <= Converter.Length() - 3)
 			{
 				// Treat all %XX as straight byte
 				Value = FParse::HexDigit(UTF8Data[CharIdx + 1]) << 4;
diff --git a/Engine/Source/Runtime/Online/HTTPServer/Private/HttpConnectionRequestReadContext.cpp b/Engine/Source/Runtime/Online/HTTPServer/Private/HttpConnectionRequestReadContext.cpp
--- a/Engine/Source/Runtime/Online/HTTPServer/Private/HttpConnectionRequestReadContext.cpp
+++ b/Engine/Source/Runtime/Online/HTTPServer/Private/HttpConnectionRequestReadContext.cpp
@@ -332,9 +332,9 @@ FString FHttpConnectionRequestReadContext::UrlDecode(const FString &EncodedStrin
 		if (UTF8Data[CharIdx] == '%')
 		{
 			int32 Value = 0;
-			if (UTF8Data[CharIdx + 1] == 'u')
+			if (CharIdx < Converter.Length() - 1 && UTF8Data[CharIdx + 1] == 'u')
 			{
-				if (CharIdx + 6 <= Converter.Length())
+				if (CharIdx <= Converter.Length() - 6)
 				{
 					// Treat all %uXXXX as code point
 					Value = FParse::HexDigit(UTF8Data[CharIdx + 2]) << 12;
@@ -358,7 +358,7 @@ FString FHttpConnectionRequestReadContext::UrlDecode(const FString &EncodedStrin
 					continue;
 				}
 			}
-			else if (CharIdx + 3 <= Converter.Length())
+			else if (CharIdx <= Converter.Length() - 3)
 			{
 				// Treat all %XX as straight byte
 				Value = FParse::HexDigit(UTF8Data[CharIdx + 1]) << 4;
