Fix untagged http allocations for http server.

diff --git a/Engine/Source/Runtime/Online/HTTPServer/Private/HttpConnection.cpp b/Engine/Source/Runtime/Online/HTTPServer/Private/HttpConnection.cpp
--- a/Engine/Source/Runtime/Online/HTTPServer/Private/HttpConnection.cpp
+++ b/Engine/Source/Runtime/Online/HTTPServer/Private/HttpConnection.cpp
@@ -146,7 +146,7 @@ void FHttpConnection::ContinueRead(float DeltaTime)
 
 void FHttpConnection::CompleteRead(const TSharedPtr<FHttpServerRequest>& Request)
 {
-	TArray<FString>* ConnectionHeaders = Request->Headers.Find(FHttpServerHeaderKeys::CONNECTION);
+	TArray<FString>* ConnectionHeaders = Request->Headers.Find(UE_HTTP_SERVER_HEADER_KEYS_CONNECTION);
 	if (ConnectionHeaders)
 	{
 		bKeepAlive = ResolveKeepAlive(Request->HttpVersion, *ConnectionHeaders);
@@ -204,8 +204,7 @@ void FHttpConnection::ProcessRequest(const TSharedPtr<FHttpServerRequest>& Reque
 
 	if (!bRequestHandled)
 	{
-		const FString& ResponseError(FHttpServerErrorStrings::NotFound);
-		auto Response = FHttpServerResponse::Error(EHttpServerResponseCodes::NotFound, ResponseError);
+		auto Response = FHttpServerResponse::Error(EHttpServerResponseCodes::NotFound, UE_HTTP_SERVER_ERROR_STR_ROUTE_HANDLER_NOT_FOUND);
 		OnProcessingComplete(MoveTemp(Response));
 	}
 }
@@ -221,7 +220,7 @@ void FHttpConnection::BeginWrite(TUniquePtr<FHttpServerResponse>&& Response, uin
 	{
 		FString KeepAliveTimeoutStr = FString::Printf(TEXT("timeout=%f"), ConnectionKeepAliveTimeout);
 		TArray<FString> KeepAliveTimeoutValue = { MoveTemp(KeepAliveTimeoutStr) };
-		Response->Headers.Add(FHttpServerHeaderKeys::KEEP_ALIVE, MoveTemp(KeepAliveTimeoutValue));
+		Response->Headers.Add(UE_HTTP_SERVER_HEADER_KEYS_KEEP_ALIVE, MoveTemp(KeepAliveTimeoutValue));
 	}
 
 	WriteContext.ResetContext(MoveTemp(Response));
diff --git a/Engine/Source/Runtime/Online/HTTPServer/Private/HttpConnectionRequestReadContext.cpp b/Engine/Source/Runtime/Online/HTTPServer/Private/HttpConnectionRequestReadContext.cpp
--- a/Engine/Source/Runtime/Online/HTTPServer/Private/HttpConnectionRequestReadContext.cpp
+++ b/Engine/Source/Runtime/Online/HTTPServer/Private/HttpConnectionRequestReadContext.cpp
@@ -34,7 +34,7 @@ EHttpConnectionContextState FHttpConnectionRequestReadContext::ReadStream(float
 	int32 BytesRead = 0;
 	if (!Socket->Recv(ByteBuffer, sizeof(ByteBuffer) - 1, BytesRead, ESocketReceiveFlags::None))
 	{
-		AddError(FHttpServerErrorStrings::SocketRecvFailure);
+		AddError(UE_HTTP_SERVER_ERROR_STR_SOCKET_RECV_FAILURE);
 		return EHttpConnectionContextState::Error;
 	}
 
@@ -131,7 +131,7 @@ bool FHttpConnectionRequestReadContext::ParseBody(uint8* ByteBuffer, int32 Buffe
 	if (BufferLen > IncomingRequestBodyBytesToRead)
 	{
 		// Error - Sent data size exceeds expected 
-		AddError(FHttpServerErrorStrings::MismatchedContentLengthBodyTooLarge, EHttpServerResponseCodes::BadRequest);
+		AddError(UE_HTTP_SERVER_ERROR_STR_MISMATCHED_CONTENT_LENGTH_BODY_TOO_LARGE, EHttpServerResponseCodes::BadRequest);
 		return false;
 	}
 
@@ -157,7 +157,7 @@ bool FHttpConnectionRequestReadContext::IsRequestValid(const FHttpServerRequest&
 		// Enforce content length missing or 0
 		if (bContentLengthSpecified && 0 != RequestContentLength)
 		{
-			AddError(FHttpServerErrorStrings::InvalidContentLengthHeader, EHttpServerResponseCodes::BadRequest);
+			AddError(UE_HTTP_SERVER_ERROR_STR_INVALID_CONTENT_LENGTH_HEADER, EHttpServerResponseCodes::BadRequest);
 			return false;
 		}
 		break;
@@ -167,13 +167,13 @@ bool FHttpConnectionRequestReadContext::IsRequestValid(const FHttpServerRequest&
 		// Content length must be set
 		if (!bContentLengthSpecified)
 		{
-			AddError(FHttpServerErrorStrings::MissingContentLengthHeader, EHttpServerResponseCodes::LengthRequired);
+			AddError(UE_HTTP_SERVER_ERROR_STR_MISSING_CONTENT_LENGTH_HEADER, EHttpServerResponseCodes::LengthRequired);
 			return false;
 		}
 		// Content length must be valid
 		if(RequestContentLength < 0)
 		{
-			AddError(FHttpServerErrorStrings::InvalidContentLengthHeader, EHttpServerResponseCodes::LengthRequired);
+			AddError(UE_HTTP_SERVER_ERROR_STR_INVALID_CONTENT_LENGTH_HEADER, EHttpServerResponseCodes::LengthRequired);
 			return false;
 		}
 		break;
@@ -184,7 +184,7 @@ bool FHttpConnectionRequestReadContext::IsRequestValid(const FHttpServerRequest&
 
 bool FHttpConnectionRequestReadContext::ParseContentLength(const FHttpServerRequest& InRequest, int32& OutContentLength)
 {
-	const TArray<FString>* ContentLengthValues = InRequest.Headers.Find(FHttpServerHeaderKeys::CONTENT_LENGTH);
+	const TArray<FString>* ContentLengthValues = InRequest.Headers.Find(UE_HTTP_SERVER_HEADER_KEYS_CONTENT_LENGTH);
 
 	if (ContentLengthValues && ContentLengthValues->Num() > 0)
 	{
@@ -202,7 +202,7 @@ TSharedPtr<FHttpServerRequest> FHttpConnectionRequestReadContext::BuildRequest(c
 
 	if (0 == ParsedHeader.Num())
 	{
-		AddError(FHttpServerErrorStrings::MissingRequestHeaders, EHttpServerResponseCodes::BadRequest);
+		AddError(UE_HTTP_SERVER_ERROR_STR_MISSING_REQUEST_HEADERS, EHttpServerResponseCodes::BadRequest);
 		return nullptr;
 	}
 
@@ -212,7 +212,7 @@ TSharedPtr<FHttpServerRequest> FHttpConnectionRequestReadContext::BuildRequest(c
 	HttpMethod.ParseIntoArrayWS(HttpMethodTokens);
 	if (HttpMethodTokens.Num() < 3)
 	{
-		AddError(FHttpServerErrorStrings::MalformedRequestHeaders, EHttpServerResponseCodes::BadRequest);
+		AddError(UE_HTTP_SERVER_ERROR_STR_MALFORMED_REQUEST_HEADER, EHttpServerResponseCodes::BadRequest);
 		return nullptr;
 	}
 
@@ -252,7 +252,7 @@ TSharedPtr<FHttpServerRequest> FHttpConnectionRequestReadContext::BuildRequest(c
 	else
 	{
 		// Unknown Verb
-		AddError(FHttpServerErrorStrings::UnknownRequestVerb, EHttpServerResponseCodes::BadMethod);
+		AddError(UE_HTTP_SERVER_ERROR_STR_UNKNOWN_REQUEST_VERB, EHttpServerResponseCodes::BadMethod);
 		return nullptr;
 	}
 
@@ -284,7 +284,7 @@ TSharedPtr<FHttpServerRequest> FHttpConnectionRequestReadContext::BuildRequest(c
 	const FString& RequestHttpVersion = HttpMethodTokens[2];
 	if(!HttpVersion::FromString(RequestHttpVersion, Request->HttpVersion))
 	{
-		AddError(FHttpServerErrorStrings::UnsupportedHttpVersion, EHttpServerResponseCodes::VersionNotSup);
+		AddError(UE_HTTP_SERVER_ERROR_STR_UNSUPPORTED_HTTP_VERSION, EHttpServerResponseCodes::VersionNotSup);
 		return nullptr;
 	}
 
diff --git a/Engine/Source/Runtime/Online/HTTPServer/Private/HttpConnectionResponseWriteContext.cpp b/Engine/Source/Runtime/Online/HTTPServer/Private/HttpConnectionResponseWriteContext.cpp
--- a/Engine/Source/Runtime/Online/HTTPServer/Private/HttpConnectionResponseWriteContext.cpp
+++ b/Engine/Source/Runtime/Online/HTTPServer/Private/HttpConnectionResponseWriteContext.cpp
@@ -27,7 +27,7 @@ void FHttpConnectionResponseWriteContext::ResetContext(TUniquePtr<FHttpServerRes
 	{
 		// Affix Content-Length Header
 		TArray<FString> ContentLengthValue = { FString::FromInt(Response->Body.Num()) };
-		Response->Headers.Add(FHttpServerHeaderKeys::CONTENT_LENGTH, MoveTemp(ContentLengthValue));
+		Response->Headers.Add(UE_HTTP_SERVER_HEADER_KEYS_CONTENT_LENGTH, MoveTemp(ContentLengthValue));
 
 		// Serialize Headers
 		HeaderBytes.Append(SerializeHeadersUtf8(Response->HttpVersion, Response->Code, Response->Headers));
@@ -46,7 +46,7 @@ EHttpConnectionContextState FHttpConnectionResponseWriteContext::WriteStream(flo
 		int32 DataLen = HeaderBytes.Num() - HeaderBytesWritten;
 		if (!WriteBytes(DataOffset, DataLen, BytesWritten))
 		{
-			AddError(FHttpServerErrorStrings::SocketSendFailure);
+			AddError(UE_HTTP_SERVER_ERROR_STR_SOCKET_SEND_FAILURE);
 			return EHttpConnectionContextState::Error;
 		}
 		HeaderBytesWritten += BytesWritten;
@@ -58,7 +58,7 @@ EHttpConnectionContextState FHttpConnectionResponseWriteContext::WriteStream(flo
 		int32 DataLen = Response->Body.Num() - BodyBytesWritten;
 		if (!WriteBytes(DataOffset, DataLen, BytesWritten))
 		{
-			AddError(FHttpServerErrorStrings::SocketSendFailure);
+			AddError(UE_HTTP_SERVER_ERROR_STR_SOCKET_SEND_FAILURE);
 			return EHttpConnectionContextState::Error;
 		}
 		BodyBytesWritten += BytesWritten;
diff --git a/Engine/Source/Runtime/Online/HTTPServer/Private/HttpServerConstantsPrivate.h b/Engine/Source/Runtime/Online/HTTPServer/Private/HttpServerConstantsPrivate.h
--- a/Engine/Source/Runtime/Online/HTTPServer/Private/HttpServerConstantsPrivate.h
+++ b/Engine/Source/Runtime/Online/HTTPServer/Private/HttpServerConstantsPrivate.h
@@ -3,33 +3,26 @@
 
 #include "CoreMinimal.h"
 
-namespace FHttpServerHeaderKeys
-{
-	const FString CONTENT_TYPE = TEXT("content-type");
-	const FString CONTENT_LENGTH = TEXT("content-length");
-	const FString CONNECTION = TEXT("connection");
-	const FString KEEP_ALIVE = TEXT("keep-alive");
-};
+#define UE_HTTP_SERVER_HEADER_KEYS_CONTENT_TYPE TEXT("content-type")
+#define UE_HTTP_SERVER_HEADER_KEYS_CONTENT_LENGTH TEXT("content-length")
+#define UE_HTTP_SERVER_HEADER_KEYS_CONNECTION TEXT("connection")
+#define UE_HTTP_SERVER_HEADER_KEYS_KEEP_ALIVE TEXT("keep-alive")
 
-namespace FHttpServerErrorStrings
-{
-	// Connection
-	const FString SocketClosedFailure = TEXT("errors.com.epicgames.httpserver.socket_closed_failure");
-	const FString SocketRecvFailure = TEXT("errors.com.epicgames.httpserver.socket_recv_failure");
-	const FString SocketSendFailure = TEXT("errors.com.epicgames.httpserver.socket_send_failure");
+// Connection
+#define UE_HTTP_SERVER_ERROR_STR_SOCKET_CLOSED_FAILURE TEXT("errors.com.epicgames.httpserver.socket_closed_failure")
+#define UE_HTTP_SERVER_ERROR_STR_SOCKET_RECV_FAILURE TEXT("errors.com.epicgames.httpserver.socket_recv_failure")
+#define UE_HTTP_SERVER_ERROR_STR_SOCKET_SEND_FAILURE TEXT("errors.com.epicgames.httpserver.socket_send_failure")
 
-	// Routing
-	const FString NotFound = TEXT("errors.com.epicgames.httpserver.route_handler_not_found");
-
-	// Serialization
-	const FString MalformedRequestSize = TEXT("errors.com.epicgames.httpserver.malformed_request_size");
-	const FString MalformedRequestHeaders = TEXT("errors.com.epicgames.httpserver.malformed_request_header");
-	const FString MissingRequestHeaders = TEXT("errors.com.epicgames.httpserver.missing_request_headers");
-	const FString MalformedRequestBody = TEXT("errors.com.epicgames.httpserver.malformed_request_body");
-	const FString UnknownRequestVerb = TEXT("errors.com.epicgames.httpserver.unknown_request_verb");
-	const FString UnsupportedHttpVersion = TEXT("errors.com.epicgames.httpserver.unsupported_http_version");
-	const FString InvalidContentLengthHeader = TEXT("errors.com.epicgames.httpserver.invalid_content_length_header");
-	const FString MissingContentLengthHeader = TEXT("errors.com.epicgames.httpserver.missing_content_length_header");
-	const FString MismatchedContentLengthBodyTooLarge =  TEXT("errors.com.epicgames.httpserver.mismatched_content_length_body_too_large");
-};
+// Routing
+#define UE_HTTP_SERVER_ERROR_STR_ROUTE_HANDLER_NOT_FOUND TEXT("errors.com.epicgames.httpserver.route_handler_not_found")
 
+// Serialization
+#define UE_HTTP_SERVER_ERROR_STR_MALFORMED_REQUEST_SIZE TEXT("errors.com.epicgames.httpserver.malformed_request_size")
+#define UE_HTTP_SERVER_ERROR_STR_MALFORMED_REQUEST_HEADER TEXT("errors.com.epicgames.httpserver.malformed_request_header")
+#define UE_HTTP_SERVER_ERROR_STR_MISSING_REQUEST_HEADERS TEXT("errors.com.epicgames.httpserver.missing_request_headers")
+#define UE_HTTP_SERVER_ERROR_STR_MALFORMED_REQUEST_BODY TEXT("errors.com.epicgames.httpserver.malformed_request_body")
+#define UE_HTTP_SERVER_ERROR_STR_UNKNOWN_REQUEST_VERB TEXT("errors.com.epicgames.httpserver.unknown_request_verb")
+#define UE_HTTP_SERVER_ERROR_STR_UNSUPPORTED_HTTP_VERSION TEXT("errors.com.epicgames.httpserver.unsupported_http_version")
+#define UE_HTTP_SERVER_ERROR_STR_INVALID_CONTENT_LENGTH_HEADER TEXT("errors.com.epicgames.httpserver.invalid_content_length_header")
+#define UE_HTTP_SERVER_ERROR_STR_MISSING_CONTENT_LENGTH_HEADER TEXT("errors.com.epicgames.httpserver.missing_content_length_header")
+#define UE_HTTP_SERVER_ERROR_STR_MISMATCHED_CONTENT_LENGTH_BODY_TOO_LARGE TEXT("errors.com.epicgames.httpserver.mismatched_content_length_body_too_large")
diff --git a/Engine/Source/Runtime/Online/HTTPServer/Private/HttpServerResponse.cpp b/Engine/Source/Runtime/Online/HTTPServer/Private/HttpServerResponse.cpp
--- a/Engine/Source/Runtime/Online/HTTPServer/Private/HttpServerResponse.cpp
+++ b/Engine/Source/Runtime/Online/HTTPServer/Private/HttpServerResponse.cpp
@@ -15,7 +15,7 @@ TUniquePtr<FHttpServerResponse> FHttpServerResponse::Create(const FString& Text,
 
 	FString Utf8CharsetContentType = FString::Printf(TEXT("%s;charset=utf-8"), *ContentType);
 	TArray<FString> ContentTypeValue = { MoveTemp(Utf8CharsetContentType) };
-	Response->Headers.Add(FHttpServerHeaderKeys::CONTENT_TYPE, MoveTemp(ContentTypeValue));
+	Response->Headers.Add(UE_HTTP_SERVER_HEADER_KEYS_CONTENT_TYPE, MoveTemp(ContentTypeValue));
 
 	return Response;
 }
@@ -26,7 +26,7 @@ TUniquePtr<FHttpServerResponse> FHttpServerResponse::Create(TArray<uint8>&& RawB
 	Response->Code = EHttpServerResponseCodes::Ok;
 
 	TArray<FString> ContentTypeValue = { MoveTemp(ContentType) };
-	Response->Headers.Add(FHttpServerHeaderKeys::CONTENT_TYPE, MoveTemp(ContentTypeValue));
+	Response->Headers.Add(UE_HTTP_SERVER_HEADER_KEYS_CONTENT_TYPE, MoveTemp(ContentTypeValue));
 	return Response;
 }
 
@@ -37,7 +37,7 @@ TUniquePtr<FHttpServerResponse> FHttpServerResponse::Create(const TArrayView<uin
 	Response->Body.Append(RawBytes.GetData(), RawBytes.Num());
 
 	TArray<FString> ContentTypeValue = { MoveTemp(ContentType) };
-	Response->Headers.Add(FHttpServerHeaderKeys::CONTENT_TYPE, MoveTemp(ContentTypeValue));
+	Response->Headers.Add(UE_HTTP_SERVER_HEADER_KEYS_CONTENT_TYPE, MoveTemp(ContentTypeValue));
 	return Response;
 }
 
