Use strict check for the size of received buffer in UDPEchoImpl.

diff --git a/Engine/Source/Runtime/Online/ICMP/Private/UDPPing.cpp b/Engine/Source/Runtime/Online/ICMP/Private/UDPPing.cpp
--- a/Engine/Source/Runtime/Online/ICMP/Private/UDPPing.cpp
+++ b/Engine/Source/Runtime/Online/ICMP/Private/UDPPing.cpp
@@ -1404,7 +1404,7 @@ FIcmpEchoResult UDPEchoImpl(ISocketSubsystem* SocketSub, const FString& TargetAd
 							TSharedRef<FInternetAddr> RecvAddr = SocketSub->CreateInternetAddr();
 							if (Socket->RecvFrom(ResultBuffer, ResultPacketSize, BytesRead, *RecvAddr))
 							{
-								if (BytesRead > 0)
+								if (BytesRead == ResultPacketSize)
 								{
 									uint64 NowTime = FPlatformTime::Cycles64();
 
