BuildPatchServices | Build file streaming implementation

diff --git a/Engine/Source/Runtime/Online/BuildPatchServices/BuildPatchServices.Build.cs b/Engine/Source/Runtime/Online/BuildPatchServices/BuildPatchServices.Build.cs
--- a/Engine/Source/Runtime/Online/BuildPatchServices/BuildPatchServices.Build.cs
+++ b/Engine/Source/Runtime/Online/BuildPatchServices/BuildPatchServices.Build.cs
@@ -26,6 +26,7 @@ public class BuildPatchServices : ModuleRules
 				"AnalyticsET",
 				"HTTP",
 				"Json",
+				"VirtualFileCache",
 			}
 		);
 
diff --git a/Engine/Source/Runtime/Online/BuildPatchServices/Private/BuildPatchManifest.h b/Engine/Source/Runtime/Online/BuildPatchServices/Private/BuildPatchManifest.h
--- a/Engine/Source/Runtime/Online/BuildPatchServices/Private/BuildPatchManifest.h
+++ b/Engine/Source/Runtime/Online/BuildPatchServices/Private/BuildPatchManifest.h
@@ -116,6 +116,7 @@ public:
 	virtual int64 GetFileSize(const FString& Filename) const override;
 	virtual int64 GetFileSize(const TArray<FString>& Filenames) const override;
 	virtual int64 GetFileSize(const TSet  <FString>& Filenames) const override;
+	virtual bool GetFileHash(const FString& Filename, FSHAHash& OutHash) const override;
 	virtual TSet<FString> GetFileTagList() const override;
 	virtual void GetFileTagList(TSet<FString>& Tags) const override;
 	virtual void GetOutdatedFiles(const IBuildManifestRef& OldManifest, TSet<FString>& OutdatedFiles) const override;
@@ -279,14 +280,6 @@ public:
 	 */
 	virtual bool GetFileHash(const FGuid& FileGuid, FSHAHash& OutHash) const; // DEPRECATE ME
 
-	/**
-	 * Gets the file hash for a given file
-	 * @param Filename		IN		The filename in the build
-	 * @param OutHash		OUT		Receives the hash value if found
-	 * @return	true if we had the hash for this file
-	 */
-	virtual bool GetFileHash(const FString& Filename, FSHAHash& OutHash) const;
-
 	/**
 	 * Gets the file hash for given file data. Valid for non-chunked manifest
 	 * @param FileGuid		IN		The guid of the file data to get hash for
diff --git a/Engine/Source/Runtime/Online/BuildPatchServices/Private/BuildPatchServicesModule.cpp b/Engine/Source/Runtime/Online/BuildPatchServices/Private/BuildPatchServicesModule.cpp
--- a/Engine/Source/Runtime/Online/BuildPatchServices/Private/BuildPatchServicesModule.cpp
+++ b/Engine/Source/Runtime/Online/BuildPatchServices/Private/BuildPatchServicesModule.cpp
@@ -113,6 +113,14 @@ void FBuildPatchServicesModule::ShutdownModule()
 	GLog->Log(ELogVerbosity::VeryVerbose, TEXT( "BuildPatchServicesModule: Finished shutting down" ) );
 }
 
+IBuildInstallStreamerRef FBuildPatchServicesModule::CreateBuildInstallStreamer(BuildPatchServices::FBuildInstallStreamerConfiguration Configuration)
+{
+	FBuildInstallStreamerRef Streamer = MakeShareable(FBuildInstallStreamerFactory::Create(MoveTemp(Configuration)));
+	FBuildInstallStreamerWeakPtr WeakStreamer = Streamer;
+	AsyncHelpers::ExecuteOnGameThread<void>([this, WeakStreamer = MoveTemp(WeakStreamer)] { WeakBuildInstallStreamers.Add(WeakStreamer); });
+	return Streamer;
+}
+
 IBuildInstallerRef FBuildPatchServicesModule::CreateBuildInstaller(BuildPatchServices::FBuildInstallerConfiguration Configuration, FBuildPatchInstallerDelegate CompleteDelegate) const
 {
 	// Override prereq install using the config/commandline value to force skip them.
@@ -214,6 +222,16 @@ bool FBuildPatchServicesModule::Tick(float Delta)
 		}
 	}
 
+	// Tick running streamers.
+	for (auto StreamerIter = WeakBuildInstallStreamers.CreateIterator(); StreamerIter; ++StreamerIter)
+	{
+		const FBuildInstallStreamerPtr& Streamer = StreamerIter->Pin();
+		if (!Streamer.IsValid() || !Streamer->Tick())
+		{
+			StreamerIter.RemoveCurrent();
+		}
+	}
+
 	// More ticks.
 	return true;
 }
@@ -379,6 +397,15 @@ void FBuildPatchServicesModule::PreExit()
 		const FBuildPatchInstallerRef& Installer = *InstallerIter;
 		Installer->PreExit();
 	}
+	// Inform streamers
+	for (auto StreamerIter = WeakBuildInstallStreamers.CreateIterator(); StreamerIter; ++StreamerIter)
+	{
+		const FBuildInstallStreamerPtr& Streamer = StreamerIter->Pin();
+		if (Streamer.IsValid())
+		{
+			Streamer->PreExit();
+		}
+	}
 
 	// Release our ptr to analytics
 	Analytics.Reset();
diff --git a/Engine/Source/Runtime/Online/BuildPatchServices/Private/BuildPatchServicesModule.h b/Engine/Source/Runtime/Online/BuildPatchServices/Private/BuildPatchServicesModule.h
--- a/Engine/Source/Runtime/Online/BuildPatchServices/Private/BuildPatchServicesModule.h
+++ b/Engine/Source/Runtime/Online/BuildPatchServices/Private/BuildPatchServicesModule.h
@@ -4,6 +4,7 @@
 
 #include "CoreMinimal.h"
 #include "Misc/Guid.h"
+#include "Installer/BuildInstallStreamer.h"
 #include "Interfaces/IBuildPatchServicesModule.h"
 #include "BuildPatchInstaller.h"
 #include "Containers/Ticker.h"
@@ -35,6 +36,7 @@ public:
 	// IModuleInterface interface end.
 
 	// IBuildPatchServicesModule interface begin.
+	virtual IBuildInstallStreamerRef CreateBuildInstallStreamer(BuildPatchServices::FBuildInstallStreamerConfiguration Configuration) override;
 	virtual IBuildInstallerRef CreateBuildInstaller(BuildPatchServices::FBuildInstallerConfiguration Configuration, FBuildPatchInstallerDelegate OnComplete) const override;
 	virtual BuildPatchServices::IBuildStatisticsRef CreateBuildStatistics(const IBuildInstallerRef& Installer) const override;
 	virtual BuildPatchServices::IPatchDataEnumerationRef CreatePatchDataEnumeration(BuildPatchServices::FPatchDataEnumerationConfiguration Configuration) const override;
@@ -139,6 +141,9 @@ private:
 	// Holds available installations used for recycling install data
 	TMultiMap<FString, FBuildPatchAppManifestRef> AvailableInstallations;
 
+	// Array of running streamers
+	TArray<FBuildInstallStreamerWeakPtr> WeakBuildInstallStreamers;
+
 	// Handle to the registered Tick delegate
 	FTSTicker::FDelegateHandle TickDelegateHandle;
 
diff --git a/Engine/Source/Runtime/Online/BuildPatchServices/Private/Core/AsyncHelpers.h b/Engine/Source/Runtime/Online/BuildPatchServices/Private/Core/AsyncHelpers.h
--- a/Engine/Source/Runtime/Online/BuildPatchServices/Private/Core/AsyncHelpers.h
+++ b/Engine/Source/Runtime/Online/BuildPatchServices/Private/Core/AsyncHelpers.h
@@ -1,10 +1,11 @@
 // Copyright Epic Games, Inc. All Rights Reserved.
 #pragma once
 
+#include "Async/Async.h"
+#include "Containers/Queue.h"
 #include "CoreMinimal.h"
 #include "HAL/ThreadSafeCounter.h"
 #include "HAL/ThreadSafeCounter64.h"
-#include "Async/Async.h"
 
 namespace BuildPatchServices
 {
@@ -63,6 +64,22 @@ namespace BuildPatchServices
 			}
 			return Promise->GetFuture();
 		}
+
+		template<typename ResultType>
+		static TFuture<ResultType> ExecuteOnCustomThread(const TFunction<ResultType()>& Function, TQueue<TFunction<void()>, EQueueMode::Spsc>& CustomThreadQueue)
+		{
+			TSharedRef<TPromise<ResultType>, ESPMode::ThreadSafe> Promise = MakeShareable(new TPromise<ResultType>());
+			TFunction<void()> PromiseKeeper = MakePromiseKeeper(Promise, Function);
+			if (!IsInGameThread())
+			{
+				CustomThreadQueue.Enqueue(MoveTemp(PromiseKeeper));
+			}
+			else
+			{
+				PromiseKeeper();
+			}
+			return Promise->GetFuture();
+		}
 	}
 
 	/**
diff --git a/Engine/Source/Runtime/Online/BuildPatchServices/Private/Core/ProcessTimer.h b/Engine/Source/Runtime/Online/BuildPatchServices/Private/Core/ProcessTimer.h
--- a/Engine/Source/Runtime/Online/BuildPatchServices/Private/Core/ProcessTimer.h
+++ b/Engine/Source/Runtime/Online/BuildPatchServices/Private/Core/ProcessTimer.h
@@ -11,7 +11,7 @@ namespace BuildPatchServices
 	 * We template the dependency for cycles so that the class can be nicely tested.
 	 * Under normal circumstances, use FProcessTimer, which is a typedef below.
 	 */
-	template<typename FCyclesProvider>
+	template<typename FCyclesProvider, bool bThreadSafe = true>
 	class TProcessTimer
 	{
 	public:
@@ -33,12 +33,19 @@ namespace BuildPatchServices
 		 */
 		double GetSeconds()
 		{
-			FScopeLock ScopeLock(&ThreadLock);
+			if (bThreadSafe)
+			{
+				ThreadLock.Lock();
+			}
 			double Seconds = FCyclesProvider::CyclesToSeconds(Cycles);
 			if (bIsRunning && !bIsPaused)
 			{
 				Seconds += FCyclesProvider::CyclesToSeconds(FCyclesProvider::GetCycles() - StartCycles);
 			}
+			if (bThreadSafe)
+			{
+				ThreadLock.Unlock();
+			}
 			return Seconds;
 		}
 
@@ -47,7 +54,10 @@ namespace BuildPatchServices
 		 */
 		void Start()
 		{
-			FScopeLock ScopeLock(&ThreadLock);
+			if (bThreadSafe)
+			{
+				ThreadLock.Lock();
+			}
 			if (!bIsRunning)
 			{
 				bIsRunning = true;
@@ -56,6 +66,10 @@ namespace BuildPatchServices
 					StartCycles = FCyclesProvider::GetCycles();
 				}
 			}
+			if (bThreadSafe)
+			{
+				ThreadLock.Unlock();
+			}
 		}
 
 		/**
@@ -63,7 +77,10 @@ namespace BuildPatchServices
 		 */
 		void Stop()
 		{
-			FScopeLock ScopeLock(&ThreadLock);
+			if (bThreadSafe)
+			{
+				ThreadLock.Lock();
+			}
 			if (bIsRunning)
 			{
 				bIsRunning = false;
@@ -72,6 +89,10 @@ namespace BuildPatchServices
 					Cycles += FCyclesProvider::GetCycles() - StartCycles;
 				}
 			}
+			if (bThreadSafe)
+			{
+				ThreadLock.Unlock();
+			}
 		}
 
 		/**
@@ -80,7 +101,10 @@ namespace BuildPatchServices
 		 */
 		void SetPause(bool bPause)
 		{
-			FScopeLock ScopeLock(&ThreadLock);
+			if (bThreadSafe)
+			{
+				ThreadLock.Lock();
+			}
 			if (bIsPaused != bPause)
 			{
 				bIsPaused = bPause;
@@ -100,6 +124,26 @@ namespace BuildPatchServices
 					}
 				}
 			}
+			if (bThreadSafe)
+			{
+				ThreadLock.Unlock();
+			}
+		}
+
+		void Reset()
+		{
+			if (bThreadSafe)
+			{
+				ThreadLock.Lock();
+			}
+			StartCycles = 0;
+			Cycles = 0;
+			bIsRunning = false;
+			bIsPaused = false;
+			if (bThreadSafe)
+			{
+				ThreadLock.Unlock();
+			}
 		}
 
 	private:
diff --git a/Engine/Source/Runtime/Online/BuildPatchServices/Private/Installer/BuildInstallStreamer.cpp b/Engine/Source/Runtime/Online/BuildPatchServices/Private/Installer/BuildInstallStreamer.cpp
new file mode 100644
--- /dev/null
+++ b/Engine/Source/Runtime/Online/BuildPatchServices/Private/Installer/BuildInstallStreamer.cpp
@@ -0,0 +1,395 @@
+// Copyright Epic Games, Inc. All Rights Reserved.
+
+#include "Installer/BuildInstallStreamer.h"
+
+#include "Algo/RemoveIf.h"
+#include "BuildPatchSettings.h"
+#include "Common/ChunkDataSizeProvider.h"
+#include "Common/FileSystem.h"
+#include "Common/HttpManager.h"
+#include "Core/AsyncHelpers.h"
+#include "Core/Platform.h"
+#include "Core/ProcessTimer.h"
+#include "IBuildManifestSet.h"
+#include "Installer/CloudChunkSource.h"
+#include "Installer/ChunkEvictionPolicy.h"
+#include "Installer/ChunkReferenceTracker.h"
+#include "Installer/DownloadService.h"
+#include "Installer/InstallerAnalytics.h"
+#include "Installer/InstallerError.h"
+#include "Installer/MemoryChunkStore.h"
+#include "Installer/MessagePump.h"
+#include "Installer/Statistics/CloudChunkSourceStatistics.h"
+#include "Installer/Statistics/DownloadServiceStatistics.h"
+#include "Installer/Statistics/FileConstructorStatistics.h"
+#include "Installer/Statistics/FileOperationTracker.h"
+#include "Installer/Statistics/MemoryChunkStoreStatistics.h"
+#include "Installer/VirtualFileConstructor.h"
+#include "VirtualFileCache.h"
+
+DEFINE_LOG_CATEGORY_STATIC(LogBuildInstallStreamer, Log, All);
+
+namespace BuildPatchServices
+{
+	namespace
+	{
+		FDownloadConnectionCountConfig BuildConnectionCountConfig()
+		{
+			// TODO: Configure for scaling, move installer's config builder to shared code?
+			FDownloadConnectionCountConfig Config;
+			Config.bDisableConnectionScaling = true;
+			return Config;
+		}
+	}
+
+	FBuildInstallStreamer::FBuildInstallStreamer(FBuildInstallStreamerConfiguration InConfiguration)
+		: Configuration(MoveTemp(InConfiguration))
+		, BuildPatchManifest(StaticCastSharedPtr<FBuildPatchAppManifest>(Configuration.Manifest))
+		, ManifestSet(FBuildManifestSetFactory::Create({ FBuildPatchInstallerAction(FInstallerAction::MakeInstall(Configuration.Manifest.ToSharedRef())) }))
+		, HttpManager(FHttpManagerFactory::Create())
+		, FileSystem(FFileSystemFactory::Create())
+		, ChunkDataSerialization(FChunkDataSerializationFactory::Create(FileSystem.Get()))
+		, Platform(FPlatformFactory::Create())
+		, InstallerAnalytics(FInstallerAnalyticsFactory::Create(nullptr))
+		, InstallerError(FInstallerErrorFactory::Create())
+		, FileOperationTracker(FFileOperationTrackerFactory::Create(FTSTicker::GetCoreTicker()))
+		, DownloadSpeedRecorder(FSpeedRecorderFactory::Create())
+		, FileReadSpeedRecorder(FSpeedRecorderFactory::Create())
+		, FileWriteSpeedRecorder(FSpeedRecorderFactory::Create())
+		, ChunkDataSizeProvider(FChunkDataSizeProviderFactory::Create())
+		, DownloadServiceStatistics(FDownloadServiceStatisticsFactory::Create(DownloadSpeedRecorder.Get(), ChunkDataSizeProvider.Get(), InstallerAnalytics.Get()))
+		, MemoryChunkStoreStatistics(FMemoryChunkStoreStatisticsFactory::Create(FileOperationTracker.Get()))
+		, CloudChunkSourceStatistics(FCloudChunkSourceStatisticsFactory::Create(InstallerAnalytics.Get(), &BuildProgress, FileOperationTracker.Get()))
+		, FileConstructorStatistics(FFileConstructorStatisticsFactory::Create(FileReadSpeedRecorder.Get(), FileWriteSpeedRecorder.Get(), &BuildProgress, FileOperationTracker.Get()))
+		, DownloadConnectionCount(FDownloadConnectionCountFactory::Create(BuildConnectionCountConfig(), DownloadServiceStatistics.Get()))
+		, DownloadService(FDownloadServiceFactory::Create(FTSTicker::GetCoreTicker(), HttpManager.Get(), FileSystem.Get(), DownloadServiceStatistics.Get(), InstallerAnalytics.Get()))
+		, MessagePump(FMessagePumpFactory::Create())
+		, bIsShuttingDown(false)
+		, RequestTrigger(FPlatformProcess::GetSynchEventFromPool(false))
+		, CloudTrigger(FPlatformProcess::GetSynchEventFromPool(false))
+	{
+		RequestWorker = Async(EAsyncExecution::ThreadIfForkSafe, [this]()
+			{
+				RequestWorkerThread();
+			});
+		CloudWorker = Async(EAsyncExecution::ThreadIfForkSafe, [this]()
+			{
+				CloudWorkerThread();
+			});
+	}
+
+	FBuildInstallStreamer::~FBuildInstallStreamer()
+	{
+		// If we destruct on the main thread we need to use the PreExit functionality to properly close down threads.
+		if (IsInGameThread())
+		{
+			PreExit();
+		}
+		else
+		{
+			bIsShuttingDown = true;
+			InstallerError->SetError(EBuildPatchInstallError::ApplicationClosing, ApplicationClosedErrorCodes::ApplicationClosed);
+			RequestTrigger->Trigger();
+			CloudTrigger->Trigger();
+			RequestWorker.Wait();
+			CloudWorker.Wait();
+		}
+		FPlatformProcess::ReturnSynchEventToPool(RequestTrigger);
+		FPlatformProcess::ReturnSynchEventToPool(CloudTrigger);
+	}
+
+	void FBuildInstallStreamer::QueueFilesByTag(TSet<FString> Tags, FBuildPatchStreamCompleteDelegate OnComplete)
+	{
+		check(!bIsShuttingDown);
+		UE_LOG(LogBuildInstallStreamer, Verbose, TEXT("Receive request %s"), *FString::Join(Tags, TEXT(",")));
+		RequestQueue.Enqueue({ MoveTemp(Tags), true, MoveTemp(OnComplete) });
+		RequestTrigger->Trigger();
+	}
+
+	void FBuildInstallStreamer::QueueFilesByName(TSet<FString> Files, FBuildPatchStreamCompleteDelegate OnComplete)
+	{
+		check(!bIsShuttingDown);
+		UE_LOG(LogBuildInstallStreamer, Verbose, TEXT("Receive request %s"), *FString::Join(Files, TEXT(",")));
+		RequestQueue.Enqueue({ MoveTemp(Files), false, MoveTemp(OnComplete) });
+		RequestTrigger->Trigger();
+	}
+
+	void FBuildInstallStreamer::RegisterMessageHandler(FMessageHandler* MessageHandler)
+	{
+	}
+
+	void FBuildInstallStreamer::UnregisterMessageHandler(FMessageHandler* MessageHandler)
+	{
+	}
+
+	const FBuildInstallStreamerConfiguration& FBuildInstallStreamer::GetConfiguration() const
+	{
+		return Configuration;
+	}
+
+	void FBuildInstallStreamer::Initialise()
+	{
+		// Queue update to chunk data size cache on main thread
+		AsyncHelpers::ExecuteOnCustomThread<void>(
+			[this]()
+			{
+				ChunkDataSizeProvider->AddManifestData(BuildPatchManifest);
+			},
+			TickQueue).Wait();
+	}
+
+	void FBuildInstallStreamer::RequestWorkerThread()
+	{
+		Initialise();
+
+		TArray<FStreamRequest> Requests;
+		double RequestTime = 0;
+
+		auto CompleteRequest = [this, &Requests, &RequestTime]()
+		{
+			for (FStreamRequest& Request : Requests)
+			{
+				FBuildPatchStreamResult Result;
+				Result.Request = MoveTemp(Request.Get<0>());
+				Result.ErrorType = InstallerError->GetErrorType();
+				Result.ErrorCode = InstallerError->GetErrorCode();
+				// TODO: Not accurate.. How to fix this
+				Result.TotalDownloaded = DownloadServiceStatistics->GetBytesDownloaded();
+				UE_LOG(LogBuildInstallStreamer, Verbose, TEXT("Completed request %s in %s"), *FString::Join(Result.Request, TEXT(",")), *FPlatformTime::PrettyTime(RequestTime));
+				Request.Get<2>().ExecuteIfBound(MoveTemp(Result));
+			}
+		};
+
+		auto CancelRequest = [this, &Requests]()
+		{
+			for (FStreamRequest& Request : Requests)
+			{
+				FBuildPatchStreamResult Result;
+				Result.Request = MoveTemp(Request.Get<0>());
+				Result.ErrorType = InstallerError->GetErrorType();
+				Result.ErrorCode = InstallerError->GetErrorCode();
+				Result.TotalDownloaded = 0;
+				UE_LOG(LogBuildInstallStreamer, Verbose, TEXT("Cancelled request %s"), *FString::Join(Result.Request, TEXT(",")));
+				Request.Get<2>().ExecuteIfBound(MoveTemp(Result));
+			}
+		};
+
+		while (!bIsShuttingDown)
+		{
+			// If we were asked to cancel all requests, then do this first
+			if (InstallerError->IsCancelled())
+			{
+				while (RequestQueue.Dequeue(Requests.AddDefaulted_GetRef())) {}
+				Requests.Pop();
+				if (Requests.Num() > 0)
+				{
+					CallOnDelegateThread(CancelRequest);
+					Requests.Empty();
+				}
+			}
+
+			// Get queued requests to process.
+			if (Configuration.bShouldBatch)
+			{
+				while (RequestQueue.Dequeue(Requests.AddDefaulted_GetRef())) {}
+				Requests.Pop();
+			}
+			else if(!RequestQueue.Dequeue(Requests.AddDefaulted_GetRef()))
+			{
+				Requests.Pop();
+			}
+			// Run the streaming for received requests.
+			if (Requests.Num() > 0)
+			{
+				for (const FStreamRequest& Request : Requests)
+				{
+					UE_LOG(LogBuildInstallStreamer, Verbose, TEXT("Begin request %s"), *FString::Join(Request.Get<0>(), TEXT(",")));
+				}
+
+				TProcessTimer<class FStatsCollector, false> RequestTimer;
+				RequestTimer.Start();
+				AsyncHelpers::ExecuteOnCustomThread<void>([this]() { DownloadServiceStatistics->Reset(); }, TickQueue);
+				TSharedPtr<IVirtualFileCache, ESPMode::ThreadSafe> VirtualFileCache = IVirtualFileCache::CreateVirtualFileCache();
+				InstallerError->Reset();
+
+				// Setup file build list.
+				FVirtualFileConstructorConfiguration VFCConfig;
+				for (const FStreamRequest& Request : Requests)
+				{
+					const bool bIsTagRequest = Request.Get<1>();
+					if (bIsTagRequest)
+					{
+						BuildPatchManifest->GetTaggedFileList(Request.Get<0>(), VFCConfig.FilesToConstruct);
+					}
+					else
+					{
+						VFCConfig.FilesToConstruct.Append(Request.Get<0>());
+					}
+				}
+
+				// Remove all files that already exist.
+				for (auto FileIt = VFCConfig.FilesToConstruct.CreateIterator(); FileIt; ++FileIt)
+				{
+					const FFileManifest* FileManifest = ManifestSet->GetNewFileManifest(*FileIt);
+					if (FileManifest != nullptr && VirtualFileCache->DoesChunkExist(FileManifest->FileHash))
+					{
+						FileIt.RemoveCurrent();
+					}
+				}
+
+				// Composition root.
+				FCloudSourceConfig CloudSourceConfig(Configuration.CloudDirectories);
+				CloudSourceConfig.bBeginDownloadsOnFirstGet = false;
+				CloudSourceConfig.bRunOwnThread = false;
+				TUniquePtr<IChunkReferenceTracker> ChunkReferenceTracker(FChunkReferenceTrackerFactory::Create(
+					ManifestSet.Get(),
+					VFCConfig.FilesToConstruct));
+				TUniquePtr<IChunkEvictionPolicy> MemoryEvictionPolicy(FChunkEvictionPolicyFactory::Create(
+					ChunkReferenceTracker.Get()));
+				TUniquePtr<IMemoryChunkStore> CloudChunkStore(FMemoryChunkStoreFactory::Create(
+					FMath::Clamp<int32>(CloudSourceConfig.PreFetchMaximum, 64, 512),
+					MemoryEvictionPolicy.Get(),
+					nullptr,
+					MemoryChunkStoreStatistics.Get()));
+				TUniquePtr<ICloudChunkSource> CloudChunkSource(FCloudChunkSourceFactory::Create(
+					MoveTemp(CloudSourceConfig),
+					Platform.Get(),
+					CloudChunkStore.Get(),
+					DownloadService.Get(),
+					ChunkReferenceTracker.Get(),
+					ChunkDataSerialization.Get(),
+					MessagePump.Get(),
+					InstallerError.Get(),
+					DownloadConnectionCount.Get(),
+					CloudChunkSourceStatistics.Get(),
+					ManifestSet.Get(),
+					ChunkReferenceTracker->GetReferencedChunks()));
+				// TODO: Use promise and future so we can wait for this call to end before destructing, either here or in class.
+				CloudQueue.Enqueue([&CloudChunkSource]() { CloudChunkSource->ThreadRun(); });
+				CloudTrigger->Trigger();
+				FVirtualFileConstructorDependencies VFCDepends{
+					ManifestSet.Get(),
+					VirtualFileCache.Get(),
+					CloudChunkSource.Get(),
+					ChunkReferenceTracker.Get(),
+					InstallerError.Get(),
+					FileConstructorStatistics.Get()
+				};
+				TUniquePtr<FVirtualFileConstructor> VirtualFileConstructor(FVirtualFileConstructorFactory::Create(MoveTemp(VFCConfig), MoveTemp(VFCDepends)));
+				TArray<IControllable*> Controllables;
+				Controllables.Add(CloudChunkSource.Get());
+				int32 ErrorHandle = InstallerError->RegisterForErrors([&CloudChunkSource]()
+					{
+						CloudChunkSource->Abort();
+					});
+
+				// Run the construction.
+				bool bSuccess = VirtualFileConstructor->Run();
+				RequestTimer.Stop();
+				RequestTime = RequestTimer.GetSeconds();
+
+				// Keep a ref during OnComplete execution, the provided delegate can release all other refs causing our dctor to block.
+				// This way, our destructor will instead be called once we exit this request scope causing this thread to quit also.
+				FBuildInstallStreamerRef SharedThis = AsShared();
+				CallOnDelegateThread(CompleteRequest);
+				InstallerError->UnregisterForErrors(ErrorHandle);
+				Requests.Empty();
+			}
+			else
+			{
+				// Else wait on trigger.
+				RequestTrigger->Wait();
+			}
+		}
+		// Ensure we call all remaining requests delegates as cancelled.
+		while (RequestQueue.Dequeue(Requests.AddDefaulted_GetRef())) {}
+		Requests.Pop();
+		if (Requests.Num() > 0)
+		{
+			CallOnDelegateThread(CancelRequest);
+		}
+	}
+
+	void FBuildInstallStreamer::CloudWorkerThread()
+	{
+		while (!bIsShuttingDown)
+		{
+			TFunction<void()> Task;
+			if (CloudQueue.Dequeue(Task))
+			{
+				Task();
+			}
+			else
+			{
+				CloudTrigger->Wait();
+			}
+		}
+	}
+
+	void FBuildInstallStreamer::CallOnDelegateThread(const TFunction<void()>& Callback)
+	{
+		if (Configuration.bMainThreadDelegates)
+		{
+			AsyncHelpers::ExecuteOnCustomThread<void>(Callback, TickQueue).Wait();
+		}
+		else
+		{
+			Callback();
+		}
+	}
+
+	bool FBuildInstallStreamer::Tick()
+	{
+		const bool bKeepTicking = true;
+
+		MessagePump->PumpMessages(MessageHandlers);
+
+		TFunction<void()> TickFunc;
+		while (TickQueue.Dequeue(TickFunc))
+		{
+			TickFunc();
+		}
+
+		return bKeepTicking;
+	}
+
+	void FBuildInstallStreamer::CancelAllRequests()
+	{
+		const bool bInMainThread = IsInGameThread();
+		bool bAllRequestsFlushed = false;
+		// If we are asked to cancel from the game thread, we need to tick ourselves so we can pump cancelled request delegates.
+		// We can add our own empty request right away so we know when all requests in the queue at time of this function call have been executed.
+		// If we are not in the game thread, then there needs to be ticks externally.
+		if (bInMainThread)
+		{
+			QueueFilesByName({}, FBuildPatchStreamCompleteDelegate::CreateLambda([&](FBuildPatchStreamResult) { bAllRequestsFlushed = true; }));
+		}
+		InstallerError->SetError(EBuildPatchInstallError::UserCanceled, UserCancelErrorCodes::UserRequested);
+		if (bInMainThread)
+		{
+			while (!bAllRequestsFlushed)
+			{
+				Tick();
+			}
+		}
+	}
+
+	void FBuildInstallStreamer::PreExit()
+	{
+		bIsShuttingDown = true;
+		InstallerError->SetError(EBuildPatchInstallError::ApplicationClosing, ApplicationClosedErrorCodes::ApplicationClosed);
+		RequestTrigger->Trigger();
+		CloudTrigger->Trigger();
+		// We need to tick until the thread exits here (we must be in main thread now).
+		check(IsInGameThread());
+		while (!RequestWorker.IsReady() || !CloudWorker.IsReady())
+		{
+			Tick();
+		}
+	}
+
+	FBuildInstallStreamer* FBuildInstallStreamerFactory::Create(FBuildInstallStreamerConfiguration Configuration)
+	{
+		return new FBuildInstallStreamer(MoveTemp(Configuration));
+	}
+}
diff --git a/Engine/Source/Runtime/Online/BuildPatchServices/Private/Installer/BuildInstallStreamer.h b/Engine/Source/Runtime/Online/BuildPatchServices/Private/Installer/BuildInstallStreamer.h
new file mode 100644
--- /dev/null
+++ b/Engine/Source/Runtime/Online/BuildPatchServices/Private/Installer/BuildInstallStreamer.h
@@ -0,0 +1,111 @@
+// Copyright Epic Games, Inc. All Rights Reserved.
+#pragma once
+
+#include "Async/Async.h"
+#include "BuildPatchManifest.h"
+#include "BuildPatchProgress.h"
+#include "Containers/Queue.h"
+#include "Containers/Set.h"
+#include "Containers/UnrealString.h"
+#include "HAL/Event.h"
+#include "HAL/ThreadSafeBool.h"
+#include "Interfaces/IBuildInstallStreamer.h"
+#include "Interfaces/IBuildPatchServicesModule.h"
+#include "Templates/UniquePtr.h"
+#include "Templates/Tuple.h"
+
+namespace BuildPatchServices
+{
+	class IBuildManifestSet;
+	class IHttpManager;
+	class IFileSystem;
+	class IChunkDataSerialization;
+	class IPlatform;
+	class IInstallerError;
+	class IInstallerAnalytics;
+	class IFileOperationTracker;
+	class ISpeedRecorder;
+	class IChunkDataSizeProvider;
+	class IDownloadServiceStatistics;
+	class IMemoryChunkStoreStatistics;
+	class ICloudChunkSourceStatistics;
+	class IFileConstructorStatistics;
+	class IDownloadConnectionCount;
+	class IDownloadService;
+	class IMessagePump;
+
+	class FBuildInstallStreamer
+		: public IBuildInstallStreamer
+		, public TSharedFromThis<FBuildInstallStreamer, ESPMode::ThreadSafe>
+	{
+	public:
+		FBuildInstallStreamer(FBuildInstallStreamerConfiguration Configuration);
+		~FBuildInstallStreamer();
+
+	public:
+		// IBuildInstallStreamer interface begin.
+		virtual void QueueFilesByTag(TSet<FString> Tags, FBuildPatchStreamCompleteDelegate OnComplete) override;
+		virtual void QueueFilesByName(TSet<FString> Files, FBuildPatchStreamCompleteDelegate OnComplete) override;
+		virtual void RegisterMessageHandler(FMessageHandler* MessageHandler) override;
+		virtual void UnregisterMessageHandler(FMessageHandler* MessageHandler) override;
+		virtual const FBuildInstallStreamerConfiguration& GetConfiguration() const override;
+		// IBuildInstallStreamer interface end.
+
+		bool Tick();
+		void CancelAllRequests();
+		void PreExit();
+
+	private:
+		void Initialise();
+		void RequestWorkerThread();
+		void CloudWorkerThread();
+		void CallOnDelegateThread(const TFunction<void()>& Callback);
+
+	private:
+		const FBuildInstallStreamerConfiguration Configuration;
+		FBuildPatchAppManifestPtr BuildPatchManifest;
+		FBuildPatchProgress BuildProgress;
+		TUniquePtr<IBuildManifestSet> ManifestSet;
+		TUniquePtr<IHttpManager> HttpManager;
+		TUniquePtr<IFileSystem> FileSystem;
+		TUniquePtr<IChunkDataSerialization> ChunkDataSerialization;
+		TUniquePtr<IPlatform> Platform;
+		TUniquePtr<IInstallerAnalytics> InstallerAnalytics;
+		TUniquePtr<IInstallerError> InstallerError;
+		TUniquePtr<IFileOperationTracker> FileOperationTracker;
+		TUniquePtr<ISpeedRecorder> DownloadSpeedRecorder;
+		TUniquePtr<ISpeedRecorder> FileReadSpeedRecorder;
+		TUniquePtr<ISpeedRecorder> FileWriteSpeedRecorder;
+		TUniquePtr<IChunkDataSizeProvider> ChunkDataSizeProvider;
+		TUniquePtr<IDownloadServiceStatistics> DownloadServiceStatistics;
+		TUniquePtr<IMemoryChunkStoreStatistics> MemoryChunkStoreStatistics;
+		TUniquePtr<ICloudChunkSourceStatistics> CloudChunkSourceStatistics;
+		TUniquePtr<IFileConstructorStatistics> FileConstructorStatistics;
+		TUniquePtr<IDownloadConnectionCount> DownloadConnectionCount;
+		TUniquePtr<IDownloadService> DownloadService;
+		TUniquePtr<IMessagePump> MessagePump;
+		TArray<FMessageHandler*> MessageHandlers;
+		FThreadSafeBool bIsShuttingDown;
+		TFuture<void> RequestWorker;
+		TFuture<void> CloudWorker;
+		typedef TTuple<TSet<FString>, bool, FBuildPatchStreamCompleteDelegate> FStreamRequest;
+		TQueue<FStreamRequest, EQueueMode::Spsc> RequestQueue;
+		TQueue<TFunction<void()>, EQueueMode::Spsc> CloudQueue;
+		TQueue<TFunction<void()>, EQueueMode::Spsc> TickQueue;
+		FEvent* RequestTrigger;
+		FEvent* CloudTrigger;
+	};
+
+	/**
+	 * A factory for creating an IBuildInstallStreamer instance.
+	 */
+	class FBuildInstallStreamerFactory
+	{
+	public:
+		static FBuildInstallStreamer* Create(FBuildInstallStreamerConfiguration Configuration);
+	};
+}
+
+typedef TSharedPtr< BuildPatchServices::FBuildInstallStreamer, ESPMode::ThreadSafe > FBuildInstallStreamerPtr;
+typedef TSharedRef< BuildPatchServices::FBuildInstallStreamer, ESPMode::ThreadSafe > FBuildInstallStreamerRef;
+typedef TWeakPtr< BuildPatchServices::FBuildInstallStreamer, ESPMode::ThreadSafe > FBuildInstallStreamerWeakPtr;
\ No newline at end of file
diff --git a/Engine/Source/Runtime/Online/BuildPatchServices/Private/Installer/CloudChunkSource.cpp b/Engine/Source/Runtime/Online/BuildPatchServices/Private/Installer/CloudChunkSource.cpp
--- a/Engine/Source/Runtime/Online/BuildPatchServices/Private/Installer/CloudChunkSource.cpp
+++ b/Engine/Source/Runtime/Online/BuildPatchServices/Private/Installer/CloudChunkSource.cpp
@@ -145,13 +145,16 @@ namespace BuildPatchServices
 		virtual void SetUnavailableChunksCallback(TFunction<void(TSet<FGuid>)> Callback) override;
 		// IChunkSource interface end.
 
+		// ICloudChunkSource interface begin.
+		virtual void ThreadRun() override;
+		// ICloudChunkSource interface end.
+
 	private:
 		void EnsureAquiring(const FGuid& DataId);
 		const FString& GetCloudRoot(int32 RetryNum) const;
 		float GetRetryDelay(int32 RetryNum);
 		EBuildPatchDownloadHealth GetDownloadHealth(bool bIsDisconnected, float ChunkSuccessRate);
 		FGuid GetNextTask(const TMap<FGuid, FTaskInfo>& TaskInfos, const TMap<int32, FGuid>& InFlightDownloads, const TSet<FGuid>& TotalRequiredChunks, const TSet<FGuid>& PriorityRequests, const TSet<FGuid>& FailedDownloads, const TSet<FGuid>& Stored, TArray<FGuid>& DownloadQueue, EBuildPatchDownloadHealth DownloadHealth);
-		void ThreadRun();
 		void OnDownloadProgress(int32 RequestId, int32 BytesSoFar);
 		void OnDownloadComplete(int32 RequestId, const FDownloadRef& Download);
 
@@ -168,6 +171,7 @@ namespace BuildPatchServices
 		ICloudChunkSourceStat* CloudChunkSourceStat;
 		IBuildManifestSet* ManifestSet;
 		const TSet<FGuid> InitialDownloadSet;
+		TPromise<void> Promise;
 		TFuture<void> Future;
 		FDownloadProgressDelegate OnDownloadProgressDelegate;
 		FDownloadCompleteDelegate OnDownloadCompleteDelegate;
@@ -234,6 +238,7 @@ namespace BuildPatchServices
 		, CloudChunkSourceStat(InCloudChunkSourceStat)
 		, ManifestSet(InManifestSet)
 		, InitialDownloadSet(MoveTemp(InInitialDownloadSet))
+		, Promise()
 		, Future()
 		, OnDownloadProgressDelegate(FDownloadProgressDelegate::CreateThreadSafeSP(DownloadDelegates, &FDownloadDelegates::OnDownloadProgress))
 		, OnDownloadCompleteDelegate(FDownloadCompleteDelegate::CreateThreadSafeSP(DownloadDelegates, &FDownloadDelegates::OnDownloadComplete))
@@ -246,8 +251,15 @@ namespace BuildPatchServices
 		, RequestedDownloads()
 		, DownloadCount(InDownloadConnectionCount)
 	{
-		TFunction<void()> Task = [this]() { return ThreadRun(); };
-		Future = Async(EAsyncExecution::ThreadIfForkSafe, MoveTemp(Task));
+		if (Configuration.bRunOwnThread)
+		{
+			TFunction<void()> Task = [this]() { return ThreadRun(); };
+			Future = Async(EAsyncExecution::ThreadIfForkSafe, MoveTemp(Task));
+		}
+		else
+		{
+			Future = Promise.GetFuture();
+		}
 	}
 
 	FCloudChunkSource::~FCloudChunkSource()
@@ -718,6 +730,9 @@ namespace BuildPatchServices
 		CloudChunkSourceStat->OnDownloadHealthUpdated(TrackedDownloadHealth);
 		CloudChunkSourceStat->OnSuccessRateUpdated(ChunkSuccessRate.GetOverall());
 		CloudChunkSourceStat->OnActiveRequestCountUpdated(0);
+
+		// The promise should always be set, even if not needed as destruction of an unset promise will assert.
+		Promise.SetValue();
 	}
 
 	void FCloudChunkSource::OnDownloadProgress(int32 RequestId, int32 BytesSoFar)
diff --git a/Engine/Source/Runtime/Online/BuildPatchServices/Private/Installer/CloudChunkSource.h b/Engine/Source/Runtime/Online/BuildPatchServices/Private/Installer/CloudChunkSource.h
--- a/Engine/Source/Runtime/Online/BuildPatchServices/Private/Installer/CloudChunkSource.h
+++ b/Engine/Source/Runtime/Online/BuildPatchServices/Private/Installer/CloudChunkSource.h
@@ -31,6 +31,12 @@ namespace BuildPatchServices
 	{
 	public:
 		virtual ~ICloudChunkSource() {}
+
+		/**
+		 * Should be called from the desired thread if bRunOwnThread is false.
+		 * It will block until until aborted by the owning system.
+		 */
+		virtual void ThreadRun() = 0;
 	};
 
 	/**
@@ -60,6 +66,8 @@ namespace BuildPatchServices
 		bool bBeginDownloadsOnFirstGet;
 		// The minimum time to allow a http download before assessing it as affected by TCP zero window issue.
 		float TcpZeroWindowMinimumSeconds;
+		// Whether cloud source should run its own thread.
+		bool bRunOwnThread;
 
 		/**
 		 * Constructor which sets usual defaults, and takes params for values that cannot use a default.
@@ -73,6 +81,7 @@ namespace BuildPatchServices
 			, DisconnectedDelay(5.0f)
 			, bBeginDownloadsOnFirstGet(true)
 			, TcpZeroWindowMinimumSeconds(20.0f)
+			, bRunOwnThread(true)
 		{
 			const float RetryFloats[] = {0.5f, 1.0f, 1.0f, 3.0f, 3.0f, 10.0f, 10.0f, 20.0f, 20.0f, 30.0f};
 			RetryDelayTimes.Empty(UE_ARRAY_COUNT(RetryFloats));
diff --git a/Engine/Source/Runtime/Online/BuildPatchServices/Private/Installer/InstallerError.cpp b/Engine/Source/Runtime/Online/BuildPatchServices/Private/Installer/InstallerError.cpp
--- a/Engine/Source/Runtime/Online/BuildPatchServices/Private/Installer/InstallerError.cpp
+++ b/Engine/Source/Runtime/Online/BuildPatchServices/Private/Installer/InstallerError.cpp
@@ -79,6 +79,7 @@ namespace BuildPatchServices
 		virtual void SetError(EBuildPatchInstallError InErrorType, const TCHAR* InErrorSubType, uint32 InErrorCode, FText InErrorText) override;
 		virtual int32 RegisterForErrors(FOnErrorDelegate Delegate) override;
 		virtual void UnregisterForErrors(int32 Handle) override;
+		virtual void Reset() override;
 		// IInstallerError interface end.
 
 	private:
@@ -185,6 +186,14 @@ namespace BuildPatchServices
 		RegisteredDelegates.Remove(Handle);
 	}
 
+	void FInstallerError::Reset()
+	{
+		FScopeLock ScopeLock(&ThreadLockCs);
+		ErrorType = EBuildPatchInstallError::NoError;
+		ErrorCode = InstallErrorPrefixes::ErrorTypeStrings[static_cast<int32>(ErrorType)];
+		ErrorText = GetStandardErrorText(ErrorType);
+	}
+
 	IInstallerError* FInstallerErrorFactory::Create()
 	{
 		return new FInstallerError();
diff --git a/Engine/Source/Runtime/Online/BuildPatchServices/Private/Installer/InstallerError.h b/Engine/Source/Runtime/Online/BuildPatchServices/Private/Installer/InstallerError.h
--- a/Engine/Source/Runtime/Online/BuildPatchServices/Private/Installer/InstallerError.h
+++ b/Engine/Source/Runtime/Online/BuildPatchServices/Private/Installer/InstallerError.h
@@ -197,6 +197,11 @@ namespace BuildPatchServices
 		 * @param Handle        The value that was returned RegisterForErrors for the delegate to be removed.
 		 */
 		virtual void UnregisterForErrors(int32 Handle) = 0;
+
+		/**
+		 * Reset internal state to be ready for a new process.
+		 */
+		virtual void Reset() = 0;
 	};
 
 	/**
diff --git a/Engine/Source/Runtime/Online/BuildPatchServices/Private/Installer/Statistics/DownloadServiceStatistics.cpp b/Engine/Source/Runtime/Online/BuildPatchServices/Private/Installer/Statistics/DownloadServiceStatistics.cpp
--- a/Engine/Source/Runtime/Online/BuildPatchServices/Private/Installer/Statistics/DownloadServiceStatistics.cpp
+++ b/Engine/Source/Runtime/Online/BuildPatchServices/Private/Installer/Statistics/DownloadServiceStatistics.cpp
@@ -35,6 +35,7 @@ namespace BuildPatchServices
 		virtual int32 GetNumCurrentDownloads() const override;
 		virtual TArray<FDownload> GetCurrentDownloads() const override;
 		virtual TPair<double, uint32> GetImmediateAverageSpeedPerRequest(uint32 MinCount) override;
+		virtual void Reset() override;
 		// IDownloadServiceStatistics interface end.
 
 	private:
@@ -157,6 +158,21 @@ namespace BuildPatchServices
 		}
 		return TPair<double, uint32>(Result, Count);
 	}
+
+	void FDownloadServiceStatistics::Reset()
+	{
+		checkSlow(IsInGameThread());
+
+		TotalBytesReceived = 0;
+		NumSuccessfulDownloads.Set(0);
+		NumFailedDownloads.Set(0);
+		Downloads.Empty();
+
+		FScopeLock Lock(&AverageSpeedCriticalSection);
+		AccumulatedRequestSpeed = 0;
+		AverageSpeedSampleCount = 0;
+	}
+
 	IDownloadServiceStatistics* FDownloadServiceStatisticsFactory::Create(ISpeedRecorder* SpeedRecorder, IDataSizeProvider* DataSizeProvider, IInstallerAnalytics* InstallerAnalytics)
 	{
 		check(SpeedRecorder != nullptr);
diff --git a/Engine/Source/Runtime/Online/BuildPatchServices/Private/Installer/Statistics/DownloadServiceStatistics.h b/Engine/Source/Runtime/Online/BuildPatchServices/Private/Installer/Statistics/DownloadServiceStatistics.h
--- a/Engine/Source/Runtime/Online/BuildPatchServices/Private/Installer/Statistics/DownloadServiceStatistics.h
+++ b/Engine/Source/Runtime/Online/BuildPatchServices/Private/Installer/Statistics/DownloadServiceStatistics.h
@@ -53,6 +53,11 @@ namespace BuildPatchServices
 		 *         and the count of requests completed since the last call.
 		 */
 		virtual TPair<double, uint32> GetImmediateAverageSpeedPerRequest(uint32 MinCount) = 0;
+
+		/**
+		 * Resets all internal statistics.
+		 */
+		virtual void Reset() = 0;
 	};
 
 	/**
diff --git a/Engine/Source/Runtime/Online/BuildPatchServices/Private/Installer/VirtualFileConstructor.cpp b/Engine/Source/Runtime/Online/BuildPatchServices/Private/Installer/VirtualFileConstructor.cpp
new file mode 100644
--- /dev/null
+++ b/Engine/Source/Runtime/Online/BuildPatchServices/Private/Installer/VirtualFileConstructor.cpp
@@ -0,0 +1,108 @@
+// Copyright Epic Games, Inc. All Rights Reserved.
+
+#include "Installer/VirtualFileConstructor.h"
+
+#include "BuildPatchFileConstructor.h"
+#include "Data/ManifestData.h"
+#include "IBuildManifestSet.h"
+#include "Installer/ChunkSource.h"
+#include "Installer/ChunkReferenceTracker.h"
+#include "Installer/InstallerError.h"
+#include "Interfaces/IBuildInstaller.h"
+#include "Logging/LogMacros.h"
+#include "VirtualFileCache.h"
+
+// TODO: Move EBuildPatchInstallError into it's own header header (replacing Interfaces/IBuildInstaller.h include)
+// TODO: Move IFileConstructorStat into it's own header (replacing BuildPatchFileConstructor.h include)
+
+DEFINE_LOG_CATEGORY_STATIC(LogVirtualFileConstructor, Log, All);
+
+namespace BuildPatchServices
+{
+	FVirtualFileConstructor::FVirtualFileConstructor(FVirtualFileConstructorConfiguration InConfiguration, FVirtualFileConstructorDependencies InDependencies)
+		: Configuration(MoveTemp(InConfiguration))
+		, Dependencies(MoveTemp(InDependencies))
+	{
+	}
+
+	FVirtualFileConstructor::~FVirtualFileConstructor()
+	{
+	}
+
+	bool FVirtualFileConstructor::Run()
+	{
+		// TODO: Implement abort / cancel
+		const bool bShouldAbort = false;
+
+		bool bFileSuccess = true;
+		for (auto FileIt = Configuration.FilesToConstruct.CreateConstIterator(); FileIt && bFileSuccess; ++FileIt)
+		{
+			FSHA1 HashState;
+			FSHAHash HashValue;
+			TArray<uint8> FileBytes;
+			const FString& FileToConstruct = *FileIt;
+			const FFileManifest* FileManifest = Dependencies.ManifestSet->GetNewFileManifest(FileToConstruct);
+			bFileSuccess = FileManifest != nullptr;
+			if (bFileSuccess)
+			{
+				FileBytes.Reserve(FileManifest->FileSize);
+				Dependencies.FileConstructorStat->OnFileStarted(FileToConstruct, FileManifest->FileSize);
+				Dependencies.FileConstructorStat->OnFileProgress(FileToConstruct, 0);
+				for (int32 ChunkPartIdx = 0; ChunkPartIdx < FileManifest->ChunkParts.Num() && bFileSuccess && !bShouldAbort; ++ChunkPartIdx)
+				{
+					const FChunkPart& ChunkPart = FileManifest->ChunkParts[ChunkPartIdx];
+					Dependencies.FileConstructorStat->OnChunkGet(ChunkPart.Guid);
+					IChunkDataAccess* ChunkDataAccess = Dependencies.ChunkSource->Get(ChunkPart.Guid);
+					bFileSuccess = ChunkDataAccess != nullptr;
+					if (bFileSuccess)
+					{
+						uint8* Data;
+						ChunkDataAccess->GetDataLock(&Data, nullptr);
+						uint8* DataStart = &Data[ChunkPart.Offset];
+						HashState.Update(DataStart, ChunkPart.Size);
+						FileBytes.Append(DataStart, ChunkPart.Size);
+						ChunkDataAccess->ReleaseDataLock();
+						Dependencies.FileConstructorStat->OnFileProgress(FileToConstruct, FileBytes.Num());
+						bFileSuccess = Dependencies.ChunkReferenceTracker->PopReference(ChunkPart.Guid);
+						if (!bFileSuccess)
+						{
+							UE_LOG(LogVirtualFileConstructor, Error, TEXT("Failed %s due to untracked chunk %s"), *FileToConstruct, *ChunkPart.Guid.ToString());
+							Dependencies.InstallerError->SetError(EBuildPatchInstallError::FileConstructionFail, ConstructionErrorCodes::TrackingError);
+						}
+					}
+					else
+					{
+						UE_LOG(LogVirtualFileConstructor, Error, TEXT("Failed %s due to missing chunk %s"), *FileToConstruct, *ChunkPart.Guid.ToString());
+						Dependencies.InstallerError->SetError(EBuildPatchInstallError::FileConstructionFail, ConstructionErrorCodes::MissingChunkData);
+					}
+				}
+				if (bFileSuccess)
+				{
+					HashState.Final();
+					HashState.GetHash(HashValue.Hash);
+					bFileSuccess = HashValue == FileManifest->FileHash;
+					if (bFileSuccess)
+					{
+						Dependencies.VirtualFileCache->WriteData(HashValue, FileBytes.GetData(), FileBytes.Num());
+					}
+					else
+					{
+						UE_LOG(LogVirtualFileConstructor, Error, TEXT("Verify failed after constructing %s"), *FileToConstruct);
+					}
+				}
+				Dependencies.FileConstructorStat->OnFileCompleted(FileToConstruct, bFileSuccess);
+			}
+			else
+			{
+				UE_LOG(LogVirtualFileConstructor, Error, TEXT("Missing file manifest for %s"), *FileToConstruct);
+				Dependencies.InstallerError->SetError(EBuildPatchInstallError::FileConstructionFail, ConstructionErrorCodes::MissingFileInfo);
+			}
+		}
+		return bFileSuccess;
+	}
+
+	FVirtualFileConstructor* FVirtualFileConstructorFactory::Create(FVirtualFileConstructorConfiguration Configuration, FVirtualFileConstructorDependencies Dependencies)
+	{
+		return new FVirtualFileConstructor(MoveTemp(Configuration), MoveTemp(Dependencies));
+	}
+}
diff --git a/Engine/Source/Runtime/Online/BuildPatchServices/Private/Installer/VirtualFileConstructor.h b/Engine/Source/Runtime/Online/BuildPatchServices/Private/Installer/VirtualFileConstructor.h
new file mode 100644
--- /dev/null
+++ b/Engine/Source/Runtime/Online/BuildPatchServices/Private/Installer/VirtualFileConstructor.h
@@ -0,0 +1,55 @@
+// Copyright Epic Games, Inc. All Rights Reserved.
+#pragma once
+
+#include "Containers/Array.h"
+#include "Containers/UnrealString.h"
+#include "Containers/Set.h"
+
+struct IVirtualFileCache;
+namespace BuildPatchServices
+{
+	class IBuildManifestSet;
+	class IChunkSource;
+	class IChunkReferenceTracker;
+	class IInstallerError;
+	class IFileConstructorStat;
+
+	struct FVirtualFileConstructorConfiguration
+	{
+	public:
+		TSet<FString> FilesToConstruct;
+	};
+
+	struct FVirtualFileConstructorDependencies
+	{
+	public:
+		IBuildManifestSet* ManifestSet;
+		IVirtualFileCache* VirtualFileCache;
+		IChunkSource* ChunkSource;
+		IChunkReferenceTracker* ChunkReferenceTracker;
+		IInstallerError* InstallerError;
+		IFileConstructorStat* FileConstructorStat;
+	};
+
+	class FVirtualFileConstructor
+	{
+	public:
+		FVirtualFileConstructor(FVirtualFileConstructorConfiguration Configuration, FVirtualFileConstructorDependencies Dependencies);
+		~FVirtualFileConstructor();
+
+	public:
+		bool Run();
+
+	private:
+
+	private:
+		const FVirtualFileConstructorConfiguration Configuration;
+		const FVirtualFileConstructorDependencies Dependencies;
+	};
+	
+	class FVirtualFileConstructorFactory
+	{
+	public:
+		static FVirtualFileConstructor* Create(FVirtualFileConstructorConfiguration Configuration, FVirtualFileConstructorDependencies Dependencies);
+	};
+}
\ No newline at end of file
diff --git a/Engine/Source/Runtime/Online/BuildPatchServices/Private/Tests/Fake/InstallerError.fake.h b/Engine/Source/Runtime/Online/BuildPatchServices/Private/Tests/Fake/InstallerError.fake.h
--- a/Engine/Source/Runtime/Online/BuildPatchServices/Private/Tests/Fake/InstallerError.fake.h
+++ b/Engine/Source/Runtime/Online/BuildPatchServices/Private/Tests/Fake/InstallerError.fake.h
@@ -2,6 +2,7 @@
 #pragma once
 
 #include "Installer/InstallerError.h"
+#include "Interfaces/IBuildInstaller.h"
 
 #if WITH_DEV_AUTOMATION_TESTS
 
@@ -60,6 +61,17 @@ namespace BuildPatchServices
 			Delegates.Remove(Handle);
 		}
 
+		virtual void Reset() override
+		{
+			bHasError = false;
+			bIsCancelled = false;
+			bCanRetry = true;
+			ErrorType = EBuildPatchInstallError::NoError;
+			ErrorCode = TEXT("");
+			ErrorText = FText::GetEmpty();
+			// To mimic real imp, delegates are kept.
+		}
+
 	public:
 		bool bHasError;
 		bool bIsCancelled;
diff --git a/Engine/Source/Runtime/Online/BuildPatchServices/Public/BuildPatchSettings.h b/Engine/Source/Runtime/Online/BuildPatchServices/Public/BuildPatchSettings.h
--- a/Engine/Source/Runtime/Online/BuildPatchServices/Public/BuildPatchSettings.h
+++ b/Engine/Source/Runtime/Online/BuildPatchServices/Public/BuildPatchSettings.h
@@ -268,6 +268,23 @@ namespace BuildPatchServices
 		bool bAllowConcurrentExecution;
 	};
 
+	/**
+	 * Defines a list of all the options of a build streamer class.
+	 */
+	struct BUILDPATCHSERVICES_API FBuildInstallStreamerConfiguration
+	{
+	public:
+		IBuildManifestPtr Manifest;
+		// The list of chunk database filenames that will be used to pull patch data from.
+		TArray<FString> ChunkDatabaseFiles;
+		// The list of cloud directory roots that will be used to pull patch data from. Empty array will use module's global setting..
+		TArray<FString> CloudDirectories;
+		// Whether the streamer should batch all requests into each cycle.
+		bool bShouldBatch = true;
+		// Whether completion delegates should be called on main thread, or the streamer's thread,
+		bool bMainThreadDelegates = true;
+	};
+
 	/**
 	 * Defines a list of all options for the build chunking task.
 	 */
diff --git a/Engine/Source/Runtime/Online/BuildPatchServices/Public/Interfaces/IBuildInstallStreamer.h b/Engine/Source/Runtime/Online/BuildPatchServices/Public/Interfaces/IBuildInstallStreamer.h
new file mode 100644
--- /dev/null
+++ b/Engine/Source/Runtime/Online/BuildPatchServices/Public/Interfaces/IBuildInstallStreamer.h
@@ -0,0 +1,59 @@
+// Copyright Epic Games, Inc. All Rights Reserved.
+#pragma once
+
+#include "CoreMinimal.h"
+#include "BuildPatchState.h"
+#include "BuildPatchMessage.h"
+#include "Interfaces/IBuildInstaller.h"
+
+namespace BuildPatchServices
+{
+	struct FBuildInstallStreamerConfiguration;
+}
+
+typedef TSharedPtr<class IBuildInstallStreamer, ESPMode::ThreadSafe> IBuildInstallStreamerPtr;
+typedef TSharedRef<class IBuildInstallStreamer, ESPMode::ThreadSafe> IBuildInstallStreamerRef;
+typedef TWeakPtr<class IBuildInstallStreamer, ESPMode::ThreadSafe> IBuildInstallStreamerWeakPtr;
+
+struct FBuildPatchStreamResult
+{
+public:
+	TSet<FString> Request;
+	EBuildPatchInstallError ErrorType;
+	FString ErrorCode;
+	int64 TotalDownloaded;
+};
+
+DECLARE_DELEGATE_OneParam(FBuildPatchStreamCompleteDelegate, FBuildPatchStreamResult /*Result*/);
+
+class IBuildInstallStreamer
+{
+public:
+	/**
+	 * Virtual destructor.
+	 */
+	virtual ~IBuildInstallStreamer() { }
+	
+	virtual void QueueFilesByTag(TSet<FString> Tags, FBuildPatchStreamCompleteDelegate OnComplete) = 0;
+	virtual void QueueFilesByName(TSet<FString> Files, FBuildPatchStreamCompleteDelegate OnComplete) = 0;
+	virtual void CancelAllRequests() = 0;
+
+	/**
+	 * Registers a message handler with the installer.
+	 * @param MessageHandler    Ptr to the message handler to add. Must not be null.
+	 */
+	virtual void RegisterMessageHandler(BuildPatchServices::FMessageHandler* MessageHandler) = 0;
+
+	/**
+	 * Unregisters a message handler, will no longer receive HandleMessage calls.
+	 * @param MessageHandler    Ptr to the message handler to remove.
+	 */
+	virtual void UnregisterMessageHandler(BuildPatchServices::FMessageHandler* MessageHandler) = 0;
+
+	/**
+	 * Get the configuration object.
+	 * 
+	 * @returns a const reference to the configuration.
+	 */
+	virtual const BuildPatchServices::FBuildInstallStreamerConfiguration& GetConfiguration() const = 0;
+};
\ No newline at end of file
diff --git a/Engine/Source/Runtime/Online/BuildPatchServices/Public/Interfaces/IBuildManifest.h b/Engine/Source/Runtime/Online/BuildPatchServices/Public/Interfaces/IBuildManifest.h
--- a/Engine/Source/Runtime/Online/BuildPatchServices/Public/Interfaces/IBuildManifest.h
+++ b/Engine/Source/Runtime/Online/BuildPatchServices/Public/Interfaces/IBuildManifest.h
@@ -190,6 +190,14 @@ public:
 	 */
 	virtual int64 GetFileSize(const TSet<FString>& Filenames) const = 0;
 
+	/**
+	 * Gets the file hash for a given file
+	 * @param Filename		IN		The filename in the build
+	 * @param OutHash		OUT		Receives the hash value if found
+	 * @return	true if we had the hash for this file
+	 */
+	virtual bool GetFileHash(const FString& Filename, FSHAHash& OutHash) const = 0;
+	
 	/**
 	 * Get the list of install tags in this manifest
 	 * @return the tags referenced.
diff --git a/Engine/Source/Runtime/Online/BuildPatchServices/Public/Interfaces/IBuildPatchServicesModule.h b/Engine/Source/Runtime/Online/BuildPatchServices/Public/Interfaces/IBuildPatchServicesModule.h
--- a/Engine/Source/Runtime/Online/BuildPatchServices/Public/Interfaces/IBuildPatchServicesModule.h
+++ b/Engine/Source/Runtime/Online/BuildPatchServices/Public/Interfaces/IBuildPatchServicesModule.h
@@ -5,6 +5,7 @@
 #include "Modules/ModuleInterface.h"
 #include "Interfaces/IBuildInstaller.h"
 #include "Interfaces/IBuildStatistics.h"
+#include "Interfaces/IBuildInstallStreamer.h"
 #include "Interfaces/IPatchDataEnumeration.h"
 #include "BuildPatchSettings.h"
 
@@ -37,6 +38,13 @@ public:
 	 */
 	virtual ~IBuildPatchServicesModule() { }
 
+	/**
+	 * Factory providing construction of a build streamer class.
+	 * @param Configuration     The configuration for the installer.
+	 * @return an instance of an IBuildInstallStreamer implementation.
+	 */
+	virtual IBuildInstallStreamerRef CreateBuildInstallStreamer(BuildPatchServices::FBuildInstallStreamerConfiguration Configuration) = 0;
+
 	/**
 	 * Factory providing construction of a build installer class.
 	 * @param Configuration     The configuration for the installer.
