Integrate url request filter to http manager.

diff --git a/Engine/Source/Runtime/Core/Private/Misc/URLRequestFilter.cpp b/Engine/Source/Runtime/Core/Private/Misc/URLRequestFilter.cpp
--- a/Engine/Source/Runtime/Core/Private/Misc/URLRequestFilter.cpp
+++ b/Engine/Source/Runtime/Core/Private/Misc/URLRequestFilter.cpp
@@ -18,6 +18,23 @@ namespace UE::Core
 
 FURLRequestFilter::FURLRequestFilter(const TCHAR* ConfigSectionRootName, const FString& ConfigFileName)
 {
+	UpdateConfig(ConfigSectionRootName, ConfigFileName);
+}
+
+FURLRequestFilter::FURLRequestFilter(const FRequestMap& InAllowedRequests)
+	: AllowedRequests(InAllowedRequests)
+{
+}
+
+FURLRequestFilter::FURLRequestFilter(FRequestMap&& InAllowedRequests)
+	: AllowedRequests(MoveTemp(InAllowedRequests))
+{
+}
+
+void FURLRequestFilter::UpdateConfig(const TCHAR* ConfigSectionRootName, const FString& ConfigFileName)
+{
+	AllowedRequests.Empty();
+
 	FConfigFile* File = GConfig->FindConfigFile(ConfigFileName);
 	if (File != nullptr)
 	{
@@ -62,16 +79,6 @@ FURLRequestFilter::FURLRequestFilter(const TCHAR* ConfigSectionRootName, const F
 	}
 }
 
-FURLRequestFilter::FURLRequestFilter(const FRequestMap& InAllowedRequests)
-	: AllowedRequests(InAllowedRequests)
-{
-}
-
-FURLRequestFilter::FURLRequestFilter(FRequestMap&& InAllowedRequests)
-	: AllowedRequests(MoveTemp(InAllowedRequests))
-{
-}
-
 bool FURLRequestFilter::IsRequestAllowed(FStringView InURL) const
 {
 	if (AllowedRequests.IsEmpty())
diff --git a/Engine/Source/Runtime/Core/Public/Misc/URLRequestFilter.h b/Engine/Source/Runtime/Core/Public/Misc/URLRequestFilter.h
--- a/Engine/Source/Runtime/Core/Public/Misc/URLRequestFilter.h
+++ b/Engine/Source/Runtime/Core/Public/Misc/URLRequestFilter.h
@@ -51,6 +51,16 @@ public:
 	FURLRequestFilter(const FRequestMap& InAllowedRequests);
 	FURLRequestFilter(FRequestMap&& InAllowedRequests);
 
+	/**
+	 * Update the allow list by corresponding section in a config file
+	 */
+	void UpdateConfig(const TCHAR* ConfigSectionRootName, const FString& ConfigFileName);
+
+	/**
+	 * Check if allow list is empty
+	 */
+	bool IsEmpty() const { return AllowedRequests.IsEmpty(); }
+
 	/**
 	 * Matches InURL against the configued list of allowed URL schemes and domains. Returns true if the URL is allowed.
 	 */
diff --git a/Engine/Source/Runtime/Online/HTTP/Private/HttpManager.cpp b/Engine/Source/Runtime/Online/HTTP/Private/HttpManager.cpp
--- a/Engine/Source/Runtime/Online/HTTP/Private/HttpManager.cpp
+++ b/Engine/Source/Runtime/Online/HTTP/Private/HttpManager.cpp
@@ -134,6 +134,13 @@ FString FHttpManager::CreateCorrelationId() const
 
 bool FHttpManager::IsDomainAllowed(const FString& Url) const
 {
+	if (!URLRequestFilter.IsEmpty())
+	{
+		return URLRequestFilter.IsRequestAllowed(Url);
+	}
+
+PRAGMA_DISABLE_DEPRECATION_WARNINGS
+
 #if !UE_BUILD_SHIPPING
 #if !(UE_GAME || UE_SERVER)
 	// Allowed domain filtering is opt-in in non-shipping non-game/server builds
@@ -152,7 +159,7 @@ bool FHttpManager::IsDomainAllowed(const FString& Url) const
 #endif
 #endif // !UE_BUILD_SHIPPING
 
-	// check to see if the Domain is allowed (either on the list or the list was empty)
+	// Check to see if the Domain is allowed (either on the list or the list was empty)
 	const TArray<FString>& AllowedDomains = FHttpModule::Get().GetAllowedDomains();
 	if (AllowedDomains.Num() > 0)
 	{
@@ -167,6 +174,8 @@ bool FHttpManager::IsDomainAllowed(const FString& Url) const
 		return false;
 	}
 	return true;
+
+PRAGMA_ENABLE_DEPRECATION_WARNINGS
 }
 
 /*static*/
@@ -193,6 +202,8 @@ void FHttpManager::OnEndFramePostFork()
 
 void FHttpManager::UpdateConfigs()
 {
+	URLRequestFilter.UpdateConfig(TEXT("Online.HttpManager"), GEngineIni);
+
 	ReloadFlushTimeLimits();
 
 	if (Thread)
diff --git a/Engine/Source/Runtime/Online/HTTP/Public/HttpManager.h b/Engine/Source/Runtime/Online/HTTP/Public/HttpManager.h
--- a/Engine/Source/Runtime/Online/HTTP/Public/HttpManager.h
+++ b/Engine/Source/Runtime/Online/HTTP/Public/HttpManager.h
@@ -14,6 +14,7 @@
 #include "HttpPackage.h"
 #include "Interfaces/IHttpRequest.h"
 #include "Misc/EnumRange.h"
+#include "Misc/URLRequestFilter.h"
 #include "Templates/Function.h"
 #include "Templates/SharedPointer.h"
 
@@ -252,6 +253,9 @@ protected:
 	/** Delegate that will get called when a request completes */
 	FHttpManagerRequestCompletedDelegate RequestCompletedDelegate;
 
+	/** Url request filter, if specified in the config, it will launch http request only when the url is in the allowlist */
+	UE::Core::FURLRequestFilter URLRequestFilter;
+
 	struct FHttpFlushTimeLimit
 	{
 		/**
diff --git a/Engine/Source/Runtime/Online/HTTP/Public/HttpModule.h b/Engine/Source/Runtime/Online/HTTP/Public/HttpModule.h
--- a/Engine/Source/Runtime/Online/HTTP/Public/HttpModule.h
+++ b/Engine/Source/Runtime/Online/HTTP/Public/HttpModule.h
@@ -291,6 +291,7 @@ public:
 	/**
 	 * @returns the list of domains allowed to be visited in a shipping build
 	 */
+	UE_DEPRECATED(5.3, "GetAllowedDomains has been deprecated. URLRequestFilter should be used instead.")
 	inline const TArray<FString>& GetAllowedDomains() const
 	{
 		return AllowedDomains;
