Improve POSIX ICMP code:
* Check to make sure ReadSize != -1
* Account for IP options
* Add a bounds check to make sure icmp_id and icmp_seq are received
* Add console command so we can test those changes

diff --git a/Engine/Source/Runtime/Online/ICMP/Private/IcmpModule.cpp b/Engine/Source/Runtime/Online/ICMP/Private/IcmpModule.cpp
--- a/Engine/Source/Runtime/Online/ICMP/Private/IcmpModule.cpp
+++ b/Engine/Source/Runtime/Online/ICMP/Private/IcmpModule.cpp
@@ -1,6 +1,8 @@
 // Copyright Epic Games, Inc. All Rights Reserved.
 
 #include "IcmpModule.h"
+#include "Icmp.h"
+#include "HAL/IConsoleManager.h"
 
 IMPLEMENT_MODULE(FIcmpModule, Icmp);
 
@@ -14,3 +16,35 @@ void FIcmpModule::ShutdownModule()
 {
 }
 
+
+bool FIcmpModule::Exec(UWorld* InWorld, const TCHAR* Cmd, FOutputDevice& Ar)
+{
+	// Ignore any execs that don't start with Icmp
+	if (FParse::Command(&Cmd, TEXT("Icmp")))
+	{
+		if (FParse::Command(&Cmd, TEXT("Ping")))
+		{
+			FString Url;
+			FParse::Token(Cmd, Url, true);
+			if (Url.IsEmpty())
+			{
+				UE_LOG(LogConsoleResponse, Warning, TEXT("Url is missing, console command is Icmp Ping URL, defaulting to 127.0.0.1"));
+				Url = TEXT("127.0.0.1");
+			}
+			FIcmp::IcmpEcho(Url, 1.0f, [](FIcmpEchoResult Result)
+				{
+					if (Result.Status == EIcmpResponseStatus::Success)
+					{
+						const double CurPing = static_cast<double>(Result.Time) * 1000.0;
+						UE_LOG(LogConsoleResponse, Display, TEXT("Ping success: Reply from %s time=%.02f ms"), *Result.ReplyFrom, CurPing);
+					}
+					else
+					{
+						UE_LOG(LogConsoleResponse, Display, TEXT("Ping failure"));
+					}
+				});
+			return true;
+		}
+	}
+	return false;
+}
diff --git a/Engine/Source/Runtime/Online/ICMP/Private/IcmpPosix.cpp b/Engine/Source/Runtime/Online/ICMP/Private/IcmpPosix.cpp
--- a/Engine/Source/Runtime/Online/ICMP/Private/IcmpPosix.cpp
+++ b/Engine/Source/Runtime/Online/ICMP/Private/IcmpPosix.cpp
@@ -59,14 +59,15 @@ uint32 HtoNL(uint32 val)
 
 FIcmpEchoResult IcmpEchoImpl(ISocketSubsystem* SocketSub, const FString& TargetAddress, float Timeout)
 {
-	static const SIZE_T IpHeaderSize = sizeof(struct ip);
+	static const SIZE_T IpHeaderFixedPartSize = sizeof(struct ip);
+	static const SIZE_T IpHeaderVariablePartMaxSize = 40;
 	static const SIZE_T IcmpHeaderSize = sizeof(struct icmp);
 
 	// The packet we send is just the ICMP header plus our payload
 	static const SIZE_T PacketSize = IcmpHeaderSize + IcmpPosix::IcmpPayloadSize;
 
 	// The result read back will need a room for the IP header as well the icmp echo reply packet
-	static const SIZE_T ResultPacketSize = IpHeaderSize + PacketSize;
+	static const SIZE_T MaxResultPacketSize = IpHeaderFixedPartSize + IpHeaderVariablePartMaxSize + PacketSize;
 	static int PingSequence = 0;
 
 	FIcmpEchoResult Result;
@@ -109,7 +110,7 @@ FIcmpEchoResult IcmpEchoImpl(ISocketSubsystem* SocketSub, const FString& TargetA
 	// Calculate the IP checksum
 	PacketHeader->icmp_cksum = CalculateChecksum(Packet, PacketSize);
 
-	uint8 ResultBuffer[ResultPacketSize];
+	uint8 ResultBuffer[MaxResultPacketSize];
 
 	// We can only have one ping in flight at once, as otherwise we risk swapping echo replies between requests
 	{
@@ -139,7 +140,7 @@ FIcmpEchoResult IcmpEchoImpl(ISocketSubsystem* SocketSub, const FString& TargetA
 				}
 				else if (NumReady == 1)
 				{
-					int ReadSize = recv(IcmpSocket, ResultBuffer, ResultPacketSize, 0);
+					int ReadSize = recv(IcmpSocket, ResultBuffer, MaxResultPacketSize, 0);
 
 					double EndTime = FPlatformTime::Seconds();
 
@@ -148,19 +149,27 @@ FIcmpEchoResult IcmpEchoImpl(ISocketSubsystem* SocketSub, const FString& TargetA
 
 					TimeLeft = FPlatformMath::Max(0.0, (double)Timeout - Result.Time);
 
-					if (ReadSize > IpHeaderSize)
+					if (ReadSize == -1)
 					{
-						struct ip* IpHeader = reinterpret_cast<struct ip*>(ResultBuffer);
-						if (IpHeader->ip_p != IPPROTO_ICMP)
+						Result.Status = EIcmpResponseStatus::InternalError;
+						bDone = true;
+					}
+					else if (ReadSize > IpHeaderFixedPartSize)
+					{
+						struct ip* IpHeaderFixedPart = reinterpret_cast<struct ip*>(ResultBuffer);
+						if (IpHeaderFixedPart->ip_p != IPPROTO_ICMP)
 						{
 							// We got a non-ICMP packet back??!
 						}
 						else
 						{
-							PacketHeader = reinterpret_cast<struct icmp*>(ResultBuffer + IpHeaderSize);
-							Result.ReplyFrom = IcmpPosix::IpToString(&IpHeader->ip_src.s_addr);
-							switch (PacketHeader->icmp_type)
+							SIZE_T IpHeaderSize = IpHeaderFixedPart->ip_hl;
+							if (ReadSize >= IpHeaderSize + IcmpHeaderSize)
 							{
+								PacketHeader = reinterpret_cast<struct icmp*>(ResultBuffer + IpHeaderSize);
+								Result.ReplyFrom = IcmpPosix::IpToString(&IpHeaderFixedPart->ip_src.s_addr);
+								switch (PacketHeader->icmp_type)
+								{
 								case ICMP_ECHOREPLY:
 									if (Result.ReplyFrom == ResolvedAddress &&
 										PacketHeader->icmp_id == SentId && PacketHeader->icmp_seq == SentSeq)
@@ -176,6 +185,12 @@ FIcmpEchoResult IcmpEchoImpl(ISocketSubsystem* SocketSub, const FString& TargetA
 									break;
 								default:
 									break;
+								}
+							}
+							else
+							{
+								Result.Status = EIcmpResponseStatus::InternalError;
+								bDone = true;
 							}
 						}
 					}
@@ -189,4 +204,4 @@ FIcmpEchoResult IcmpEchoImpl(ISocketSubsystem* SocketSub, const FString& TargetA
 	return Result;
 }
 
-#endif //PLATFORM_USES_POSIX_ICMP
+#endif //PLATFORM_USES_POSIX_ICMP
\ No newline at end of file
diff --git a/Engine/Source/Runtime/Online/ICMP/Public/IcmpModule.h b/Engine/Source/Runtime/Online/ICMP/Public/IcmpModule.h
--- a/Engine/Source/Runtime/Online/ICMP/Public/IcmpModule.h
+++ b/Engine/Source/Runtime/Online/ICMP/Public/IcmpModule.h
@@ -18,7 +18,7 @@ ICMP_API DECLARE_LOG_CATEGORY_EXTERN(LogIcmp, Display, All);
  * Module for Icmp service utilities
  */
 class FIcmpModule : 
-	public IModuleInterface
+	public IModuleInterface, public FSelfRegisteringExec
 {
 
 public:
@@ -45,16 +45,29 @@ public:
 
 private:
 
+	// FSelfRegisteringExec
+
+	/**
+	 * Handle exec commands starting with "Icmp"
+	 *
+	 * @param InWorld	the world context
+	 * @param Cmd		the exec command being executed
+	 * @param Ar		the archive to log results to
+	 *
+	 * @return true if the handler consumed the input, false to continue searching handlers
+	 */
+	virtual bool Exec(UWorld* InWorld, const TCHAR* Cmd, FOutputDevice& Ar) override;
+
 	// IModuleInterface
 
 	/**
-	 * Called when voice module is loaded
+	 * Called when icmp module is loaded
 	 * Initialize platform specific parts of template handling
 	 */
 	virtual void StartupModule() override;
 	
 	/**
-	 * Called when voice module is unloaded
+	 * Called when icmp module is unloaded
 	 * Shutdown platform specific parts of template handling
 	 */
 	virtual void ShutdownModule() override;
