Make sure to check the current state in dtor before actually close web socket connection from FWinHttpWebSocket.

diff --git a/Engine/Source/Runtime/Online/WebSockets/Private/WinHttp/WinHttpWebSocket.cpp b/Engine/Source/Runtime/Online/WebSockets/Private/WinHttp/WinHttpWebSocket.cpp
--- a/Engine/Source/Runtime/Online/WebSockets/Private/WinHttp/WinHttpWebSocket.cpp
+++ b/Engine/Source/Runtime/Online/WebSockets/Private/WinHttp/WinHttpWebSocket.cpp
@@ -50,12 +50,7 @@ FWinHttpWebSocket::~FWinHttpWebSocket()
 	{
 		if (WebSocket->IsValid())
 		{
-			// We can gracefully close if we're still connected
-			if (!WebSocket->CloseConnection(UE_WEBSOCKET_CLOSE_NORMAL_CLOSURE, FString()))
-			{
-				// If we can't gracefully close, just tear down the connection
-				WebSocket->CancelRequest();
-			}
+			Close(UE_WEBSOCKET_CLOSE_NORMAL_CLOSURE, FString());
 		}
 
 		WebSocket.Reset();
@@ -115,7 +110,7 @@ void FWinHttpWebSocket::Close(const int32 Code, const FString& Reason)
 		case EWebSocketConnectionState::Closed:
 		{
 			// Not connected, ignore close request
-			UE_LOG(LogWebSockets, Verbose, TEXT("WinHttp WebSocket[%p]: Closed socket while in %s state, ignoring"), this, LexToString(State));
+			UE_LOG(LogWebSockets, Verbose, TEXT("WinHttp WebSocket[%p]: Close socket while in %s state, ignoring"), this, LexToString(State));
 			return;
 		}
 		case EWebSocketConnectionState::Connecting:
@@ -132,7 +127,12 @@ void FWinHttpWebSocket::Close(const int32 Code, const FString& Reason)
 
 			if (WebSocket.IsValid())
 			{
-				WebSocket->CloseConnection(Code, Reason);
+				// We can gracefully close if we're still connected
+				if (!WebSocket->CloseConnection(Code, Reason))
+				{
+					// If we can't gracefully close, just tear down the connection
+					WebSocket->CancelRequest();
+				}
 			}
 			else
 			{
