Moving WinHttp files to /Internal/ so other plugins can access them.
HTTP.Build.cs needs to be modified for builds to correctly pick up the moves.

diff --git a/Engine/Source/Runtime/Online/HTTP/HTTP.Build.cs b/Engine/Source/Runtime/Online/HTTP/HTTP.Build.cs
--- a/Engine/Source/Runtime/Online/HTTP/HTTP.Build.cs
+++ b/Engine/Source/Runtime/Online/HTTP/HTTP.Build.cs
@@ -41,7 +41,7 @@ public class HTTP : ModuleRules
 
 		PublicDependencyModuleNames.AddRange(
 			new string[] {
-				"Core"
+				"Core",
 			}
 			);
 
diff --git a/Engine/Source/Runtime/Online/HTTP/Internal/WinHttp/Support/WinHttpConnection.h b/Engine/Source/Runtime/Online/HTTP/Internal/WinHttp/Support/WinHttpConnection.h
new file mode 100644
--- /dev/null
+++ b/Engine/Source/Runtime/Online/HTTP/Internal/WinHttp/Support/WinHttpConnection.h
@@ -0,0 +1,88 @@
+// Copyright Epic Games, Inc. All Rights Reserved.
+
+#pragma once
+
+#if WITH_WINHTTP
+
+#include "CoreMinimal.h"
+#include "WinHttp/Support/WinHttpTypes.h"
+
+using HINTERNET = void*;
+using LPVOID = void*;
+
+class HTTP_API IWinHttpConnection
+	: public TSharedFromThis<IWinHttpConnection, ESPMode::ThreadSafe>
+{
+public:
+	virtual ~IWinHttpConnection() = default;
+	
+	/**
+	 * [Call on Any Thread]
+	 *
+	 * Is this object available to use? If this is false, StartRequest will always fail.
+	 *
+	 * @return True if the connection is valid to use, false if not
+	 */
+	virtual bool IsValid() const = 0;
+
+	/**
+	 * [Call on Any Thread]
+	 *
+	 * Get the request URL for this connection
+	 *
+	 * @return A reference to the request URL for this connection
+	 */
+	virtual const FString& GetRequestUrl() const = 0;
+
+	/**
+	 * [Call on Any Thread in Request Callback]
+	 *
+	 * Get the request handle for this connection. This is only safe to call in a callback for this connection.
+	 *
+	 * @return The request handle, or nullptr if not available
+	 */
+	virtual void* GetHandle() = 0;
+
+	/**
+	 * [Call on Main Thread only]
+	 *
+	 * Start this connection if possible. No callbacks will be fired if StartRequest does not return true.
+	 *
+	 * @return True if the connection was able to start, false if not
+	 */
+	virtual bool StartRequest() = 0;
+
+	/**
+	 * [Call on Any Thread]
+	 *
+	 * Cancel this connection and start the teardown of any handles owned by this connection
+	 *
+	 * @return True if the connection is able to cancel, false if it can't
+	 */
+	virtual bool CancelRequest() = 0;
+
+	/**
+	 * [Call on Any Thread]
+	 *
+	 * Has this connection reached a final state (success or failure)?
+	 *
+	 * @return True if the connection has finished, false if it hasn't started or is still processing
+	 */
+	virtual bool IsComplete() const = 0;
+
+	/**
+	 * [Call on Main Thread only]
+	 *
+	 * Process and call delegates events that are queued for the main thread
+	 */
+	virtual void PumpMessages() = 0;
+
+	/**
+	 * [Call on Any Thread]
+	 *
+	 * Pump the state machine and advance the request's state
+	 */
+	virtual void PumpStates() = 0;
+};
+
+#endif // WITH_WINHTTP
diff --git a/Engine/Source/Runtime/Online/HTTP/Internal/WinHttp/Support/WinHttpConnectionHttp.h b/Engine/Source/Runtime/Online/HTTP/Internal/WinHttp/Support/WinHttpConnectionHttp.h
new file mode 100644
--- /dev/null
+++ b/Engine/Source/Runtime/Online/HTTP/Internal/WinHttp/Support/WinHttpConnectionHttp.h
@@ -0,0 +1,344 @@
+// Copyright Epic Games, Inc. All Rights Reserved.
+
+#pragma once
+
+#if WITH_WINHTTP
+
+#include "CoreMinimal.h"
+#include "Containers/StringView.h"
+#include "HAL/CriticalSection.h"
+
+#include "WinHttp/Support/WinHttpConnection.h"
+#include "WinHttp/Support/WinHttpHandle.h"
+#include "Interfaces/IHttpRequest.h"
+#include "Interfaces/IHttpResponse.h"
+
+#include <atomic>
+
+class FWinHttpSession;
+class FWinHttpRequestContextHttp;
+class FRequestPayload;
+
+using FStringKeyValueMap = TMap<FString, FString>;
+
+DECLARE_DELEGATE_TwoParams(FWinHttpConnectionHttpOnDataTransferred, int32 /*BytesSent*/, int32 /*BytesReceived*/);
+DECLARE_DELEGATE_TwoParams(FWinHttpConnectionHttpOnHeaderReceived, const FString& /*HeaderKey*/, const FString& /*HeaderValue*/);
+DECLARE_DELEGATE_OneParam(FWinHttpConnectionHttpOnRequestComplete, EHttpRequestStatus::Type /*CompletionStatus*/);
+
+class HTTP_API FWinHttpConnectionHttp
+	: public IWinHttpConnection
+{
+public:
+	static TSharedPtr<FWinHttpConnectionHttp, ESPMode::ThreadSafe> CreateHttpConnection(
+		FWinHttpSession& Session,
+		const FString& Verb,
+		const FString& Url,
+		const TMap<FString, FString>& Headers,
+		const TSharedPtr<FRequestPayload, ESPMode::ThreadSafe>& Payload);
+
+	virtual ~FWinHttpConnectionHttp();
+	FWinHttpConnectionHttp(const FWinHttpConnectionHttp& Other) = delete;
+	FWinHttpConnectionHttp(FWinHttpConnectionHttp&& Other) = delete;
+	FWinHttpConnectionHttp& operator=(const FWinHttpConnectionHttp& Other) = delete;
+	FWinHttpConnectionHttp& operator=(FWinHttpConnectionHttp&& Other) = delete;
+
+	//~ Begin IWinHttpConnection Interface
+	virtual bool IsValid() const override;
+	virtual const FString& GetRequestUrl() const override;
+	virtual void* GetHandle() override;
+	virtual bool StartRequest() override;
+	virtual bool CancelRequest() override;
+	virtual bool IsComplete() const override;
+	virtual void PumpMessages() override;
+	virtual void PumpStates() override;
+	//~ End IWinHttpConnection Interface
+
+	/** Provides the current known response code. */
+	EHttpResponseCodes::Type GetResponseCode() const { return ResponseCode; }
+	/**
+	 * Gets the last chunk that is ready for processing on the game thread.
+	 * The caller should either use MoveTemp or call Reset after retrieving the data.
+	 */
+	TArray<uint8>& GetLastChunk() { return GameThreadChunk; }
+	/** Provides the current known response content length. */
+	int32 GetResponseContentLength() const { return ResponseContentLength; }
+	/** Provides the current known headers received. */
+	const TMap<FString, FString>& GetHeadersReceived() const { return HeadersReceived; }
+
+	/// Request Events
+
+	/**
+	 * Bind a handler that is called each tick we receive or send data
+	 *
+	 * @param Handler The Handler to call when we receive or send data
+	 */
+	void SetDataTransferredHandler(FWinHttpConnectionHttpOnDataTransferred&& Handler);
+	/**
+	 * Bind a handler that is called each tick we receive a new header
+	 *
+	 * @param Handler The Handler to call when we receive a new header
+	 */
+	void SetHeaderReceivedHandler(FWinHttpConnectionHttpOnHeaderReceived&& Handler);
+	/**
+	 * Bind a handler that is called when the request completes
+	 *
+	 * @param Handler The Handler to call when the request completes
+	 */
+	void SetRequestCompletedHandler(FWinHttpConnectionHttpOnRequestComplete&& Handler);
+
+protected:
+	FWinHttpConnectionHttp(
+		FWinHttpSession& Session,
+		const FString& Url,
+		const FString& Verb,
+		const bool bIsSecure,
+		const FString& Domain,
+		const TOptional<uint16> Port,
+		const FString& PathAndQuery,
+		const TMap<FString, FString>& Headers,
+		const TSharedPtr<FRequestPayload, ESPMode::ThreadSafe>& Payload);
+
+	/// Request Setup
+
+	/**
+	 * Add to our request's headers. These new headers will be added to the existing headers set, and will
+	 * replace any values that already exist with the same header name.
+	 *
+	 * @param Headers Key/Value Map of headers to add to our request
+	 * @return True if we could set the headers, false otherwise
+	 */
+	bool SetHeaders(const TMap<FString, FString>& Headers);
+
+	/**
+	 * Add to our request's headers. Theis new header will be added to the existing headers set, and will
+	 * replace any values that already exist with the same header name.
+	 *
+	 * @param Key Key of header to add to our request
+	 * @param Value Value of header to set on our request
+	 * @return True if we could set the header, false otherwise
+	 */
+	bool SetHeader(const FString& Key, const FString& Value);
+	/**
+	 * Set the payload object of our request. This may only be called before StartRequest is called.
+	 *
+	 * @param NewPayload the Payload object to set (it should not be modified during the request!)
+	 * @return True if we were able to set the payload, false otherwise
+	 */
+	bool SetPayload(const TSharedRef<FRequestPayload, ESPMode::ThreadSafe>& NewPayload);
+
+	/// Sending Request
+
+	/**
+	 * Reset our data sent and fire off a send request again
+	 *
+	 * @return True if we could start the request, or false otherwise
+	 */
+	bool SendRequest();
+	/**
+	 * Increment our bytes-sent counters
+	 *
+	 * @param AmountSent The amount of bytes to increment our counters by
+	 */
+	void IncrementSentByteCounts(const int32 AmountSent);
+	/**
+	 * Increment our bytes-received counters
+	 *
+	 * @param AmountReceived The amount of bytes to increment our counters by
+	 */
+	void IncrementReceivedByteCounts(const int32 AmountReceived);
+	/** 
+	 * Check if we have more data to send
+	 *
+	 * @return True if we have data to send, or false otherwise
+	 */
+	bool HasRequestBodyToSend() const;
+	/**
+	 * Gives WinHttp the next check of data.
+	 *
+	 * @param True if we asked WinHttp to send the data, or false if we could not send the next chunk
+	 */
+	bool SendAdditionalRequestBody();
+
+	/// Request Response
+
+	/**
+	 * Tell WinHttp that we are ready for the response of our request
+	 *
+	 * @return True if the action completed successfully, false otherwise
+	 */
+	bool RequestResponse();
+	/**
+	 * Read our HTTP response headers and store the results
+	 *
+	 * @return True if the action completed successfully, false otherwise
+	 */
+	bool ProcessResponseHeaders();
+	/**
+	 * Ask WinHttp to asynchronously tell us how large the next chunk will be. If it's 0, we're done.
+	 *
+	 * @return True if the action completed successfully, false otherwise
+	 */
+	bool RequestNextResponseBodyChunkSize();
+	/**
+	 * Ask WinHttp to asynchronously write the next chunk directly into our response object
+	 *
+	 * @return True if the action completed successfully, false otherwise
+	 */
+	bool RequestNextResponseBodyChunkData();
+
+	/// Request State
+
+	/** Set our request as finished and set state to the specified state. This also releases the WinHttp handles. */
+	virtual bool FinishRequest(const EHttpRequestStatus::Type FinalState);
+
+	/// Callback handling
+
+	/** Handles storing the server's IP Address if needed */
+	virtual void HandleConnectedToServer(const wchar_t* ServerIP);
+	/** Handles validating request now that we have server information (certificates, etc) */
+	void HandleSendingRequest();
+	/** Handles SOMETHING not sure yet*/
+	void HandleWriteComplete(const uint32 NumBytesSent);
+	/** Handles completing the request so we may read the headers and body of the response */
+	void HandleSendRequestComplete();
+	/** Handles reading the response headers */
+	virtual void HandleHeadersAvailable();
+	/** Handles writing the http response body to our response body */
+	void HandleDataAvailable(const uint32 NumBytesAvailable);
+	/** Handles querying the next chunk of our http response body */
+	void HandleReadComplete(const uint32 NumBytesRead);
+	/** Handles shutting down the request when there was an error */
+	void HandleRequestError(const uint32 ErrorApiId, const uint32 ErrorCode);
+	/** Handles our request object closing, ending this request */
+	virtual void HandleHandleClosing();
+
+private:
+	/**
+	 * Handle receiving HTTP status callbacks from WinHttp
+	 *
+	 * @param ResourceHandle The handle of the resource which generated this Status
+	 * @param InternetStatus The event that is being triggered
+	 * @param StatusInformation Pointer to optional data for this event, if there is any
+	 * @param StatusInformationLength The length of data available to read in StatusInformation or 0 if not used
+	 */
+	void HandleHttpStatusCallback(HINTERNET ResourceHandle, EWinHttpCallbackStatus Status, void* StatusInformation, uint32 StatusInformationLength);
+	
+	/** Mark our callback as a friend so they can call the above status callback function */
+	friend void CALLBACK UE_WinHttpStatusHttpCallback(HINTERNET, DWORD_PTR, DWORD, LPVOID, DWORD);
+
+	/** Reset both PayloadBuffer and Payload when we don't need them anymore, releasing the memory */
+	void ReleasePayloadData();
+
+protected:
+	/** Critical section used to synchronize access to the below properties across multiple threads */
+	mutable FCriticalSection SyncObject;
+
+	/// Members for internal request state
+	/** Handle to our connection object (do not reuse across requests!) */
+	FWinHttpHandle ConnectionHandle;
+	/** Handle to our request object (do not reuse across requests!)*/
+	FWinHttpHandle RequestHandle;
+	/** Keep-alive for this object, to ensure it does not destruct while a request is in progress */
+	TSharedPtr<FWinHttpConnectionHttp, ESPMode::ThreadSafe> KeepAlive;
+
+	/** The current request state */
+	TOptional<EHttpRequestStatus::Type> FinalState;
+	/** Has this request successfully connected to the server? */
+	bool bConnectedToServer = false;
+	/** Has this request been cancelled? */
+	bool bRequestCancelled = false;
+	/** If true, we have a pending event that needs to be broadcast on the game thread */
+	std::atomic<bool> bHasPendingDelegate{false};
+
+private:
+	FString RequestUrl;
+
+	enum class EState : uint8
+	{
+		/// Initial state
+
+		/** We have not started the request */
+		WaitToStart,
+
+		/// Send Request states
+
+		/** Send our request (or resend it due to redirect to a different host) */
+		SendRequest,
+		/** Waiting for our request headers and/or body to be completely written to the connection */
+		WaitForSendComplete,
+		/** Send an additional chunk of request body */
+		SendAdditionalRequestBody,
+
+		/// Receive Header states
+
+		/** Request the response now that we've finished writing our request */
+		RequestResponse,
+		/** Wait for the response headers to be received from the server */
+		WaitForResponseHeaders,
+		/** Read values from our response headers */
+		ProcessResponseHeaders,
+
+		/// Receive Body states
+
+		/** Request the size of the next chunk of our response body */
+		RequestNextResponseBodyChunkSize,
+		/** Wait for a size of the next chunk from the response body */
+		WaitForNextResponseBodyChunkSize,
+		/** Request the data for the next chunk from the response body to written to our buffer */
+		RequestNextResponseBodyChunkData,
+		/** Waiting for data from the next chunk to be written to our buffer */
+		WaitForNextResponseBodyChunkData,
+
+		/// Final state
+
+		/** We are done */
+		Finished
+	};
+
+	EState CurrentAction = EState::WaitToStart;
+
+	/// Members for sending request
+
+	/** Payload body*/
+	TSharedPtr<FRequestPayload, ESPMode::ThreadSafe> Payload;
+	/** Amount of bytes from payload that has been sent successfully */
+	int32 NumBytesSuccessfullySent = 0;
+	/** Data from the payload currently trying to be written */
+	TArray<uint8> PayloadBuffer;
+
+	/// Members for request response
+
+	/** The response code received from our request */
+	EHttpResponseCodes::Type ResponseCode = EHttpResponseCodes::Unknown;
+	/** The headers we received from our request */
+	TMap<FString, FString> HeadersReceived;
+	/** The current chunk being built in the WinHTTP thread. Always access within a critical section. */
+	TArray<uint8> CurrentChunk;
+	/** The current chunk being built for consumption on the game thread. Only access on the game thread. */
+	TArray<uint8> GameThreadChunk;
+	/** The content length from the response headers, or 0. */
+	int32 ResponseContentLength = 0;
+	/** */
+	int32 BytesWrittenToGameThreadChunk = 0;
+	/** The amount of bytes that are available to read from WinHttpReadData now */
+	TOptional<int32> ResponseBytesAvailable;
+	
+	/// Members for reporting updates to main thread
+
+	/** Called on the game thread during tick if we have sent or received data since last tick */
+	FWinHttpConnectionHttpOnDataTransferred OnDataTransferredHandler;
+	/** Data to report to the game thread about how many bytes have been sent since last tick (not total received) */
+	TOptional<int32> BytesToReportSent;
+	/** Data to report to the game thread about how many bytes have been received since last tick (not total received) */
+	TOptional<int32> BytesToReportReceived;
+
+	/** Called on the game thread during tick if we have received header data since last tick */
+	FWinHttpConnectionHttpOnHeaderReceived OnHeaderReceivedHandler;
+	/** List of headers we have received that we haven't told the game thread about yet */
+	TArray<FString> HeaderKeysToReportReceived;
+
+	/** Called on the game thread if the request has completed since last tick */
+	FWinHttpConnectionHttpOnRequestComplete OnRequestCompleteHandler;
+};
+
+#endif // WITH_WINHTTP
diff --git a/Engine/Source/Runtime/Online/HTTP/Internal/WinHttp/Support/WinHttpErrorHelper.h b/Engine/Source/Runtime/Online/HTTP/Internal/WinHttp/Support/WinHttpErrorHelper.h
new file mode 100644
--- /dev/null
+++ b/Engine/Source/Runtime/Online/HTTP/Internal/WinHttp/Support/WinHttpErrorHelper.h
@@ -0,0 +1,116 @@
+// Copyright Epic Games, Inc. All Rights Reserved.
+
+#pragma once
+
+#if WITH_WINHTTP
+
+#include "CoreMinimal.h"
+
+class HTTP_API FWinHttpErrorHelper
+{
+public:
+	/**
+	 * Log that an WinHttpOpen call failed with the specified ErrorCode
+	 *
+	 * @param ErrorCode The error that occured
+	 */
+	static void LogWinHttpOpenFailure(const uint32 ErrorCode);
+
+	/**
+	 * Log that an WinHttpCloseHandle call failed with the specified ErrorCode
+	 *
+	 * @param ErrorCode The error that occured
+	 */
+	static void LogWinHttpCloseHandleFailure(const uint32 ErrorCode);
+
+	/**
+	 * Log that an WinHttpSetTimeouts call failed with the specified ErrorCode
+	 *
+	 * @param ErrorCode The error that occured
+	 */
+	static void LogWinHttpSetTimeoutsFailure(const uint32 ErrorCode);
+
+	/**
+	 * Log that an WinHttpConnect call failed with the specified ErrorCode
+	 *
+	 * @param ErrorCode The error that occured
+	 */
+	static void LogWinHttpConnectFailure(const uint32 ErrorCode);
+
+	/**
+	 * Log that an WinHttpOpenRequest call failed with the specified ErrorCode
+	 *
+	 * @param ErrorCode The error that occured
+	 */
+	static void LogWinHttpOpenRequestFailure(const uint32 ErrorCode);
+
+	/**
+	 * Log that an WinHttpSetOption call failed with the specified ErrorCode
+	 *
+	 * @param ErrorCode The error that occured
+	 */
+	static void LogWinHttpSetOptionFailure(const uint32 ErrorCode);
+
+	/**
+	 * Log that an WinHttpAddRequestHeaders call failed with the specified ErrorCode
+	 *
+	 * @param ErrorCode The error that occured
+	 */
+	static void LogWinHttpAddRequestHeadersFailure(const uint32 ErrorCode);
+
+	/**
+	 * Log that an WinHttpSetStatusCallback call failed with the specified ErrorCode
+	 *
+	 * @param ErrorCode The error that occured
+	 */
+	static void LogWinHttpSetStatusCallbackFailure(const uint32 ErrorCode);
+
+	/**
+	 * Log that an WinHttpSendRequest call failed with the specified ErrorCode
+	 *
+	 * @param ErrorCode The error that occured
+	 */
+	static void LogWinHttpSendRequestFailure(const uint32 ErrorCode);
+
+	/**
+	 * Log that an WinHttpReceiveResponse call failed with the specified ErrorCode
+	 *
+	 * @param ErrorCode The error that occured
+	 */
+	static void LogWinHttpReceiveResponseFailure(const uint32 ErrorCode);
+
+	/**
+	 * Log that an WinHttpQueryHeaders call failed with the specified ErrorCode
+	 *
+	 * @param ErrorCode The error that occured
+	 */
+	static void LogWinHttpQueryHeadersFailure(const uint32 ErrorCode);
+
+	/**
+	 * Log that an WinHttpQueryDataAvailable call failed with the specified ErrorCode
+	 *
+	 * @param ErrorCode The error that occured
+	 */
+	static void LogWinHttpQueryDataAvailableFailure(const uint32 ErrorCode);
+
+	/**
+	 * Log that an WinHttpReadData call failed with the specified ErrorCode
+	 *
+	 * @param ErrorCode The error that occured
+	 */
+	static void LogWinHttpReadDataFailure(const uint32 ErrorCode);
+
+	/**
+	 * Log that an WinHttpWriteData call failed with the specified ErrorCode
+	 *
+	 * @param ErrorCode The error that occured
+	 */
+	static void LogWinHttpWriteDataFailure(const uint32 ErrorCode);
+
+private:
+	/** This class should not be instantiated */
+	FWinHttpErrorHelper() = delete;
+};
+
+
+#endif // WITH_WINHTTP
diff --git a/Engine/Source/Runtime/Online/HTTP/Internal/WinHttp/Support/WinHttpHandle.h b/Engine/Source/Runtime/Online/HTTP/Internal/WinHttp/Support/WinHttpHandle.h
new file mode 100644
--- /dev/null
+++ b/Engine/Source/Runtime/Online/HTTP/Internal/WinHttp/Support/WinHttpHandle.h
@@ -0,0 +1,76 @@
+// Copyright Epic Games, Inc. All Rights Reserved.
+
+#pragma once
+
+#if WITH_WINHTTP
+
+#include "CoreMinimal.h"
+
+using HINTERNET = void*;
+
+class HTTP_API FWinHttpHandle
+{
+public:
+	/**
+	 * Construct an invalid object
+	 */
+	FWinHttpHandle() = default;
+
+	/**
+	 * Wrap a new HINTERNET Handle
+	 */
+	explicit FWinHttpHandle(HINTERNET NewHandle);
+
+	/**
+	 * Destroy any currently held object
+	 */
+	~FWinHttpHandle();
+
+	// Copy/Move constructors
+	FWinHttpHandle(const FWinHttpHandle& Other) = delete;
+	FWinHttpHandle(FWinHttpHandle&& Other);
+	FWinHttpHandle& operator=(const FWinHttpHandle& Other) = delete;
+	FWinHttpHandle& operator=(FWinHttpHandle&& Other);
+
+	/**
+	 * Wrap a new handle (destroying any previously held object)
+	 *
+	 * @param NewHandle The new handle to wrap (must not be wrapped by anything else)
+	 * @return A reference to this object that now wraps NewHandle
+	 */
+	FWinHttpHandle& operator=(HINTERNET NewHandle);
+
+	/**
+	 * Destroy our current handle and reset our state to holding nothing
+	 */
+	void Reset();
+
+	/**
+	 * Do we contain a valid handle?
+	 *
+	 * @return True if we have a valid handle, false otherwise
+	 */
+	explicit operator bool() const;
+
+	/**
+	 * Do we contain a valid handle?
+	 *
+	 * @return True if we have a valid handle, false otherwise
+	 */
+	bool IsValid() const;
+
+	/**
+	 * Get the underlying handle for use
+	 *
+	 * @return The HINTERNET handle we're wrapping
+	 */
+	HINTERNET Get() const;
+
+protected:
+	/**
+	 * The handle we're wrap (if we have one)
+	 */
+	HINTERNET Handle = nullptr;
+};
+
+#endif // WITH_WINHTTP
diff --git a/Engine/Source/Runtime/Online/HTTP/Internal/WinHttp/Support/WinHttpSession.h b/Engine/Source/Runtime/Online/HTTP/Internal/WinHttp/Support/WinHttpSession.h
new file mode 100644
--- /dev/null
+++ b/Engine/Source/Runtime/Online/HTTP/Internal/WinHttp/Support/WinHttpSession.h
@@ -0,0 +1,56 @@
+// Copyright Epic Games, Inc. All Rights Reserved.
+
+#pragma once
+
+#if WITH_WINHTTP
+
+#include "CoreMinimal.h"
+#include "WinHttp/Support/WinHttpHandle.h"
+
+class IWinHttpConnection;
+
+using HINTERNET = void*;
+
+class HTTP_API FWinHttpSession
+{
+public:
+	/**
+	 * Construct a new WinHttp session with the specified security protocols flags
+	 */
+	FWinHttpSession(const uint32 SecurityProtocolFlags, const bool bForceSecureConnections);
+
+	/**
+	 * FWinHttpSession is move-only
+	 */
+	FWinHttpSession(const FWinHttpSession& Other) = delete;
+	FWinHttpSession(FWinHttpSession&& Other) = default;
+	FWinHttpSession& operator=(const FWinHttpSession& Other) = delete;
+	FWinHttpSession& operator=(FWinHttpSession&& Other) = default;
+
+	/**
+	 * Did this session initialize successfully?
+	 */
+	bool IsValid() const;
+
+	/**
+	 * Get the underlying session handle
+	 *
+	 * @return the HINTERNET for this session
+	 */
+	HINTERNET Get() const;
+
+	/**
+	 * Are we only allowed to make secure connection requests (HTTPS, etc)
+	 *
+	 * @return True if the platform does not allow for insecure messages
+	 */
+	bool AreOnlySecureConnectionsAllowed() const;
+
+private:
+	/** The handle for our session that we are wrapping */
+	FWinHttpHandle SessionHandle;
+	/** Should we force connections to be secure? */
+	bool bForceSecureConnections = false;
+};
+
+#endif // WITH_WINHTTP
diff --git a/Engine/Source/Runtime/Online/HTTP/Internal/WinHttp/Support/WinHttpTypes.h b/Engine/Source/Runtime/Online/HTTP/Internal/WinHttp/Support/WinHttpTypes.h
new file mode 100644
--- /dev/null
+++ b/Engine/Source/Runtime/Online/HTTP/Internal/WinHttp/Support/WinHttpTypes.h
@@ -0,0 +1,54 @@
+// Copyright Epic Games, Inc. All Rights Reserved.
+
+#if WITH_WINHTTP
+
+#pragma once
+
+#include "CoreMinimal.h"
+
+#pragma warning(push)
+#pragma warning(disable : 28285) // Disable static analysis syntax error in WinHttpSetHeaders macros
+#include "Windows/AllowWindowsPlatformTypes.h"
+THIRD_PARTY_INCLUDES_START
+#include <winhttp.h>
+THIRD_PARTY_INCLUDES_END
+#include "Windows/HideWindowsPlatformTypes.h"
+#pragma warning(pop)
+
+enum class EWinHttpCallbackStatus : uint32
+{
+	ResolvingName = WINHTTP_CALLBACK_STATUS_RESOLVING_NAME,
+	NameResolved = WINHTTP_CALLBACK_STATUS_NAME_RESOLVED,
+	ConnectingToServer = WINHTTP_CALLBACK_STATUS_CONNECTING_TO_SERVER,
+	ConnectedToServer = WINHTTP_CALLBACK_STATUS_CONNECTED_TO_SERVER,
+	SendingRequest = WINHTTP_CALLBACK_STATUS_SENDING_REQUEST,
+	RequestSent = WINHTTP_CALLBACK_STATUS_REQUEST_SENT,
+	ReceivingResponse = WINHTTP_CALLBACK_STATUS_RECEIVING_RESPONSE,
+	ResponseReceived = WINHTTP_CALLBACK_STATUS_RESPONSE_RECEIVED,
+	ClosingConnection = WINHTTP_CALLBACK_STATUS_CLOSING_CONNECTION,
+	ConnectionClosed = WINHTTP_CALLBACK_STATUS_CONNECTION_CLOSED,
+	HandleCreated = WINHTTP_CALLBACK_STATUS_HANDLE_CREATED,
+	HandleClosing = WINHTTP_CALLBACK_STATUS_HANDLE_CLOSING,
+	DetectingProxy = WINHTTP_CALLBACK_STATUS_DETECTING_PROXY,
+	Redirect = WINHTTP_CALLBACK_STATUS_REDIRECT,
+	IntermediateResponse = WINHTTP_CALLBACK_STATUS_INTERMEDIATE_RESPONSE,
+	SecureFailure = WINHTTP_CALLBACK_STATUS_SECURE_FAILURE,
+	HeadersAvailable = WINHTTP_CALLBACK_STATUS_HEADERS_AVAILABLE,
+	DataAvailable = WINHTTP_CALLBACK_STATUS_DATA_AVAILABLE,
+	ReadComplete = WINHTTP_CALLBACK_STATUS_READ_COMPLETE,
+	WriteComplete = WINHTTP_CALLBACK_STATUS_WRITE_COMPLETE,
+	RequestError = WINHTTP_CALLBACK_STATUS_REQUEST_ERROR,
+	SendRequestComplete = WINHTTP_CALLBACK_STATUS_SENDREQUEST_COMPLETE,
+	GetProxyForUrlComplete = WINHTTP_CALLBACK_STATUS_GETPROXYFORURL_COMPLETE,
+	CloseComplete = WINHTTP_CALLBACK_STATUS_CLOSE_COMPLETE,
+	ShutdownComplete = WINHTTP_CALLBACK_STATUS_SHUTDOWN_COMPLETE,
+	SettingsWriteComplete = WINHTTP_CALLBACK_STATUS_SETTINGS_WRITE_COMPLETE,
+	SettingsReadComplete = WINHTTP_CALLBACK_STATUS_SETTINGS_READ_COMPLETE
+};
+
+HTTP_API const TCHAR* LexToString(const EWinHttpCallbackStatus Status);
+HTTP_API bool IsValidStatus(const EWinHttpCallbackStatus Status);
+
+HTTP_API DECLARE_LOG_CATEGORY_EXTERN(LogWinHttp, Log, All);
+
+#endif // WITH_WINHTTP
diff --git a/Engine/Source/Runtime/Online/HTTP/Internal/WinHttp/WinHttpHttpManager.h b/Engine/Source/Runtime/Online/HTTP/Internal/WinHttp/WinHttpHttpManager.h
new file mode 100644
--- /dev/null
+++ b/Engine/Source/Runtime/Online/HTTP/Internal/WinHttp/WinHttpHttpManager.h
@@ -0,0 +1,76 @@
+// Copyright Epic Games, Inc. All Rights Reserved.
+
+#if WITH_WINHTTP
+
+#pragma once
+
+#include "CoreMinimal.h"
+#include "HttpManager.h"
+
+class FWinHttpSession;
+
+DECLARE_DELEGATE_OneParam(FWinHttpQuerySessionComplete, FWinHttpSession* /*HttpSessionPtr*/);
+
+class IWinHttpConnection;
+
+class HTTP_API FWinHttpHttpManager
+	: public FHttpManager
+{
+public:
+	static FWinHttpHttpManager* GetManager();
+
+	FWinHttpHttpManager();
+	virtual ~FWinHttpHttpManager();
+
+	/**
+	 * Asynchronously finds an existing WinHttp session for the provided URL, or creates a new one for it.
+	 *
+	 * @param Url The URL to find or create a WinHttp session for
+	 * @param Delegate The delegate that is called with the WinHttp session pointer if successful, or null otherwise
+	 */
+	virtual void QuerySessionForUrl(const FString& Url, FWinHttpQuerySessionComplete&& Delegate);
+
+	/**
+	 * Validate the provided connection before we start sending our request.
+	 *
+	 * NOTE: this is called on multiple threads, and should be written in a way that handles this safely!
+	 * 
+	 * @param Connection the connection to validate
+	 * @return True if the request validated successfully, false if not
+	 */
+	virtual bool ValidateRequestCertificates(IWinHttpConnection& Connection);
+
+	/**
+	 * Release any resources that were pinned for this request.
+	 *
+	 * NOTE: this is a no-op on Windows builds, but does things on other platforms.
+	 *
+	 * @param Connection The connection to release resources for
+	 */
+	virtual void ReleaseRequestResources(IWinHttpConnection& Connection);
+
+	//~ Begin FHttpManager Interface
+	virtual void OnBeforeFork() override;
+	//~ End FHttpManager Interface
+
+	/**
+	 * Handle the application preparing to suspend
+	 */
+	virtual void HandleApplicationSuspending();
+
+	/**
+	 * Handle the application preparing to suspend
+	 */
+	virtual void HandleApplicationResuming();
+
+protected:
+	FWinHttpSession* FindOrCreateSession(const uint32 SecurityProtocols);
+
+protected:
+	bool bPlatformForcesSecureConnections = false;
+
+	/** Map of Security Flags to WinHttp Session objects */
+	TMap<uint32, TUniquePtr<FWinHttpSession>> ActiveSessions;
+};
+
+#endif // WITH_WINHTTP
diff --git a/Engine/Source/Runtime/Online/HTTP/Private/WinHttp/Support/WinHttpConnection.h b/Engine/Source/Runtime/Online/HTTP/Private/WinHttp/Support/WinHttpConnection.h
deleted file mode 100644
diff --git a/Engine/Source/Runtime/Online/HTTP/Private/WinHttp/Support/WinHttpConnectionHttp.h b/Engine/Source/Runtime/Online/HTTP/Private/WinHttp/Support/WinHttpConnectionHttp.h
deleted file mode 100644
diff --git a/Engine/Source/Runtime/Online/HTTP/Private/WinHttp/Support/WinHttpErrorHelper.h b/Engine/Source/Runtime/Online/HTTP/Private/WinHttp/Support/WinHttpErrorHelper.h
deleted file mode 100644
diff --git a/Engine/Source/Runtime/Online/HTTP/Private/WinHttp/Support/WinHttpHandle.h b/Engine/Source/Runtime/Online/HTTP/Private/WinHttp/Support/WinHttpHandle.h
deleted file mode 100644
diff --git a/Engine/Source/Runtime/Online/HTTP/Private/WinHttp/Support/WinHttpSession.h b/Engine/Source/Runtime/Online/HTTP/Private/WinHttp/Support/WinHttpSession.h
deleted file mode 100644
diff --git a/Engine/Source/Runtime/Online/HTTP/Private/WinHttp/Support/WinHttpTypes.h b/Engine/Source/Runtime/Online/HTTP/Private/WinHttp/Support/WinHttpTypes.h
deleted file mode 100644
diff --git a/Engine/Source/Runtime/Online/HTTP/Private/WinHttp/WinHttpHttpManager.h b/Engine/Source/Runtime/Online/HTTP/Private/WinHttp/WinHttpHttpManager.h
deleted file mode 100644
diff --git a/Engine/Source/Runtime/Online/WebSockets/WebSockets.Build.cs b/Engine/Source/Runtime/Online/WebSockets/WebSockets.Build.cs
--- a/Engine/Source/Runtime/Online/WebSockets/WebSockets.Build.cs
+++ b/Engine/Source/Runtime/Online/WebSockets/WebSockets.Build.cs
@@ -120,13 +120,6 @@ public class WebSockets : ModuleRules
 				bWithWinHttpWebSockets = true;
 
 				AddEngineThirdPartyPrivateStaticDependencies(Target, "WinHttp");
-
-				// We need to access the WinHttp folder in HTTP
-				PrivateIncludePaths.AddRange(
-					new string[] {
-						"Runtime/Online/HTTP/Private",
-					}
-				);
 			}
 		}
 
