Adding CSV analytics for background downloading

diff --git a/Engine/Source/Runtime/Online/BackgroundHTTP/Private/BackgroundHttpManagerImpl.cpp b/Engine/Source/Runtime/Online/BackgroundHTTP/Private/BackgroundHttpManagerImpl.cpp
--- a/Engine/Source/Runtime/Online/BackgroundHTTP/Private/BackgroundHttpManagerImpl.cpp
+++ b/Engine/Source/Runtime/Online/BackgroundHTTP/Private/BackgroundHttpManagerImpl.cpp
@@ -16,6 +16,7 @@
 
 
 DEFINE_LOG_CATEGORY(LogBackgroundHttpManager);
+CSV_DEFINE_CATEGORY(BackgroundDownload, true);
 
 FBackgroundHttpManagerImpl::FBackgroundHttpManagerImpl()
 	: PendingStartRequests()
@@ -398,6 +399,9 @@ bool FBackgroundHttpManagerImpl::Tick(float DeltaTime)
 	QUICK_SCOPE_CYCLE_COUNTER(STAT_FBackgroundHttpManagerImpl_Tick);
 
 	ensureAlwaysMsgf(IsInGameThread(), TEXT("Called from un-expected thread! Potential error in an implementation of background downloads!"));
+	CSV_CUSTOM_STAT(BackgroundDownload, MaxActiveDownloads, MaxActiveDownloads, ECsvCustomStatOp::Set);
+	CSV_CUSTOM_STAT(BackgroundDownload, PendingStartRequests, PendingStartRequests.Num(), ECsvCustomStatOp::Set);
+	CSV_CUSTOM_STAT(BackgroundDownload, NumCurrentlyActiveRequests, NumCurrentlyActiveRequests, ECsvCustomStatOp::Set);
 	
 	ActivatePendingRequests();
 	
