Adding CSV analytics for background downloading

diff --git a/Engine/Source/Runtime/Online/BackgroundHTTP/Private/BackgroundHttpManagerImpl.cpp b/Engine/Source/Runtime/Online/BackgroundHTTP/Private/BackgroundHttpManagerImpl.cpp
--- a/Engine/Source/Runtime/Online/BackgroundHTTP/Private/BackgroundHttpManagerImpl.cpp
+++ b/Engine/Source/Runtime/Online/BackgroundHTTP/Private/BackgroundHttpManagerImpl.cpp
@@ -13,9 +13,12 @@
 #include "Misc/Paths.h"
 #include "Misc/ScopeRWLock.h"
 #include "Stats/Stats.h"
+#include "ProfilingDebugging/CsvProfiler.h"
 
 
 DEFINE_LOG_CATEGORY(LogBackgroundHttpManager);
+CSV_DECLARE_CATEGORY_MODULE_EXTERN(BACKGROUNDHTTP_API, BackgroundDownload);
+CSV_DEFINE_CATEGORY(BackgroundDownload, true);
 
 FBackgroundHttpManagerImpl::FBackgroundHttpManagerImpl()
 	: PendingStartRequests()
@@ -398,7 +401,9 @@ bool FBackgroundHttpManagerImpl::Tick(float DeltaTime)
 	QUICK_SCOPE_CYCLE_COUNTER(STAT_FBackgroundHttpManagerImpl_Tick);
 
 	ensureAlwaysMsgf(IsInGameThread(), TEXT("Called from un-expected thread! Potential error in an implementation of background downloads!"));
-	
+	CSV_CUSTOM_STAT(BackgroundDownload, MaxActiveDownloads, MaxActiveDownloads, ECsvCustomStatOp::Set);
+	CSV_CUSTOM_STAT(BackgroundDownload, PendingStartRequests, PendingStartRequests.Num(), ECsvCustomStatOp::Set);
+	CSV_CUSTOM_STAT(BackgroundDownload, NumCurrentlyActiveRequests, NumCurrentlyActiveRequests, ECsvCustomStatOp::Set);
 	ActivatePendingRequests();
 	
 	//for now we are saving data every tick, could change this to be on an interval later if required
