[BuildPatchServices] Check if destination file system is throttled and if so, stall file write.

diff --git a/Engine/Source/Runtime/Online/BuildPatchServices/Private/BuildPatchFileConstructor.cpp b/Engine/Source/Runtime/Online/BuildPatchServices/Private/BuildPatchFileConstructor.cpp
--- a/Engine/Source/Runtime/Online/BuildPatchServices/Private/BuildPatchFileConstructor.cpp
+++ b/Engine/Source/Runtime/Online/BuildPatchServices/Private/BuildPatchFileConstructor.cpp
@@ -26,6 +26,21 @@ using namespace BuildPatchServices;
 // incase of previous partial write.
 #define NUM_BYTES_RESUME_IGNORE     1024
 
+static int32 SleepTimeWhenFileSystemThrottledSeconds = 15;
+static FAutoConsoleVariableRef CVarSleepTimeWhenFileSystemThrottledSeconds(
+	TEXT("BuildPatchFileConstructor.SleepTimeWhenFileSystemThrottledSeconds"),
+	SleepTimeWhenFileSystemThrottledSeconds,
+	TEXT("The amount of time to sleep if the destination filesystem is throttled."),
+	ECVF_Default);
+
+static bool bStallWhenFileSystemThrottled = false;
+static FAutoConsoleVariableRef CVarStallWhenFileSystemThrottled(
+	TEXT("BuildPatchFileConstructor.bStallWhenFileSystemThrottled"),
+	bStallWhenFileSystemThrottled,
+	TEXT("Whether to stall if the file system is throttled"),
+	ECVF_Default);
+
+
 // Helper functions wrapping common code.
 namespace FileConstructorHelpers
 {
@@ -835,6 +850,18 @@ bool FBuildPatchFileConstructor::ConstructFileFromChunks(const FFileManifest& Fi
 
 bool FBuildPatchFileConstructor::InsertChunkData(const FChunkPart& ChunkPart, FArchive& DestinationFile, FSHA1& HashState, EConstructionError& ConstructionError)
 {
+	if (bStallWhenFileSystemThrottled)
+	{
+		int64 AvailableBytes = FileSystem->GetAllowedBytesToWriteThrottledStorage(*DestinationFile.GetArchiveName());
+		while (ChunkPart.Size > AvailableBytes)
+		{
+			UE_LOG(LogBuildPatchServices, Display, TEXT("Avaliable write bytes to write throttled storage exhausted (%s).  Sleeping %ds.  Bytes needed: %u, bytes available: %lld")
+				, *DestinationFile.GetArchiveName(), SleepTimeWhenFileSystemThrottledSeconds, ChunkPart.Size, AvailableBytes);
+			FPlatformProcess::Sleep(SleepTimeWhenFileSystemThrottledSeconds);
+			AvailableBytes = FileSystem->GetAllowedBytesToWriteThrottledStorage(*DestinationFile.GetArchiveName());
+		}
+	}
+
 	uint8* Data;
 	uint8* DataStart;
 	ConstructionError = EConstructionError::None;
diff --git a/Engine/Source/Runtime/Online/BuildPatchServices/Private/Common/FileSystem.cpp b/Engine/Source/Runtime/Online/BuildPatchServices/Private/Common/FileSystem.cpp
--- a/Engine/Source/Runtime/Online/BuildPatchServices/Private/Common/FileSystem.cpp
+++ b/Engine/Source/Runtime/Online/BuildPatchServices/Private/Common/FileSystem.cpp
@@ -272,6 +272,7 @@ namespace BuildPatchServices
 		virtual void FindFiles(TArray<FString>& FoundFiles, const TCHAR* Directory, const TCHAR* FileExtension = nullptr) const override;
 		virtual void FindFilesRecursively(TArray<FString>& FoundFiles, const TCHAR* Directory, const TCHAR* FileExtension = nullptr) const override;
 		virtual void ParallelFindFilesRecursively(TArray<FString>& FoundFiles, const TCHAR* Directory, const TCHAR* FileExtension = nullptr, EAsyncExecution AsyncExecution = EAsyncExecution::ThreadPool) const override;
+		virtual int64 GetAllowedBytesToWriteThrottledStorage(const TCHAR* DestinationPath = nullptr) const override;
 		// IFileSystem interface end.
 
 	private:
@@ -389,6 +390,11 @@ namespace BuildPatchServices
 		DirectoryEnumerator.GetFiles(FoundFiles);
 	}
 
+	int64 FFileSystem::GetAllowedBytesToWriteThrottledStorage(const TCHAR* DestinationPath) const
+	{
+		return PlatformFile.GetAllowedBytesToWriteThrottledStorage(DestinationPath);
+	}
+
 	IFileSystem* FFileSystemFactory::Create()
 	{
 		return new FFileSystem();
diff --git a/Engine/Source/Runtime/Online/BuildPatchServices/Private/Common/FileSystem.h b/Engine/Source/Runtime/Online/BuildPatchServices/Private/Common/FileSystem.h
--- a/Engine/Source/Runtime/Online/BuildPatchServices/Private/Common/FileSystem.h
+++ b/Engine/Source/Runtime/Online/BuildPatchServices/Private/Common/FileSystem.h
@@ -201,6 +201,16 @@ namespace BuildPatchServices
 		 * @param AsyncExecution    The threaded execution method to use, i.e. on Task Graph, Thread Pool, or a new thread per directory.
 		 */
 		virtual void ParallelFindFilesRecursively(TArray<FString>& FoundFiles, const TCHAR* Directory, const TCHAR* FileExtension = nullptr, EAsyncExecution AsyncExecution = EAsyncExecution::ThreadPool) const = 0;
+
+		/**
+		 * Returns the number of bytes that are currently allowed to be written to throttled write storage (if the platform
+		 * has such restrictions)
+		 * @param DestinationPath	If specified, the file system can optionally take into account the destination of
+		 *							the file to determine the current limit
+		 * @returns					The number of bytes that are allowed to be written to write throttled storage.
+		 *							If there is no limit, INT64_MAX is returned
+		 */
+		virtual int64 GetAllowedBytesToWriteThrottledStorage(const TCHAR* DestinationPath = nullptr) const = 0;
 	};
 
 	/**
diff --git a/Engine/Source/Runtime/Online/BuildPatchServices/Private/Tests/Mock/FileSystem.mock.h b/Engine/Source/Runtime/Online/BuildPatchServices/Private/Tests/Mock/FileSystem.mock.h
--- a/Engine/Source/Runtime/Online/BuildPatchServices/Private/Tests/Mock/FileSystem.mock.h
+++ b/Engine/Source/Runtime/Online/BuildPatchServices/Private/Tests/Mock/FileSystem.mock.h
@@ -152,6 +152,12 @@ namespace BuildPatchServices
 			FindFilesRecursively(FoundFiles, Directory, FileExtension);
 		}
 
+		virtual int64 GetAllowedBytesToWriteThrottledStorage(const TCHAR* DestinationPath = nullptr) const override
+		{
+			MOCK_FUNC_NOT_IMPLEMENTED("FMockFileSystem::GetAllowedBytesToWriteThrottledStorage");
+			return INT64_MAX;
+		}
+
 	public:
 		mutable FCriticalSection ThreadLock;
 		mutable TArray<FCreateFileReader> RxCreateFileReader;
