Adding support to use runtime args "-httpnoproxy=aaa,bbb" or config "[HTTP]HttpNoProxy=aaa,bbb": Even if proxy got set, make libcurl impl not use proxy for specified domains.

Tested by setting up a local web server and local proxy server on Windows, added test code in local Windows DS to send GET request to web server "http://127.0.0.1:8000", while passing "-httpproxy=my_pc_ip:7000" as runtime args for setting proxy. With "-httpnoproxy=127.0.0.1,localhost", proxy server will not log the request. Without it the proxy server will log that request.

diff --git a/Engine/Source/Runtime/Online/HTTP/Private/Curl/CurlHttp.cpp b/Engine/Source/Runtime/Online/HTTP/Private/Curl/CurlHttp.cpp
--- a/Engine/Source/Runtime/Online/HTTP/Private/Curl/CurlHttp.cpp
+++ b/Engine/Source/Runtime/Online/HTTP/Private/Curl/CurlHttp.cpp
@@ -125,6 +125,12 @@ FCurlHttpRequest::FCurlHttpRequest()
 		curl_easy_setopt(EasyHandle, CURLOPT_PROXY, TCHAR_TO_ANSI(*ProxyAddress));
 	}
 
+	const FString& HttpNoProxy = FHttpModule::Get().GetHttpNoProxy();
+	if (!HttpNoProxy.IsEmpty())
+	{
+		curl_easy_setopt(EasyHandle, CURLOPT_NOPROXY, TCHAR_TO_ANSI(*HttpNoProxy));
+	}
+
 	if (FCurlHttpManager::CurlRequestOptions.bDontReuseConnections)
 	{
 		curl_easy_setopt(EasyHandle, CURLOPT_FORBID_REUSE, 1L);
diff --git a/Engine/Source/Runtime/Online/HTTP/Private/HttpModule.cpp b/Engine/Source/Runtime/Online/HTTP/Private/HttpModule.cpp
--- a/Engine/Source/Runtime/Online/HTTP/Private/HttpModule.cpp
+++ b/Engine/Source/Runtime/Online/HTTP/Private/HttpModule.cpp
@@ -45,6 +45,11 @@ void FHttpModule::UpdateConfigs()
 	GConfig->GetFloat(TEXT("HTTP"), TEXT("HttpThreadIdleFrameTimeInSeconds"), HttpThreadIdleFrameTimeInSeconds, GEngineIni);
 	GConfig->GetFloat(TEXT("HTTP"), TEXT("HttpThreadIdleMinimumSleepTimeInSeconds"), HttpThreadIdleMinimumSleepTimeInSeconds, GEngineIni);
 
+	if (!FParse::Value(FCommandLine::Get(), TEXT("HttpNoProxy="), HttpNoProxy))
+	{
+		GConfig->GetString(TEXT("HTTP"), TEXT("HttpNoProxy"), HttpNoProxy, GEngineIni);
+	}
+
 	AllowedDomains.Empty();
 	GConfig->GetArray(TEXT("HTTP"), TEXT("AllowedDomains"), AllowedDomains, GEngineIni);
 
diff --git a/Engine/Source/Runtime/Online/HTTP/Public/HttpModule.h b/Engine/Source/Runtime/Online/HTTP/Public/HttpModule.h
--- a/Engine/Source/Runtime/Online/HTTP/Public/HttpModule.h
+++ b/Engine/Source/Runtime/Online/HTTP/Public/HttpModule.h
@@ -281,6 +281,14 @@ public:
 		ProxyAddress = InProxyAddress;
 	}
 
+	/**
+	 * @returns the domains which won't use proxy even if ProxyAddress is set
+	 */
+	inline const FString& GetHttpNoProxy() const
+	{
+		return HttpNoProxy;
+	}
+
 	/**
 	 * Method to check dynamic proxy setting support.
 	 * @returns Whether this http implementation supports dynamic proxy setting.
@@ -361,6 +369,8 @@ private:
 	static FHttpModule* Singleton;
 	/** The address to use for proxy, in format IPADDRESS:PORT */
 	FString ProxyAddress;
+	/** The domains which won't use proxy even if the ProxyAddress is set, in format "127.0.0.1,localhost,example.com" */
+	FString HttpNoProxy;
 	/** Whether or not the http implementation we are using supports dynamic proxy setting. */
 	bool bSupportsDynamicProxy;
 	/** List of domains that can be accessed. If Empty then no filtering is applied */
