Make "-DisableSSLCertificatePinning" can work in Shipping servers.

diff --git a/Engine/Source/Runtime/Online/SSL/Private/SslCertificateManager.cpp b/Engine/Source/Runtime/Online/SSL/Private/SslCertificateManager.cpp
--- a/Engine/Source/Runtime/Online/SSL/Private/SslCertificateManager.cpp
+++ b/Engine/Source/Runtime/Online/SSL/Private/SslCertificateManager.cpp
@@ -195,7 +195,7 @@ void FSslCertificateManager::SetPinnedPublicKeys(const FString& Domain, const FS
 
 bool FSslCertificateManager::VerifySslCertificates(X509_STORE_CTX* Context, const FString& Domain) const
 {
-#if !UE_BUILD_SHIPPING
+#if !UE_BUILD_SHIPPING || WITH_SERVER_CODE
 	static const bool bPinningDisabled = FParse::Param(FCommandLine::Get(), TEXT("DisableSSLCertificatePinning"));
 	if (bPinningDisabled)
 	{
@@ -274,7 +274,7 @@ bool FSslCertificateManager::VerifySslCertificates(X509_STORE_CTX* Context, cons
 
 bool FSslCertificateManager::VerifySslCertificates(TArray<TArray<uint8, TFixedAllocator<PUBLIC_KEY_DIGEST_SIZE>>>& Digests, const FString& Domain) const
 {
-#if !UE_BUILD_SHIPPING
+#if !UE_BUILD_SHIPPING || WITH_SERVER_CODE
 	static const bool bPinningDisabled = FParse::Param(FCommandLine::Get(), TEXT("DisableSSLCertificatePinning"));
 	if (bPinningDisabled)
 	{
