Refactor to enable WinHttp websockets to be created with an empty protocol list, which will result in the corresponding header not being added.

diff --git a/Engine/Source/Runtime/Online/WebSockets/Private/WinHttp/Support/WinHttpConnectionWebSocket.cpp b/Engine/Source/Runtime/Online/WebSockets/Private/WinHttp/Support/WinHttpConnectionWebSocket.cpp
--- a/Engine/Source/Runtime/Online/WebSockets/Private/WinHttp/Support/WinHttpConnectionWebSocket.cpp
+++ b/Engine/Source/Runtime/Online/WebSockets/Private/WinHttp/Support/WinHttpConnectionWebSocket.cpp
@@ -74,12 +74,6 @@ TSharedPtr<FWinHttpConnectionWebSocket, ESPMode::ThreadSafe> FWinHttpConnectionW
 		return nullptr;
 	}
 
-	if (InProtocols.Num() == 0)
-	{
-		UE_LOG(LogWinHttp, Warning, TEXT("Attempted to create a WinHttp WebSocket with an empty protocols list"));
-		return nullptr;
-	}
-
 	TOptional<uint16> Port = FGenericPlatformHttp::GetUrlPort(InUrl);
 	TSharedPtr<FWinHttpConnectionWebSocket, ESPMode::ThreadSafe> WebSocket = MakeShareable(new FWinHttpConnectionWebSocket(InSession, InUrl, bIsSecure, Domain, Port, PathAndQuery, InProtocols, InUpgradeHeaders));
 	if (!WebSocket->IsValid())
@@ -360,7 +354,10 @@ FWinHttpConnectionWebSocket::FWinHttpConnectionWebSocket(
 		return;
 	}
 
-	SetHeader(TEXT("Sec-WebSocket-Protocol"), FString::Join(InProtocols, TEXT(", ")));
+	if (!InProtocols.IsEmpty())
+	{
+		SetHeader(TEXT("Sec-WebSocket-Protocol"), FString::Join(InProtocols, TEXT(", ")));
+	}
 }
 
 bool FWinHttpConnectionWebSocket::IsReadInProgress() const
