- use FGenericPlatformHttp UrlDecode implementation, the HttpConnectionRequestReadContext one was a copy paste of it anyway.

diff --git a/Engine/Source/Runtime/Online/HTTPServer/Private/HttpConnectionRequestReadContext.cpp b/Engine/Source/Runtime/Online/HTTPServer/Private/HttpConnectionRequestReadContext.cpp
--- a/Engine/Source/Runtime/Online/HTTPServer/Private/HttpConnectionRequestReadContext.cpp
+++ b/Engine/Source/Runtime/Online/HTTPServer/Private/HttpConnectionRequestReadContext.cpp
@@ -6,6 +6,7 @@
 #include "HttpServerRequest.h"
 #include "HttpConnectionContext.h"
 #include "IPAddress.h"
+#include "PlatformHttp.h"
 #include "Sockets.h"
 #include "SocketSubsystem.h"
 
@@ -291,8 +292,8 @@ TSharedPtr<FHttpServerRequest> FHttpConnectionRequestReadContext::BuildRequest(c
 			int32 Equalsindex = 0;
 			if (QueryParamPair.FindChar(TCHAR('='), Equalsindex))
 			{
-				FString QueryParamKey = UrlDecode(QueryParamPair.Mid(0, Equalsindex));
-				FString QueryParamValue = UrlDecode(QueryParamPair.Mid(Equalsindex + 1));
+				FString QueryParamKey = FGenericPlatformHttp::UrlDecode(QueryParamPair.Mid(0, Equalsindex));
+				FString QueryParamValue = FGenericPlatformHttp::UrlDecode(QueryParamPair.Mid(Equalsindex + 1));
 				Request->QueryParams.Emplace(MoveTemp(QueryParamKey), MoveTemp(QueryParamValue));
 			}
 		}
@@ -336,70 +337,4 @@ TSharedPtr<FHttpServerRequest> FHttpConnectionRequestReadContext::BuildRequest(c
 	}
 
 	return Request;
-}
-
-FString FHttpConnectionRequestReadContext::UrlDecode(const FString &EncodedString)
-{
-	FTCHARToUTF8 Converter(*EncodedString);
-	const UTF8CHAR* UTF8Data = (UTF8CHAR*)Converter.Get();
-
-	TArray<UTF8CHAR> Data;
-	Data.Reserve(EncodedString.Len());
-
-	for (int32 CharIdx = 0; CharIdx < Converter.Length();)
-	{
-		if (UTF8Data[CharIdx] == '%')
-		{
-			int32 Value = 0;
-			if (CharIdx < Converter.Length() - 1 && UTF8Data[CharIdx + 1] == 'u')
-			{
-				if (CharIdx <= Converter.Length() - 6)
-				{
-					// Treat all %uXXXX as code point
-					Value = FParse::HexDigit(UTF8Data[CharIdx + 2]) << 12;
-					Value += FParse::HexDigit(UTF8Data[CharIdx + 3]) << 8;
-					Value += FParse::HexDigit(UTF8Data[CharIdx + 4]) << 4;
-					Value += FParse::HexDigit(UTF8Data[CharIdx + 5]);
-					CharIdx += 6;
-
-					UTF8CHAR Buffer[8] = { };
-					UTF8CHAR* BufferPtr = Buffer;
-					const int32 Len = UE_ARRAY_COUNT(Buffer);
-					const UTF8CHAR* BufferEnd = FPlatformString::Convert(BufferPtr, Len, (UTF32CHAR*)&Value, 1);
-					check(BufferEnd);
-
-					Data.Append(Buffer, (int32)(BufferEnd - BufferPtr));
-				}
-				else
-				{
-					// Not enough in the buffer for valid decoding, skip it
-					CharIdx++;
-					continue;
-				}
-			}
-			else if (CharIdx <= Converter.Length() - 3)
-			{
-				// Treat all %XX as straight byte
-				Value = FParse::HexDigit(UTF8Data[CharIdx + 1]) << 4;
-				Value += FParse::HexDigit(UTF8Data[CharIdx + 2]);
-				CharIdx += 3;
-				Data.Add((UTF8CHAR)(Value));
-			}
-			else
-			{
-				// Not enough in the buffer for valid decoding, skip it
-				CharIdx++;
-				continue;
-			}
-		}
-		else
-		{
-			// Non escaped characters
-			Data.Add(UTF8Data[CharIdx]);
-			CharIdx++;
-		}
-	}
-
-	Data.Add(UTF8TEXT('\0'));
-	return FString(UTF8_TO_TCHAR(Data.GetData()));
 }
\ No newline at end of file
diff --git a/Engine/Source/Runtime/Online/HTTPServer/Private/HttpConnectionRequestReadContext.h b/Engine/Source/Runtime/Online/HTTPServer/Private/HttpConnectionRequestReadContext.h
--- a/Engine/Source/Runtime/Online/HTTPServer/Private/HttpConnectionRequestReadContext.h
+++ b/Engine/Source/Runtime/Online/HTTPServer/Private/HttpConnectionRequestReadContext.h
@@ -113,15 +113,6 @@ private:
 	 */
 	static bool ParseContentLength(const FHttpServerRequest& InRequest, int32& OutContentLength);
 
-	/**
-	 * Returns a decoded version of the percent-encoded passed in string
-	 *
-	 * @param EncodedString The percent encoded string to convert to string
-	 * @return The decoded string
-	 */
-	static FString UrlDecode(const FString& EncodedString);
-
-
 private:
 
 	/** The underlying connection */
