Make sure Background notifications are deliverd right away

diff --git a/Engine/Source/Runtime/Online/BackgroundHTTP/Private/BackgroundHttpNotificationObject.cpp b/Engine/Source/Runtime/Online/BackgroundHTTP/Private/BackgroundHttpNotificationObject.cpp
--- a/Engine/Source/Runtime/Online/BackgroundHTTP/Private/BackgroundHttpNotificationObject.cpp
+++ b/Engine/Source/Runtime/Online/BackgroundHTTP/Private/BackgroundHttpNotificationObject.cpp
@@ -82,9 +82,10 @@ FBackgroundHttpNotificationObject::~FBackgroundHttpNotificationObject()
 	{
 		if (!bNotifyOnlyOnFullSuccess || (NumFailedDownloads == 0))
 		{
+			//Used to be 15s, probably for reasons to do with Ios. 
 			//make a notification 1 second from now
 			FDateTime TargetTime = FDateTime::Now();
-			TargetTime += FTimespan::FromSeconds(15);
+			TargetTime += FTimespan::FromSeconds(1);
 
 			if (nullptr != PlatformNotificationService)
 			{
