Unbind http added&completed delegates in http manager when app exit;
Now game should use Shutdown() instead of Flush(EHttpFlushReason::Shutdown) when app exit;
Move RequestAddedDelegate to internal scope, remove corresponding deprecated code.

diff --git a/Engine/Source/Runtime/Online/HTTP/Private/HttpManager.cpp b/Engine/Source/Runtime/Online/HTTP/Private/HttpManager.cpp
--- a/Engine/Source/Runtime/Online/HTTP/Private/HttpManager.cpp
+++ b/Engine/Source/Runtime/Online/HTTP/Private/HttpManager.cpp
@@ -35,6 +35,14 @@ const TCHAR* LexToString(const EHttpFlushReason& FlushReason)
 	return TEXT("Invalid");
 }
 
+namespace
+{
+	bool ShouldOutputHttpWarnings()
+	{
+		return !IsRunningCommandlet() && !FApp::IsUnattended();
+	}
+}
+
 PRAGMA_DISABLE_DEPRECATION_WARNINGS
 FHttpManager::FHttpManager()
 	: FTSTickerObjectBase(0.0f, FTSBackgroundableTicker::GetCoreTicker())
@@ -65,6 +73,59 @@ void FHttpManager::Initialize()
 	UpdateConfigs();
 }
 
+void FHttpManager::Shutdown()
+{
+	FScopeLock ScopeLock(&RequestLock);
+
+	// Don't emit these tracking logs in commandlet runs. Build system traps warnings during cook, and these are not truly fatal, but useful for tracking down shutdown issues.
+	UE_CLOG(ShouldOutputHttpWarnings() && Requests.Num(), LogHttp, Warning, TEXT("[FHttpManager::Shutdown] Unbinding delegates for %d outstanding Http Requests:"), Requests.Num());
+
+	// Clear delegates since they may point to deleted instances
+	for (TArray<FHttpRequestRef>::TIterator It(Requests); It; ++It)
+	{
+		FHttpRequestRef& Request = *It;
+		Request->OnProcessRequestComplete().Unbind();
+		Request->OnRequestProgress().Unbind();
+		Request->OnHeaderReceived().Unbind();
+
+		// Don't emit these tracking logs in commandlet runs. Build system traps warnings during cook, and these are not truly fatal, but useful for tracking down shutdown issues.
+		UE_CLOG(ShouldOutputHttpWarnings(), LogHttp, Warning, TEXT("	verb=[%s] url=[%s] refs=[%d] status=%s"), *Request->GetVerb(), *Request->GetURL(), Request.GetSharedReferenceCount(), EHttpRequestStatus::ToString(Request->GetStatus()));
+	}
+
+	// Clear general delegates since they may point to deleted instances
+	RequestAddedDelegate.Unbind();
+	RequestCompletedDelegate.Unbind();
+
+	// Flush all requests
+	Flush(EHttpFlushReason::Shutdown);
+}
+
+bool FHttpManager::HasAnyBoundDelegate() const
+{
+	FScopeLock ScopeLock(&RequestLock);
+
+	for (TArray<FHttpRequestRef>::TConstIterator It(Requests); It; ++It)
+	{
+		const FHttpRequestRef& Request = *It;
+		if (Request->OnProcessRequestComplete().IsBound())
+		{
+			return true;
+		}
+	}
+
+	if (RequestAddedDelegate.IsBound())
+	{
+		return true;
+	}
+
+	if (RequestCompletedDelegate.IsBound())
+	{
+		return true;
+	}
+
+	return false;
+}
+
 void FHttpManager::ReloadFlushTimeLimits()
 {
 	FlushTimeLimitsMap.Reset();
@@ -225,24 +286,13 @@ FHttpThread* FHttpManager::CreateHttpThread()
 	return new FHttpThread();
 }
 
-void FHttpManager::Flush(bool bShutdown)
-{
-	Flush(bShutdown ? EHttpFlushReason::Shutdown : EHttpFlushReason::Default);
-}
-
-namespace
-{
-	bool ShouldOutputHttpWarnings()
-	{
-		return !IsRunningCommandlet() && !FApp::IsUnattended();
-	}
-}
-
 void FHttpManager::Flush(EHttpFlushReason FlushReason)
 {
 	QUICK_SCOPE_CYCLE_COUNTER(STAT_FHttpManager_Flush);
 
 	FScopeLock ScopeLock(&RequestLock);
+
+	checkf(FlushReason != EHttpFlushReason::Shutdown || !HasAnyBoundDelegate(), TEXT("Use Shutdown() instead of Flush(EHttpFlushReason::Shutdown) directly."));
 	
 	// This variable is set to indicate that flush is happening.
 	// While flushing is in progress, the RequestLock is held and threads are blocked when trying to submit new requests.
@@ -258,25 +308,6 @@ void FHttpManager::Flush(EHttpFlushReason FlushReason)
 
 	UE_CLOG(!IsRunningCommandlet(), LogHttp, Verbose, TEXT("[FHttpManager::Flush] FlushReason [%s] FlushTimeSoftLimitSeconds [%.3fs] FlushTimeHardLimitSeconds [%.3fs] SecondsToSleepForOutstandingThreadedRequests [%.3fs]"), LexToString(FlushReason), FlushTimeSoftLimitSeconds, FlushTimeHardLimitSeconds, SecondsToSleepForOutstandingThreadedRequests);
 
-	// Clear all delegates bound to ongoing Http requests
-	if (FlushReason == EHttpFlushReason::Shutdown)
-	{
-		// Don't emit these tracking logs in commandlet runs. Build system traps warnings during cook, and these are not truly fatal, but useful for tracking down shutdown issues.
-		UE_CLOG(ShouldOutputHttpWarnings() && Requests.Num(), LogHttp, Warning, TEXT("[FHttpManager::Flush] FlushReason was Shutdown. Unbinding delegates for %d outstanding Http Requests:"), Requests.Num());
-
-		// Clear delegates since they may point to deleted instances
-		for (TArray<FHttpRequestRef>::TIterator It(Requests); It; ++It)
-		{
-			FHttpRequestRef& Request = *It;
-			Request->OnProcessRequestComplete().Unbind();
-			Request->OnRequestProgress().Unbind();
-			Request->OnHeaderReceived().Unbind();
-
-			// Don't emit these tracking logs in commandlet runs. Build system traps warnings during cook, and these are not truly fatal, but useful for tracking down shutdown issues.
-			UE_CLOG(ShouldOutputHttpWarnings(), LogHttp, Warning, TEXT("	verb=[%s] url=[%s] refs=[%d] status=%s"), *Request->GetVerb(), *Request->GetURL(), Request.GetSharedReferenceCount(), EHttpRequestStatus::ToString(Request->GetStatus()));
-		}
-	}
-
 	UE_CLOG(!IsRunningCommandlet() && Requests.Num(), LogHttp, Verbose, TEXT("[FHttpManager::Flush] Cleanup starts for %d outstanding Http Requests."), Requests.Num());
 
 	double BeginWaitTime = FPlatformTime::Seconds();
@@ -410,9 +441,7 @@ void FHttpManager::AddRequest(const FHttpRequestRef& Request)
 	FScopeLock ScopeLock(&RequestLock);
 	check(!bFlushing);
 	Requests.Add(Request);
-PRAGMA_DISABLE_DEPRECATION_WARNINGS // Valid direct access of RequestAddedDelegate
 	RequestAddedDelegate.ExecuteIfBound(Request);
-PRAGMA_ENABLE_DEPRECATION_WARNINGS
 }
 
 void FHttpManager::RemoveRequest(const FHttpRequestRef& Request)
@@ -456,9 +485,7 @@ bool FHttpManager::IsValidRequest(const IHttpRequest* RequestPtr) const
 
 void FHttpManager::SetRequestAddedDelegate(const FHttpManagerRequestAddedDelegate& Delegate)
 {
-PRAGMA_DISABLE_DEPRECATION_WARNINGS // Valid direct access of RequestAddedDelegate
 	RequestAddedDelegate = Delegate;
-PRAGMA_ENABLE_DEPRECATION_WARNINGS
 }
 
 void FHttpManager::SetRequestCompletedDelegate(const FHttpManagerRequestCompletedDelegate& Delegate)
diff --git a/Engine/Source/Runtime/Online/HTTP/Private/HttpModule.cpp b/Engine/Source/Runtime/Online/HTTP/Private/HttpModule.cpp
--- a/Engine/Source/Runtime/Online/HTTP/Private/HttpModule.cpp
+++ b/Engine/Source/Runtime/Online/HTTP/Private/HttpModule.cpp
@@ -134,7 +134,7 @@ void FHttpModule::ShutdownModule()
 	if (HttpManager != nullptr)
 	{
 		// block on any http requests that have already been queued up
-		HttpManager->Flush(EHttpFlushReason::Shutdown);
+		HttpManager->Shutdown();
 	}
 
 	// at least on Linux, the code in HTTP manager (e.g. request destructors) expects platform to be initialized yet
diff --git a/Engine/Source/Runtime/Online/HTTP/Public/HttpManager.h b/Engine/Source/Runtime/Online/HTTP/Public/HttpManager.h
--- a/Engine/Source/Runtime/Online/HTTP/Public/HttpManager.h
+++ b/Engine/Source/Runtime/Online/HTTP/Public/HttpManager.h
@@ -74,6 +74,11 @@ public:
 	 */
 	void Initialize();
 
+	/**
+	 * Shutdown logic should be called before quiting
+	 */
+	void Shutdown();
+
 	/**
 	 * Adds an Http request instance to the manager for tracking/ticking
 	 * Manager should always have a list of requests currently being processed
@@ -82,10 +87,6 @@ public:
 	 */
 	void AddRequest(const FHttpRequestRef& Request);
 
-	/** Delegate that will get called once request added */
-	UE_DEPRECATED(5.2, "Direct access to RequestAddedDelegate is deprecated. Please use SetRequestAddedDelegate")
-	FHttpManagerRequestAddedDelegate RequestAddedDelegate;
-
 	void SetRequestAddedDelegate(const FHttpManagerRequestAddedDelegate& Delegate);
 	void SetRequestCompletedDelegate(const FHttpManagerRequestCompletedDelegate& Delegate);
 
@@ -106,14 +107,6 @@ public:
 	*/
 	bool IsValidRequest(const IHttpRequest* RequestPtr) const;
 
-	/**
-	 * Block until all pending requests are finished processing
-	 *
-	 * @param bShutdown true if final flush during shutdown
-	 */
-	UE_DEPRECATED(5.0, "This method is deprecated. Please use Flush(EHttpFlushReason FlushReason) instead.")
-	void Flush(bool bShutdown);
-
 	/**
 	 * Block until all pending requests are finished processing
 	 *
@@ -235,6 +228,8 @@ protected:
 
 	void ReloadFlushTimeLimits();
 
+	bool HasAnyBoundDelegate() const;
+
 protected:
 	/** List of Http requests that are actively being processed */
 	TArray<FHttpRequestRef> Requests;
@@ -250,6 +245,9 @@ protected:
 	// This variable is set to true in Flush(EHttpFlushReason), and prevents new Http requests from being launched
 	bool bFlushing;
 
+	/** Delegate that will get called once request added */
+	FHttpManagerRequestAddedDelegate RequestAddedDelegate;
+
 	/** Delegate that will get called when a request completes */
 	FHttpManagerRequestCompletedDelegate RequestCompletedDelegate;
 
