Some online modules adjusted so Windows Platform Extensions can re-use the Win64 versions.

#rnx

diff --git a/Engine/Source/Runtime/Online/XMPP/XMPP.Build.cs b/Engine/Source/Runtime/Online/XMPP/XMPP.Build.cs
--- a/Engine/Source/Runtime/Online/XMPP/XMPP.Build.cs
+++ b/Engine/Source/Runtime/Online/XMPP/XMPP.Build.cs
@@ -9,7 +9,7 @@ public class XMPP : ModuleRules
 	protected virtual bool bTargetPlatformSupportsStrophe
 	{
 		get =>
-			Target.Platform == UnrealTargetPlatform.Win64 ||
+			Target.IsInPlatformGroup(UnrealPlatformGroup.Windows) ||
 			Target.Platform == UnrealTargetPlatform.Android ||
 			Target.Platform == UnrealTargetPlatform.IOS ||
 			Target.Platform == UnrealTargetPlatform.Mac ||
@@ -54,7 +54,7 @@ public class XMPP : ModuleRules
 			PrivateDefinitions.Add("WITH_XMPP_STROPHE=0");
 		}
 
-		if (Target.Platform == UnrealTargetPlatform.Win64 ||
+		if (Target.IsInPlatformGroup(UnrealPlatformGroup.Windows) ||
 			Target.Platform == UnrealTargetPlatform.Mac ||
 			bRequireOpenSSL)
 		{
