BPT | Reject installation if a manifest has any file paths that escape the install directory

diff --git a/Engine/Source/Runtime/Online/BuildPatchServices/Private/BuildPatchInstaller.cpp b/Engine/Source/Runtime/Online/BuildPatchServices/Private/BuildPatchInstaller.cpp
--- a/Engine/Source/Runtime/Online/BuildPatchServices/Private/BuildPatchInstaller.cpp
+++ b/Engine/Source/Runtime/Online/BuildPatchServices/Private/BuildPatchInstaller.cpp
@@ -731,6 +731,24 @@ namespace BuildPatchServices
 			BuildStats.FailureType = InstallerError->GetErrorType();
 		}
 
+		// Check for any filepath violations.
+		FString InstallDirectoryWithSlash = Configuration.InstallDirectory;
+		FPaths::NormalizeDirectoryName(InstallDirectoryWithSlash);
+		FPaths::CollapseRelativeDirectories(InstallDirectoryWithSlash);
+		InstallDirectoryWithSlash /= TEXT("/");
+		TSet<FString> ExpectedFiles;
+		ManifestSet->GetExpectedFiles(ExpectedFiles);
+		for (const FString& ExpectedFile : ExpectedFiles)
+		{
+			const FString InstallConstructionFile = FPaths::ConvertRelativePathToFull(Configuration.InstallDirectory, ExpectedFile);
+			if (!InstallConstructionFile.StartsWith(InstallDirectoryWithSlash))
+			{
+				UE_LOG(LogBuildPatchServices, Error, TEXT("Installer setup: Filepath in manifest escaped install directory. %s -> %s"), *ExpectedFile, *InstallConstructionFile);
+				InstallerError->SetError(EBuildPatchInstallError::InitializationError, InitializationErrorCodes::InvalidDataInManifest);
+				bInstallerInitSuccess = false;
+			}
+		}
+
 		bIsInited = true;
 		return bInstallerInitSuccess;
 	}
diff --git a/Engine/Source/Runtime/Online/BuildPatchServices/Private/Installer/InstallerError.h b/Engine/Source/Runtime/Online/BuildPatchServices/Private/Installer/InstallerError.h
--- a/Engine/Source/Runtime/Online/BuildPatchServices/Private/Installer/InstallerError.h
+++ b/Engine/Source/Runtime/Online/BuildPatchServices/Private/Installer/InstallerError.h
@@ -18,6 +18,7 @@ namespace BuildPatchServices
 		static const TCHAR* InvalidInstallTags = TEXT("04");
 		static const TCHAR* ChunkReferenceTracking = TEXT("05");
 		static const TCHAR* MissingPrereqForPrereqOnlyInstall = TEXT("06");
+		static const TCHAR* InvalidDataInManifest = TEXT("07");
 	}
 
 	/**
