Ensure iOS HTTP traffic isn't cached to disk where not required

diff --git a/Engine/Source/Runtime/Online/HTTP/Private/Apple/AppleHTTPNSUrlConnection.cpp b/Engine/Source/Runtime/Online/HTTP/Private/Apple/AppleHTTPNSUrlConnection.cpp
--- a/Engine/Source/Runtime/Online/HTTP/Private/Apple/AppleHTTPNSUrlConnection.cpp
+++ b/Engine/Source/Runtime/Online/HTTP/Private/Apple/AppleHTTPNSUrlConnection.cpp
@@ -825,6 +825,13 @@ float FAppleHttpNSUrlConnectionRequest::GetElapsedTime() const
 }
 #endif
 
+- (NSCachedURLResponse *)connection:(NSURLConnection *)connection willCacheResponse:(NSCachedURLResponse *)cachedResponse
+{
+	// All FAppleHttpNSUrlConnectionRequest use NSURLRequestReloadIgnoringLocalCacheData
+	// NSURLRequestReloadIgnoringLocalCacheData disables loading of data from cache, but responses can still be stored in cache
+	// Returning nil in this delegate disables caching the responses
+	return nil;
+}
 
 -(void) connectionDidFinishLoading:(NSURLConnection *)connection
 {
@@ -845,8 +852,6 @@ float FAppleHttpNSUrlConnectionRequest::GetElapsedTime() const
 @end
 
 
-
-
 /****************************************************************************
  * FAppleHttpNSUrlConnectionResponse implementation
  **************************************************************************/
diff --git a/Engine/Source/Runtime/Online/HTTP/Private/Apple/AppleHTTPNSUrlSession.cpp b/Engine/Source/Runtime/Online/HTTP/Private/Apple/AppleHTTPNSUrlSession.cpp
--- a/Engine/Source/Runtime/Online/HTTP/Private/Apple/AppleHTTPNSUrlSession.cpp
+++ b/Engine/Source/Runtime/Online/HTTP/Private/Apple/AppleHTTPNSUrlSession.cpp
@@ -49,7 +49,8 @@ enum class EAppleHttpRequestResponseState: uint8
 - (void)URLSession:(NSURLSession *)session dataTask:(NSURLSessionDataTask *)dataTask didReceiveData:(NSData *)data;
 /** Sent as the last message related to a specific task.  A nil Error implies that no error occurred and this task is complete. */
 - (void)URLSession:(NSURLSession *)session task:(NSURLSessionTask *)task didCompleteWithError:(nullable NSError *)error;
-
+/** Asks the delegate if it needs to store responses in the cache. */
+- (void)URLSession:(NSURLSession *)session dataTask:(NSURLSessionDataTask *)dataTask willCacheResponse:(NSCachedURLResponse *)proposedResponse completionHandler:(void (^)(NSCachedURLResponse *cachedResponse))completionHandler;
 @end
 
 @implementation FAppleHttpNSUrlSessionResponseDelegate
@@ -106,7 +107,7 @@ enum class EAppleHttpRequestResponseState: uint8
 	UE_LOG(LogHttp, Verbose, TEXT("URLSession:dataTask:didReceiveData with %u bytes. After Append, Payload Length = %d: %p"), [data length], Payload.Num(), self);
 }
 
-- (void)URLSession:(NSURLSession *)session task:(NSURLSessionTask *)task didCompleteWithError:(nullable NSError *)error;
+- (void)URLSession:(NSURLSession *)session task:(NSURLSessionTask *)task didCompleteWithError:(nullable NSError *)error
 {
 	if (error == nil)
 	{
@@ -161,6 +162,14 @@ enum class EAppleHttpRequestResponseState: uint8
 		}
 	}
 }
+
+- (void)URLSession:(NSURLSession *)session dataTask:(NSURLSessionDataTask *)dataTask willCacheResponse:(NSCachedURLResponse *)proposedResponse completionHandler:(void (^)(NSCachedURLResponse *cachedResponse))completionHandler
+{
+	// All FAppleHttpNSUrlSessionRequest use NSURLRequestReloadIgnoringLocalCacheData
+	// NSURLRequestReloadIgnoringLocalCacheData disables loading of data from cache, but responses can still be stored in cache
+	// Passing nil to this handler disables caching the responses
+	completionHandler(nil);
+}
 @end
 
 /****************************************************************************
