Adding scope counters and some logging to FHttp* classes to track down root cause of server deferred tick time regression.

diff --git a/Engine/Source/Runtime/Online/HTTPServer/Private/HttpConnection.cpp b/Engine/Source/Runtime/Online/HTTPServer/Private/HttpConnection.cpp
--- a/Engine/Source/Runtime/Online/HTTPServer/Private/HttpConnection.cpp
+++ b/Engine/Source/Runtime/Online/HTTPServer/Private/HttpConnection.cpp
@@ -35,6 +35,7 @@ FHttpConnection::~FHttpConnection()
 
 void FHttpConnection::Tick(float DeltaTime)
 {
+	QUICK_SCOPE_CYCLE_COUNTER(STAT_FHttpConnection_Tick);
 	const float AwaitReadTimeout = bKeepAlive ?
 		ConnectionKeepAliveTimeout : ConnectionTimeout;
 
@@ -95,6 +96,7 @@ void FHttpConnection::TransferState(EHttpConnectionState CurrentState, EHttpConn
 
 void FHttpConnection::BeginRead(float DeltaTime)
 {
+	QUICK_SCOPE_CYCLE_COUNTER(STAT_FHttpConnection_BeginRead);
 	const FTimespan WaitTime = FTimespan::FromMilliseconds(1);
 	if (!Socket->Wait(ESocketWaitConditions::WaitForRead, WaitTime))
 	{
@@ -105,6 +107,10 @@ void FHttpConnection::BeginRead(float DeltaTime)
 	}
 	else
 	{
+		UE_LOG(LogHttpConnection, Verbose,
+			TEXT("SecondsWaitingForReadableSocket\t [%d][%u]-%u : %f"),
+				OriginPort, ConnectionId, LastRequestNumber, ReadContext.GetSecondsWaitingForReadableSocket());
+
 		ReadContext.ResetSecondsWaitingForReadableSocket();
 
 		// The socket is reachable, however there may not be data in the pipe
@@ -125,6 +131,7 @@ void FHttpConnection::BeginRead(float DeltaTime)
 
 void FHttpConnection::ContinueRead(float DeltaTime)
 {
+	QUICK_SCOPE_CYCLE_COUNTER(STAT_FHttpConnection_ContinueRead);
 	check(State == EHttpConnectionState::Reading);
 
 	auto ReaderState = ReadContext.ReadStream(DeltaTime);
@@ -146,6 +153,7 @@ void FHttpConnection::ContinueRead(float DeltaTime)
 
 void FHttpConnection::CompleteRead(const TSharedPtr<FHttpServerRequest>& Request)
 {
+	QUICK_SCOPE_CYCLE_COUNTER(STAT_FHttpConnection_CompleteRead);
 	TArray<FString>* ConnectionHeaders = Request->Headers.Find(UE_HTTP_SERVER_HEADER_KEYS_CONNECTION);
 	if (ConnectionHeaders)
 	{
@@ -182,6 +190,7 @@ void FHttpConnection::CompleteRead(const TSharedPtr<FHttpServerRequest>& Request
 
 void FHttpConnection::ProcessRequest(const TSharedPtr<FHttpServerRequest>& Request, const FHttpResultCallback& OnProcessingComplete)
 {
+	QUICK_SCOPE_CYCLE_COUNTER(STAT_FHttpConnection_ProcessRequest);
 	TransferState(EHttpConnectionState::Reading, EHttpConnectionState::AwaitingProcessing);
 
 	UE_LOG(LogHttpConnection, Verbose,
@@ -211,6 +220,7 @@ void FHttpConnection::ProcessRequest(const TSharedPtr<FHttpServerRequest>& Reque
 
 void FHttpConnection::BeginWrite(TUniquePtr<FHttpServerResponse>&& Response, uint32 RequestNumber)
 {
+	QUICK_SCOPE_CYCLE_COUNTER(STAT_FHttpConnection_BeginWrite);
 	// Ensure the passed-in request number is the one we expect
 	check(RequestNumber == LastRequestNumber);
 
@@ -229,6 +239,7 @@ void FHttpConnection::BeginWrite(TUniquePtr<FHttpServerResponse>&& Response, uin
 
 void FHttpConnection::ContinueWrite(float DeltaTime)
 {
+	QUICK_SCOPE_CYCLE_COUNTER(STAT_FHttpConnection_ContinueWrite);
 	check(State == EHttpConnectionState::Writing);
 
 	auto WriterState = WriteContext.WriteStream(DeltaTime);
@@ -249,6 +260,7 @@ void FHttpConnection::ContinueWrite(float DeltaTime)
 
 void FHttpConnection::CompleteWrite()
 {
+	QUICK_SCOPE_CYCLE_COUNTER(STAT_FHttpConnection_CompleteWrite);
 	check(EHttpConnectionState::Writing == State);
 
 	if (bKeepAlive && !bGracefulDestroyRequested)
diff --git a/Engine/Source/Runtime/Online/HTTPServer/Private/HttpConnectionRequestReadContext.cpp b/Engine/Source/Runtime/Online/HTTPServer/Private/HttpConnectionRequestReadContext.cpp
--- a/Engine/Source/Runtime/Online/HTTPServer/Private/HttpConnectionRequestReadContext.cpp
+++ b/Engine/Source/Runtime/Online/HTTPServer/Private/HttpConnectionRequestReadContext.cpp
@@ -10,6 +10,8 @@
 #include "Sockets.h"
 #include "SocketSubsystem.h"
 
+DEFINE_LOG_CATEGORY(LogHttpConnectionRequestReadContext)
+
 FHttpConnectionRequestReadContext::FHttpConnectionRequestReadContext(FSocket* InSocket)
 	: Socket(InSocket)
 {
@@ -47,6 +49,8 @@ EHttpConnectionContextState FHttpConnectionRequestReadContext::ReadStream(float
 	}
 	else
 	{
+		UE_LOG(LogHttpConnectionRequestReadContext, Verbose,
+			TEXT("ElapsedIdleTime\t %f"), ElapsedIdleTime);
 		ElapsedIdleTime = 0.0f;
 	}
 
diff --git a/Engine/Source/Runtime/Online/HTTPServer/Private/HttpConnectionRequestReadContext.h b/Engine/Source/Runtime/Online/HTTPServer/Private/HttpConnectionRequestReadContext.h
--- a/Engine/Source/Runtime/Online/HTTPServer/Private/HttpConnectionRequestReadContext.h
+++ b/Engine/Source/Runtime/Online/HTTPServer/Private/HttpConnectionRequestReadContext.h
@@ -8,6 +8,7 @@ struct FHttpConnection;
 
 class FSocket;
 
+DECLARE_LOG_CATEGORY_EXTERN(LogHttpConnectionRequestReadContext, Log, All);
 
 struct FHttpConnectionRequestReadContext final : public FHttpConnectionContext
 {
diff --git a/Engine/Source/Runtime/Online/HTTPServer/Private/HttpConnectionResponseWriteContext.cpp b/Engine/Source/Runtime/Online/HTTPServer/Private/HttpConnectionResponseWriteContext.cpp
--- a/Engine/Source/Runtime/Online/HTTPServer/Private/HttpConnectionResponseWriteContext.cpp
+++ b/Engine/Source/Runtime/Online/HTTPServer/Private/HttpConnectionResponseWriteContext.cpp
@@ -85,6 +85,8 @@ bool FHttpConnectionResponseWriteContext::WriteBytes(const uint8* Bytes, int32 B
 
 	if (OutBytesWritten > 0)
 	{
+		UE_LOG(LogHttpConnectionResponseWriteContext, Verbose,
+			TEXT("ElapsedIdleTime\t %f"), ElapsedIdleTime);
 		ElapsedIdleTime = 0.0f;
 	}
 
diff --git a/Engine/Source/Runtime/Online/HTTPServer/Private/HttpListener.cpp b/Engine/Source/Runtime/Online/HTTPServer/Private/HttpListener.cpp
--- a/Engine/Source/Runtime/Online/HTTPServer/Private/HttpListener.cpp
+++ b/Engine/Source/Runtime/Online/HTTPServer/Private/HttpListener.cpp
@@ -183,6 +183,7 @@ bool FHttpListener::HasPendingConnections() const
 // --------------------------------------------------------------------------------------------
 void FHttpListener::AcceptConnections()
 {
+	QUICK_SCOPE_CYCLE_COUNTER(STAT_FHttpListener_AcceptConnections);
 	check(ListenSocket);
 
 	for (int32 i = 0; i < Config.MaxConnectionsAcceptPerFrame; ++i)
@@ -226,6 +227,7 @@ void FHttpListener::AcceptConnections()
 
 void FHttpListener::TickConnections(float DeltaTime)
 {
+	QUICK_SCOPE_CYCLE_COUNTER(STAT_FHttpListener_TickConnections);
 	for (const auto& Connection : Connections)
 	{
 		check(Connection.IsValid());
@@ -254,6 +256,7 @@ void FHttpListener::TickConnections(float DeltaTime)
 
 void FHttpListener::RemoveDestroyedConnections()
 {
+	QUICK_SCOPE_CYCLE_COUNTER(STAT_FHttpListener_RemoveDestroyedConnections);
 	for (auto ConnectionsIter = Connections.CreateIterator(); ConnectionsIter; ++ConnectionsIter)
 	{
 		// Remove any destroyed connections
