Use warning log instead of ensure in CurlHttpManager when shutdown.

diff --git a/Engine/Source/Runtime/Online/HTTP/Private/Curl/CurlHttpManager.cpp b/Engine/Source/Runtime/Online/HTTP/Private/Curl/CurlHttpManager.cpp
--- a/Engine/Source/Runtime/Online/HTTP/Private/Curl/CurlHttpManager.cpp
+++ b/Engine/Source/Runtime/Online/HTTP/Private/Curl/CurlHttpManager.cpp
@@ -341,7 +341,7 @@ void FCurlHttpManager::ShutdownCurl()
 	if (GShareHandle != nullptr)
 	{
 		CURLSHcode ShareCleanupCode = curl_share_cleanup(GShareHandle);
-		ensureMsgf(ShareCleanupCode == CURLSHE_OK, TEXT("CurlShareCleanup failed. ReturnValue=[%d]"), static_cast<int32>(ShareCleanupCode));
+		UE_CLOG(ShareCleanupCode != CURLSHE_OK, LogHttp, Warning, TEXT("curl_share_cleanup failed. ReturnValue=[%d]"), static_cast<int32>(ShareCleanupCode));
 		GShareHandle = nullptr;
 	}
 #endif
@@ -349,7 +349,7 @@ void FCurlHttpManager::ShutdownCurl()
 	if (GMultiHandle != nullptr)
 	{
 		CURLMcode MutliCleanupCode = curl_multi_cleanup(GMultiHandle);
-		ensureMsgf(MutliCleanupCode == CURLM_OK, TEXT("CurlMultiCleanup failed. ReturnValue=[%d]"), static_cast<int32>(MutliCleanupCode));
+		UE_CLOG(MutliCleanupCode != CURLM_OK, LogHttp, Warning, TEXT("curl_multi_cleanup failed. ReturnValue=[%d]"), static_cast<int32>(MutliCleanupCode));
 		GMultiHandle = nullptr;
 	}
 
