Fixed mismatched copy constructor/assignment operator in FIcmpEchoResult.

diff --git a/Engine/Source/Runtime/Online/ICMP/Public/Icmp.h b/Engine/Source/Runtime/Online/ICMP/Public/Icmp.h
--- a/Engine/Source/Runtime/Online/ICMP/Public/Icmp.h
+++ b/Engine/Source/Runtime/Online/ICMP/Public/Icmp.h
@@ -55,13 +55,6 @@ struct FIcmpEchoResult
 		, ReplyFrom()
 		, Time(-1)
 	{}
-
-	FIcmpEchoResult(const FIcmpEchoResult& Other)
-		: Status(Other.Status)
-		, ResolvedAddress(Other.ResolvedAddress)
-		, ReplyFrom(Other.ReplyFrom)
-		, Time(Other.Time)
-	{}
 };
 
 typedef TFunction<void(FIcmpEchoResult)> FIcmpEchoResultCallback;
