CSV Telemetry for Cosmetic streaming

diff --git a/Engine/Source/Runtime/Online/BuildPatchServices/Private/Installer/BuildInstallStreamer.cpp b/Engine/Source/Runtime/Online/BuildPatchServices/Private/Installer/BuildInstallStreamer.cpp
--- a/Engine/Source/Runtime/Online/BuildPatchServices/Private/Installer/BuildInstallStreamer.cpp
+++ b/Engine/Source/Runtime/Online/BuildPatchServices/Private/Installer/BuildInstallStreamer.cpp
@@ -26,9 +26,12 @@
 #include "Installer/Statistics/MemoryChunkStoreStatistics.h"
 #include "Installer/VirtualFileConstructor.h"
 #include "VirtualFileCache.h"
+#include "ProfilingDebugging/CsvProfiler.h"
 
 DEFINE_LOG_CATEGORY_STATIC(LogBuildInstallStreamer, Log, All);
 
+CSV_DEFINE_CATEGORY(CosmeticStreamingCsv, true);
+
 namespace BuildPatchServices
 {
 	namespace
@@ -67,6 +70,8 @@ namespace BuildPatchServices
 		, bIsShuttingDown(false)
 		, RequestTrigger(FPlatformProcess::GetSynchEventFromPool(false))
 		, CloudTrigger(FPlatformProcess::GetSynchEventFromPool(false))
+		, BuildStats(FBuildInstallStreamerStats())
+		
 	{
 		RequestWorker = Async(EAsyncExecution::ThreadIfForkSafe, [this]()
 			{
@@ -147,6 +152,7 @@ namespace BuildPatchServices
 
 		auto CompleteRequest = [this, &Requests, &RequestTime]()
 		{
+			uint64 TotalDownload = DownloadServiceStatistics->GetBytesDownloaded();
 			for (FStreamRequest& Request : Requests)
 			{
 				FBuildPatchStreamResult Result;
@@ -154,10 +160,13 @@ namespace BuildPatchServices
 				Result.ErrorType = InstallerError->GetErrorType();
 				Result.ErrorCode = InstallerError->GetErrorCode();
 				// TODO: Not accurate.. How to fix this
-				Result.TotalDownloaded = DownloadServiceStatistics->GetBytesDownloaded();
+				Result.TotalDownloaded = TotalDownload;
+				
+				BuildStats.AmountRequestsCompleted++;
 				UE_LOG(LogBuildInstallStreamer, Verbose, TEXT("Completed request %s in %s"), *FString::Join(Result.Request, TEXT(",")), *FPlatformTime::PrettyTime(RequestTime));
 				Request.Get<2>().ExecuteIfBound(MoveTemp(Result));
 			}
+			BuildStats.MegaBytesDownloaded += (TotalDownload / 1024.0f / 1024.0f);
 		};
 
 		auto CancelRequest = [this, &Requests]()
@@ -169,6 +178,7 @@ namespace BuildPatchServices
 				Result.ErrorType = InstallerError->GetErrorType();
 				Result.ErrorCode = InstallerError->GetErrorCode();
 				Result.TotalDownloaded = 0;
+				BuildStats.AmountRequestsCancelled++;
 				UE_LOG(LogBuildInstallStreamer, Verbose, TEXT("Cancelled request %s"), *FString::Join(Result.Request, TEXT(",")));
 				Request.Get<2>().ExecuteIfBound(MoveTemp(Result));
 			}
@@ -179,6 +189,7 @@ namespace BuildPatchServices
 			// If we were asked to cancel all requests, then do this first
 			if (InstallerError->IsCancelled())
 			{
+				
 				while (RequestQueue.Dequeue(Requests.AddDefaulted_GetRef())) {}
 				Requests.Pop();
 				if (Requests.Num() > 0)
@@ -203,6 +214,7 @@ namespace BuildPatchServices
 			{
 				for (const FStreamRequest& Request : Requests)
 				{
+					BuildStats.AmountRequestsMade++;
 					UE_LOG(LogBuildInstallStreamer, Verbose, TEXT("Begin request %s"), *FString::Join(Request.Get<0>(), TEXT(",")));
 				}
 
@@ -211,6 +223,13 @@ namespace BuildPatchServices
 				AsyncHelpers::ExecuteOnCustomThread<void>([this]() { DownloadServiceStatistics->Reset(); }, TickQueue);
 				TSharedPtr<IVirtualFileCache, ESPMode::ThreadSafe> VirtualFileCache = IVirtualFileCache::CreateVirtualFileCache();
 				InstallerError->Reset();
+#if CSV_PROFILER
+				//Used and Total size is calulcated on VFC Start.
+				//Used + Requested write will always result in new total. Doing this to recalculating all blocks.
+				BuildStats.VFCCachedUsedSize = (VirtualFileCache->GetUsedSize() / 1024.0f / 1024.0f);
+				BuildStats.VFCCachedTotalSize = (VirtualFileCache->GetTotalSize()  / 1024.0f / 1024.0f);
+#endif
+
 
 				// Setup file build list.
 				FVirtualFileConstructorConfiguration VFCConfig;
@@ -235,6 +254,10 @@ namespace BuildPatchServices
 					{
 						FileIt.RemoveCurrent();
 					}
+					else
+					{
+						BuildStats.VFCRequestedFileWrite += (FileManifest->FileSize / 1024.0f / 1024.0f);
+					}
 				}
 
 				// Composition root.
@@ -350,6 +373,16 @@ namespace BuildPatchServices
 			TickFunc();
 		}
 
+#if !UE_BUILD_SHIPPING
+		CSV_CUSTOM_STAT(CosmeticStreamingCsv, BuildStats.VFCRequestedFileWrite,		BuildStats.VFCRequestedFileWrite, ECsvCustomStatOp::Set);
+		CSV_CUSTOM_STAT(CosmeticStreamingCsv, BuildStats.VFCCachedUsedSize,			BuildStats.VFCCachedUsedSize + BuildStats.VFCRequestedFileWrite, ECsvCustomStatOp::Set);
+		CSV_CUSTOM_STAT(CosmeticStreamingCsv, BuildStats.VFCCachedTotalSize,		BuildStats.VFCCachedTotalSize, ECsvCustomStatOp::Set);
+		CSV_CUSTOM_STAT(CosmeticStreamingCsv, BuildStats.MegaBytesDownloaded,		BuildStats.MegaBytesDownloaded, ECsvCustomStatOp::Set);
+		CSV_CUSTOM_STAT(CosmeticStreamingCsv, BuildStats.AmountRequestsCompleted,	BuildStats.AmountRequestsCompleted, ECsvCustomStatOp::Set);
+		CSV_CUSTOM_STAT(CosmeticStreamingCsv, BuildStats.AmountRequestsCancelled,	BuildStats.AmountRequestsCancelled, ECsvCustomStatOp::Set);
+		CSV_CUSTOM_STAT(CosmeticStreamingCsv, BuildStats.AmountRequestsMade,		BuildStats.AmountRequestsMade, ECsvCustomStatOp::Set);
+#endif
+
 		return bKeepTicking;
 	}
 
diff --git a/Engine/Source/Runtime/Online/BuildPatchServices/Private/Installer/BuildInstallStreamer.h b/Engine/Source/Runtime/Online/BuildPatchServices/Private/Installer/BuildInstallStreamer.h
--- a/Engine/Source/Runtime/Online/BuildPatchServices/Private/Installer/BuildInstallStreamer.h
+++ b/Engine/Source/Runtime/Online/BuildPatchServices/Private/Installer/BuildInstallStreamer.h
@@ -34,6 +34,17 @@ namespace BuildPatchServices
 	class IDownloadService;
 	class IMessagePump;
 
+	struct FBuildInstallStreamerStats
+	{
+		float MegaBytesDownloaded;
+		int32 AmountRequestsCompleted;
+		int32 AmountRequestsCancelled;
+		int32 AmountRequestsMade;
+		float VFCCachedTotalSize;
+		float VFCCachedUsedSize;
+		float VFCRequestedFileWrite;
+	};
+
 	class FBuildInstallStreamer
 		: public IBuildInstallStreamer
 		, public TSharedFromThis<FBuildInstallStreamer, ESPMode::ThreadSafe>
@@ -62,6 +73,7 @@ namespace BuildPatchServices
 		void CallOnDelegateThread(const TFunction<void()>& Callback);
 
 	private:
+
 		const FBuildInstallStreamerConfiguration Configuration;
 		FBuildPatchAppManifestPtr BuildPatchManifest;
 		FBuildPatchProgress BuildProgress;
@@ -94,6 +106,7 @@ namespace BuildPatchServices
 		TQueue<TFunction<void()>, EQueueMode::Spsc> TickQueue;
 		FEvent* RequestTrigger;
 		FEvent* CloudTrigger;
+		FBuildInstallStreamerStats BuildStats;
 	};
 
 	/**
diff --git a/Engine/Source/Runtime/VirtualFileCache/Private/VirtualFileCache.cpp b/Engine/Source/Runtime/VirtualFileCache/Private/VirtualFileCache.cpp
--- a/Engine/Source/Runtime/VirtualFileCache/Private/VirtualFileCache.cpp
+++ b/Engine/Source/Runtime/VirtualFileCache/Private/VirtualFileCache.cpp
@@ -104,3 +104,15 @@ void FVirtualFileCache::Defragment()
 	FFileTableWriter FileTable = Thread.ModifyFileTable();
 	return FileTable->Defragment();
 }
+
+int64 FVirtualFileCache::GetTotalSize() const
+{
+	FFileTableReader FileTable = Thread.ReadFileTable();
+	return FileTable->GetTotalSize();
+}
+
+int64 FVirtualFileCache::GetUsedSize() const
+{
+	FFileTableReader FileTable = Thread.ReadFileTable();
+	return FileTable->GetUsedSize();
+}
diff --git a/Engine/Source/Runtime/VirtualFileCache/Private/VirtualFileCacheInternal.cpp b/Engine/Source/Runtime/VirtualFileCache/Private/VirtualFileCacheInternal.cpp
--- a/Engine/Source/Runtime/VirtualFileCache/Private/VirtualFileCacheInternal.cpp
+++ b/Engine/Source/Runtime/VirtualFileCache/Private/VirtualFileCacheInternal.cpp
@@ -1041,6 +1041,17 @@ void FFileTable::Empty()
 	UsedSize = -1;
 }
 
+
+int64 FFileTable::GetUsedSize() const
+{
+	return UsedSize;
+}
+
+int64 FFileTable::GetTotalSize() const
+{
+	return TotalSize;
+}
+
 double FFileTable::CurrentFragmentation() const
 {
 	double RangesPerFile = 0;
@@ -1206,7 +1217,6 @@ void FFileTable::CalculateSizes()
 	UsedSize = Sum - FreeSize;
 }
 
-
 FString FFileTable::GetBlockFilename(const FBlockFile& File) const
 {
 	FString Filename = VFC_CACHE_FILE_BASE_NAME + LexToString(File.FileId);
diff --git a/Engine/Source/Runtime/VirtualFileCache/Private/VirtualFileCacheInternal.h b/Engine/Source/Runtime/VirtualFileCache/Private/VirtualFileCacheInternal.h
--- a/Engine/Source/Runtime/VirtualFileCache/Private/VirtualFileCacheInternal.h
+++ b/Engine/Source/Runtime/VirtualFileCache/Private/VirtualFileCacheInternal.h
@@ -113,6 +113,8 @@ struct IFileTableReader
 	virtual TIoStatusOr<uint64> GetSizeForChunk(const VFCKey& Id) const = 0;
 	virtual bool DoesChunkExist(const VFCKey& Id) const = 0;
 	virtual double CurrentFragmentation() const = 0;
+	virtual int64 GetUsedSize() const = 0;
+	virtual int64 GetTotalSize() const = 0;
 };
 
 struct FFileTable : IFileTableReader
@@ -180,6 +182,8 @@ public:
 	virtual TIoStatusOr<uint64> GetSizeForChunk(const VFCKey& Id) const override;
 	virtual bool DoesChunkExist(const VFCKey& Id) const override;
 	virtual double CurrentFragmentation() const override;
+	virtual int64 GetUsedSize() const override;
+	virtual int64 GetTotalSize() const override;
 	// \IFileTableReader
 
 	friend FArchive& operator<<(FArchive& Ar, FFileTable& FileTable);
@@ -412,6 +416,8 @@ public:
 	virtual TIoStatusOr<uint64> GetSizeForChunk(const VFCKey& Id) const override;
 	virtual double CurrentFragmentation() const override;
 	virtual void Defragment() override;
+	virtual int64 GetTotalSize() const override;
+	virtual int64 GetUsedSize() const override;
 
 	FString GetTableFilename() const { return Thread.GetTableFilename(); }
 
diff --git a/Engine/Source/Runtime/VirtualFileCache/Public/VirtualFileCache.h b/Engine/Source/Runtime/VirtualFileCache/Public/VirtualFileCache.h
--- a/Engine/Source/Runtime/VirtualFileCache/Public/VirtualFileCache.h
+++ b/Engine/Source/Runtime/VirtualFileCache/Public/VirtualFileCache.h
@@ -36,6 +36,8 @@ struct IVirtualFileCache
 	virtual TIoStatusOr<uint64> GetSizeForChunk(const VFCKey& Id) const = 0;
 	virtual double CurrentFragmentation() const = 0;
 	virtual void Defragment() = 0;
+	virtual int64 GetTotalSize() const = 0;
+	virtual int64 GetUsedSize() const = 0;
 
 	static VIRTUALFILECACHE_API TSharedRef<IVirtualFileCache, ESPMode::ThreadSafe> CreateVirtualFileCache();
 };
