Add ini variable to enable/disable disk overflow store for patch chunks.

diff --git a/Engine/Config/BaseEngine.ini b/Engine/Config/BaseEngine.ini
--- a/Engine/Config/BaseEngine.ini
+++ b/Engine/Config/BaseEngine.ini
@@ -3310,3 +3310,7 @@ GenericPipelineClass=/Script/InterchangePipelines.InterchangeGenericAssetsPipeli
 
 [/Script/InterchangePipelines.MaterialXPipelineSettings]
 PredefinedSurfaceShaders=((StandardSurface, "/Interchange/Functions/MX_StandardSurface.MX_StandardSurface"),(StandardSurfaceTransmission, "/Interchange/Functions/MX_TransmissionSurface.MX_TransmissionSurface"),(SurfaceUnlit, "/Interchange/Functions/MX_SurfaceUnlit.MX_SurfaceUnlit"),(UsdPreviewSurface, "/Interchange/Functions/MX_UsdPreviewSurface.MX_UsdPreviewSurface"))
+
+[BuildPatchServices]
+; Enable/Disable mechanism that provides a disk overflow store for patch chunks.
+bEnableDiskOverflowStore=true
\ No newline at end of file
diff --git a/Engine/Source/Runtime/Online/BuildPatchServices/BuildPatchServices.Build.cs b/Engine/Source/Runtime/Online/BuildPatchServices/BuildPatchServices.Build.cs
--- a/Engine/Source/Runtime/Online/BuildPatchServices/BuildPatchServices.Build.cs
+++ b/Engine/Source/Runtime/Online/BuildPatchServices/BuildPatchServices.Build.cs
@@ -1,9 +1,13 @@
 // Copyright Epic Games, Inc. All Rights Reserved.
 
+using EpicGames.Core;
 using UnrealBuildTool;
 
 public class BuildPatchServices : ModuleRules
 {
+	[ConfigFile(ConfigHierarchyType.Engine, "BuildPatchServices")]
+	bool bEnableDiskOverflowStore = true;
+
 	public BuildPatchServices(ReadOnlyTargetRules Target) : base(Target)
 	{
 		PublicDependencyModuleNames.AddRange(
@@ -35,5 +39,24 @@ public class BuildPatchServices : ModuleRules
 				"HTTP"
 			}
 		);
+
+		if (EnableDiskOverflowStore)
+		{
+			PublicDefinitions.Add("ENABLE_PATCH_DISK_OVERFLOW_STORE=1");
+		}
+		else
+		{
+			PublicDefinitions.Add("ENABLE_PATCH_DISK_OVERFLOW_STORE=0");
+		}
+
+	}
+
+	protected bool EnableDiskOverflowStore
+	{
+		get
+		{
+			ConfigCache.ReadSettings(DirectoryReference.FromFile(Target.ProjectFile), Target.Platform, this);
+			return bEnableDiskOverflowStore;
+		}
 	}
 }
diff --git a/Engine/Source/Runtime/Online/BuildPatchServices/Private/BuildPatchInstaller.cpp b/Engine/Source/Runtime/Online/BuildPatchServices/Private/BuildPatchInstaller.cpp
--- a/Engine/Source/Runtime/Online/BuildPatchServices/Private/BuildPatchInstaller.cpp
+++ b/Engine/Source/Runtime/Online/BuildPatchServices/Private/BuildPatchInstaller.cpp
@@ -53,6 +53,10 @@
 #include "Installer/Statistics/VerifierStatistics.h"
 #include "Installer/Statistics/FileOperationTracker.h"
 
+#if !defined(ENABLE_PATCH_DISK_OVERFLOW_STORE)
+#	define ENABLE_PATCH_DISK_OVERFLOW_STORE 1
+#endif
+
 DEFINE_LOG_CATEGORY_STATIC(LogBPSInstallerConfig, Log, All);
 
 namespace ConfigHelpers
@@ -1254,15 +1258,21 @@ namespace BuildPatchServices
 			TSet<FGuid> ReferencedChunks = ChunkReferenceTracker->GetReferencedChunks();
 			TUniquePtr<IChunkEvictionPolicy> MemoryEvictionPolicy(FChunkEvictionPolicyFactory::Create(
 				ChunkReferenceTracker.Get()));
+#if ENABLE_PATCH_DISK_OVERFLOW_STORE
 			TUniquePtr<IDiskChunkStore> DiskOverflowStore(FDiskChunkStoreFactory::Create(
 				FileSystem.Get(),
 				ChunkDataSerialization.Get(),
 				DiskChunkStoreStatistics.Get(),
 				FDiskChunkStoreConfig(DataStagingDir)));
+#endif // ENABLE_PATCH_DISK_OVERFLOW_STORE
 			TUniquePtr<IMemoryChunkStore> CloudChunkStore(FMemoryChunkStoreFactory::Create(
 				ChunkStoreMemorySize,
 				MemoryEvictionPolicy.Get(),
+#if ENABLE_PATCH_DISK_OVERFLOW_STORE
 				DiskOverflowStore.Get(),
+#else
+				nullptr,
+#endif // ENABLE_PATCH_DISK_OVERFLOW_STORE
 				MemoryChunkStoreStatistics.Get()));
 			TUniquePtr<IChunkDbChunkSource> ChunkDbChunkSource(FChunkDbChunkSourceFactory::Create(
 				BuildChunkDbSourceConfig(),
@@ -1333,7 +1343,9 @@ namespace BuildPatchServices
 			{
 				ChainedChunkSource->AddRepeatRequirement(LostChunk);
 			};
+#if ENABLE_PATCH_DISK_OVERFLOW_STORE
 			DiskOverflowStore->SetLostChunkCallback(LostChunkCallback);
+#endif // ENABLE_PATCH_DISK_OVERFLOW_STORE
 			CloudChunkStore->SetLostChunkCallback(LostChunkCallback);
 
 
