Fix possible nullptr dereference

diff --git a/Engine/Source/Runtime/Online/BuildPatchServices/Private/Installer/BuildInstallStreamer.cpp b/Engine/Source/Runtime/Online/BuildPatchServices/Private/Installer/BuildInstallStreamer.cpp
--- a/Engine/Source/Runtime/Online/BuildPatchServices/Private/Installer/BuildInstallStreamer.cpp
+++ b/Engine/Source/Runtime/Online/BuildPatchServices/Private/Installer/BuildInstallStreamer.cpp
@@ -254,7 +254,7 @@ namespace BuildPatchServices
 					{
 						FileIt.RemoveCurrent();
 					}
-					else
+					else if (FileManifest != nullptr)
 					{
 						BuildStats.VFCRequestedFileWrite += (FileManifest->FileSize / 1024.0f / 1024.0f);
 					}
