Fixed mismatched copy constructor/assignment operator in FSslManager.

diff --git a/Engine/Source/Runtime/Online/SSL/Private/SslManager.h b/Engine/Source/Runtime/Online/SSL/Private/SslManager.h
--- a/Engine/Source/Runtime/Online/SSL/Private/SslManager.h
+++ b/Engine/Source/Runtime/Online/SSL/Private/SslManager.h
@@ -26,8 +26,8 @@ public:
 protected:
 	/** Default constructor */
 	FSslManager();
-	/** Disable copy constructor */
-	FSslManager(const FSslManager& Copy) = delete;
+	/** Disable copies */
+	UE_NONCOPYABLE(FSslManager)
 	/** SSL ref count */
 	int32 InitCount;
 
