Make the wait time for FHttpConnection::BeginRead configurable.

diff --git a/Engine/Source/Runtime/Online/HTTPServer/Private/HttpConnection.cpp b/Engine/Source/Runtime/Online/HTTPServer/Private/HttpConnection.cpp
--- a/Engine/Source/Runtime/Online/HTTPServer/Private/HttpConnection.cpp
+++ b/Engine/Source/Runtime/Online/HTTPServer/Private/HttpConnection.cpp
@@ -27,6 +27,8 @@ FHttpConnection::FHttpConnection(FSocket* InSocket, TSharedPtr<FHttpRouter> InRo
 	,WriteContext(InSocket)
 {
 	check(nullptr != Socket);
+
+	Config = FHttpServerConfig::GetConnectionConfig();
 }
 
 FHttpConnection::~FHttpConnection()
@@ -98,7 +100,7 @@ void FHttpConnection::TransferState(EHttpConnectionState CurrentState, EHttpConn
 void FHttpConnection::BeginRead(float DeltaTime)
 {
 	QUICK_SCOPE_CYCLE_COUNTER(STAT_FHttpConnection_BeginRead);
-	const FTimespan WaitTime = FTimespan::FromMilliseconds(1);
+	const FTimespan WaitTime = FTimespan::FromMilliseconds(Config.BeginReadWaitTimeMS);
 	if (!Socket->Wait(ESocketWaitConditions::WaitForRead, WaitTime))
 	{
 		const float SecondsPollingSocketThisFrame = WaitTime.GetTotalSeconds();
diff --git a/Engine/Source/Runtime/Online/HTTPServer/Private/HttpConnection.h b/Engine/Source/Runtime/Online/HTTPServer/Private/HttpConnection.h
--- a/Engine/Source/Runtime/Online/HTTPServer/Private/HttpConnection.h
+++ b/Engine/Source/Runtime/Online/HTTPServer/Private/HttpConnection.h
@@ -6,6 +6,7 @@
 #include "HttpConnectionRequestReadContext.h"
 #include "HttpConnectionResponseWriteContext.h"
 #include "HttpResultCallback.h"
+#include "HttpServerConfig.h"
 #include "HttpServerConstants.h"
 #include "HttpServerHttpVersion.h"
 #include "Misc/Timespan.h"
@@ -252,5 +253,8 @@ private:
 
 	/** The duration (seconds) at which idle keep-alive connections are forcefully timed out */
 	static constexpr float ConnectionKeepAliveTimeout = 15.0f;
+
+	/** Connection configuration data */
+	FHttpServerConnectionConfig Config;
 };
 
diff --git a/Engine/Source/Runtime/Online/HTTPServer/Private/HttpServerConfig.cpp b/Engine/Source/Runtime/Online/HTTPServer/Private/HttpServerConfig.cpp
--- a/Engine/Source/Runtime/Online/HTTPServer/Private/HttpServerConfig.cpp
+++ b/Engine/Source/Runtime/Online/HTTPServer/Private/HttpServerConfig.cpp
@@ -54,3 +54,16 @@ const FHttpServerListenerConfig FHttpServerConfig::GetListenerConfig(uint32 Port
 	}
 	return Config;
 }
+
+const FHttpServerConnectionConfig FHttpServerConfig::GetConnectionConfig()
+{
+	static const FString IniSectionName(TEXT("HTTPServer.Connections"));
+
+	// Code default values
+	FHttpServerConnectionConfig Config;
+
+	// Apply default ini configuration
+	GConfig->GetFloat(*IniSectionName, TEXT("BeginReadWaitTimeMS"), Config.BeginReadWaitTimeMS, GEngineIni);
+
+	return Config;
+}
diff --git a/Engine/Source/Runtime/Online/HTTPServer/Private/HttpServerConfig.h b/Engine/Source/Runtime/Online/HTTPServer/Private/HttpServerConfig.h
--- a/Engine/Source/Runtime/Online/HTTPServer/Private/HttpServerConfig.h
+++ b/Engine/Source/Runtime/Online/HTTPServer/Private/HttpServerConfig.h
@@ -21,6 +21,12 @@ struct FHttpServerListenerConfig
 	bool bReuseAddressAndPort = false;
 };
 
+struct FHttpServerConnectionConfig
+{
+	/** Time in milliseconds to wait for data to be available when ticking the connection. */
+	float BeginReadWaitTimeMS = 1;
+};
+
 struct FHttpServerConfig
 {
 	/** 
@@ -30,4 +36,11 @@ struct FHttpServerConfig
 	* @return The per-port configuration if configured, or default configuration
 	*/
 	static const FHttpServerListenerConfig GetListenerConfig(uint32 Port);
+
+	/**
+	* Gets the connection configuration
+	*
+	* @return The configuration if configured, or default configuration
+	*/
+	static const FHttpServerConnectionConfig GetConnectionConfig();
 };
\ No newline at end of file
