Log warning when we hit the http request limit on servers.

diff --git a/Engine/Source/Runtime/Online/HTTP/Private/HttpThread.cpp b/Engine/Source/Runtime/Online/HTTP/Private/HttpThread.cpp
--- a/Engine/Source/Runtime/Online/HTTP/Private/HttpThread.cpp
+++ b/Engine/Source/Runtime/Online/HTTP/Private/HttpThread.cpp
@@ -276,6 +276,12 @@ void FHttpThread::Process(TArray<IHttpThreadedRequest*>& RequestsToCancel, TArra
 				RunningThreadedRequests.Add(ReadyThreadedRequest);
 				ReadyThreadedRequest->TickThreadedRequest(0.0f);
 				UE_LOG(LogHttp, Verbose, TEXT("Started running threaded request (%p). Running threaded requests (%d) Rate limited threaded requests (%d)"), ReadyThreadedRequest, RunningThreadedRequests.Num(), RateLimitedThreadedRequests.Num());
+#if WITH_SERVER_CODE
+				if (RunningThreadedRequestsCounter == RunningThreadedRequestLimit)
+				{
+					UE_LOG(LogHttp, Warning, TEXT("Reached threaded request limit (%d)"), RunningThreadedRequestsCounter);
+				}
+#endif // WITH_SERVER_CODE
 			}
 			else
 			{
