[Ensure] FHttpConnectionResponseWriteContext::SerializeHeadersUtf8

diff --git a/Engine/Source/Runtime/Online/HTTPServer/Private/HttpConnection.cpp b/Engine/Source/Runtime/Online/HTTPServer/Private/HttpConnection.cpp
--- a/Engine/Source/Runtime/Online/HTTPServer/Private/HttpConnection.cpp
+++ b/Engine/Source/Runtime/Online/HTTPServer/Private/HttpConnection.cpp
@@ -149,7 +149,7 @@ void FHttpConnection::ContinueRead(float DeltaTime)
 		break;
 
 	case EHttpConnectionContextState::Error:
-		HandleReadError(ReadContext.GetErrorCode(), *ReadContext.GetErrorStr());
+		HandleReadError(ReadContext.GetRequest()->HttpVersion, ReadContext.GetErrorCode(), *ReadContext.GetErrorStr());
 		break;
 	}
 }
@@ -311,13 +311,14 @@ void FHttpConnection::Destroy(EConnectionDestroyReason Reason)
 	}
 }
 
-void FHttpConnection::HandleReadError(EHttpServerResponseCodes ErrorCode, const TCHAR* ErrorCodeStr)
+void FHttpConnection::HandleReadError(HttpVersion::EHttpServerHttpVersion HttpVersion, EHttpServerResponseCodes ErrorCode, const TCHAR* ErrorCodeStr)
 {
 	UE_LOG(LogHttpConnection, Error, TEXT("%s"), ErrorCodeStr);
 
 	// Forcibly Reply
 	bKeepAlive = false;
 	auto Response = FHttpServerResponse::Error(ErrorCode, ErrorCodeStr);
+	Response->HttpVersion = HttpVersion;
 	BeginWrite(MoveTemp(Response), LastRequestNumber);
 }
 
diff --git a/Engine/Source/Runtime/Online/HTTPServer/Private/HttpConnection.h b/Engine/Source/Runtime/Online/HTTPServer/Private/HttpConnection.h
--- a/Engine/Source/Runtime/Online/HTTPServer/Private/HttpConnection.h
+++ b/Engine/Source/Runtime/Online/HTTPServer/Private/HttpConnection.h
@@ -164,10 +164,11 @@ private:
 	/**
 	 * Logs and responds with the caller-supplied error code
 	 *
+	 * @param HttpVersion The HTTP version to be used for the error response
 	 * @param ErrorCode The HTTP error code
 	 * @param ErrorCodeStr The machine-readable error description
 	 */
-	void HandleReadError(EHttpServerResponseCodes ErrorCode,  const TCHAR* ErrorCodeStr);
+	void HandleReadError(HttpVersion::EHttpServerHttpVersion HttpVersion, EHttpServerResponseCodes ErrorCode,  const TCHAR* ErrorCodeStr);
 
 	/**
 	 * Logs the caller-supplied error code and closes the connection
