- Delete #if 0 code

diff --git a/Engine/Source/Runtime/Online/XMPP/Private/XmppModule.cpp b/Engine/Source/Runtime/Online/XMPP/Private/XmppModule.cpp
--- a/Engine/Source/Runtime/Online/XMPP/Private/XmppModule.cpp
+++ b/Engine/Source/Runtime/Online/XMPP/Private/XmppModule.cpp
@@ -813,68 +813,3 @@ void FXmppModule::OnXmppRoomConfigured(const TSharedRef<IXmppConnection>& Connec
 
 	UE_LOG(LogXmpp, Log, TEXT("FXmppModule::OnXmppRoomConfigured - entered - user(%s) room(%s)"), *Connection->GetUserJid().Id, *RoomId);
 }
-
-// temp until WebRTC linked w/OpenSSL
-#if 0 // PLATFORM_WINDOWS || PLATFORM_LINUX || PLATFORM_MAC
-
-// newer OpenSSL version have these as #defines
-// WebRTC was linked with boringSSL which does not use the following functions as #defined
-// so un-doing the #defines and making them live functions to get linking errors fixed...
-
-#include "openssl/bio.h"
-
-#undef BIO_set_retry_read
-extern "C" void BIO_set_retry_read(BIO *b)
-{
-	BIO_set_flags(  b, (BIO_FLAGS_READ | BIO_FLAGS_SHOULD_RETRY));
-}
-
-#undef BIO_set_retry_write
-extern "C" void BIO_set_retry_write(BIO *b)
-{
-	BIO_set_flags(b, (BIO_FLAGS_WRITE | BIO_FLAGS_SHOULD_RETRY));
-}
-
-#undef BIO_clear_retry_flags
-extern "C" void BIO_clear_retry_flags(BIO *b)
-{
-	BIO_clear_flags(b, (BIO_FLAGS_RWS | BIO_FLAGS_SHOULD_RETRY));
-}
-
-#include "openssl/evp.h"
-
-#undef OpenSSL_add_all_algorithms
-extern "C" void OpenSSL_add_all_algorithms(void)
-{
-	OPENSSL_add_all_algorithms_noconf();
-	//	OPENSSL_add_all_algorithms_conf(void);
-}
-
-#define DWORD uint32 // FUnusableType DWORD
-#include "openssl/ssl.h"
-
-#undef DTLSv1_get_timeout
-extern "C" long DTLSv1_get_timeout(SSL *ssl, void *parg)
-{
-	return SSL_ctrl(ssl, DTLS_CTRL_GET_TIMEOUT, 0, (void *)parg);
-}
-
-#undef DTLSv1_handle_timeout
-extern "C" long DTLSv1_handle_timeout(SSL *ssl)
-{
-	return SSL_ctrl(ssl, DTLS_CTRL_HANDLE_TIMEOUT, 0, NULL);
-}
-
-#undef SSL_set_mode
-extern "C" long SSL_set_mode(SSL *ssl, long larg)
-{
-	return SSL_ctrl((ssl), SSL_CTRL_MODE, (larg), NULL);
-}
-
-#undef SSL_CTX_set_read_ahead
-extern "C" long SSL_CTX_set_read_ahead(SSL_CTX *ctx, long larg)
-{
-	return SSL_CTX_ctrl(ctx, SSL_CTRL_SET_READ_AHEAD, larg, NULL);
-}
-
-#endif
