- Manual merge of 11259588 to //UE5/Release-5.2
- Delete libjingle support from XMPP

diff --git a/Engine/Source/Runtime/Online/XMPP/Private/XmppJingle/XmppChatJingle.cpp b/Engine/Source/Runtime/Online/XMPP/Private/XmppJingle/XmppChatJingle.cpp
deleted file mode 100644
diff --git a/Engine/Source/Runtime/Online/XMPP/Private/XmppJingle/XmppChatJingle.h b/Engine/Source/Runtime/Online/XMPP/Private/XmppJingle/XmppChatJingle.h
deleted file mode 100644
diff --git a/Engine/Source/Runtime/Online/XMPP/Private/XmppJingle/XmppConnectionJingle.cpp b/Engine/Source/Runtime/Online/XMPP/Private/XmppJingle/XmppConnectionJingle.cpp
deleted file mode 100644
diff --git a/Engine/Source/Runtime/Online/XMPP/Private/XmppJingle/XmppConnectionJingle.h b/Engine/Source/Runtime/Online/XMPP/Private/XmppJingle/XmppConnectionJingle.h
deleted file mode 100644
diff --git a/Engine/Source/Runtime/Online/XMPP/Private/XmppJingle/XmppJingle.cpp b/Engine/Source/Runtime/Online/XMPP/Private/XmppJingle/XmppJingle.cpp
deleted file mode 100644
diff --git a/Engine/Source/Runtime/Online/XMPP/Private/XmppJingle/XmppJingle.h b/Engine/Source/Runtime/Online/XMPP/Private/XmppJingle/XmppJingle.h
deleted file mode 100644
diff --git a/Engine/Source/Runtime/Online/XMPP/Private/XmppJingle/XmppMessagesJingle.cpp b/Engine/Source/Runtime/Online/XMPP/Private/XmppJingle/XmppMessagesJingle.cpp
deleted file mode 100644
diff --git a/Engine/Source/Runtime/Online/XMPP/Private/XmppJingle/XmppMessagesJingle.h b/Engine/Source/Runtime/Online/XMPP/Private/XmppJingle/XmppMessagesJingle.h
deleted file mode 100644
diff --git a/Engine/Source/Runtime/Online/XMPP/Private/XmppJingle/XmppMultiUserChatJingle.cpp b/Engine/Source/Runtime/Online/XMPP/Private/XmppJingle/XmppMultiUserChatJingle.cpp
deleted file mode 100644
diff --git a/Engine/Source/Runtime/Online/XMPP/Private/XmppJingle/XmppMultiUserChatJingle.h b/Engine/Source/Runtime/Online/XMPP/Private/XmppJingle/XmppMultiUserChatJingle.h
deleted file mode 100644
diff --git a/Engine/Source/Runtime/Online/XMPP/Private/XmppJingle/XmppPresenceJingle.cpp b/Engine/Source/Runtime/Online/XMPP/Private/XmppJingle/XmppPresenceJingle.cpp
deleted file mode 100644
diff --git a/Engine/Source/Runtime/Online/XMPP/Private/XmppJingle/XmppPresenceJingle.h b/Engine/Source/Runtime/Online/XMPP/Private/XmppJingle/XmppPresenceJingle.h
deleted file mode 100644
diff --git a/Engine/Source/Runtime/Online/XMPP/Private/XmppJingle/XmppPubSubJingle.cpp b/Engine/Source/Runtime/Online/XMPP/Private/XmppJingle/XmppPubSubJingle.cpp
deleted file mode 100644
diff --git a/Engine/Source/Runtime/Online/XMPP/Private/XmppJingle/XmppPubSubJingle.h b/Engine/Source/Runtime/Online/XMPP/Private/XmppJingle/XmppPubSubJingle.h
deleted file mode 100644
diff --git a/Engine/Source/Runtime/Online/XMPP/Private/XmppModule.cpp b/Engine/Source/Runtime/Online/XMPP/Private/XmppModule.cpp
--- a/Engine/Source/Runtime/Online/XMPP/Private/XmppModule.cpp
+++ b/Engine/Source/Runtime/Online/XMPP/Private/XmppModule.cpp
@@ -10,9 +10,6 @@
 #include "XmppLog.h"
 #include "XmppTests.h"
 #include "XmppNull.h"
-#if WITH_XMPP_JINGLE
-#include "XmppJingle/XmppJingle.h"
-#endif
 #if WITH_XMPP_STROPHE
 #include "XmppStrophe/XmppStrophe.h"
 #include "WebSocketsModule.h"
@@ -35,9 +32,6 @@ void FXmppModule::StartupModule()
 
 	if (bEnabled)
 	{
-#if WITH_XMPP_JINGLE
-		FXmppJingle::Init();
-#endif
 #if WITH_XMPP_STROPHE
 		FXmppStrophe::Init();
 		FModuleManager::LoadModuleChecked<FWebSocketsModule>("WebSockets");
@@ -54,9 +48,6 @@ void FXmppModule::ShutdownModule()
 
 	if (bEnabled)
 	{
-#if WITH_XMPP_JINGLE
-		FXmppJingle::Cleanup();
-#endif
 #if WITH_XMPP_STROPHE
 		FXmppStrophe::Cleanup();
 #endif
@@ -698,37 +689,13 @@ TSharedRef<IXmppConnection> FXmppModule::CreateConnection(const FString& UserId)
 	}
 	else
 	{
-		bool bEnableWebsockets = false;
-		GConfig->GetBool(TEXT("XMPP"), TEXT("bEnableWebsockets"), bEnableWebsockets, GEngineIni);
-
-		bool bUseStrophe = WITH_XMPP_STROPHE && bEnableWebsockets;
-		bool bUseJingle = WITH_XMPP_JINGLE && !bUseStrophe;
-		if (!bUseJingle && !bUseStrophe)
-		{
-			// if not using websockets, use the previous default implementation (jingle if available, otherwise strophe)
-#if WITH_XMPP_JINGLE
-			bUseJingle = true;
-#elif WITH_XMPP_STROPHE
-			bUseStrophe = true;
-#endif
-		}
-
-		if (bEnabled && (bUseStrophe || bUseJingle))
-		{
 #if WITH_XMPP_STROPHE
-			if (bUseStrophe)
-			{
-				Connection = FXmppStrophe::CreateConnection();
-			}
-#endif
-#if WITH_XMPP_JINGLE
-			if (bUseJingle)
-			{
-				Connection = FXmppJingle::CreateConnection();
-			}
-#endif
+		if (bEnabled)
+		{
+			Connection = FXmppStrophe::CreateConnection();
 		}
 		else
+#endif
 		{
 			Connection = FXmppNull::CreateConnection();
 		}
diff --git a/Engine/Source/Runtime/Online/XMPP/XMPP.Build.cs b/Engine/Source/Runtime/Online/XMPP/XMPP.Build.cs
--- a/Engine/Source/Runtime/Online/XMPP/XMPP.Build.cs
+++ b/Engine/Source/Runtime/Online/XMPP/XMPP.Build.cs
@@ -4,8 +4,6 @@ using UnrealBuildTool;
 
 public class XMPP : ModuleRules
 {
-	protected virtual bool bTargetPlatformSupportsJingle { get { return false; } }
-
 	protected virtual bool bTargetPlatformSupportsStrophe
 	{
 		get =>
@@ -30,20 +28,7 @@ public class XMPP : ModuleRules
 			}
 		);
 
-		bool TargetPlatformSupportsJingle = bTargetPlatformSupportsJingle;
-		bool TargetPlatformSupportsStrophe = bTargetPlatformSupportsStrophe;
-
-		if (TargetPlatformSupportsJingle)
-		{
-			AddEngineThirdPartyPrivateStaticDependencies(Target, "WebRTC");
-			PrivateDefinitions.Add("WITH_XMPP_JINGLE=1");
-		}
-		else
-		{
-			PrivateDefinitions.Add("WITH_XMPP_JINGLE=0");
-		}
-
-		if (TargetPlatformSupportsStrophe)
+		if (bTargetPlatformSupportsStrophe)
 		{
 			AddEngineThirdPartyPrivateStaticDependencies(Target, "libstrophe");
 			PrivateDependencyModuleNames.Add("WebSockets");
