make VFS stats available to perf reports in shipping

diff --git a/Engine/Source/Runtime/Online/BuildPatchServices/Private/Installer/BuildInstallStreamer.cpp b/Engine/Source/Runtime/Online/BuildPatchServices/Private/Installer/BuildInstallStreamer.cpp
--- a/Engine/Source/Runtime/Online/BuildPatchServices/Private/Installer/BuildInstallStreamer.cpp
+++ b/Engine/Source/Runtime/Online/BuildPatchServices/Private/Installer/BuildInstallStreamer.cpp
@@ -223,13 +223,11 @@ namespace BuildPatchServices
 				AsyncHelpers::ExecuteOnCustomThread<void>([this]() { DownloadServiceStatistics->Reset(); }, TickQueue);
 				TSharedPtr<IVirtualFileCache, ESPMode::ThreadSafe> VirtualFileCache = IVirtualFileCache::CreateVirtualFileCache();
 				InstallerError->Reset();
-#if CSV_PROFILER
+
 				//Used and Total size is calulcated on VFC Start.
 				//Used + Requested write will always result in new total. Doing this to recalculating all blocks.
 				BuildStats.VFCCachedUsedSize = (VirtualFileCache->GetUsedSize() / 1024.0f / 1024.0f);
 				BuildStats.VFCCachedTotalSize = (VirtualFileCache->GetTotalSize()  / 1024.0f / 1024.0f);
-#endif
-
 
 				// Setup file build list.
 				FVirtualFileConstructorConfiguration VFCConfig;
@@ -373,7 +371,6 @@ namespace BuildPatchServices
 			TickFunc();
 		}
 
-#if !UE_BUILD_SHIPPING
 		CSV_CUSTOM_STAT(CosmeticStreamingCsv, BuildStats.VFCRequestedFileWrite,		BuildStats.VFCRequestedFileWrite, ECsvCustomStatOp::Set);
 		CSV_CUSTOM_STAT(CosmeticStreamingCsv, BuildStats.VFCCachedUsedSize,			BuildStats.VFCCachedUsedSize + BuildStats.VFCRequestedFileWrite, ECsvCustomStatOp::Set);
 		CSV_CUSTOM_STAT(CosmeticStreamingCsv, BuildStats.VFCCachedTotalSize,		BuildStats.VFCCachedTotalSize, ECsvCustomStatOp::Set);
@@ -381,7 +378,6 @@ namespace BuildPatchServices
 		CSV_CUSTOM_STAT(CosmeticStreamingCsv, BuildStats.AmountRequestsCompleted,	BuildStats.AmountRequestsCompleted, ECsvCustomStatOp::Set);
 		CSV_CUSTOM_STAT(CosmeticStreamingCsv, BuildStats.AmountRequestsCancelled,	BuildStats.AmountRequestsCancelled, ECsvCustomStatOp::Set);
 		CSV_CUSTOM_STAT(CosmeticStreamingCsv, BuildStats.AmountRequestsMade,		BuildStats.AmountRequestsMade, ECsvCustomStatOp::Set);
-#endif
 
 		return bKeepTicking;
 	}
