Fix the test "System.Online.HttpServer.Integration" by setting the expected error messages.

diff --git a/Engine/Source/Runtime/Online/HTTPServer/Private/Tests/HttpServerTests.cpp b/Engine/Source/Runtime/Online/HTTPServer/Private/Tests/HttpServerTests.cpp
--- a/Engine/Source/Runtime/Online/HTTPServer/Private/Tests/HttpServerTests.cpp
+++ b/Engine/Source/Runtime/Online/HTTPServer/Private/Tests/HttpServerTests.cpp
@@ -29,6 +29,9 @@ bool FHttpServerIntegrationTest::RunTest(const FString& Parameters)
 	TSharedPtr<IHttpRouter> DuplicateHttpRouter = FHttpServerModule::Get().GetHttpRouter(HttpRouterPort);
 	TestEqual(TEXT("HttpRouter Duplicates"), HttpRouter, DuplicateHttpRouter);
 
+	// Ensure failed port binds still return a valid router if not explicitly requested to fail (and by default)
+	TSharedPtr<IHttpRouter> ValidHttpRouterOnFail = FHttpServerModule::Get().GetHttpRouter(InvalidHttpRouterPort /*, bFailOnBindFailure = false */);
+	TestTrue(TEXT("HttpRouter is NOT null on bind failure by default"), ValidHttpRouterOnFail.IsValid());
 
 	// Ensure we can create route bindings
 	const FHttpRequestHandler RequestHandler = [this]
@@ -43,27 +46,24 @@ bool FHttpServerIntegrationTest::RunTest(const FString& Parameters)
 	FHttpRouteHandle DuplicateHandle = HttpRouter->BindRoute(HttpPath, EHttpServerRequestVerbs::VERB_GET, RequestHandler);
 	TestFalse(TEXT("HttpRouteHandle Duplicated"), DuplicateHandle.IsValid());
 
+	// Because of the ValidHttpRouterOnFail was created by FHttpServerModule::Get().GetHttpRouter(InvalidHttpRouterPort...), it will fail to listen in StartAllListeners
+	// Also after bHttpListenersEnabled got set to true by StartAllListeners, when call GetHttpRouter(InvalidHttpRouterPort...) again, it will call StartListening again in there and fail
+	AddExpectedError(TEXT("HttpListener detected invalid port"), EAutomationExpectedErrorFlags::Contains, 2);
 	FHttpServerModule::Get().StartAllListeners();
 
-	// We can't have this test otherwise it'll output the error LogHttpListener: Error: HttpListener detected invalid port 65536 and the test will failed
+	// Because of the ValidHttpRouterOnFail was created by FHttpServerModule::Get().GetHttpRouter(InvalidHttpRouterPort...)
+	AddExpectedError(TEXT("is not listening/bound and listeners are still enabled"), EAutomationExpectedErrorFlags::Contains, 1);
 	// Ensure failed port binds result in a null router instance if requested (and listeners are enabled)
-	// TSharedPtr<IHttpRouter> InvalidHttpRouterOnFail = FHttpServerModule::Get().GetHttpRouter(InvalidHttpRouterPort, /* bFailOnBindFailure = */ true);
-	// TestFalse(TEXT("HttpRouter is null on bind failure if requested"), InvalidHttpRouterOnFail.IsValid());
-	
-	 // Make a request
+	TSharedPtr<IHttpRouter> InvalidHttpRouterOnFail = FHttpServerModule::Get().GetHttpRouter(InvalidHttpRouterPort, /* bFailOnBindFailure = */ true);
+	TestFalse(TEXT("HttpRouter is null on bind failure if requested"), InvalidHttpRouterOnFail.IsValid());
+
+	// Make a request
 	/*
 	ADD_LATENT_AUTOMATION_COMMAND(FExecStringLatentCommand(TEXT("HTTP TEST 1 http://localhost:8888/TestHttpServer")));
 	*/
 
 	FHttpServerModule::Get().StopAllListeners();
 
-	// Ensure failed port binds still return a valid router if not explicitly requested to fail (and by default)
-	// Note that if there is already a router with the same port it will output an error and the test will fail
-	TSharedPtr<IHttpRouter> ValidHttpRouterOnFail = FHttpServerModule::Get().GetHttpRouter(InvalidHttpRouterPort /*, bFailOnBindFailure = false */);
-	TestTrue(TEXT("HttpRouter is NOT null on bind failure by default"), ValidHttpRouterOnFail.IsValid());
-
-
-
 	HttpRouter->UnbindRoute(HttpRouteHandle);
 	return true;
 }
