Adding additional analytics to the streaming backend so that we can better monitor the performance of the VFC and cosmetic streaming.

diff --git a/Engine/Source/Runtime/Online/BuildPatchServices/Private/Installer/BuildInstallStreamer.cpp b/Engine/Source/Runtime/Online/BuildPatchServices/Private/Installer/BuildInstallStreamer.cpp
--- a/Engine/Source/Runtime/Online/BuildPatchServices/Private/Installer/BuildInstallStreamer.cpp
+++ b/Engine/Source/Runtime/Online/BuildPatchServices/Private/Installer/BuildInstallStreamer.cpp
@@ -132,6 +132,11 @@ namespace BuildPatchServices
 		return Configuration;
 	}
 
+	const FBuildInstallStreamerStats& FBuildInstallStreamer::GetInstallStreamerStatistics() const
+	{
+		return BuildStats;
+	}
+
 	void FBuildInstallStreamer::Initialise()
 	{
 		// Queue update to chunk data size cache on main thread
@@ -153,8 +158,10 @@ namespace BuildPatchServices
 		auto CompleteRequest = [this, &Requests, &RequestTime]()
 		{
 			uint64 TotalDownload = DownloadServiceStatistics->GetBytesDownloaded();
+			uint64 PreviousTotalRequestsCompleted = BuildStats.BundleRequestsCompleted + BuildStats.FileRequestsCompleted;
 			for (FStreamRequest& Request : Requests)
 			{
+				const bool bIsTagRequest = Request.Get<1>();
 				FBuildPatchStreamResult Result;
 				Result.Request = MoveTemp(Request.Get<0>());
 				Result.ErrorType = InstallerError->GetErrorType();
@@ -162,23 +169,35 @@ namespace BuildPatchServices
 				// TODO: Not accurate.. How to fix this
 				Result.TotalDownloaded = TotalDownload;
 				
-				BuildStats.AmountRequestsCompleted++;
+				// if (bIsTagRequest) BuildStats.BundleMegaBytesDownloaded += (Result.TotalDownloaded / 1024.0f / 1024.0f);
+				// else BuildStats.FileMegaBytesDownloaded += (Result.TotalDownloaded / 1024.0f / 1024.0f);
+				
+				if (bIsTagRequest) BuildStats.BundleRequestsCompleted++;
+				else BuildStats.FileRequestsCompleted++;
+
 				UE_LOG(LogBuildInstallStreamer, Verbose, TEXT("Completed request %s in %s"), *FString::Join(Result.Request, TEXT(",")), *FPlatformTime::PrettyTime(RequestTime));
 				Request.Get<2>().ExecuteIfBound(MoveTemp(Result));
 			}
-			BuildStats.MegaBytesDownloaded += (TotalDownload / 1024.0f / 1024.0f);
+			BuildStats.TotalMegaBytesDownloaded += (TotalDownload / 1024.0f / 1024.0f);
+				
+			if (RequestTime > BuildStats.MaxRequestTime) BuildStats.MaxRequestTime = RequestTime;
+				
+			BuildStats.AverageRequestTime = ((BuildStats.AverageRequestTime * PreviousTotalRequestsCompleted) + RequestTime) /
+				(BuildStats.BundleRequestsCompleted + BuildStats.FileRequestsCompleted);
 		};
 
 		auto CancelRequest = [this, &Requests]()
 		{
 			for (FStreamRequest& Request : Requests)
 			{
+				const bool bIsTagRequest = Request.Get<1>();
 				FBuildPatchStreamResult Result;
 				Result.Request = MoveTemp(Request.Get<0>());
 				Result.ErrorType = InstallerError->GetErrorType();
 				Result.ErrorCode = InstallerError->GetErrorCode();
 				Result.TotalDownloaded = 0;
-				BuildStats.AmountRequestsCancelled++;
+				if (bIsTagRequest) BuildStats.BundleRequestsCancelled++;
+				else BuildStats.FileRequestsCancelled++;
 				UE_LOG(LogBuildInstallStreamer, Verbose, TEXT("Cancelled request %s"), *FString::Join(Result.Request, TEXT(",")));
 				Request.Get<2>().ExecuteIfBound(MoveTemp(Result));
 			}
@@ -214,7 +233,8 @@ namespace BuildPatchServices
 			{
 				for (const FStreamRequest& Request : Requests)
 				{
-					BuildStats.AmountRequestsMade++;
+					if (Request.Get<1>()) BuildStats.BundleRequestsMade++;
+					else BuildStats.FileRequestsMade++;
 					UE_LOG(LogBuildInstallStreamer, Verbose, TEXT("Begin request %s"), *FString::Join(Request.Get<0>(), TEXT(",")));
 				}
 
@@ -372,12 +392,23 @@ namespace BuildPatchServices
 		}
 
 		CSV_CUSTOM_STAT(CosmeticStreamingCsv, BuildStats.VFCRequestedFileWrite,		BuildStats.VFCRequestedFileWrite, ECsvCustomStatOp::Set);
-		CSV_CUSTOM_STAT(CosmeticStreamingCsv, BuildStats.VFCCachedUsedSize,			BuildStats.VFCCachedUsedSize + BuildStats.VFCRequestedFileWrite, ECsvCustomStatOp::Set);
+		//CSV_CUSTOM_STAT(CosmeticStreamingCsv, BuildStats.VFCActualFileWrite,		BuildStats.VFCActualFileWrite, ECsvCustomStatOp::Set);
+		CSV_CUSTOM_STAT(CosmeticStreamingCsv, BuildStats.VFCCachedUsedSize,			BuildStats.VFCCachedUsedSize, ECsvCustomStatOp::Set);
 		CSV_CUSTOM_STAT(CosmeticStreamingCsv, BuildStats.VFCCachedTotalSize,		BuildStats.VFCCachedTotalSize, ECsvCustomStatOp::Set);
-		CSV_CUSTOM_STAT(CosmeticStreamingCsv, BuildStats.MegaBytesDownloaded,		BuildStats.MegaBytesDownloaded, ECsvCustomStatOp::Set);
-		CSV_CUSTOM_STAT(CosmeticStreamingCsv, BuildStats.AmountRequestsCompleted,	BuildStats.AmountRequestsCompleted, ECsvCustomStatOp::Set);
-		CSV_CUSTOM_STAT(CosmeticStreamingCsv, BuildStats.AmountRequestsCancelled,	BuildStats.AmountRequestsCancelled, ECsvCustomStatOp::Set);
-		CSV_CUSTOM_STAT(CosmeticStreamingCsv, BuildStats.AmountRequestsMade,		BuildStats.AmountRequestsMade, ECsvCustomStatOp::Set);
+		CSV_CUSTOM_STAT(CosmeticStreamingCsv, BuildStats.TotalMegaBytesDownloaded,	BuildStats.TotalMegaBytesDownloaded, ECsvCustomStatOp::Set);
+		
+		//CSV_CUSTOM_STAT(CosmeticStreamingCsv, BuildStats.FileMegaBytesDownloaded,	BuildStats.FileMegaBytesDownloaded, ECsvCustomStatOp::Set);
+		CSV_CUSTOM_STAT(CosmeticStreamingCsv, BuildStats.FileRequestsCompleted,	    (double)BuildStats.FileRequestsCompleted, ECsvCustomStatOp::Set);
+		CSV_CUSTOM_STAT(CosmeticStreamingCsv, BuildStats.FileRequestsCancelled,	    (double)BuildStats.FileRequestsCancelled, ECsvCustomStatOp::Set);
+		CSV_CUSTOM_STAT(CosmeticStreamingCsv, BuildStats.FileRequestsMade,		    (double)BuildStats.FileRequestsMade, ECsvCustomStatOp::Set);
+		
+		//CSV_CUSTOM_STAT(CosmeticStreamingCsv, BuildStats.BundleMegaBytesDownloaded,	BuildStats.BundleMegaBytesDownloaded, ECsvCustomStatOp::Set);
+		CSV_CUSTOM_STAT(CosmeticStreamingCsv, BuildStats.BundleRequestsCompleted,	(double)BuildStats.BundleRequestsCompleted, ECsvCustomStatOp::Set);
+		CSV_CUSTOM_STAT(CosmeticStreamingCsv, BuildStats.BundleRequestsCancelled,	(double)BuildStats.BundleRequestsCancelled, ECsvCustomStatOp::Set);
+		CSV_CUSTOM_STAT(CosmeticStreamingCsv, BuildStats.BundleRequestsMade,		(double)BuildStats.BundleRequestsMade, ECsvCustomStatOp::Set);
+		
+		CSV_CUSTOM_STAT(CosmeticStreamingCsv, BuildStats.AverageRequestTime,	BuildStats.AverageRequestTime, ECsvCustomStatOp::Set);
+		CSV_CUSTOM_STAT(CosmeticStreamingCsv, BuildStats.MaxRequestTime,		BuildStats.MaxRequestTime, ECsvCustomStatOp::Set);
 
 		return bKeepTicking;
 	}
diff --git a/Engine/Source/Runtime/Online/BuildPatchServices/Private/Installer/BuildInstallStreamer.h b/Engine/Source/Runtime/Online/BuildPatchServices/Private/Installer/BuildInstallStreamer.h
--- a/Engine/Source/Runtime/Online/BuildPatchServices/Private/Installer/BuildInstallStreamer.h
+++ b/Engine/Source/Runtime/Online/BuildPatchServices/Private/Installer/BuildInstallStreamer.h
@@ -34,17 +34,6 @@ namespace BuildPatchServices
 	class IDownloadService;
 	class IMessagePump;
 
-	struct FBuildInstallStreamerStats
-	{
-		float MegaBytesDownloaded;
-		int32 AmountRequestsCompleted;
-		int32 AmountRequestsCancelled;
-		int32 AmountRequestsMade;
-		float VFCCachedTotalSize;
-		float VFCCachedUsedSize;
-		float VFCRequestedFileWrite;
-	};
-
 	class FBuildInstallStreamer
 		: public IBuildInstallStreamer
 		, public TSharedFromThis<FBuildInstallStreamer, ESPMode::ThreadSafe>
@@ -60,6 +49,7 @@ namespace BuildPatchServices
 		virtual void RegisterMessageHandler(FMessageHandler* MessageHandler) override;
 		virtual void UnregisterMessageHandler(FMessageHandler* MessageHandler) override;
 		virtual const FBuildInstallStreamerConfiguration& GetConfiguration() const override;
+		virtual const FBuildInstallStreamerStats& GetInstallStreamerStatistics() const override;
 		// IBuildInstallStreamer interface end.
 
 		bool Tick();
diff --git a/Engine/Source/Runtime/Online/BuildPatchServices/Public/BuildPatchSettings.h b/Engine/Source/Runtime/Online/BuildPatchServices/Public/BuildPatchSettings.h
--- a/Engine/Source/Runtime/Online/BuildPatchServices/Public/BuildPatchSettings.h
+++ b/Engine/Source/Runtime/Online/BuildPatchServices/Public/BuildPatchSettings.h
@@ -285,6 +285,31 @@ namespace BuildPatchServices
 		bool bMainThreadDelegates = true;
 	};
 
+	/**
+	 * The collection of statistics gathered by the build streamer class.
+	 */
+	struct BUILDPATCHSERVICES_API FBuildInstallStreamerStats
+	{
+		//float FileMegaBytesDownloaded;
+		uint64 FileRequestsCompleted;
+		uint64 FileRequestsMade;
+		uint64 FileRequestsCancelled;
+		
+		//float BundleMegaBytesDownloaded;
+		uint64 BundleRequestsCancelled;
+		uint64 BundleRequestsCompleted;
+		uint64 BundleRequestsMade;
+		
+		float TotalMegaBytesDownloaded;
+		double MaxRequestTime;
+		double AverageRequestTime;
+		
+		float VFCCachedTotalSize;
+		float VFCCachedUsedSize;
+		float VFCRequestedFileWrite;
+		float VFCActualFileWrite;
+	};
+
 	/**
 	 * Defines a list of all options for the build chunking task.
 	 */
diff --git a/Engine/Source/Runtime/Online/BuildPatchServices/Public/Interfaces/IBuildInstallStreamer.h b/Engine/Source/Runtime/Online/BuildPatchServices/Public/Interfaces/IBuildInstallStreamer.h
--- a/Engine/Source/Runtime/Online/BuildPatchServices/Public/Interfaces/IBuildInstallStreamer.h
+++ b/Engine/Source/Runtime/Online/BuildPatchServices/Public/Interfaces/IBuildInstallStreamer.h
@@ -9,6 +9,7 @@
 namespace BuildPatchServices
 {
 	struct FBuildInstallStreamerConfiguration;
+	struct FBuildInstallStreamerStats;
 }
 
 typedef TSharedPtr<class IBuildInstallStreamer, ESPMode::ThreadSafe> IBuildInstallStreamerPtr;
@@ -56,4 +57,11 @@ public:
 	 * @returns a const reference to the configuration.
 	 */
 	virtual const BuildPatchServices::FBuildInstallStreamerConfiguration& GetConfiguration() const = 0;
+
+	/**
+	 * Get the stats for this install streamer.
+	 * 
+	 * @returns a const reference to the statistics.
+	 */
+	virtual const BuildPatchServices::FBuildInstallStreamerStats& GetInstallStreamerStatistics() const = 0;
 };
\ No newline at end of file
diff --git a/Engine/Source/Runtime/VirtualFileCache/Private/VirtualFileCache.cpp b/Engine/Source/Runtime/VirtualFileCache/Private/VirtualFileCache.cpp
--- a/Engine/Source/Runtime/VirtualFileCache/Private/VirtualFileCache.cpp
+++ b/Engine/Source/Runtime/VirtualFileCache/Private/VirtualFileCache.cpp
@@ -116,3 +116,13 @@ int64 FVirtualFileCache::GetUsedSize() const
 	FFileTableReader FileTable = Thread.ReadFileTable();
 	return FileTable->GetUsedSize();
 }
+
+uint64 FVirtualFileCache::GetTotalMemCacheHits() const
+{
+	return Thread.GetTotalMemCacheHits();
+}
+
+uint64 FVirtualFileCache::GetTotalMemCacheMisses() const
+{
+	return Thread.GetTotalMemCacheMisses();
+}
diff --git a/Engine/Source/Runtime/VirtualFileCache/Private/VirtualFileCacheInternal.cpp b/Engine/Source/Runtime/VirtualFileCache/Private/VirtualFileCacheInternal.cpp
--- a/Engine/Source/Runtime/VirtualFileCache/Private/VirtualFileCacheInternal.cpp
+++ b/Engine/Source/Runtime/VirtualFileCache/Private/VirtualFileCacheInternal.cpp
@@ -49,10 +49,12 @@ void FVirtualFileCacheThread::DoOneOp(FRWOp* Op)
 		TSharedPtr<TArray<uint8>> CachedData = MemCache.ReadLockAndFindData(Op->Target);
 		if (CachedData.IsValid())
 		{
+			TotalMemCacheHits++;
 			Op->ReadResult->SetValue(*CachedData);
 		}
 		else
 		{
+			TotalMemCacheMisses++;
 			FFileTableMutator FileTable = MutateFileTable();
 			auto Result = FileTable->ReadData(Op->Target, Op->ReadOffset, Op->ReadSize);
 			if (Result.IsOk())
@@ -1334,6 +1336,16 @@ void FVirtualFileCacheThread::SetInMemoryCacheSize(int64 MaxSize)
 	MemCache.SetMaxSize(MaxSize);
 }
 
+uint64 FVirtualFileCacheThread::GetTotalMemCacheHits() const
+{
+	return TotalMemCacheHits;
+}
+
+uint64 FVirtualFileCacheThread::GetTotalMemCacheMisses() const
+{
+	return TotalMemCacheMisses;
+}
+
 FLruCacheNode* FLruCache::Find(VFCKey Key)
 {
 	TUniquePtr<FLruCacheNode>* NodePtr = NodeMap.Find(Key);
diff --git a/Engine/Source/Runtime/VirtualFileCache/Private/VirtualFileCacheInternal.h b/Engine/Source/Runtime/VirtualFileCache/Private/VirtualFileCacheInternal.h
--- a/Engine/Source/Runtime/VirtualFileCache/Private/VirtualFileCacheInternal.h
+++ b/Engine/Source/Runtime/VirtualFileCache/Private/VirtualFileCacheInternal.h
@@ -374,6 +374,9 @@ public:
 
 	void SetInMemoryCacheSize(int64 MaxSize);
 
+	uint64 GetTotalMemCacheHits() const;
+	uint64 GetTotalMemCacheMisses() const;
+
 public:
 	FVirtualFileCache* Parent;
 	FRunnableThread* Thread;
@@ -391,6 +394,8 @@ public:
 private:
 	mutable FRWLock FileTableLock;
 	FFileTable FileTableStorage;
+	uint64 TotalMemCacheHits = 0;
+	uint64 TotalMemCacheMisses = 0;
 };
 
 struct FVirtualFileCache final : IVirtualFileCache
@@ -418,6 +423,8 @@ public:
 	virtual void Defragment() override;
 	virtual int64 GetTotalSize() const override;
 	virtual int64 GetUsedSize() const override;
+	virtual uint64 GetTotalMemCacheHits() const override;
+	virtual uint64 GetTotalMemCacheMisses() const override;
 
 	FString GetTableFilename() const { return Thread.GetTableFilename(); }
 
diff --git a/Engine/Source/Runtime/VirtualFileCache/Public/VirtualFileCache.h b/Engine/Source/Runtime/VirtualFileCache/Public/VirtualFileCache.h
--- a/Engine/Source/Runtime/VirtualFileCache/Public/VirtualFileCache.h
+++ b/Engine/Source/Runtime/VirtualFileCache/Public/VirtualFileCache.h
@@ -38,6 +38,8 @@ struct IVirtualFileCache
 	virtual void Defragment() = 0;
 	virtual int64 GetTotalSize() const = 0;
 	virtual int64 GetUsedSize() const = 0;
+	virtual uint64 GetTotalMemCacheHits() const = 0;
+	virtual uint64 GetTotalMemCacheMisses() const = 0;
 
 	static VIRTUALFILECACHE_API TSharedRef<IVirtualFileCache, ESPMode::ThreadSafe> CreateVirtualFileCache();
 };
