Added LLM tags to CurlHttp allocation hooks.

  - Because most of the allocations for Curl are on an OS created thread, there is no top level tag.
  - Added a named tag of "Networking/Curl"

#tests Built all Switch client targets. Tested a Replay with the code on/off and saw a difference in the Untagged in the reports. Ran Horde Minimal Code Only Build.

diff --git a/Engine/Source/Runtime/Online/HTTP/Private/Curl/CurlHttp.h b/Engine/Source/Runtime/Online/HTTP/Private/Curl/CurlHttp.h
--- a/Engine/Source/Runtime/Online/HTTP/Private/Curl/CurlHttp.h
+++ b/Engine/Source/Runtime/Online/HTTP/Private/Curl/CurlHttp.h
@@ -3,6 +3,7 @@
 #pragma once
 
 #include "CoreMinimal.h"
+#include "HAL/LowLevelMemTracker.h"
 #include "HAL/ThreadSafeCounter.h"
 #include "Interfaces/IHttpResponse.h"
 #include "IHttpThreadedRequest.h"
@@ -54,6 +55,7 @@ namespace
 	*/
 	void* CurlMalloc(size_t Size)
 	{
+		LLM_SCOPE_BYNAME(TEXT("Networking/Curl"));
 		return FMemory::Malloc(Size);
 	}
 
@@ -64,6 +66,7 @@ namespace
 	*/
 	void CurlFree(void* Ptr)
 	{
+		LLM_SCOPE_BYNAME(TEXT("Networking/Curl"));
 		FMemory::Free(Ptr);
 	}
 
@@ -80,6 +83,7 @@ namespace
 
 		if (Size)
 		{
+			LLM_SCOPE_BYNAME(TEXT("Networking/Curl"));
 			Return = FMemory::Realloc(Ptr, Size);
 		}
 
@@ -98,6 +102,8 @@ namespace
 		check(ZeroTerminatedString);
 		if (ZeroTerminatedString)
 		{
+			LLM_SCOPE_BYNAME(TEXT("Networking/Curl"));
+
 			SIZE_T StrLen = FCStringAnsi::Strlen(ZeroTerminatedString);
 			Copy = reinterpret_cast<char*>(FMemory::Malloc(StrLen + 1));
 			if (Copy)
@@ -122,6 +128,8 @@ namespace
 		const size_t Size = NumElems * ElemSize;
 		if (Size)
 		{
+			LLM_SCOPE_BYNAME(TEXT("Networking/Curl"));
+
 			Return = FMemory::Malloc(Size);
 
 			if (Return)
