[Ensure] FHttpConnectionResponseWriteContext::SerializeHeadersUtf8 correction after CL 25332384.

diff --git a/Engine/Source/Runtime/Online/HTTPServer/Private/HttpConnection.cpp b/Engine/Source/Runtime/Online/HTTPServer/Private/HttpConnection.cpp
--- a/Engine/Source/Runtime/Online/HTTPServer/Private/HttpConnection.cpp
+++ b/Engine/Source/Runtime/Online/HTTPServer/Private/HttpConnection.cpp
@@ -149,7 +149,7 @@ void FHttpConnection::ContinueRead(float DeltaTime)
 		break;
 
 	case EHttpConnectionContextState::Error:
-		HandleReadError(ReadContext.GetRequest()->HttpVersion, ReadContext.GetErrorCode(), *ReadContext.GetErrorStr());
+		HandleReadError(ReadContext.GetErrorCode(), *ReadContext.GetErrorStr());
 		break;
 	}
 }
@@ -311,14 +311,24 @@ void FHttpConnection::Destroy(EConnectionDestroyReason Reason)
 	}
 }
 
-void FHttpConnection::HandleReadError(HttpVersion::EHttpServerHttpVersion HttpVersion, EHttpServerResponseCodes ErrorCode, const TCHAR* ErrorCodeStr)
+void FHttpConnection::HandleReadError(EHttpServerResponseCodes ErrorCode, const TCHAR* ErrorCodeStr)
 {
 	UE_LOG(LogHttpConnection, Error, TEXT("%s"), ErrorCodeStr);
 
 	// Forcibly Reply
 	bKeepAlive = false;
-	auto Response = FHttpServerResponse::Error(ErrorCode, ErrorCodeStr);
-	Response->HttpVersion = HttpVersion;
+	TUniquePtr<FHttpServerResponse> Response = FHttpServerResponse::Error(ErrorCode, ErrorCodeStr);
+	TSharedPtr<FHttpServerRequest> Request = ReadContext.GetRequest();
+	if (Request.IsValid())
+	{
+		Response->HttpVersion = Request->HttpVersion;
+	}
+	else
+	{
+		UE_LOG(LogHttpConnection, Verbose, TEXT("Http Request is null, unable to parse version."));
+		Response->HttpVersion = HttpVersion::EHttpServerHttpVersion::HTTP_VERSION_UNKNOWN;
+	}
+
 	BeginWrite(MoveTemp(Response), LastRequestNumber);
 }
 
diff --git a/Engine/Source/Runtime/Online/HTTPServer/Private/HttpConnection.h b/Engine/Source/Runtime/Online/HTTPServer/Private/HttpConnection.h
--- a/Engine/Source/Runtime/Online/HTTPServer/Private/HttpConnection.h
+++ b/Engine/Source/Runtime/Online/HTTPServer/Private/HttpConnection.h
@@ -164,11 +164,10 @@ private:
 	/**
 	 * Logs and responds with the caller-supplied error code
 	 *
-	 * @param HttpVersion The HTTP version to be used for the error response
 	 * @param ErrorCode The HTTP error code
 	 * @param ErrorCodeStr The machine-readable error description
 	 */
-	void HandleReadError(HttpVersion::EHttpServerHttpVersion HttpVersion, EHttpServerResponseCodes ErrorCode,  const TCHAR* ErrorCodeStr);
+	void HandleReadError(EHttpServerResponseCodes ErrorCode, const TCHAR* ErrorCodeStr);
 
 	/**
 	 * Logs the caller-supplied error code and closes the connection
diff --git a/Engine/Source/Runtime/Online/HTTPServer/Public/HttpServerHttpVersion.h b/Engine/Source/Runtime/Online/HTTPServer/Public/HttpServerHttpVersion.h
--- a/Engine/Source/Runtime/Online/HTTPServer/Public/HttpServerHttpVersion.h
+++ b/Engine/Source/Runtime/Online/HTTPServer/Public/HttpServerHttpVersion.h
@@ -7,12 +7,18 @@ namespace HttpVersion
 {
 	enum class EHttpServerHttpVersion : uint8
 	{
+		HTTP_VERSION_UNKNOWN,
 		HTTP_VERSION_1_0,
 		HTTP_VERSION_1_1
 	};
 
 	static bool FromString(const FString& HttpVersionStr, EHttpServerHttpVersion& OutHttpVersion)
 	{
+		if (0 == HttpVersionStr.Compare(TEXT("Unknown")))
+		{
+			OutHttpVersion = EHttpServerHttpVersion::HTTP_VERSION_UNKNOWN;
+			return true;
+		}
 		if (0 == HttpVersionStr.Compare(TEXT("HTTP/1.0")))
 		{
 			OutHttpVersion = EHttpServerHttpVersion::HTTP_VERSION_1_0;
@@ -31,6 +37,9 @@ namespace HttpVersion
 		FString Result;
 		switch (InHttpVersion)
 		{
+		case EHttpServerHttpVersion::HTTP_VERSION_UNKNOWN:
+			Result = TEXT("Unknown");
+			break;
 		case EHttpServerHttpVersion::HTTP_VERSION_1_0:
 			Result = TEXT("HTTP/1.0");
 			break;
