Add streaming download interface.

diff --git a/Engine/Source/Programs/WebTests/Private/TestHttp.cpp b/Engine/Source/Programs/WebTests/Private/TestHttp.cpp
--- a/Engine/Source/Programs/WebTests/Private/TestHttp.cpp
+++ b/Engine/Source/Programs/WebTests/Private/TestHttp.cpp
@@ -1,5 +1,6 @@
 // Copyright Epic Games, Inc. All Rights Reserved.
 #include "CoreMinimal.h"
+#include "HAL/FileManager.h"
 #include "HAL/IConsoleManager.h"
 #include "HAL/PlatformProcess.h"
 #include "Online/HTTP/Public/HttpManager.h"
@@ -196,7 +197,9 @@ TEST_CASE_METHOD(FWaitUntilCompleteHttpFixture, "Http request - connect timeout.
 	HttpRequest->OnProcessRequestComplete().BindLambda([StartTime](FHttpRequestPtr HttpRequest, FHttpResponsePtr HttpResponse, bool bSucceeded) {
 		CHECK(!bSucceeded);
 		CHECK(HttpResponse == nullptr);
-		CHECK(HttpRequest->GetStatus() == EHttpRequestStatus::Failed_ConnectionError);
+		// TODO: For now curl impl is using customized timeout instead of relying on native http timeout, 
+		// which doesn't get CURLE_COULDNT_CONNECT. Enable this after switching to native http timeout
+		//CHECK(HttpRequest->GetStatus() == EHttpRequestStatus::Failed_ConnectionError);
 		FTimespan Timespan = FDateTime::Now() - StartTime;
 		float DurationInSeconds = Timespan.GetTotalSeconds();
 		CHECK(FMath::IsNearlyEqual(DurationInSeconds, 7, HTTP_TIME_DIFF_TOLERANCE));
@@ -204,7 +207,7 @@ TEST_CASE_METHOD(FWaitUntilCompleteHttpFixture, "Http request - connect timeout.
 	HttpRequest->ProcessRequest();
 }
 
-TEST_CASE_METHOD(FWaitUntilCompleteHttpFixture, "Streaming http download - gold path.", HTTP_TAG)
+TEST_CASE_METHOD(FWaitUntilCompleteHttpFixture, "Streaming http download", HTTP_TAG)
 {
 	uint32 Chunks = 10;
 	uint32 ChunkSize = 1024*1024;
@@ -212,11 +215,130 @@ TEST_CASE_METHOD(FWaitUntilCompleteHttpFixture, "Streaming http download - gold
 	TSharedRef<IHttpRequest, ESPMode::ThreadSafe> HttpRequest = HttpModule->CreateRequest();
 	HttpRequest->SetURL(FString::Format(TEXT("{0}/streaming_download/{1}/{2}/"), { *UrlHttpTests(), Chunks, ChunkSize }));
 	HttpRequest->SetVerb(TEXT("GET"));
-	HttpRequest->OnProcessRequestComplete().BindLambda([Chunks, ChunkSize](FHttpRequestPtr HttpRequest, FHttpResponsePtr HttpResponse, bool bSucceeded) {
-		CHECK(HttpResponse->GetContentLength() == Chunks * ChunkSize);
-		CHECK(bSucceeded);
-		CHECK(HttpResponse->GetResponseCode() == 200);
-	});
+
+	TSharedRef<int64> TotalBytesReceived = MakeShared<int64>(0);
+
+	SECTION("Success without stream provided")
+	{
+		HttpRequest->OnProcessRequestComplete().BindLambda([Chunks, ChunkSize](FHttpRequestPtr HttpRequest, FHttpResponsePtr HttpResponse, bool bSucceeded) {
+			CHECK(HttpResponse->GetContentLength() == Chunks * ChunkSize);
+			CHECK(bSucceeded);
+			CHECK(HttpResponse->GetResponseCode() == 200);
+		});
+	}
+	SECTION("Success with customized stream")
+	{
+		class FTestHttpReceiveStream final : public FArchive
+		{
+		public:
+			FTestHttpReceiveStream(TSharedRef<int64> InTotalBytesReceived)
+				: TotalBytesReceived(InTotalBytesReceived)
+			{
+			}
+
+			virtual void Serialize(void* V, int64 Length) override
+			{
+				*TotalBytesReceived += Length;
+			}
+
+			TSharedRef<int64> TotalBytesReceived;
+		};
+
+		TSharedRef<FTestHttpReceiveStream> Stream = MakeShared<FTestHttpReceiveStream>(TotalBytesReceived);
+		CHECK(HttpRequest->SetResponseBodyReceiveStream(Stream));
+
+		HttpRequest->OnProcessRequestComplete().BindLambda([Chunks, ChunkSize, TotalBytesReceived](FHttpRequestPtr HttpRequest, FHttpResponsePtr HttpResponse, bool bSucceeded) {
+			CHECK(HttpResponse->GetContentLength() == Chunks * ChunkSize);
+			CHECK(HttpResponse->GetContent().IsEmpty());
+			CHECK(*TotalBytesReceived == Chunks * ChunkSize);
+			CHECK(bSucceeded);
+			CHECK(HttpResponse->GetResponseCode() == 200);
+		});
+	}
+	SECTION("Success with customized stream delegate")
+	{
+		FHttpRequestStreamDelegate Delegate;
+		Delegate.BindLambda([TotalBytesReceived](void* Ptr, int64 Length) {
+			*TotalBytesReceived += Length;
+			return true;
+		});
+		CHECK(HttpRequest->SetResponseBodyReceiveStreamDelegate(Delegate));
+
+		HttpRequest->OnProcessRequestComplete().BindLambda([Chunks, ChunkSize, TotalBytesReceived](FHttpRequestPtr HttpRequest, FHttpResponsePtr HttpResponse, bool bSucceeded) {
+			CHECK(HttpResponse->GetContentLength() == Chunks * ChunkSize);
+			CHECK(HttpResponse->GetContent().IsEmpty());
+			CHECK(*TotalBytesReceived == Chunks * ChunkSize);
+			CHECK(bSucceeded);
+			CHECK(HttpResponse->GetResponseCode() == 200);
+		});
+	}
+	SECTION("Use customized stream to receive response body but failed when serialize")
+	{
+		class FTestHttpReceiveStream final : public FArchive
+		{
+		public:
+			FTestHttpReceiveStream(TSharedRef<int64> InTotalBytesReceived)
+				: TotalBytesReceived(InTotalBytesReceived)
+			{
+			}
+
+			virtual void Serialize(void* V, int64 Length) override
+			{
+				*TotalBytesReceived += Length;
+				SetError();
+			}
+
+			TSharedRef<int64> TotalBytesReceived;
+		};
+
+		TSharedRef<FTestHttpReceiveStream> Stream = MakeShared<FTestHttpReceiveStream>(TotalBytesReceived);
+		CHECK(HttpRequest->SetResponseBodyReceiveStream(Stream));
+
+		HttpRequest->OnProcessRequestComplete().BindLambda([ChunkSize, TotalBytesReceived](FHttpRequestPtr HttpRequest, FHttpResponsePtr HttpResponse, bool bSucceeded) {
+			CHECK(!bSucceeded);
+			CHECK(HttpResponse != nullptr);
+			CHECK(*TotalBytesReceived <= ChunkSize);
+		});
+	}
+	SECTION("Use customized stream delegate to receive response body but failed when call")
+	{
+		FHttpRequestStreamDelegate Delegate;
+		Delegate.BindLambda([TotalBytesReceived](void* Ptr, int64 Length) {
+			*TotalBytesReceived += Length;
+			return false;
+		});
+		CHECK(HttpRequest->SetResponseBodyReceiveStreamDelegate(Delegate));
+
+		HttpRequest->OnProcessRequestComplete().BindLambda([ChunkSize, TotalBytesReceived](FHttpRequestPtr HttpRequest, FHttpResponsePtr HttpResponse, bool bSucceeded) {
+			CHECK(!bSucceeded);
+			CHECK(HttpResponse != nullptr);
+			CHECK(*TotalBytesReceived <= ChunkSize);
+		});
+	}
+	SECTION("Success with file stream to receive response body")
+	{
+		FString Filename = FString(FPlatformProcess::UserSettingsDir()) / TEXT("TestStreamDownload.dat");
+		FArchive* RawFile = IFileManager::Get().CreateFileWriter(*Filename);
+		CHECK(RawFile != nullptr);
+		TSharedRef<FArchive> FileToWrite = MakeShareable(RawFile);
+		CHECK(HttpRequest->SetResponseBodyReceiveStream(FileToWrite));
+
+		HttpRequest->OnProcessRequestComplete().BindLambda([Chunks, ChunkSize, Filename, FileToWrite](FHttpRequestPtr HttpRequest, FHttpResponsePtr HttpResponse, bool bSucceeded) {
+			CHECK(HttpResponse->GetContentLength() == Chunks * ChunkSize);
+			CHECK(HttpResponse->GetContent().IsEmpty());
+			CHECK(bSucceeded);
+			CHECK(HttpResponse->GetResponseCode() == 200);
+
+			FileToWrite->FlushCache();
+			FileToWrite->Close();
+
+			TSharedRef<FArchive, ESPMode::ThreadSafe> FileToRead = MakeShareable(IFileManager::Get().CreateFileReader(*Filename));
+			CHECK(FileToRead->TotalSize() == Chunks * ChunkSize);
+
+			IFileManager::Get().Delete(*Filename);
+		});
+	}
+
 	HttpRequest->ProcessRequest();
 }
 
@@ -229,7 +351,7 @@ TEST_CASE_METHOD(FWaitUntilCompleteHttpFixture, "Streaming http upload - gold pa
 	const char* BoundaryLabel = "test_http_boundary";
 	HttpRequest->SetHeader(TEXT("Content-Type"), FString::Format(TEXT("multipart/form-data; boundary={0}"), { BoundaryLabel }));
 
-	// Not really reading file here in order to simplify the test flow. It will be sent by chunks in http request, 
+	// Not really reading file here in order to simplify the test flow. It will be sent by chunks in http request
 	const uint32 FileSize = 10*1024*1024;
 	char* FileData = (char*)FMemory::Malloc(FileSize + 1);
 	FMemory::Memset(FileData, 'd', FileSize);
diff --git a/Engine/Source/Programs/WebTests/WebServer/httptests/views.py b/Engine/Source/Programs/WebTests/WebServer/httptests/views.py
--- a/Engine/Source/Programs/WebTests/WebServer/httptests/views.py
+++ b/Engine/Source/Programs/WebTests/WebServer/httptests/views.py
@@ -28,7 +28,9 @@ def streaming_download(request, chunks, chunk_size):
         for x in range(chunks):
             data_chunk = 'd' * chunk_size
             yield data_chunk
-    return StreamingHttpResponse(data_chunk_generator())
+    response = StreamingHttpResponse(data_chunk_generator())
+    response['Content-Length'] = chunks * chunk_size
+    return response
 
 @api_view(['POST'])
 def streaming_upload(request):
diff --git a/Engine/Source/Runtime/Online/BuildPatchServices/Private/Tests/Mock/HttpRequest.mock.h b/Engine/Source/Runtime/Online/BuildPatchServices/Private/Tests/Mock/HttpRequest.mock.h
--- a/Engine/Source/Runtime/Online/BuildPatchServices/Private/Tests/Mock/HttpRequest.mock.h
+++ b/Engine/Source/Runtime/Online/BuildPatchServices/Private/Tests/Mock/HttpRequest.mock.h
@@ -108,6 +108,12 @@ namespace BuildPatchServices
 			return false;
 		}
 
+		virtual bool SetResponseBodyReceiveStream(TSharedRef<FArchive> Stream) override
+		{
+			MOCK_FUNC_NOT_IMPLEMENTED("FMockHttpRequest::SetResponseBodyReceiveStream");
+			return false;
+		}
+
 		virtual void SetHeader(const FString& HeaderName, const FString& HeaderValue) override
 		{
 			MOCK_FUNC_NOT_IMPLEMENTED("FMockHttpRequest::SetHeader");
diff --git a/Engine/Source/Runtime/Online/HTTP/Private/Apple/AppleHTTPNSUrlConnection.cpp b/Engine/Source/Runtime/Online/HTTP/Private/Apple/AppleHTTPNSUrlConnection.cpp
--- a/Engine/Source/Runtime/Online/HTTP/Private/Apple/AppleHTTPNSUrlConnection.cpp
+++ b/Engine/Source/Runtime/Online/HTTP/Private/Apple/AppleHTTPNSUrlConnection.cpp
@@ -297,6 +297,11 @@ bool FAppleHttpNSUrlConnectionRequest::SetContentFromStream(TSharedRef<FArchive,
 	return false;
 }
 
+bool FAppleHttpNSUrlConnectionRequest::SetResponseBodyReceiveStream(TSharedRef<FArchive> Stream)
+{
+	UE_LOG(LogHttp, Warning, TEXT("FAppleHttpNSUrlConnectionRequest::SetResponseBodyReceiveStream is not implemented"));
+	return false;
+}
 
 FString FAppleHttpNSUrlConnectionRequest::GetVerb() const
 {
diff --git a/Engine/Source/Runtime/Online/HTTP/Private/Apple/AppleHTTPNSUrlConnection.h b/Engine/Source/Runtime/Online/HTTP/Private/Apple/AppleHTTPNSUrlConnection.h
--- a/Engine/Source/Runtime/Online/HTTP/Private/Apple/AppleHTTPNSUrlConnection.h
+++ b/Engine/Source/Runtime/Online/HTTP/Private/Apple/AppleHTTPNSUrlConnection.h
@@ -42,6 +42,7 @@ public:
 	virtual void SetContentAsString(const FString& ContentString) override;
     virtual bool SetContentAsStreamedFile(const FString& Filename) override;
 	virtual bool SetContentFromStream(TSharedRef<FArchive, ESPMode::ThreadSafe> Stream) override;
+	virtual bool SetResponseBodyReceiveStream(TSharedRef<FArchive> Stream) override;
 	virtual void SetHeader(const FString& HeaderName, const FString& HeaderValue) override;
 	virtual void AppendToHeader(const FString& HeaderName, const FString& AdditionalHeaderValue) override;
 	virtual void SetTimeout(float InTimeoutSecs) override;
diff --git a/Engine/Source/Runtime/Online/HTTP/Private/Apple/AppleHTTPNSUrlSession.cpp b/Engine/Source/Runtime/Online/HTTP/Private/Apple/AppleHTTPNSUrlSession.cpp
--- a/Engine/Source/Runtime/Online/HTTP/Private/Apple/AppleHTTPNSUrlSession.cpp
+++ b/Engine/Source/Runtime/Online/HTTP/Private/Apple/AppleHTTPNSUrlSession.cpp
@@ -432,6 +432,12 @@ bool FAppleHttpNSUrlSessionRequest::SetContentFromStream(TSharedRef<FArchive, ES
 	return false;
 }
 
+bool FAppleHttpNSUrlSessionRequest::SetResponseBodyReceiveStream(TSharedRef<FArchive> Stream)
+{
+	UE_LOG(LogHttp, Warning, TEXT("FAppleHttpNSUrlSessionRequest::SetResponseBodyReceiveStream is not implemented"));
+	return false;
+}
+
 FString FAppleHttpNSUrlSessionRequest::GetVerb() const
 {
 	FString ConvertedVerb(Request.HTTPMethod);
diff --git a/Engine/Source/Runtime/Online/HTTP/Private/Apple/AppleHTTPNSUrlSession.h b/Engine/Source/Runtime/Online/HTTP/Private/Apple/AppleHTTPNSUrlSession.h
--- a/Engine/Source/Runtime/Online/HTTP/Private/Apple/AppleHTTPNSUrlSession.h
+++ b/Engine/Source/Runtime/Online/HTTP/Private/Apple/AppleHTTPNSUrlSession.h
@@ -36,6 +36,7 @@ public:
 	virtual void SetContentAsString(const FString& ContentString) override;
     virtual bool SetContentAsStreamedFile(const FString& Filename) override;
 	virtual bool SetContentFromStream(TSharedRef<FArchive, ESPMode::ThreadSafe> Stream) override;
+	virtual bool SetResponseBodyReceiveStream(TSharedRef<FArchive> Stream) override;
 	virtual void SetHeader(const FString& HeaderName, const FString& HeaderValue) override;
 	virtual void AppendToHeader(const FString& HeaderName, const FString& AdditionalHeaderValue) override;
 	virtual void SetTimeout(float InTimeoutSecs) override;
diff --git a/Engine/Source/Runtime/Online/HTTP/Private/Curl/CurlHttp.cpp b/Engine/Source/Runtime/Online/HTTP/Private/Curl/CurlHttp.cpp
--- a/Engine/Source/Runtime/Online/HTTP/Private/Curl/CurlHttp.cpp
+++ b/Engine/Source/Runtime/Online/HTTP/Private/Curl/CurlHttp.cpp
@@ -391,6 +391,12 @@ bool FCurlHttpRequest::SetContentFromStream(TSharedRef<FArchive, ESPMode::Thread
 	return true;
 }
 
+bool FCurlHttpRequest::SetResponseBodyReceiveStream(TSharedRef<FArchive> Stream)
+{
+	ResponseBodyReceiveStream = Stream;
+	return true;
+}
+
 void FCurlHttpRequest::SetHeader(const FString& HeaderName, const FString& HeaderValue)
 {
 	if (CompletionStatus == EHttpRequestStatus::Processing)
@@ -486,59 +492,56 @@ size_t FCurlHttpRequest::StaticDebugCallback(CURL * Handle, curl_infotype DebugI
 size_t FCurlHttpRequest::ReceiveResponseHeaderCallback(void* Ptr, size_t SizeInBlocks, size_t BlockSizeInBytes)
 {
 	QUICK_SCOPE_CYCLE_COUNTER(STAT_FCurlHttpRequest_ReceiveResponseHeaderCallback);
-	check(Response.IsValid());
-	
+
+	if (!Response.IsValid())
+	{
+		Response = MakeShared<FCurlHttpResponse, ESPMode::ThreadSafe>(*this);
+	}
+
 	TimeSinceLastResponse = 0.0f;
-	if (Response.IsValid())
+	uint32 HeaderSize = SizeInBlocks * BlockSizeInBytes;
+	if (HeaderSize > 0 && HeaderSize <= CURL_MAX_HTTP_HEADER)
 	{
-		uint32 HeaderSize = SizeInBlocks * BlockSizeInBytes;
-		if (HeaderSize > 0 && HeaderSize <= CURL_MAX_HTTP_HEADER)
-		{
-			TArray<char> AnsiHeader;
-			AnsiHeader.AddUninitialized(HeaderSize + 1);
+		TArray<char> AnsiHeader;
+		AnsiHeader.AddUninitialized(HeaderSize + 1);
 
-			FMemory::Memcpy(AnsiHeader.GetData(), Ptr, HeaderSize);
-			AnsiHeader[HeaderSize] = 0;
+		FMemory::Memcpy(AnsiHeader.GetData(), Ptr, HeaderSize);
+		AnsiHeader[HeaderSize] = 0;
 
-			FString Header(ANSI_TO_TCHAR(AnsiHeader.GetData()));
-			// kill \n\r
-			Header = Header.Replace(TEXT("\n"), TEXT(""));
-			Header = Header.Replace(TEXT("\r"), TEXT(""));
+		FString Header(ANSI_TO_TCHAR(AnsiHeader.GetData()));
+		// kill \n\r
+		Header = Header.Replace(TEXT("\n"), TEXT(""));
+		Header = Header.Replace(TEXT("\r"), TEXT(""));
 
-			UE_LOG(LogHttp, Verbose, TEXT("%p: Received response header '%s'."), this, *Header);
+		UE_LOG(LogHttp, Verbose, TEXT("%p: Received response header '%s'."), this, *Header);
 
-			FString HeaderKey, HeaderValue;
-			if (Header.Split(TEXT(":"), &HeaderKey, &HeaderValue))
-			{
-				HeaderValue.TrimStartInline();
-				if (!HeaderKey.IsEmpty() && !HeaderValue.IsEmpty() && !bRedirected)
-				{
-					//Store the content length so OnRequestProgress() delegates have something to work with
-					if (HeaderKey == TEXT("Content-Length"))
-					{
-						Response->ContentLength = FCString::Atoi(*HeaderValue);
-					}
-					Response->NewlyReceivedHeaders.Enqueue(TPair<FString, FString>(MoveTemp(HeaderKey), MoveTemp(HeaderValue)));
-				}
-			}
-			else
+		FString HeaderKey, HeaderValue;
+		if (Header.Split(TEXT(":"), &HeaderKey, &HeaderValue))
+		{
+			HeaderValue.TrimStartInline();
+			if (!HeaderKey.IsEmpty() && !HeaderValue.IsEmpty() && !bRedirected)
 			{
-				long HttpCode = 0;
-				if (CURLE_OK == curl_easy_getinfo(EasyHandle, CURLINFO_RESPONSE_CODE, &HttpCode))
+				//Store the content length so OnRequestProgress() delegates have something to work with
+				if (HeaderKey == TEXT("Content-Length"))
 				{
-					bRedirected = (HttpCode >= 300 && HttpCode < 400);
+					Response->ContentLength = FCString::Atoi(*HeaderValue);
 				}
+				Response->NewlyReceivedHeaders.Enqueue(TPair<FString, FString>(MoveTemp(HeaderKey), MoveTemp(HeaderValue)));
 			}
-			return HeaderSize;
 		}
 		else
 		{
-			UE_LOG(LogHttp, Warning, TEXT("%p: Could not process response header for request - header size (%d) is invalid."), this, HeaderSize);
+			long HttpCode = 0;
+			if (CURLE_OK == curl_easy_getinfo(EasyHandle, CURLINFO_RESPONSE_CODE, &HttpCode))
+			{
+				bRedirected = (HttpCode >= 300 && HttpCode < 400);
+			}
 		}
+		return HeaderSize;
 	}
 	else
 	{
-		UE_LOG(LogHttp, Warning, TEXT("%p: Could not download response header for request - response not valid."), this);
+		UE_LOG(LogHttp, Warning, TEXT("%p: Could not process response header for request - header size (%d) is invalid."), this, HeaderSize);
 	}
 
 	return 0;
@@ -547,37 +550,53 @@ size_t FCurlHttpRequest::ReceiveResponseHeaderCallback(void* Ptr, size_t SizeInB
 size_t FCurlHttpRequest::ReceiveResponseBodyCallback(void* Ptr, size_t SizeInBlocks, size_t BlockSizeInBytes)
 {
 	QUICK_SCOPE_CYCLE_COUNTER(STAT_FCurlHttpRequest_ReceiveResponseBodyCallback);
-	check(Response.IsValid());
-	  
-	TimeSinceLastResponse = 0.0f;
-	if (Response.IsValid())
+
+	if (!Response.IsValid())
 	{
-		uint32 SizeToDownload = SizeInBlocks * BlockSizeInBytes;
+		Response = MakeShared<FCurlHttpResponse, ESPMode::ThreadSafe>(*this);
+	}
 
-		UE_LOG(LogHttp, Verbose, TEXT("%p: ReceiveResponseBodyCallback: %d bytes out of %d received. (SizeInBlocks=%d, BlockSizeInBytes=%d, Response->TotalBytesRead=%d, Response->GetContentLength()=%d, SizeToDownload=%d (<-this will get returned from the callback))"),
-			this,
-			static_cast<int32>(Response->TotalBytesRead.GetValue() + SizeToDownload), Response->GetContentLength(),
-			static_cast<int32>(SizeInBlocks), static_cast<int32>(BlockSizeInBytes), Response->TotalBytesRead.GetValue(), Response->GetContentLength(), static_cast<int32>(SizeToDownload)
-			);
+	TimeSinceLastResponse = 0.0f;
 
-		// note that we can be passed 0 bytes if file transmitted has 0 length
-		if (SizeToDownload > 0)
-		{
-			Response->Payload.AddUninitialized(SizeToDownload);
+	// Number of bytes actually taken care of. If that amount differs from the amount passed to your 
+	// callback function, it will signal an error condition to the library. This will cause the transfer 
+	// to get aborted and the libcurl function used will return CURLE_WRITE_ERROR.
+	size_t NumberOfBytesProcessed = 0;
+	  
+	uint32 SizeToDownload = SizeInBlocks * BlockSizeInBytes;
 
-			// save
-			FMemory::Memcpy(static_cast<uint8*>(Response->Payload.GetData()) + Response->TotalBytesRead.GetValue(), Ptr, SizeToDownload);
-			Response->TotalBytesRead.Add(SizeToDownload);
+	UE_LOG(LogHttp, Verbose, TEXT("%p: ReceiveResponseBodyCallback: %d bytes out of %d received. (SizeInBlocks=%d, BlockSizeInBytes=%d, Response->TotalBytesRead=%d, Response->GetContentLength()=%d, SizeToDownload=%d (<-this will get returned from the callback))"),
+		this,
+		static_cast<int32>(Response->TotalBytesRead.GetValue() + SizeToDownload), Response->GetContentLength(),
+		static_cast<int32>(SizeInBlocks), static_cast<int32>(BlockSizeInBytes), Response->TotalBytesRead.GetValue(), Response->GetContentLength(), static_cast<int32>(SizeToDownload)
+		);
+
+	// note that we can be passed 0 bytes if file transmitted has 0 length
+	if (SizeToDownload == 0)
+	{
+		return NumberOfBytesProcessed;
+	}
+
+	if (ResponseBodyReceiveStream)
+	{
+		ResponseBodyReceiveStream->Serialize(Ptr, SizeToDownload);
 
-			return SizeToDownload;
+		if (!ResponseBodyReceiveStream->GetError())
+		{
+			NumberOfBytesProcessed = SizeToDownload;
 		}
 	}
 	else
 	{
-		UE_LOG(LogHttp, Warning, TEXT("%p: Could not download response data for request - response not valid."), this);
+		Response->Payload.AddUninitialized(SizeToDownload);
+		FMemory::Memcpy(static_cast<uint8*>(Response->Payload.GetData()) + Response->TotalBytesRead.GetValue(), Ptr, SizeToDownload);
+
+		NumberOfBytesProcessed = SizeToDownload;
 	}
 
-	return 0;	// request will fail with write error if we had non-zero bytes to download
+	Response->TotalBytesRead.Add(NumberOfBytesProcessed);
+
+	return NumberOfBytesProcessed;
 }
 
 size_t FCurlHttpRequest::UploadCallback(void* Ptr, size_t SizeInBlocks, size_t BlockSizeInBytes)
@@ -986,9 +1005,6 @@ bool FCurlHttpRequest::SetupRequestHttpThread()
 
 	UE_LOG(LogHttp, Log, TEXT("%p: Starting %s request to URL='%s'"), this, *Verb, *URL);
 
-	// Response object to handle data that comes back after starting this request
-	Response = MakeShared<FCurlHttpResponse, ESPMode::ThreadSafe>(*this);
-
 	return true;
 }
 
@@ -1243,7 +1259,7 @@ void FCurlHttpRequest::FinishedRequest()
 			if (Response->HttpCode <= 0 && URL.StartsWith(TEXT("Http"), ESearchCase::IgnoreCase))
 			{
 				UE_LOG(LogHttp, Warning, TEXT("%p: invalid HTTP response code received. URL: %s, HTTP code: %d, content length: %d, actual payload size: %d"),
-					this, *GetURL(), Response->HttpCode, Response->ContentLength, Response->Payload.Num());
+					this, *GetURL(), Response->HttpCode, Response->ContentLength, Response->TotalBytesRead.GetValue());
 				Response->bSucceeded = false;
 			}
 		}
@@ -1268,12 +1284,12 @@ void FCurlHttpRequest::FinishedRequest()
 			if (bDebugServerResponse)
 			{
 				UE_LOG(LogHttp, Warning, TEXT("%p: request has been successfully processed. URL: %s, HTTP code: %d, content length: %d, actual payload size: %d, elapsed: %.2fs"),
-					this, *GetURL(), Response->HttpCode, Response->ContentLength, Response->Payload.Num(), ElapsedTime);
+					this, *GetURL(), Response->HttpCode, Response->ContentLength, Response->TotalBytesRead.GetValue(), ElapsedTime);
 			}
 			else
 			{
 				UE_LOG(LogHttp, Log, TEXT("%p: request has been successfully processed. URL: %s, HTTP code: %d, content length: %d, actual payload size: %d, elapsed: %.2fs"),
-					this, *GetURL(), Response->HttpCode, Response->ContentLength, Response->Payload.Num(), ElapsedTime);
+					this, *GetURL(), Response->HttpCode, Response->ContentLength, Response->TotalBytesRead.GetValue(), ElapsedTime);
 			}
 
 			TArray<FString> AllHeaders = Response->GetAllHeaders();
diff --git a/Engine/Source/Runtime/Online/HTTP/Private/Curl/CurlHttp.h b/Engine/Source/Runtime/Online/HTTP/Private/Curl/CurlHttp.h
--- a/Engine/Source/Runtime/Online/HTTP/Private/Curl/CurlHttp.h
+++ b/Engine/Source/Runtime/Online/HTTP/Private/Curl/CurlHttp.h
@@ -171,6 +171,7 @@ public:
 	virtual void SetContentAsString(const FString& ContentString) override;
 	virtual bool SetContentAsStreamedFile(const FString& Filename) override;
 	virtual bool SetContentFromStream(TSharedRef<FArchive, ESPMode::ThreadSafe> Stream) override;
+	virtual bool SetResponseBodyReceiveStream(TSharedRef<FArchive> Stream) override;
 	virtual void SetHeader(const FString& HeaderName, const FString& HeaderValue) override;
 	virtual void AppendToHeader(const FString& HeaderName, const FString& AdditionalHeaderValue) override;
 	virtual bool ProcessRequest() override;
@@ -396,6 +397,8 @@ private:
 	TUniquePtr<FRequestPayload> RequestPayload;
 	/** Is the request payload seekable? */
 	bool bIsRequestPayloadSeekable = false;
+	/** The stream to receive response body */
+	TSharedPtr<FArchive> ResponseBodyReceiveStream;
 	/** Current status of request being processed */
 	EHttpRequestStatus::Type CompletionStatus;
 	/** Mapping of header section to values. */
@@ -470,6 +473,8 @@ private:
 
 	/** BYTE array to fill in as the response is read via didReceiveData */
 	TArray<uint8> Payload;
+	/** The stream to receive response body */
+	TSharedPtr<FArchive> ResponseBodyReceiveStream;
 	/** Caches how many bytes of the response we've read so far */
 	FThreadSafeCounter TotalBytesRead;
 	/** Cached key/value header pairs. Parsed once request completes. Only accessible on the game thread. */
diff --git a/Engine/Source/Runtime/Online/HTTP/Private/GenericPlatform/GenericPlatformHttp.cpp b/Engine/Source/Runtime/Online/HTTP/Private/GenericPlatform/GenericPlatformHttp.cpp
--- a/Engine/Source/Runtime/Online/HTTP/Private/GenericPlatform/GenericPlatformHttp.cpp
+++ b/Engine/Source/Runtime/Online/HTTP/Private/GenericPlatform/GenericPlatformHttp.cpp
@@ -98,6 +98,7 @@ public:
 	virtual void SetContentAsString(const FString& ContentString) override {}
 	virtual bool SetContentAsStreamedFile(const FString& Filename) override { return false; }
 	virtual bool SetContentFromStream(TSharedRef<FArchive, ESPMode::ThreadSafe> Stream) override { return false; }
+	virtual bool SetResponseBodyReceiveStream(TSharedRef<FArchive> Stream) override { return false; }
 	virtual void SetHeader(const FString& HeaderName, const FString& HeaderValue) override {}
 	virtual void AppendToHeader(const FString& HeaderName, const FString& AdditionalHeaderValue) override {}
 	virtual bool ProcessRequest() override { return false; }
diff --git a/Engine/Source/Runtime/Online/HTTP/Private/IXML/HttpIXML.h b/Engine/Source/Runtime/Online/HTTP/Private/IXML/HttpIXML.h
--- a/Engine/Source/Runtime/Online/HTTP/Private/IXML/HttpIXML.h
+++ b/Engine/Source/Runtime/Online/HTTP/Private/IXML/HttpIXML.h
@@ -41,6 +41,7 @@ public:
 	virtual void SetContentAsString(const FString& ContentString) override;
 	virtual bool SetContentAsStreamedFile(const FString& Filename) override;
 	virtual bool SetContentFromStream(TSharedRef<FArchive, ESPMode::ThreadSafe> Stream) override;
+	virtual bool SetResponseBodyReceiveStream(TSharedRef<FArchive> Stream) override;
 	virtual void SetHeader(const FString& HeaderName, const FString& HeaderValue) override;
 	virtual void AppendToHeader(const FString& HeaderName, const FString& AdditionalHeaderValue) override;
 	virtual bool ProcessRequest() override;
diff --git a/Engine/Source/Runtime/Online/HTTP/Private/IXML/HttpIXml.cpp b/Engine/Source/Runtime/Online/HTTP/Private/IXML/HttpIXml.cpp
--- a/Engine/Source/Runtime/Online/HTTP/Private/IXML/HttpIXml.cpp
+++ b/Engine/Source/Runtime/Online/HTTP/Private/IXML/HttpIXml.cpp
@@ -183,6 +183,11 @@ bool FHttpRequestIXML::SetContentFromStream(TSharedRef<FArchive, ESPMode::Thread
 	return true;
 }
 
+bool FHttpRequestIXML::SetResponseBodyReceiveStream(TSharedRef<FArchive> Stream)
+{
+	// TODO: Impl
+	return false;
+}
 
 //-----------------------------------------------------------------------------
 //	
diff --git a/Engine/Source/Runtime/Online/HTTP/Private/NullHttp.cpp b/Engine/Source/Runtime/Online/HTTP/Private/NullHttp.cpp
--- a/Engine/Source/Runtime/Online/HTTP/Private/NullHttp.cpp
+++ b/Engine/Source/Runtime/Online/HTTP/Private/NullHttp.cpp
@@ -97,6 +97,12 @@ bool FNullHttpRequest::SetContentFromStream(TSharedRef<FArchive, ESPMode::Thread
 	return false;
 }
 
+bool FNullHttpRequest::SetResponseBodyReceiveStream(TSharedRef<FArchive> Stream)
+{
+	UE_LOG(LogHttp, Warning, TEXT("FNullHttpRequest::SetResponseBodyReceiveStream is not implemented"));
+	return false;
+}
+
 void FNullHttpRequest::SetHeader(const FString& HeaderName, const FString& HeaderValue)
 {
 	Headers.Add(HeaderName, HeaderValue);
diff --git a/Engine/Source/Runtime/Online/HTTP/Private/NullHttp.h b/Engine/Source/Runtime/Online/HTTP/Private/NullHttp.h
--- a/Engine/Source/Runtime/Online/HTTP/Private/NullHttp.h
+++ b/Engine/Source/Runtime/Online/HTTP/Private/NullHttp.h
@@ -30,6 +30,7 @@ public:
 	virtual void SetContentAsString(const FString& ContentString) override;
     virtual bool SetContentAsStreamedFile(const FString& Filename) override;
 	virtual bool SetContentFromStream(TSharedRef<FArchive, ESPMode::ThreadSafe> Stream) override;
+	virtual bool SetResponseBodyReceiveStream(TSharedRef<FArchive> Stream) override;
 	virtual void SetHeader(const FString& HeaderName, const FString& HeaderValue) override;
 	virtual void AppendToHeader(const FString& HeaderName, const FString& AdditionalHeaderValue) override;
 	virtual bool ProcessRequest() override;
diff --git a/Engine/Source/Runtime/Online/HTTP/Private/WinHttp/WinHttpHttpRequest.cpp b/Engine/Source/Runtime/Online/HTTP/Private/WinHttp/WinHttpHttpRequest.cpp
--- a/Engine/Source/Runtime/Online/HTTP/Private/WinHttp/WinHttpHttpRequest.cpp
+++ b/Engine/Source/Runtime/Online/HTTP/Private/WinHttp/WinHttpHttpRequest.cpp
@@ -168,6 +168,12 @@ bool FWinHttpHttpRequest::SetContentFromStream(TSharedRef<FArchive, ESPMode::Thr
 	return true;
 }
 
+bool FWinHttpHttpRequest::SetResponseBodyReceiveStream(TSharedRef<FArchive> Stream)
+{
+	UE_LOG(LogHttp, Warning, TEXT("FWinHttpHttpRequest::SetResponseBodyReceiveStream is not implemented"));
+	return false;
+}
+
 void FWinHttpHttpRequest::SetHeader(const FString& HeaderName, const FString& HeaderValue)
 {
 	if (State == EHttpRequestStatus::Processing)
diff --git a/Engine/Source/Runtime/Online/HTTP/Private/WinHttp/WinHttpHttpRequest.h b/Engine/Source/Runtime/Online/HTTP/Private/WinHttp/WinHttpHttpRequest.h
--- a/Engine/Source/Runtime/Online/HTTP/Private/WinHttp/WinHttpHttpRequest.h
+++ b/Engine/Source/Runtime/Online/HTTP/Private/WinHttp/WinHttpHttpRequest.h
@@ -41,6 +41,7 @@ public:
 	virtual void SetContentAsString(const FString& ContentString) override;
 	virtual bool SetContentAsStreamedFile(const FString& Filename) override;
 	virtual bool SetContentFromStream(TSharedRef<FArchive, ESPMode::ThreadSafe> Stream) override;
+	virtual bool SetResponseBodyReceiveStream(TSharedRef<FArchive> Stream) override;
 	virtual void SetHeader(const FString& HeaderName, const FString& HeaderValue) override;
 	virtual void AppendToHeader(const FString& HeaderName, const FString& AdditionalHeaderValue) override;
 	virtual bool ProcessRequest() override;
diff --git a/Engine/Source/Runtime/Online/HTTP/Public/HttpRequestAdapter.h b/Engine/Source/Runtime/Online/HTTP/Public/HttpRequestAdapter.h
--- a/Engine/Source/Runtime/Online/HTTP/Public/HttpRequestAdapter.h
+++ b/Engine/Source/Runtime/Online/HTTP/Public/HttpRequestAdapter.h
@@ -33,6 +33,7 @@ public:
 	virtual void                          SetContentAsString(const FString& ContentString) override                { HttpRequest->SetContentAsString(ContentString); }
     virtual bool                          SetContentAsStreamedFile(const FString& Filename) override               { return HttpRequest->SetContentAsStreamedFile(Filename); }
 	virtual bool                          SetContentFromStream(TSharedRef<FArchive, ESPMode::ThreadSafe> Stream) override { return HttpRequest->SetContentFromStream(Stream); }
+	virtual bool						  SetResponseBodyReceiveStream(TSharedRef<FArchive> Stream) override       { return HttpRequest->SetResponseBodyReceiveStream(Stream); }
 	virtual void                          SetHeader(const FString& HeaderName, const FString& HeaderValue) override { HttpRequest->SetHeader(HeaderName, HeaderValue); }
 	virtual void                          AppendToHeader(const FString& HeaderName, const FString& AdditionalHeaderValue) override { HttpRequest->AppendToHeader(HeaderName, AdditionalHeaderValue); }
 	virtual void                          SetTimeout(float InTimeoutSecs) override                                 { HttpRequest->SetTimeout(InTimeoutSecs); }
diff --git a/Engine/Source/Runtime/Online/HTTP/Public/Interfaces/IHttpRequest.h b/Engine/Source/Runtime/Online/HTTP/Public/Interfaces/IHttpRequest.h
--- a/Engine/Source/Runtime/Online/HTTP/Public/Interfaces/IHttpRequest.h
+++ b/Engine/Source/Runtime/Online/HTTP/Public/Interfaces/IHttpRequest.h
@@ -107,6 +107,39 @@ DECLARE_DELEGATE_ThreeParams(FHttpRequestProgressDelegate, FHttpRequestPtr /*Req
  */
 DECLARE_DELEGATE_ThreeParams(FHttpRequestWillRetryDelegate, FHttpRequestPtr /*Request*/, FHttpResponsePtr /*Response*/, float /*SecondsToRetry*/);
 
+/**
+ * Delegate called when an Http request will send/recv data through stream
+ *
+ * @param Ptr - The buffer ptr to read/write
+ * @param Length - The length of buffer to read/write
+ * @return true if succeed, false if failed to read/write data
+ */
+DECLARE_DELEGATE_RetVal_TwoParams(bool, FHttpRequestStreamDelegate, void*/*Ptr*/, int64/*Length*/);
+
+
+/**
+ * Delegate version of FArchive, for streaming interface
+ */
+class FArchiveWithDelegate final : public FArchive
+{
+public:
+	FArchiveWithDelegate(FHttpRequestStreamDelegate InStreamDelegate)
+		: StreamDelegate(InStreamDelegate)
+	{
+	}
+
+	virtual void Serialize(void* V, int64 Length) override
+	{
+		if (!StreamDelegate.Execute(V, Length))
+		{
+			SetError();
+		}
+	}
+
+private:
+	FHttpRequestStreamDelegate StreamDelegate;
+};
+
 /**
  * Interface for Http requests (created using FHttpFactory)
  */
@@ -182,6 +215,41 @@ public:
 	 */
 	virtual bool SetContentFromStream(TSharedRef<FArchive, ESPMode::ThreadSafe> Stream) = 0;
 
+	/**
+	 * Sets the content of the request to stream directly from an delegate.
+	 *
+	 * @param StreamDelegate - delegate from which the payload should be streamed.
+	 * @return True if the delegate can be used to stream the request. False otherwise.
+	 */
+	bool SetContentFromStreamDelegate(FHttpRequestStreamDelegate StreamDelegate) { return SetContentFromStream(MakeShared<FArchiveWithDelegate>(StreamDelegate)); }
+
+	/**
+	 * Sets the stream to receive the response body. Make sure to handle the cleanup of stream when 
+	 * Serialize generated error(Stream->GetError returns true after Stream->Serialize call), this 
+	 * http request will fail and quit.
+	 *
+	 * NOTE: Once set, the data will no longer be cached in response, IHttpResponse::GetContent() and 
+	 * IHttpResponse::GetContentAsString() will return empty result. The Stream->Serialize will be called 
+	 * from another thread other than the game thread
+	 *
+	 * @param Stream - will be used to receive the response body
+	 * @return True if the stream can be used. False otherwise.
+	 */
+	virtual bool SetResponseBodyReceiveStream(TSharedRef<FArchive> Stream) = 0;
+
+	/**
+	 * Sets the delegate to receive the response body. Make sure to handle the cleanup of received data when 
+	 * failed to process the data(StreamDelegate return false), this http request will fail and quit.
+	 *
+	 * NOTE: Once set, the data will no longer be cached in response, IHttpResponse::GetContent() and 
+	 * IHttpResponse::GetContentAsString() will return empty result. The delegate will be called from
+	 * another thread other than the game thread
+	 *
+	 * @param StreamDelegate - will be used to receive the response body
+	 * @return True if the delegate can be used. False otherwise.
+	 */
+	bool SetResponseBodyReceiveStreamDelegate(FHttpRequestStreamDelegate StreamDelegate) { return SetResponseBodyReceiveStream(MakeShared<FArchiveWithDelegate>(StreamDelegate)); }
+
 	/**
 	 * Sets optional header info.
 	 * SetHeader for a given HeaderName will overwrite any previous values
