Refactor: extract some common functions into a new base class HttpRequestCommon, where code can be shared between different platforms, but not different http impls like adaptor/mock/null impl.

diff --git a/Engine/Source/Programs/WebTests/Private/TestHttp.cpp b/Engine/Source/Programs/WebTests/Private/TestHttp.cpp
--- a/Engine/Source/Programs/WebTests/Private/TestHttp.cpp
+++ b/Engine/Source/Programs/WebTests/Private/TestHttp.cpp
@@ -364,3 +364,22 @@ TEST_CASE_METHOD(FWaitUntilCompleteHttpFixture, "Streaming http upload - gold pa
 	});
 	HttpRequest->ProcessRequest();
 }
+
+TEST_CASE_METHOD(FWaitUntilCompleteHttpFixture, "Pre check http request", HTTP_TAG)
+{
+	// Pre check will fail when domain is not allowed
+	UE::Core::FURLRequestFilter::FRequestMap SchemeMap;
+	SchemeMap.Add(TEXT("http"), TArray<FString>{TEXT("epicgames.com")});
+	UE::Core::FURLRequestFilter Filter{SchemeMap};
+	HttpModule->GetHttpManager().SetURLRequestFilter(Filter);
+
+	TSharedRef<IHttpRequest, ESPMode::ThreadSafe> HttpRequest = HttpModule->CreateRequest();
+	HttpRequest->SetVerb(TEXT("GET"));
+	HttpRequest->SetURL(UrlToTestMethods()); // Not in allowed domains: epicgames.com
+	HttpRequest->OnProcessRequestComplete().BindLambda([](FHttpRequestPtr HttpRequest, FHttpResponsePtr HttpResponse, bool bSucceeded) {
+		CHECK(!bSucceeded);
+	});
+
+	HttpRequest->ProcessRequest();
+}
+
diff --git a/Engine/Source/Runtime/Online/HTTP/Private/Apple/AppleHTTPNSUrlConnection.cpp b/Engine/Source/Runtime/Online/HTTP/Private/Apple/AppleHTTPNSUrlConnection.cpp
--- a/Engine/Source/Runtime/Online/HTTP/Private/Apple/AppleHTTPNSUrlConnection.cpp
+++ b/Engine/Source/Runtime/Online/HTTP/Private/Apple/AppleHTTPNSUrlConnection.cpp
@@ -26,7 +26,6 @@ FAppleHttpNSUrlConnectionRequest::FAppleHttpNSUrlConnectionRequest()
 :	Connection(nullptr)
 ,	bIsPayloadFile(false)
 ,	RequestPayloadByteLength(0)
-,	CompletionStatus(EHttpRequestStatus::NotStarted)
 ,	ProgressBytesSent(0)
 ,	StartRequestTime(0.0)
 ,	ElapsedTime(0.0f)
@@ -337,34 +336,8 @@ bool FAppleHttpNSUrlConnectionRequest::ProcessRequest()
 {
 	SCOPED_AUTORELEASE_POOL;
 	UE_LOG(LogHttp, Verbose, TEXT("FAppleHttpNSUrlConnectionRequest::ProcessRequest()"));
-	bool bStarted = false;
-
-	FString Scheme(Request.URL.scheme);
-	Scheme = Scheme.ToLower();
 
-	// Prevent overlapped requests using the same instance
-	if (CompletionStatus == EHttpRequestStatus::Processing)
-	{
-		UE_LOG(LogHttp, Warning, TEXT("ProcessRequest failed. Still processing last request."));
-	}
-	else if(GetURL().Len() == 0)
-	{
-		UE_LOG(LogHttp, Warning, TEXT("ProcessRequest failed. No URL was specified."));
-	}
-	else if( Scheme != TEXT("http") && Scheme != TEXT("https"))
-	{
-		UE_LOG(LogHttp, Warning, TEXT("ProcessRequest failed. URL '%s' is not a valid HTTP request. %p"), *GetURL(), this);
-	}
-	else if (!FHttpModule::Get().GetHttpManager().IsDomainAllowed(GetURL()))
-	{
-		UE_LOG(LogHttp, Warning, TEXT("ProcessRequest failed. URL '%s' is not using an allowed domain. %p"), *GetURL(), this);
-	}
-	else
-	{
-		bStarted = StartRequest();
-	}
-
-	if( !bStarted )
+	if (!PreCheck() || !StartRequest())
 	{
 		// Ensure we run on game thread
 		if (!IsInGameThread())
@@ -378,9 +351,11 @@ bool FAppleHttpNSUrlConnectionRequest::ProcessRequest()
 		{
 			FinishedRequest();
 		}
+
+		return false;
 	}
 
-	return bStarted;
+	return true;
 }
 
 bool FAppleHttpNSUrlConnectionRequest::StartRequest()
@@ -510,14 +485,6 @@ void FAppleHttpNSUrlConnectionRequest::CancelRequest()
 	}
 }
 
-
-EHttpRequestStatus::Type FAppleHttpNSUrlConnectionRequest::GetStatus() const
-{
-	UE_LOG(LogHttp, Verbose, TEXT("FAppleHttpNSUrlConnectionRequest::GetStatus()"));
-	return CompletionStatus;
-}
-
-
 const FHttpResponsePtr FAppleHttpNSUrlConnectionRequest::GetResponse() const
 {
 	return Response;
diff --git a/Engine/Source/Runtime/Online/HTTP/Private/Apple/AppleHTTPNSUrlConnection.h b/Engine/Source/Runtime/Online/HTTP/Private/Apple/AppleHTTPNSUrlConnection.h
--- a/Engine/Source/Runtime/Online/HTTP/Private/Apple/AppleHTTPNSUrlConnection.h
+++ b/Engine/Source/Runtime/Online/HTTP/Private/Apple/AppleHTTPNSUrlConnection.h
@@ -4,7 +4,7 @@
 
 #include "CoreMinimal.h"
 #include "Containers/UnrealString.h"
-#include "GenericPlatform/HttpRequestImpl.h"
+#include "GenericPlatform/FHttpRequestCommon.h"
 #include "Interfaces/IHttpResponse.h"
 #include "HttpManager.h"
 #include "PlatformHttp.h"
@@ -16,7 +16,7 @@ class IHttpRequest;
 /**
  * Apple implementation of an Http request
  */
-class FAppleHttpNSUrlConnectionRequest : public FHttpRequestImpl
+class FAppleHttpNSUrlConnectionRequest : public FHttpRequestCommon
 {
 public:
 	// implementation friends
@@ -50,7 +50,6 @@ public:
 	virtual TOptional<float> GetTimeout() const override;
 	virtual bool ProcessRequest() override;
 	virtual void CancelRequest() override;
-	virtual EHttpRequestStatus::Type GetStatus() const override;
 	virtual const FHttpResponsePtr GetResponse() const override;
 	virtual void Tick(float DeltaSeconds) override;
 	virtual float GetElapsedTime() const override;
@@ -110,9 +109,6 @@ private:
 	/** BYTE array for content which we now own */
 	TArray<uint8> ContentData;
 
-	/** Current status of request being processed */
-	EHttpRequestStatus::Type CompletionStatus;
-
 	/** Number of bytes sent to progress update */
 	int32 ProgressBytesSent;
 
diff --git a/Engine/Source/Runtime/Online/HTTP/Private/Apple/AppleHTTPNSUrlSession.cpp b/Engine/Source/Runtime/Online/HTTP/Private/Apple/AppleHTTPNSUrlSession.cpp
--- a/Engine/Source/Runtime/Online/HTTP/Private/Apple/AppleHTTPNSUrlSession.cpp
+++ b/Engine/Source/Runtime/Online/HTTP/Private/Apple/AppleHTTPNSUrlSession.cpp
@@ -471,34 +471,8 @@ bool FAppleHttpNSUrlSessionRequest::ProcessRequest()
 {
 	SCOPED_AUTORELEASE_POOL;
 	UE_LOG(LogHttp, Verbose, TEXT("FAppleHttpNSUrlSessionRequest::ProcessRequest()"));
-	bool bStarted = false;
-
-	FString Scheme(Request.URL.scheme);
-	Scheme = Scheme.ToLower();
 
-	// Prevent overlapped requests using the same instance
-	if (CompletionStatus == EHttpRequestStatus::Processing)
-	{
-		UE_LOG(LogHttp, Warning, TEXT("ProcessRequest failed. Still processing last request."));
-	}
-	else if(GetURL().Len() == 0)
-	{
-		UE_LOG(LogHttp, Warning, TEXT("ProcessRequest failed. No URL was specified."));
-	}
-	else if( Scheme != TEXT("http") && Scheme != TEXT("https"))
-	{
-		UE_LOG(LogHttp, Warning, TEXT("ProcessRequest failed. URL '%s' is not a valid HTTP request. %p"), *GetURL(), this);
-	}
-	else if (!FHttpModule::Get().GetHttpManager().IsDomainAllowed(GetURL()))
-	{
-		UE_LOG(LogHttp, Warning, TEXT("ProcessRequest failed. URL '%s' is not using an allowed domain. %p"), *GetURL(), this);
-	}
-	else
-	{
-		bStarted = StartRequest();
-	}
-
-	if( !bStarted )
+	if (!PreCheck() || !StartRequest())
 	{
 		// Ensure we run on game thread
 		if (!IsInGameThread())
@@ -512,9 +486,11 @@ bool FAppleHttpNSUrlSessionRequest::ProcessRequest()
 		{
 			FinishedRequest();
 		}
+
+		return false;
 	}
 
-	return bStarted;
+	return true;
 }
 
 bool FAppleHttpNSUrlSessionRequest::StartRequest()
@@ -643,12 +619,6 @@ void FAppleHttpNSUrlSessionRequest::CancelRequest()
 	}
 }
 
-EHttpRequestStatus::Type FAppleHttpNSUrlSessionRequest::GetStatus() const
-{
-	UE_LOG(LogHttp, Verbose, TEXT("FAppleHttpNSUrlSessionRequest::GetStatus()"));
-	return CompletionStatus;
-}
-
 const FHttpResponsePtr FAppleHttpNSUrlSessionRequest::GetResponse() const
 {
 	return Response;
diff --git a/Engine/Source/Runtime/Online/HTTP/Private/Apple/AppleHTTPNSUrlSession.h b/Engine/Source/Runtime/Online/HTTP/Private/Apple/AppleHTTPNSUrlSession.h
--- a/Engine/Source/Runtime/Online/HTTP/Private/Apple/AppleHTTPNSUrlSession.h
+++ b/Engine/Source/Runtime/Online/HTTP/Private/Apple/AppleHTTPNSUrlSession.h
@@ -3,14 +3,14 @@
 #pragma once
 
 #include "CoreMinimal.h"
-#include "GenericPlatform/HttpRequestImpl.h"
+#include "GenericPlatform/HttpRequestCommon.h"
 #include "Interfaces/IHttpResponse.h"
 #include "PlatformHttp.h"
 
 /**
  * Apple implementation of an Http request
  */
-class FAppleHttpNSUrlSessionRequest : public FHttpRequestImpl
+class FAppleHttpNSUrlSessionRequest : public FHttpRequestCommon
 {
 public:
 	// implementation friends
@@ -44,7 +44,6 @@ public:
 	virtual TOptional<float> GetTimeout() const override;
 	virtual bool ProcessRequest() override;
 	virtual void CancelRequest() override;
-	virtual EHttpRequestStatus::Type GetStatus() const override;
 	virtual const FHttpResponsePtr GetResponse() const override;
 	virtual void Tick(float DeltaSeconds) override;
 	virtual float GetElapsedTime() const override;
@@ -106,9 +105,6 @@ private:
 	/** Array used to retrieve back content set on the ObjC request when calling GetContent*/
 	mutable TArray<uint8> StorageForGetContent;
 
-	/** Current status of request being processed */
-	EHttpRequestStatus::Type CompletionStatus;
-
 	/** Start of the request */
 	double StartRequestTime;
 
diff --git a/Engine/Source/Runtime/Online/HTTP/Private/Curl/CurlHttp.cpp b/Engine/Source/Runtime/Online/HTTP/Private/Curl/CurlHttp.cpp
--- a/Engine/Source/Runtime/Online/HTTP/Private/Curl/CurlHttp.cpp
+++ b/Engine/Source/Runtime/Online/HTTP/Private/Curl/CurlHttp.cpp
@@ -69,7 +69,6 @@ FCurlHttpRequest::FCurlHttpRequest()
 	,	bRedirected(false)
 	,	CurlAddToMultiResult(CURLM_OK)
 	,	CurlCompletionResult(CURLE_OK)
-	,	CompletionStatus(EHttpRequestStatus::NotStarted)
 	,	ElapsedTime(0.0f)
 	,	TimeSinceLastResponse(0.0f)
 	,	bAnyHttpActivity(false)
@@ -781,18 +780,6 @@ bool FCurlHttpRequest::SetupRequest()
 		UE_LOG(LogHttp, Verbose, TEXT("Http disabled. Skipping request. url=%s"), *GetURL());
 		return false;
 	}
-	// Prevent overlapped requests using the same instance
-	else if (CompletionStatus == EHttpRequestStatus::Processing)
-	{
-		UE_LOG(LogHttp, Warning, TEXT("ProcessRequest failed. Still processing last request."));
-		return false;
-	}
-	// Nothing to do without a valid URL
-	else if (URL.IsEmpty())
-	{
-		UE_LOG(LogHttp, Log, TEXT("Cannot process HTTP request: URL is empty"));
-		return false;
-	}
 	if ((GetVerb().IsEmpty() || GetVerb().Equals(TEXT("GET"), ESearchCase::IgnoreCase))
 		&& (RequestPayload.IsValid() && RequestPayload->GetContentLength() > 0))
 	{
@@ -1017,28 +1004,7 @@ bool FCurlHttpRequest::ProcessRequest()
 	Response = nullptr;
 	LastReportedBytesRead = 0;
 
-	bool bStarted = false;
-	if (!FHttpModule::Get().GetHttpManager().IsDomainAllowed(URL))
-	{
-		UE_LOG(LogHttp, Warning, TEXT("ProcessRequest failed. URL '%s' is not using an allowed domain. %p"), *URL, this);
-	}
-	else if (!SetupRequest())
-	{
-		UE_LOG(LogHttp, Warning, TEXT("Could not perform game thread setup, processing HTTP request failed. Increase verbosity for additional information."));
-	}
-	else
-	{
-		// Clear the info cache log so we don't output messages from previous requests when reusing/retrying a request
-		const FScopeLock CacheLock(&InfoMessageCacheCriticalSection);
-		for (FString& Line : InfoMessageCache)
-		{
-			Line.Reset();
-		}
-
-		bStarted = true;
-	}
-
-	if (!bStarted)
+	if (!PreCheck() || !SetupRequest())
 	{
 		if (!IsInGameThread())
 		{
@@ -1054,19 +1020,27 @@ bool FCurlHttpRequest::ProcessRequest()
 			// Cleanup and call delegate
 			FinishedRequest();
 		}
+
+		return false;
 	}
-	else
-	{
-		QUICK_SCOPE_CYCLE_COUNTER(STAT_CurlHttpAddThreadedRequest);
-		// Mark as in-flight to prevent overlapped requests using the same object
-		CompletionStatus = EHttpRequestStatus::Processing;
-		// Add to global list while being processed so that the ref counted request does not get deleted
-		FHttpModule::Get().GetHttpManager().AddThreadedRequest(SharedThis(this));
 
-		UE_LOG(LogHttp, Verbose, TEXT("%p: request (easy handle:%p) has been added to threaded queue for processing"), this, EasyHandle);
+	// Clear the info cache log so we don't output messages from previous requests when reusing/retrying a request
+	{
+		const FScopeLock CacheLock(&InfoMessageCacheCriticalSection);
+		for (FString& Line : InfoMessageCache)
+		{
+			Line.Reset();
+		}
 	}
 
-	return bStarted;
+	QUICK_SCOPE_CYCLE_COUNTER(STAT_CurlHttpAddThreadedRequest);
+	// Mark as in-flight to prevent overlapped requests using the same object
+	CompletionStatus = EHttpRequestStatus::Processing;
+	// Add to global list while being processed so that the ref counted request does not get deleted
+	FHttpModule::Get().GetHttpManager().AddThreadedRequest(SharedThis(this));
+
+	UE_LOG(LogHttp, Verbose, TEXT("%p: request (easy handle:%p) has been added to threaded queue for processing"), this, EasyHandle);
+	return true;
 }
 
 bool FCurlHttpRequest::StartThreadedRequest()
@@ -1156,11 +1130,6 @@ void FCurlHttpRequest::CancelRequest()
 	}
 }
 
-EHttpRequestStatus::Type FCurlHttpRequest::GetStatus() const
-{
-	return CompletionStatus;
-}
-
 const FHttpResponsePtr FCurlHttpRequest::GetResponse() const
 {
 	return Response;
diff --git a/Engine/Source/Runtime/Online/HTTP/Private/Curl/CurlHttp.h b/Engine/Source/Runtime/Online/HTTP/Private/Curl/CurlHttp.h
--- a/Engine/Source/Runtime/Online/HTTP/Private/Curl/CurlHttp.h
+++ b/Engine/Source/Runtime/Online/HTTP/Private/Curl/CurlHttp.h
@@ -176,7 +176,6 @@ public:
 	virtual void AppendToHeader(const FString& HeaderName, const FString& AdditionalHeaderValue) override;
 	virtual bool ProcessRequest() override;
 	virtual void CancelRequest() override;
-	virtual EHttpRequestStatus::Type GetStatus() const override;
 	virtual const FHttpResponsePtr GetResponse() const override;
 	virtual void Tick(float DeltaSeconds) override;
 	virtual float GetElapsedTime() const override;
@@ -399,8 +398,6 @@ private:
 	bool bIsRequestPayloadSeekable = false;
 	/** The stream to receive response body */
 	TSharedPtr<FArchive> ResponseBodyReceiveStream;
-	/** Current status of request being processed */
-	EHttpRequestStatus::Type CompletionStatus;
 	/** Mapping of header section to values. */
 	TMap<FString, FString> Headers;
 	/** Total elapsed time in seconds since the start of the request */
diff --git a/Engine/Source/Runtime/Online/HTTP/Private/GenericPlatform/HttpRequestCommon.cpp b/Engine/Source/Runtime/Online/HTTP/Private/GenericPlatform/HttpRequestCommon.cpp
new file mode 100644
--- /dev/null
+++ b/Engine/Source/Runtime/Online/HTTP/Private/GenericPlatform/HttpRequestCommon.cpp
@@ -0,0 +1,44 @@
+// Copyright Epic Games, Inc. All Rights Reserved.
+
+#include "GenericPlatform/HttpRequestCommon.h"
+#include "Http.h"
+#include "HttpManager.h"
+
+EHttpRequestStatus::Type FHttpRequestCommon::GetStatus() const
+{
+	return CompletionStatus;
+}
+
+bool FHttpRequestCommon::PreCheck() const
+{
+	// Prevent overlapped requests using the same instance
+	if (CompletionStatus == EHttpRequestStatus::Processing)
+	{
+		UE_LOG(LogHttp, Warning, TEXT("ProcessRequest failed. Still processing last request."));
+		return false;
+	}
+
+	// Nothing to do without a valid URL
+	if (GetURL().IsEmpty())
+	{
+		UE_LOG(LogHttp, Warning, TEXT("ProcessRequest failed. No URL was specified."));
+		return false;
+	}
+
+	// Make sure the URL is parsed correctly with a valid HTTP scheme
+	const FString URL = GetURL();
+	if (!URL.StartsWith(TEXT("http://")) && !URL.StartsWith(TEXT("https://")))
+	{
+		UE_LOG(LogHttp, Warning, TEXT("ProcessRequest failed. URL '%s' is not a valid HTTP request."), *GetURL());
+		return false;
+	}
+
+	if (!FHttpModule::Get().GetHttpManager().IsDomainAllowed(GetURL()))
+	{
+		UE_LOG(LogHttp, Warning, TEXT("ProcessRequest failed. URL '%s' is not using an allowed domain."), *GetURL());
+		return false;
+	}
+
+	return true;
+}
+
diff --git a/Engine/Source/Runtime/Online/HTTP/Private/IHttpThreadedRequest.h b/Engine/Source/Runtime/Online/HTTP/Private/IHttpThreadedRequest.h
--- a/Engine/Source/Runtime/Online/HTTP/Private/IHttpThreadedRequest.h
+++ b/Engine/Source/Runtime/Online/HTTP/Private/IHttpThreadedRequest.h
@@ -3,9 +3,9 @@
 #pragma once
 
 #include "CoreMinimal.h"
-#include "GenericPlatform/HttpRequestImpl.h"
+#include "GenericPlatform/HttpRequestCommon.h"
 
-class IHttpThreadedRequest : public FHttpRequestImpl
+class IHttpThreadedRequest : public FHttpRequestCommon
 {
 public:
 	// Called on http thread
diff --git a/Engine/Source/Runtime/Online/HTTP/Private/WinHttp/WinHttpHttpRequest.cpp b/Engine/Source/Runtime/Online/HTTP/Private/WinHttp/WinHttpHttpRequest.cpp
--- a/Engine/Source/Runtime/Online/HTTP/Private/WinHttp/WinHttpHttpRequest.cpp
+++ b/Engine/Source/Runtime/Online/HTTP/Private/WinHttp/WinHttpHttpRequest.cpp
@@ -337,11 +337,6 @@ void FWinHttpHttpRequest::CancelRequest()
 	}
 }
 
-EHttpRequestStatus::Type FWinHttpHttpRequest::GetStatus() const
-{
-	return CompletionStatus;
-}
-
 const FHttpResponsePtr FWinHttpHttpRequest::GetResponse() const
 {
 	return Response;
diff --git a/Engine/Source/Runtime/Online/HTTP/Private/WinHttp/WinHttpHttpRequest.h b/Engine/Source/Runtime/Online/HTTP/Private/WinHttp/WinHttpHttpRequest.h
--- a/Engine/Source/Runtime/Online/HTTP/Private/WinHttp/WinHttpHttpRequest.h
+++ b/Engine/Source/Runtime/Online/HTTP/Private/WinHttp/WinHttpHttpRequest.h
@@ -46,7 +46,6 @@ public:
 	virtual void AppendToHeader(const FString& HeaderName, const FString& AdditionalHeaderValue) override;
 	virtual bool ProcessRequest() override;
 	virtual void CancelRequest() override;
-	virtual EHttpRequestStatus::Type GetStatus() const override;
 	virtual const FHttpResponsePtr GetResponse() const override;
 	virtual void Tick(float DeltaSeconds) override;
 	virtual float GetElapsedTime() const override;
@@ -102,9 +101,6 @@ private:
 	/** Current status of request being processed */
 	EHttpRequestStatus::Type State = EHttpRequestStatus::NotStarted;
 
-	/** Status of request available via GetStatus */
-	EHttpRequestStatus::Type CompletionStatus = EHttpRequestStatus::NotStarted;
-
 	/** */
 	TSharedPtr<FWinHttpConnectionHttp, ESPMode::ThreadSafe> Connection;
 
diff --git a/Engine/Source/Runtime/Online/HTTP/Public/GenericPlatform/HttpRequestCommon.h b/Engine/Source/Runtime/Online/HTTP/Public/GenericPlatform/HttpRequestCommon.h
new file mode 100644
--- /dev/null
+++ b/Engine/Source/Runtime/Online/HTTP/Public/GenericPlatform/HttpRequestCommon.h
@@ -0,0 +1,25 @@
+// Copyright Epic Games, Inc. All Rights Reserved.
+
+#pragma once
+
+#include "GenericPlatform/HttpRequestImpl.h"
+
+/**
+ * Contains implementation of some common functions that don't vary between implementations of different platforms
+ */
+class HTTP_API FHttpRequestCommon : public FHttpRequestImpl
+{
+public:
+	// IHttpRequest
+	virtual EHttpRequestStatus::Type GetStatus() const override;
+
+protected:
+	/**
+	 * Check if this request is valid or allowed, before actually process the request
+	 */
+	bool PreCheck() const;
+
+protected:
+	/** Current status of request being processed */
+	EHttpRequestStatus::Type CompletionStatus = EHttpRequestStatus::NotStarted;
+};
diff --git a/Engine/Source/Runtime/Online/HTTP/Public/HttpManager.h b/Engine/Source/Runtime/Online/HTTP/Public/HttpManager.h
--- a/Engine/Source/Runtime/Online/HTTP/Public/HttpManager.h
+++ b/Engine/Source/Runtime/Online/HTTP/Public/HttpManager.h
@@ -218,6 +218,13 @@ public:
 	 */
 	void AddGameThreadTask(TFunction<void()>&& Task);
 
+	/**
+	 * Set url request filter through code, instead of setting it through config.
+	 *
+	 * @param InURLRequestFilter The request filter to set
+	 */
+	void SetURLRequestFilter(const UE::Core::FURLRequestFilter& InURLRequestFilter) { URLRequestFilter = InURLRequestFilter; }
+
 protected:
 	/** 
 	 * Create HTTP thread object
