Fix the Mac Editor compile error caused by typo in http change.

diff --git a/Engine/Source/Runtime/Online/HTTP/Private/Apple/AppleHTTPNSUrlConnection.h b/Engine/Source/Runtime/Online/HTTP/Private/Apple/AppleHTTPNSUrlConnection.h
--- a/Engine/Source/Runtime/Online/HTTP/Private/Apple/AppleHTTPNSUrlConnection.h
+++ b/Engine/Source/Runtime/Online/HTTP/Private/Apple/AppleHTTPNSUrlConnection.h
@@ -4,7 +4,7 @@
 
 #include "CoreMinimal.h"
 #include "Containers/UnrealString.h"
-#include "GenericPlatform/FHttpRequestCommon.h"
+#include "GenericPlatform/HttpRequestCommon.h"
 #include "Interfaces/IHttpResponse.h"
 #include "HttpManager.h"
 #include "PlatformHttp.h"
