Fix the compile error in MAC editor.

diff --git a/Engine/Source/Runtime/Online/HTTP/Private/Apple/AppleHTTPNSUrlSession.cpp b/Engine/Source/Runtime/Online/HTTP/Private/Apple/AppleHTTPNSUrlSession.cpp
--- a/Engine/Source/Runtime/Online/HTTP/Private/Apple/AppleHTTPNSUrlSession.cpp
+++ b/Engine/Source/Runtime/Online/HTTP/Private/Apple/AppleHTTPNSUrlSession.cpp
@@ -181,7 +181,6 @@ FAppleHttpNSUrlSessionRequest::FAppleHttpNSUrlSessionRequest(NSURLSession* InSes
 ,   Task(nil)
 ,	bIsPayloadFile(false)
 ,	ContentBytesLength(0)
-,	CompletionStatus(EHttpRequestStatus::NotStarted)
 ,	StartRequestTime(0.0)
 ,	ElapsedTime(0.0f)
 {
