Log deprecation message when using the UseNSUrlConnection command line argument to remove FAppleHttpNSUrlConnectionRequest in 5.4

diff --git a/Engine/Source/Runtime/Online/HTTP/Private/Apple/ApplePlatformHttp.cpp b/Engine/Source/Runtime/Online/HTTP/Private/Apple/ApplePlatformHttp.cpp
--- a/Engine/Source/Runtime/Online/HTTP/Private/Apple/ApplePlatformHttp.cpp
+++ b/Engine/Source/Runtime/Online/HTTP/Private/Apple/ApplePlatformHttp.cpp
@@ -223,6 +223,10 @@ void FApplePlatformHttp::Init()
 #endif
 
 	bool bUseNSUrlConnection = FParse::Param(FCommandLine::Get(), TEXT("UseNSUrlConnection"));
+	if (bUseNSUrlConnection)
+	{
+		UE_LOG(LogHttp, Warning, TEXT("UseNSUrlConnection command line argument is deprecated. It will be removed in UE 5.4"));
+	}
 	bUseNSUrlSession = !bUseNSUrlConnection;
 
 	if (bUseNSUrlSession)
