Fix http requests not being removed and completion delegate not being called from HttpManager for non threaded requests
Fix some issues on Apple http request implementations
 - To much progress notification
 - non finished requests triggered a check
 - delegates not being called

diff --git a/Engine/Source/Runtime/Online/HTTP/Private/Apple/AppleHTTPNSUrlConnection.cpp b/Engine/Source/Runtime/Online/HTTP/Private/Apple/AppleHTTPNSUrlConnection.cpp
--- a/Engine/Source/Runtime/Online/HTTP/Private/Apple/AppleHTTPNSUrlConnection.cpp
+++ b/Engine/Source/Runtime/Online/HTTP/Private/Apple/AppleHTTPNSUrlConnection.cpp
@@ -29,6 +29,8 @@ FAppleHttpNSUrlConnectionRequest::FAppleHttpNSUrlConnectionRequest()
 ,	ProgressBytesSent(0)
 ,	StartRequestTime(0.0)
 ,	ElapsedTime(0.0f)
+,   LastReportedBytesWritten(0)
+,   LastReportedBytesRead(0)
 {
 	UE_LOG(LogHttp, Verbose, TEXT("FAppleHttpNSUrlConnectionRequest::FAppleHttpNSUrlConnectionRequest()"));
 	Request = [[NSMutableURLRequest alloc] init];
@@ -54,7 +56,7 @@ FAppleHttpNSUrlConnectionRequest::FAppleHttpNSUrlConnectionRequest()
 FAppleHttpNSUrlConnectionRequest::~FAppleHttpNSUrlConnectionRequest()
 {
 	UE_LOG(LogHttp, Verbose, TEXT("FAppleHttpNSUrlConnectionRequest::~FAppleHttpNSUrlConnectionRequest()"));
-	check(Connection == nullptr);
+	DiscardExistingRequest();
 	[Request release];
 }
 
@@ -377,6 +379,11 @@ bool FAppleHttpNSUrlConnectionRequest::StartRequest()
 		[Request setValue:Tag forHTTPHeaderField:@"User-Agent"];
 	}
 
+	LastReportedBytesWritten = 0;
+	LastReportedBytesRead = 0;
+
+	DiscardExistingRequest();
+
 	Response = MakeShareable( new FAppleHttpNSUrlConnectionResponse( *this ) );
 
 #pragma clang diagnostic push
@@ -418,7 +425,6 @@ void FAppleHttpNSUrlConnectionRequest::FinishedRequest()
 
 		// TODO: Try to broadcast OnHeaderReceived when we receive headers instead of here at the end
 		BroadcastResponseHeadersReceived();
-		OnProcessRequestComplete().ExecuteIfBound(SharedThis(this), Response, true);
 	}
 	else
 	{
@@ -431,20 +437,28 @@ void FAppleHttpNSUrlConnectionRequest::FinishedRequest()
 		}
 
 		Response = nullptr;
-		OnProcessRequestComplete().ExecuteIfBound(SharedThis(this), nullptr, false);
 	}
 
 	// Clean up session/request handles that may have been created
 	CleanupRequest();
 
-	// Remove from global list since processing is now complete
-	if (FHttpModule::Get().GetHttpManager().IsValidRequest(this))
+	OnProcessRequestComplete().ExecuteIfBound(SharedThis(this), Response, Response != nullptr);
+}
+
+void FAppleHttpNSUrlConnectionRequest::DiscardExistingRequest()
+{
+	if (Connection != nil)
 	{
-		FHttpModule::Get().GetHttpManager().RemoveRequest(SharedThis(this));
+		if (CompletionStatus == EHttpRequestStatus::Processing)
+		{
+			[Connection cancel];
+		}
+		[Connection release];
+		Connection = nil;
 	}
+	Response = nullptr;
 }
 
-
 void FAppleHttpNSUrlConnectionRequest::CleanupRequest()
 {
 	UE_LOG(LogHttp, Verbose, TEXT("FAppleHttpNSUrlConnectionRequest::CleanupRequest()"));
@@ -498,9 +512,11 @@ void FAppleHttpNSUrlConnectionRequest::Tick(float DeltaSeconds)
 		{
 			const int32 BytesWritten = Response->GetNumBytesWritten();
 			const int32 BytesRead = Response->GetNumBytesReceived();
-			if (BytesWritten > 0 || BytesRead > 0)
+			if (BytesWritten != LastReportedBytesWritten || BytesRead != LastReportedBytesRead)
 			{
 				OnRequestProgress().ExecuteIfBound(SharedThis(this), BytesWritten, BytesRead);
+				LastReportedBytesWritten = BytesWritten;
+				LastReportedBytesRead = BytesRead;
 			}
 		}
 		if (Response->IsReady())
@@ -581,6 +597,7 @@ float FAppleHttpNSUrlConnectionRequest::GetElapsedTime() const
 	// Determine if the specific error was failing to connect to the host.
 	switch ([error code])
 	{
+		case NSURLErrorTimedOut:
 		case NSURLErrorCannotFindHost:
 		case NSURLErrorCannotConnectToHost:
 		case NSURLErrorDNSLookupFailed:
diff --git a/Engine/Source/Runtime/Online/HTTP/Private/Apple/AppleHTTPNSUrlConnection.h b/Engine/Source/Runtime/Online/HTTP/Private/Apple/AppleHTTPNSUrlConnection.h
--- a/Engine/Source/Runtime/Online/HTTP/Private/Apple/AppleHTTPNSUrlConnection.h
+++ b/Engine/Source/Runtime/Online/HTTP/Private/Apple/AppleHTTPNSUrlConnection.h
@@ -86,6 +86,10 @@ private:
 	 */
 	void CleanupRequest();
 
+	/**
+	 * Cleans up request without triggering additional callbacks
+	*/
+	void DiscardExistingRequest();
 
 private:
 	/** This is the NSMutableURLRequest, all our Apple functionality will deal with this. */
@@ -117,6 +121,12 @@ private:
 
 	/** Time taken to complete/cancel the request. */
 	float ElapsedTime;
+
+	/** Last reported bytes written */
+	int32 LastReportedBytesWritten;
+
+	/** Last reported bytesread */
+	int32 LastReportedBytesRead;
 };
 
 
diff --git a/Engine/Source/Runtime/Online/HTTP/Private/Apple/AppleHTTPNSUrlSession.cpp b/Engine/Source/Runtime/Online/HTTP/Private/Apple/AppleHTTPNSUrlSession.cpp
--- a/Engine/Source/Runtime/Online/HTTP/Private/Apple/AppleHTTPNSUrlSession.cpp
+++ b/Engine/Source/Runtime/Online/HTTP/Private/Apple/AppleHTTPNSUrlSession.cpp
@@ -124,6 +124,7 @@ enum class EAppleHttpRequestResponseState: uint8
 		// Determine if the specific error was failing to connect to the host.
 		switch ([error code])
 		{
+			case NSURLErrorTimedOut:
 			case NSURLErrorCannotFindHost:
 			case NSURLErrorCannotConnectToHost:
 			case NSURLErrorDNSLookupFailed:
@@ -183,6 +184,8 @@ FAppleHttpNSUrlSessionRequest::FAppleHttpNSUrlSessionRequest(NSURLSession* InSes
 ,	ContentBytesLength(0)
 ,	StartRequestTime(0.0)
 ,	ElapsedTime(0.0f)
+,   LastReportedBytesWritten(0)
+,   LastReportedBytesRead(0)
 {
 	UE_LOG(LogHttp, Verbose, TEXT("FAppleHttpNSUrlSessionRequest::FAppleHttpNSUrlSessionRequest()"));
 	Request = [[NSMutableURLRequest alloc] init];
@@ -202,7 +205,7 @@ FAppleHttpNSUrlSessionRequest::FAppleHttpNSUrlSessionRequest(NSURLSession* InSes
 FAppleHttpNSUrlSessionRequest::~FAppleHttpNSUrlSessionRequest()
 {
 	UE_LOG(LogHttp, Verbose, TEXT("FAppleHttpNSUrlSessionRequest::~FAppleHttpNSUrlSessionRequest()"));
-	check(Task == nil);
+	DiscardExistingRequest();
 	[Request release];
     [Session release];
 }
@@ -511,6 +514,11 @@ bool FAppleHttpNSUrlSessionRequest::StartRequest()
 		[Request setValue:Tag forHTTPHeaderField:@"User-Agent"];
 	}
 
+	LastReportedBytesWritten = 0;
+	LastReportedBytesRead = 0;
+
+	DiscardExistingRequest();
+
 	Task = [Session dataTaskWithRequest: Request];
 	
 	if (Task != nil)
@@ -553,7 +561,6 @@ void FAppleHttpNSUrlSessionRequest::FinishedRequest()
 
 		// TODO: Try to broadcast OnHeaderReceived when we receive headers instead of here at the end
 		BroadcastResponseHeadersReceived();
-		OnProcessRequestComplete().ExecuteIfBound(SharedThis(this), Response, true);
 	}
 	else
 	{
@@ -566,17 +573,26 @@ void FAppleHttpNSUrlSessionRequest::FinishedRequest()
 		}
 
 		Response = nullptr;
-		OnProcessRequestComplete().ExecuteIfBound(SharedThis(this), nullptr, false);
 	}
 
 	// Clean up session/request handles that may have been created
 	CleanupRequest();
 
-	// Remove from global list since processing is now complete
-	if (FHttpModule::Get().GetHttpManager().IsValidRequest(this))
+	OnProcessRequestComplete().ExecuteIfBound(SharedThis(this), Response, Response != nullptr);
+}
+
+void FAppleHttpNSUrlSessionRequest::DiscardExistingRequest()
+{
+	if (Task != nil)
 	{
-		FHttpModule::Get().GetHttpManager().RemoveRequest(SharedThis(this));
+		if (CompletionStatus == EHttpRequestStatus::Processing)
+		{
+			[Task cancel];
+		}
+		[Task release];
+		Task = nil;
 	}
+	Response = nullptr;
 }
 
 void FAppleHttpNSUrlSessionRequest::CleanupRequest()
@@ -631,9 +647,11 @@ void FAppleHttpNSUrlSessionRequest::Tick(float DeltaSeconds)
 		{
 			const int32 BytesWritten = Response->GetNumBytesWritten();
 			const int32 BytesRead = Response->GetNumBytesReceived();
-			if (BytesWritten > 0 || BytesRead > 0)
+			if (BytesWritten != LastReportedBytesWritten || BytesRead != LastReportedBytesRead)
 			{
 				OnRequestProgress().ExecuteIfBound(SharedThis(this), BytesWritten, BytesRead);
+				LastReportedBytesWritten = BytesWritten;
+				LastReportedBytesRead = BytesRead;
 			}
 		}
 		if (Response->IsReady())
diff --git a/Engine/Source/Runtime/Online/HTTP/Private/Apple/AppleHTTPNSUrlSession.h b/Engine/Source/Runtime/Online/HTTP/Private/Apple/AppleHTTPNSUrlSession.h
--- a/Engine/Source/Runtime/Online/HTTP/Private/Apple/AppleHTTPNSUrlSession.h
+++ b/Engine/Source/Runtime/Online/HTTP/Private/Apple/AppleHTTPNSUrlSession.h
@@ -82,6 +82,10 @@ private:
 	 */
 	void CleanupRequest();
 
+	/**
+	 * Cleans up request without triggering additional callbacks
+	*/
+	void DiscardExistingRequest();
 
 private:
 	/** This is the NSMutableURLRequest, all our Apple functionality will deal with this. */
@@ -110,6 +114,12 @@ private:
 
 	/** Time taken to complete/cancel the request. */
 	float ElapsedTime;
+
+	/** Last reported bytes written */
+	int32 LastReportedBytesWritten;
+
+	/** Last reported bytesread */
+	int32 LastReportedBytesRead;
 };
 
 @class FAppleHttpNSUrlSessionResponseDelegate;
diff --git a/Engine/Source/Runtime/Online/HTTP/Private/HttpManager.cpp b/Engine/Source/Runtime/Online/HTTP/Private/HttpManager.cpp
--- a/Engine/Source/Runtime/Online/HTTP/Private/HttpManager.cpp
+++ b/Engine/Source/Runtime/Online/HTTP/Private/HttpManager.cpp
@@ -410,17 +410,15 @@ bool FHttpManager::Tick(float DeltaSeconds)
 		Task();
 	}
 
-	FScopeLock ScopeLock(&RequestLock);
-
-	// Tick each active request
-	for (TArray<FHttpRequestRef>::TIterator It(Requests); It; ++It)
-	{
-		FHttpRequestRef Request = *It;
-		Request->Tick(DeltaSeconds);
-	}
-
 	if (Thread)
 	{
+		FScopeLock ScopeLock(&RequestLock);
+		// Tick each active request
+		for (const FHttpRequestRef& Request: Requests)
+		{
+			Request->Tick(DeltaSeconds);
+		}
+
 		TArray<IHttpThreadedRequest*> CompletedThreadedRequests;
 		Thread->GetCompletedRequests(CompletedThreadedRequests);
 
@@ -436,6 +434,27 @@ bool FHttpManager::Tick(float DeltaSeconds)
 			BroadcastHttpRequestCompleted(CompletedRequestRef);
 		}
 	}
+	else
+	{
+		TArray<FHttpRequestRef> CompletedRequests;
+		
+		FScopeLock ScopeLock(&RequestLock);
+		// Tick each active request
+		for (const FHttpRequestRef& Request: Requests)
+		{
+			Request->Tick(DeltaSeconds);
+			if (EHttpRequestStatus::IsFinished(Request->GetStatus()))
+			{
+				CompletedRequests.Add(Request);
+			}
+		}
+
+		for (const FHttpRequestRef& CompletedRequest: CompletedRequests)
+		{
+            Requests.Remove(CompletedRequest);
+            BroadcastHttpRequestCompleted(CompletedRequest);
+        }		
+	}
 	// keep ticking
 	return true;
 }
diff --git a/Engine/Source/Runtime/Online/HTTP/Private/NullHttp.cpp b/Engine/Source/Runtime/Online/HTTP/Private/NullHttp.cpp
--- a/Engine/Source/Runtime/Online/HTTP/Private/NullHttp.cpp
+++ b/Engine/Source/Runtime/Online/HTTP/Private/NullHttp.cpp
@@ -184,15 +184,11 @@ float FNullHttpRequest::GetElapsedTime() const
 void FNullHttpRequest::FinishedRequest()
 {
 	CompletionStatus = EHttpRequestStatus::Failed;
-	FHttpRequestRef Request = SharedThis(this);
-	FHttpManager& HttpManager = FHttpModule::Get().GetHttpManager();
-	HttpManager.BroadcastHttpRequestCompleted(Request);
-	HttpManager.RemoveRequest(Request);
 
 	UE_LOG(LogHttp, Log, TEXT("Finished request %p. no response %s url=%s elapsed=%.3f"),
 		this, *GetVerb(), *GetURL(), ElapsedTime);
 
-	OnProcessRequestComplete().ExecuteIfBound(Request, NULL, false);
+	OnProcessRequestComplete().ExecuteIfBound(SharedThis(this), NULL, false);
 }
 
 void FNullHttpRequest::SetDelegateThreadPolicy(EHttpRequestDelegateThreadPolicy InThreadPolicy)
