- Reduce number of heap [re]allocations made by FGenericPlatformHttp::UrlEncode

diff --git a/Engine/Source/Runtime/Online/HTTP/Private/GenericPlatform/GenericPlatformHttp.cpp b/Engine/Source/Runtime/Online/HTTP/Private/GenericPlatform/GenericPlatformHttp.cpp
--- a/Engine/Source/Runtime/Online/HTTP/Private/GenericPlatform/GenericPlatformHttp.cpp
+++ b/Engine/Source/Runtime/Online/HTTP/Private/GenericPlatform/GenericPlatformHttp.cpp
@@ -5,6 +5,7 @@
 #include "Misc/Paths.h"
 #include "Misc/EngineVersion.h"
 #include "Misc/App.h"
+#include "String/BytesToHex.h"
 
 namespace
 {
@@ -154,28 +155,37 @@ static bool IsAllowedChar(UTF8CHAR LookupChar)
 
 FString FGenericPlatformHttp::UrlEncode(const FStringView UnencodedString)
 {
-	FTCHARToUTF8 Converter(UnencodedString.GetData(), UnencodedString.Len());	//url encoding must be encoded over each utf-8 byte
-	const UTF8CHAR* UTF8Data = (UTF8CHAR*) Converter.Get();	//converter uses ANSI instead of UTF8CHAR - not sure why - but other code seems to just do this cast. In this case it really doesn't matter
-	FString EncodedString = TEXT("");
-	
-	TCHAR Buffer[2] = { 0, 0 };
+	FString EncodedString;
+	EncodedString.Reserve(UnencodedString.Len()); // This is a minimum bound. Some characters might be outside the ascii set and require %hh encoding. Some characters might be multi-byte and require several %hh%hh%hh
 
-	for (int32 ByteIdx = 0, Length = Converter.Length(); ByteIdx < Length; ++ByteIdx)
+	UTF8CHAR Utf8ConvertedChar[4] = {};
+	TCHAR HexChars[3] = { TCHAR('%') };
+
+	for (const TCHAR& InChar : UnencodedString)
 	{
-		UTF8CHAR ByteToEncode = UTF8Data[ByteIdx];
-		
-		if (IsAllowedChar(ByteToEncode))
+		verify(FPlatformString::Convert(Utf8ConvertedChar, sizeof(Utf8ConvertedChar), &InChar, 1));
+		for (int32 ByteIdx = 0; ByteIdx < sizeof(Utf8ConvertedChar); ++ByteIdx)
 		{
-			Buffer[0] = ByteToEncode;
-			FString TmpString = Buffer;
-			EncodedString += TmpString;
-		}
-		else if (ByteToEncode != '\0')
-		{
-			EncodedString += TEXT("%");
-			EncodedString += FString::Printf(TEXT("%.2X"), ByteToEncode);
+			UTF8CHAR ByteToEncode = Utf8ConvertedChar[ByteIdx];
+			Utf8ConvertedChar[ByteIdx] = UTF8CHAR('\0');
+			if (ByteToEncode == '\0')
+			{
+				break;
+			}
+			else if (IsAllowedChar(ByteToEncode))
+			{
+				// We use InChar here as it is the same value ByteToEncode would convert back to anyway
+				// Note this relies on the fact that IsAllowedChar is only possible to be true for single-byte UTF8 characters.
+				EncodedString.AppendChar(InChar);
+			}
+			else
+			{
+				UE::String::BytesToHex(MakeArrayView((uint8*)&ByteToEncode, 1), &HexChars[1]);
+				EncodedString.AppendChars(HexChars, 3);
+			}
 		}
 	}
+
 	return EncodedString;
 }
 
