Implement streaming interface in AppleHTTPNSUrlSession

diff --git a/Engine/Source/Runtime/Online/HTTP/Private/Apple/AppleHTTPNSUrlSession.cpp b/Engine/Source/Runtime/Online/HTTP/Private/Apple/AppleHTTPNSUrlSession.cpp
--- a/Engine/Source/Runtime/Online/HTTP/Private/Apple/AppleHTTPNSUrlSession.cpp
+++ b/Engine/Source/Runtime/Online/HTTP/Private/Apple/AppleHTTPNSUrlSession.cpp
@@ -29,7 +29,19 @@ enum class EAppleHttpRequestResponseState: uint8
 {
 	/** Holds the payload as we receive it. */
 	@public TArray<uint8> Payload;
+	
+	/** This stream is shared between the request and the delegate. This field is accessed through a thread out of out control 
+	 * from the delegates methods and through mehthods in FAppleHttpNSUrlSessionRequest. Once we cancel or have an error this can 
+	 * is nullified as soon as possible to avoid receiving more data after completion delegates were triggered */
+	TSharedPtr<FArchive> ResponseBodyReceiveStream;
+	
+	/** critical section to properly clear stream */
+	FCriticalSection ResponseStreamLock;
+
+	/** flag meant to reduce locking on ResponseStreamLock*/
+	@public BOOL bInitializedWithValidStream;
 }
+
 /** A handle for the response */
 @property(retain) NSHTTPURLResponse* Response;
 /** The total number of bytes written out during the request/response */
@@ -59,17 +71,28 @@ enum class EAppleHttpRequestResponseState: uint8
 @synthesize BytesWritten;
 @synthesize BytesReceived;
 
--(FAppleHttpNSUrlSessionResponseDelegate*) init
+-(FAppleHttpNSUrlSessionResponseDelegate*) initWithResponseStream:(TSharedPtr<FArchive>)ResponseStream
 {
 	self = [super init];
 	
 	BytesWritten = 0;
 	BytesReceived = 0;
 	ResponseState = EAppleHttpRequestResponseState::NotReady;
+	ResponseBodyReceiveStream = ResponseStream;
+	bInitializedWithValidStream = (ResponseStream != nullptr);
 	
 	return self;
 }
 
+- (void) ClearResponseStream
+{
+	if (bInitializedWithValidStream)
+	{
+	    FScopeLock Lock(&ResponseStreamLock);
+		ResponseBodyReceiveStream = nullptr;
+	}
+}
+
 - (void) dealloc
 {
 	[Response release];
@@ -88,23 +111,47 @@ enum class EAppleHttpRequestResponseState: uint8
 	
 	self.Response = (NSHTTPURLResponse*)response;
 	int32 ExpectedResponseLength = (int32)response.expectedContentLength;
-	if(ExpectedResponseLength != NSURLResponseUnknownLength)
+	if(!bInitializedWithValidStream && ExpectedResponseLength != NSURLResponseUnknownLength)
 	{
 		Payload.Empty(ExpectedResponseLength);
 	}
-	UE_LOG(LogHttp, Verbose, TEXT("URLSession:dataTask:didReceiveResponse:completionHandler: expectedContentLength = %lld. Length = %d: %p"), response.expectedContentLength, Payload.Max(), self);
+	UE_LOG(LogHttp, Verbose, TEXT("URLSession:dataTask:didReceiveResponse:completionHandler: expectedContentLength = %lld. Length = %d: %p"), ExpectedResponseLength, Payload.Max(), self);
 	completionHandler(NSURLSessionResponseAllow);
 }
 
 - (void)URLSession:(NSURLSession *)session dataTask:(NSURLSessionDataTask *)dataTask didReceiveData:(NSData *)data
 {
-	[data enumerateByteRangesUsingBlock:^(const void *bytes, NSRange byteRange, BOOL *stop) {
-		Payload.Append((const uint8*)bytes, byteRange.length);
-	}];
+	__block int64 NewBytesReceived = 0;
+	if (bInitializedWithValidStream)
+	{
+		if (FScopeLock Lock(&ResponseStreamLock); ResponseBodyReceiveStream)
+		{
+			__block bool bHadError = false;
+			[data enumerateByteRangesUsingBlock:^(const void *bytes, NSRange byteRange, BOOL *stop) {
+				NewBytesReceived += byteRange.length;
+				ResponseBodyReceiveStream->Serialize(const_cast<void*>(bytes), byteRange.length);
+				bHadError = ResponseBodyReceiveStream->GetError();
+				*stop = bHadError? YES : NO;
+			}];
+			
+			if (bHadError)
+			{
+				[dataTask cancel];
+				ResponseBodyReceiveStream = nullptr;
+			}
+		}
+	}
+	else
+	{
+		[data enumerateByteRangesUsingBlock:^(const void *bytes, NSRange byteRange, BOOL *stop) {
+			NewBytesReceived += byteRange.length;
+			Payload.Append((const uint8*)bytes, byteRange.length);
+		}];
+	}
 	// Keep BytesReceived as a separated value to avoid concurrent accesses to Payload
-	self.BytesReceived = Payload.Num();
+	self.BytesReceived += NewBytesReceived;
 	
-	UE_LOG(LogHttp, Verbose, TEXT("URLSession:dataTask:didReceiveData with %u bytes. After Append, Payload Length = %d: %p"), [data length], Payload.Num(), self);
+	UE_LOG(LogHttp, Verbose, TEXT("URLSession:dataTask:didReceiveData with %d bytes. After Append, Payload Length = %d: %p"), NewBytesReceived, self.BytesReceived, self);
 }
 
 - (void)URLSession:(NSURLSession *)session task:(NSURLSessionTask *)task didCompleteWithError:(nullable NSError *)error
@@ -182,6 +229,7 @@ FAppleHttpNSUrlSessionRequest::FAppleHttpNSUrlSessionRequest(NSURLSession* InSes
 ,   Task(nil)
 ,	bIsPayloadFile(false)
 ,	ContentBytesLength(0)
+,   bCanceled(false)
 ,	StartRequestTime(0.0)
 ,	ElapsedTime(0.0f)
 ,   LastReportedBytesWritten(0)
@@ -205,7 +253,7 @@ FAppleHttpNSUrlSessionRequest::FAppleHttpNSUrlSessionRequest(NSURLSession* InSes
 FAppleHttpNSUrlSessionRequest::~FAppleHttpNSUrlSessionRequest()
 {
 	UE_LOG(LogHttp, Verbose, TEXT("FAppleHttpNSUrlSessionRequest::~FAppleHttpNSUrlSessionRequest()"));
-	DiscardExistingRequest();
+	CleanupRequest();
 	[Request release];
     [Session release];
 }
@@ -436,8 +484,8 @@ bool FAppleHttpNSUrlSessionRequest::SetContentFromStream(TSharedRef<FArchive, ES
 
 bool FAppleHttpNSUrlSessionRequest::SetResponseBodyReceiveStream(TSharedRef<FArchive> Stream)
 {
-	UE_LOG(LogHttp, Warning, TEXT("FAppleHttpNSUrlSessionRequest::SetResponseBodyReceiveStream is not implemented"));
-	return false;
+	ResponseBodyReceiveStream = Stream;
+	return true;
 }
 
 FString FAppleHttpNSUrlSessionRequest::GetVerb() const
@@ -514,10 +562,12 @@ bool FAppleHttpNSUrlSessionRequest::StartRequest()
 		[Request setValue:Tag forHTTPHeaderField:@"User-Agent"];
 	}
 
+	CleanupRequest();
+
 	LastReportedBytesWritten = 0;
 	LastReportedBytesRead = 0;
-
-	DiscardExistingRequest();
+	bCanceled = false;
+	Response = nullptr;	
 
 	Task = [Session dataTaskWithRequest: Request];
 	
@@ -554,8 +604,14 @@ void FAppleHttpNSUrlSessionRequest::FinishedRequest()
 {
 	UE_LOG(LogHttp, Verbose, TEXT("FAppleHttpNSUrlSessionRequest::FinishedRequest()"));
 	ElapsedTime = (float)(FPlatformTime::Seconds() - StartRequestTime);
-	if( Response.IsValid() && Response->IsReady() && !Response->HadError())
+	
+	// Clean up session/request handles that may have been created
+	CleanupRequest();
+
+	bool bSuccess = false;
+	if (!bCanceled && Response.IsValid() && Response->IsReady() && !Response->HadError())
 	{
+		bSuccess = true;
 		UE_LOG(LogHttp, Verbose, TEXT("Request succeeded"));
 		CompletionStatus = EHttpRequestStatus::Succeeded;
 
@@ -570,42 +626,28 @@ void FAppleHttpNSUrlSessionRequest::FinishedRequest()
 		if (Response.IsValid() && Response->HadConnectionError())
 		{
 			CompletionStatus = EHttpRequestStatus::Failed_ConnectionError;
+			Response = nullptr;
 		}
-
-		Response = nullptr;
 	}
 
-	// Clean up session/request handles that may have been created
-	CleanupRequest();
-
-	OnProcessRequestComplete().ExecuteIfBound(SharedThis(this), Response, Response != nullptr);
-}
-
-void FAppleHttpNSUrlSessionRequest::DiscardExistingRequest()
-{
-	if (Task != nil)
-	{
-		if (CompletionStatus == EHttpRequestStatus::Processing)
-		{
-			[Task cancel];
-		}
-		[Task release];
-		Task = nil;
-	}
-	Response = nullptr;
+	OnProcessRequestComplete().ExecuteIfBound(SharedThis(this), Response, bSuccess);
 }
 
 void FAppleHttpNSUrlSessionRequest::CleanupRequest()
 {
 	UE_LOG(LogHttp, Verbose, TEXT("FAppleHttpNSUrlSessionRequest::CleanupRequest()"));
 	
-	if(CompletionStatus == EHttpRequestStatus::Processing)
+	if (Response != nullptr)
 	{
-		CancelRequest();
+		Response->CleanSharedObjects();
 	}
-	
 	if(Task != nil)
 	{
+		if(CompletionStatus == EHttpRequestStatus::Processing)
+		{
+			[Task cancel];
+		}
+		
 		[Task release];
 		Task = nil;
 	}
@@ -617,20 +659,31 @@ void FAppleHttpNSUrlSessionRequest::CancelRequest()
 
 	if( Task != nil )
 	{
+		if (bCanceled)
+		{
+			return;
+		}
+		
+		bCanceled = true;		
 		[Task cancel];
-	}
 
-	// Ensure we run on game thread
-	if (!IsInGameThread())
-	{
-		FHttpModule::Get().GetHttpManager().AddGameThreadTask([StrongThis = StaticCastSharedRef<FAppleHttpNSUrlSessionRequest>(AsShared())]()
+		if (Response != nullptr)
 		{
-			StrongThis->FinishedRequest();
-		});
-	}
-	else
-	{
-		FinishedRequest();
+			Response->CleanSharedObjects();
+		}
+
+		// Ensure we run on game thread
+		if (!IsInGameThread())
+		{
+			FHttpModule::Get().GetHttpManager().AddGameThreadTask([StrongThis = StaticCastSharedRef<FAppleHttpNSUrlSessionRequest>(AsShared())]()
+			{
+				StrongThis->FinishedRequest();
+			});
+		}
+		else
+		{
+			FinishedRequest();
+		}
 	}
 }
 
@@ -674,7 +727,7 @@ FAppleHttpNSUrlSessionResponse::FAppleHttpNSUrlSessionResponse(const FAppleHttpN
 	: Request( InRequest )
 {
 	UE_LOG(LogHttp, Verbose, TEXT("FAppleHttpNSUrlSessionResponse::FAppleHttpNSUrlSessionResponse()"));
-	ResponseDelegate = [[FAppleHttpNSUrlSessionResponseDelegate alloc] init];
+	ResponseDelegate = [[FAppleHttpNSUrlSessionResponseDelegate alloc] initWithResponseStream: InRequest.ResponseBodyReceiveStream];
 }
 
 FAppleHttpNSUrlSessionResponse::~FAppleHttpNSUrlSessionResponse()
@@ -685,6 +738,11 @@ FAppleHttpNSUrlSessionResponse::~FAppleHttpNSUrlSessionResponse()
 	ResponseDelegate = nil;
 }
 
+void FAppleHttpNSUrlSessionResponse::CleanSharedObjects()
+{
+	[ResponseDelegate ClearResponseStream];
+}
+
 FString FAppleHttpNSUrlSessionResponse::GetURL() const
 {
 	UE_LOG(LogHttp, Verbose, TEXT("FAppleHttpNSUrlSessionResponse::GetURL()"));
@@ -752,15 +810,13 @@ const TArray<uint8>& FAppleHttpNSUrlSessionResponse::GetContent() const
 {
 	if( !IsReady() )
 	{
-		const static TArray<uint8> EmptyPayload;
 		UE_LOG(LogHttp, Warning, TEXT("Payload is incomplete. Response still processing. %p"), &Request);
-		return EmptyPayload;
 	}
 	else
 	{
 		UE_LOG(LogHttp, Verbose, TEXT("FAppleHttpNSUrlSessionResponse::GetContent() - Num: %i"), ResponseDelegate->Payload.Num());
-		return ResponseDelegate->Payload;
 	}
+	return ResponseDelegate->Payload;
 }
 
 FString FAppleHttpNSUrlSessionResponse::GetContentAsString() const
diff --git a/Engine/Source/Runtime/Online/HTTP/Private/Apple/AppleHTTPNSUrlSession.h b/Engine/Source/Runtime/Online/HTTP/Private/Apple/AppleHTTPNSUrlSession.h
--- a/Engine/Source/Runtime/Online/HTTP/Private/Apple/AppleHTTPNSUrlSession.h
+++ b/Engine/Source/Runtime/Online/HTTP/Private/Apple/AppleHTTPNSUrlSession.h
@@ -82,11 +82,6 @@ private:
 	 */
 	void CleanupRequest();
 
-	/**
-	 * Cleans up request without triggering additional callbacks
-	*/
-	void DiscardExistingRequest();
-
 private:
 	/** This is the NSMutableURLRequest, all our Apple functionality will deal with this. */
 	NSMutableURLRequest* Request;
@@ -109,6 +104,12 @@ private:
 	/** Array used to retrieve back content set on the ObjC request when calling GetContent*/
 	mutable TArray<uint8> StorageForGetContent;
 
+	/** The stream to receive response body */
+	TSharedPtr<FArchive> ResponseBodyReceiveStream;
+
+	/** flag to mark wether we tried canceled this request */
+	bool bCanceled;
+
 	/** Start of the request */
 	double StartRequestTime;
 
@@ -187,6 +188,10 @@ public:
 	*/
 	const int32 GetNumBytesWritten() const;
 
+	/**
+	 * Cleans internal shared objects between request and response
+	 */
+	void CleanSharedObjects();
 	/**
 	 * Constructor
 	 *
