Make AppleHTTPNSUrlSession threaded

diff --git a/Engine/Source/Programs/WebTests/Private/TestHttp.cpp b/Engine/Source/Programs/WebTests/Private/TestHttp.cpp
--- a/Engine/Source/Programs/WebTests/Private/TestHttp.cpp
+++ b/Engine/Source/Programs/WebTests/Private/TestHttp.cpp
@@ -205,6 +205,7 @@ TEST_CASE_METHOD(FWaitUntilCompleteHttpFixture, "Http Methods", HTTP_TAG)
 
 	HttpRequest->OnProcessRequestComplete().BindLambda([](FHttpRequestPtr HttpRequest, FHttpResponsePtr HttpResponse, bool bSucceeded) {
 		CHECK(bSucceeded);
+		REQUIRE(HttpResponse != nullptr);
 		CHECK(HttpResponse->GetResponseCode() == 200);
 	});
 	HttpRequest->ProcessRequest();
@@ -217,6 +218,7 @@ TEST_CASE_METHOD(FWaitUntilCompleteHttpFixture, "Get large response content with
 	HttpRequest->SetVerb(TEXT("GET"));
 	HttpRequest->OnProcessRequestComplete().BindLambda([](FHttpRequestPtr HttpRequest, FHttpResponsePtr HttpResponse, bool bSucceeded) {
 		CHECK(bSucceeded);
+		REQUIRE(HttpResponse != nullptr);
 		CHECK(HttpResponse->GetResponseCode() == 200);
 	});
 	HttpRequest->ProcessRequest();
@@ -256,9 +258,10 @@ TEST_CASE_METHOD(FWaitUntilCompleteHttpFixture, "Streaming http download", HTTP_
 	SECTION("Success without stream provided")
 	{
 		HttpRequest->OnProcessRequestComplete().BindLambda([Chunks, ChunkSize](FHttpRequestPtr HttpRequest, FHttpResponsePtr HttpResponse, bool bSucceeded) {
-			CHECK(HttpResponse->GetContentLength() == Chunks * ChunkSize);
 			CHECK(bSucceeded);
+			REQUIRE(HttpResponse != nullptr);
 			CHECK(HttpResponse->GetResponseCode() == 200);
+			CHECK(HttpResponse->GetContentLength() == Chunks * ChunkSize);
 		});
 	}
 	SECTION("Success with customized stream")
@@ -283,11 +286,12 @@ TEST_CASE_METHOD(FWaitUntilCompleteHttpFixture, "Streaming http download", HTTP_
 		CHECK(HttpRequest->SetResponseBodyReceiveStream(Stream));
 
 		HttpRequest->OnProcessRequestComplete().BindLambda([Chunks, ChunkSize, TotalBytesReceived](FHttpRequestPtr HttpRequest, FHttpResponsePtr HttpResponse, bool bSucceeded) {
+			CHECK(bSucceeded);
+			REQUIRE(HttpResponse != nullptr);
+			CHECK(HttpResponse->GetResponseCode() == 200);
 			CHECK(HttpResponse->GetContentLength() == Chunks * ChunkSize);
 			CHECK(HttpResponse->GetContent().IsEmpty());
 			CHECK(*TotalBytesReceived == Chunks * ChunkSize);
-			CHECK(bSucceeded);
-			CHECK(HttpResponse->GetResponseCode() == 200);
 		});
 	}
 	SECTION("Success with customized stream delegate")
@@ -300,11 +304,12 @@ TEST_CASE_METHOD(FWaitUntilCompleteHttpFixture, "Streaming http download", HTTP_
 		CHECK(HttpRequest->SetResponseBodyReceiveStreamDelegate(Delegate));
 
 		HttpRequest->OnProcessRequestComplete().BindLambda([Chunks, ChunkSize, TotalBytesReceived](FHttpRequestPtr HttpRequest, FHttpResponsePtr HttpResponse, bool bSucceeded) {
+			CHECK(bSucceeded);
+			REQUIRE(HttpResponse != nullptr);
+			CHECK(HttpResponse->GetResponseCode() == 200);
 			CHECK(HttpResponse->GetContentLength() == Chunks * ChunkSize);
 			CHECK(HttpResponse->GetContent().IsEmpty());
 			CHECK(*TotalBytesReceived == Chunks * ChunkSize);
-			CHECK(bSucceeded);
-			CHECK(HttpResponse->GetResponseCode() == 200);
 		});
 	}
 	SECTION("Use customized stream to receive response body but failed when serialize")
@@ -359,9 +364,10 @@ TEST_CASE_METHOD(FWaitUntilCompleteHttpFixture, "Streaming http download", HTTP_
 		CHECK(HttpRequest->SetResponseBodyReceiveStream(FileToWrite));
 
 		HttpRequest->OnProcessRequestComplete().BindLambda([Chunks, ChunkSize, Filename, FileToWrite](FHttpRequestPtr HttpRequest, FHttpResponsePtr HttpResponse, bool bSucceeded) {
+			CHECK(bSucceeded);
+			REQUIRE(HttpResponse != nullptr);
 			CHECK(HttpResponse->GetContentLength() == Chunks * ChunkSize);
 			CHECK(HttpResponse->GetContent().IsEmpty());
-			CHECK(bSucceeded);
 			CHECK(HttpResponse->GetResponseCode() == 200);
 
 			FileToWrite->FlushCache();
@@ -433,6 +439,7 @@ TEST_CASE_METHOD(FWaitUntilCompleteHttpFixture, "Streaming http upload - gold pa
 
 	HttpRequest->OnProcessRequestComplete().BindLambda([](FHttpRequestPtr HttpRequest, FHttpResponsePtr HttpResponse, bool bSucceeded) {
 		CHECK(bSucceeded);
+		REQUIRE(HttpResponse != nullptr);
 		CHECK(HttpResponse->GetResponseCode() == 200);
 	});
 	HttpRequest->ProcessRequest();
@@ -574,4 +581,23 @@ TEST_CASE_METHOD(FWaitThreadedHttpFixture, "Threaded http request pre check will
 	StartTestHttpThread();
 }
 
+
+TEST_CASE_METHOD(FWaitUntilCompleteHttpFixture, "Cancel http request connect before timeout", HTTP_TAG)
+{
+	TSharedRef<IHttpRequest, ESPMode::ThreadSafe> HttpRequest = HttpModule->CreateRequest();
+	HttpRequest->SetURL(UrlWithInvalidPortToTestConnectTimeout());
+	HttpRequest->SetVerb(TEXT("GET"));
+	HttpRequest->SetTimeout(7);
+	FDateTime StartTime = FDateTime::Now();
+	HttpRequest->OnProcessRequestComplete().BindLambda([StartTime](FHttpRequestPtr HttpRequest, FHttpResponsePtr HttpResponse, bool bSucceeded) {
+		CHECK(!bSucceeded);
+
+		FTimespan Timespan = FDateTime::Now() - StartTime;
+		float DurationInSeconds = Timespan.GetTotalSeconds();
+		CHECK(DurationInSeconds < 2);
+	});
+	HttpRequest->ProcessRequest();
+	FPlatformProcess::Sleep(0.5);
+	HttpRequest->CancelRequest();
+}
 // TODO: Add cancel test, with multiple cancel calls
diff --git a/Engine/Source/Runtime/Online/HTTP/Private/Apple/AppleHTTPNSUrlSession.cpp b/Engine/Source/Runtime/Online/HTTP/Private/Apple/AppleHTTPNSUrlSession.cpp
--- a/Engine/Source/Runtime/Online/HTTP/Private/Apple/AppleHTTPNSUrlSession.cpp
+++ b/Engine/Source/Runtime/Online/HTTP/Private/Apple/AppleHTTPNSUrlSession.cpp
@@ -230,10 +230,9 @@ FAppleHttpNSUrlSessionRequest::FAppleHttpNSUrlSessionRequest(NSURLSession* InSes
 ,	bIsPayloadFile(false)
 ,	ContentBytesLength(0)
 ,   bCanceled(false)
-,	StartRequestTime(0.0)
 ,	ElapsedTime(0.0f)
-,   LastReportedBytesWritten(0)
-,   LastReportedBytesRead(0)
+,	LastReportedBytesWritten(0)
+,	LastReportedBytesRead(0)
 {
 	UE_LOG(LogHttp, Verbose, TEXT("FAppleHttpNSUrlSessionRequest::FAppleHttpNSUrlSessionRequest()"));
 	Request = [[NSMutableURLRequest alloc] init];
@@ -524,19 +523,7 @@ bool FAppleHttpNSUrlSessionRequest::ProcessRequest()
 
 	if (!PreCheck() || !StartRequest())
 	{
-		// Ensure we run on game thread
-		if (!IsInGameThread())
-		{
-			FHttpModule::Get().GetHttpManager().AddGameThreadTask([StrongThis = StaticCastSharedRef<FAppleHttpNSUrlSessionRequest>(AsShared())]()
-			{
-				StrongThis->FinishedRequest();
-			});
-		}
-		else
-		{
-			FinishedRequest();
-		}
-
+		FinishRequestNotInHttpManager();
 		return false;
 	}
 
@@ -549,7 +536,7 @@ bool FAppleHttpNSUrlSessionRequest::StartRequest()
 	UE_LOG(LogHttp, Verbose, TEXT("FAppleHttpNSUrlSessionRequest::StartRequest()"));
 	bool bStarted = false;
 
-	// set the content-length and user-agent
+	// set the content-length and user-agent (it is possible that the OS ignores this value)
 	if(GetContentLength() > 0)
 	{
 		[Request setValue:[NSString stringWithFormat:@"%d", GetContentLength()] forHTTPHeaderField:@"Content-Length"];
@@ -566,8 +553,9 @@ bool FAppleHttpNSUrlSessionRequest::StartRequest()
 
 	LastReportedBytesWritten = 0;
 	LastReportedBytesRead = 0;
+	ElapsedTime = 0.0f;
 	bCanceled = false;
-	Response = nullptr;	
+	Response = nullptr;
 
 	Task = [Session dataTaskWithRequest: Request];
 	
@@ -585,7 +573,7 @@ bool FAppleHttpNSUrlSessionRequest::StartRequest()
 		[[Task retain] resume];
 		UE_LOG(LogHttp, Verbose, TEXT("[NSURLSessionTask resume]"));
 
-		FHttpModule::Get().GetHttpManager().AddRequest(SharedThis(this));
+		FHttpModule::Get().GetHttpManager().AddThreadedRequest(SharedThis(this));
 	}
 	else
 	{
@@ -593,18 +581,15 @@ bool FAppleHttpNSUrlSessionRequest::StartRequest()
 		CompletionStatus = EHttpRequestStatus::Failed_ConnectionError;
 	}
 
-	StartRequestTime = FPlatformTime::Seconds();
-	// reset the elapsed time.
-	ElapsedTime = 0.0f;
-
 	return bStarted;
 }
 
-void FAppleHttpNSUrlSessionRequest::FinishedRequest()
+void FAppleHttpNSUrlSessionRequest::FinishRequest()
 {
-	UE_LOG(LogHttp, Verbose, TEXT("FAppleHttpNSUrlSessionRequest::FinishedRequest()"));
-	ElapsedTime = (float)(FPlatformTime::Seconds() - StartRequestTime);
-	
+	check(IsInGameThread() || DelegateThreadPolicy == EHttpRequestDelegateThreadPolicy::CompleteOnHttpThread);
+
+	UE_LOG(LogHttp, Verbose, TEXT("FAppleHttpNSUrlSessionRequest::FinishRequest()"));
+
 	// Clean up session/request handles that may have been created
 	CleanupRequest();
 
@@ -641,13 +626,13 @@ void FAppleHttpNSUrlSessionRequest::CleanupRequest()
 	{
 		Response->CleanSharedObjects();
 	}
+
 	if(Task != nil)
 	{
-		if(CompletionStatus == EHttpRequestStatus::Processing)
+		if (CompletionStatus == EHttpRequestStatus::Processing)
 		{
 			[Task cancel];
 		}
-		
 		[Task release];
 		Task = nil;
 	}
@@ -657,32 +642,29 @@ void FAppleHttpNSUrlSessionRequest::CancelRequest()
 {
 	UE_LOG(LogHttp, Verbose, TEXT("FAppleHttpNSUrlSessionRequest::CancelRequest()"));
 
-	if( Task != nil )
+	if (Task != nil)
 	{
 		if (bCanceled)
 		{
 			return;
 		}
 		
-		bCanceled = true;		
+		bCanceled = true;
 		[Task cancel];
-
+		
 		if (Response != nullptr)
 		{
 			Response->CleanSharedObjects();
 		}
 
-		// Ensure we run on game thread
-		if (!IsInGameThread())
+		FHttpManager& HttpManager = FHttpModule::Get().GetHttpManager();
+		if (HttpManager.IsValidRequest(this))
 		{
-			FHttpModule::Get().GetHttpManager().AddGameThreadTask([StrongThis = StaticCastSharedRef<FAppleHttpNSUrlSessionRequest>(AsShared())]()
-			{
-				StrongThis->FinishedRequest();
-			});
+			HttpManager.CancelThreadedRequest(SharedThis(this));
 		}
 		else
 		{
-			FinishedRequest();
+			FinishRequestNotInHttpManager();
 		}
 	}
 }
@@ -707,10 +689,6 @@ void FAppleHttpNSUrlSessionRequest::Tick(float DeltaSeconds)
 				LastReportedBytesRead = BytesRead;
 			}
 		}
-		if (Response->IsReady())
-		{
-			FinishedRequest();
-		}
 	}
 }
 
@@ -719,6 +697,22 @@ float FAppleHttpNSUrlSessionRequest::GetElapsedTime() const
 	return ElapsedTime;
 }
 
+bool FAppleHttpNSUrlSessionRequest::StartThreadedRequest()
+{
+	return true;
+}
+
+bool FAppleHttpNSUrlSessionRequest::IsThreadedRequestComplete()
+{
+	TSharedPtr<FAppleHttpNSUrlSessionResponse> AliveResponse = Response;
+	return (AliveResponse.IsValid() && AliveResponse->IsReady());
+}
+
+void FAppleHttpNSUrlSessionRequest::TickThreadedRequest(float DeltaSeconds)
+{
+	ElapsedTime += DeltaSeconds;
+}
+
 /****************************************************************************
  * FAppleHttpNSUrlSessionResponse implementation
  **************************************************************************/
diff --git a/Engine/Source/Runtime/Online/HTTP/Private/Apple/AppleHTTPNSUrlSession.h b/Engine/Source/Runtime/Online/HTTP/Private/Apple/AppleHTTPNSUrlSession.h
--- a/Engine/Source/Runtime/Online/HTTP/Private/Apple/AppleHTTPNSUrlSession.h
+++ b/Engine/Source/Runtime/Online/HTTP/Private/Apple/AppleHTTPNSUrlSession.h
@@ -3,14 +3,14 @@
 #pragma once
 
 #include "CoreMinimal.h"
-#include "GenericPlatform/HttpRequestCommon.h"
+#include "IHttpThreadedRequest.h"
 #include "Interfaces/IHttpResponse.h"
 #include "PlatformHttp.h"
 
 /**
  * Apple implementation of an Http request
  */
-class FAppleHttpNSUrlSessionRequest : public FHttpRequestCommon
+class FAppleHttpNSUrlSessionRequest : public IHttpThreadedRequest
 {
 public:
 	// implementation friends
@@ -49,6 +49,13 @@ public:
 	virtual float GetElapsedTime() const override;
 	//~ End IHttpRequest Interface
 
+	//~ Begin IHttpRequestThreaded Interface
+	virtual bool StartThreadedRequest() override;
+	virtual void FinishRequest() override;
+	virtual bool IsThreadedRequestComplete() override;
+	virtual void TickThreadedRequest(float DeltaSeconds) override;
+	//~ End IHttpRequestThreaded Interface
+
 	/**
 	 * Constructor
 	 *
@@ -71,12 +78,6 @@ private:
 	 */
 	bool StartRequest();
 
-	/**
-	 * Process state for a finished request that no longer needs to be ticked
-	 * Calls the completion delegate
-	 */
-	void FinishedRequest();
-
 	/**
 	 * Close session/request handles and unregister callbacks
 	 */
@@ -109,10 +110,7 @@ private:
 
 	/** flag to mark wether we tried canceled this request */
 	bool bCanceled;
-
-	/** Start of the request */
-	double StartRequestTime;
-
+	
 	/** Time taken to complete/cancel the request. */
 	float ElapsedTime;
 
diff --git a/Engine/Source/Runtime/Online/HTTP/Private/Apple/ApplePlatformHttp.cpp b/Engine/Source/Runtime/Online/HTTP/Private/Apple/ApplePlatformHttp.cpp
--- a/Engine/Source/Runtime/Online/HTTP/Private/Apple/ApplePlatformHttp.cpp
+++ b/Engine/Source/Runtime/Online/HTTP/Private/Apple/ApplePlatformHttp.cpp
@@ -284,3 +284,8 @@ IHttpRequest* FApplePlatformHttp::ConstructRequest()
 		return new FAppleHttpNSUrlConnectionRequest();
 	}
 }
+
+bool FApplePlatformHttp::UsesThreadedHttp()
+{
+	return bUseNSUrlSession;
+}
diff --git a/Engine/Source/Runtime/Online/HTTP/Public/Apple/ApplePlatformHttp.h b/Engine/Source/Runtime/Online/HTTP/Public/Apple/ApplePlatformHttp.h
--- a/Engine/Source/Runtime/Online/HTTP/Public/Apple/ApplePlatformHttp.h
+++ b/Engine/Source/Runtime/Online/HTTP/Public/Apple/ApplePlatformHttp.h
@@ -37,6 +37,13 @@ public:
 	 */
 	static IHttpRequest* ConstructRequest();
 
+	/**
+	 * Check if a platform uses the HTTP thread
+	 *
+	 * @return true if the platform uses threaded HTTP, false if not
+	 */
+	static bool UsesThreadedHttp();
+
 private:
 	/** Flag to allow fall back to use NSUrlConnection instead of NSUrlSession. Assigned from commandline */
 	static inline bool bUseNSUrlSession = false;
