Add 64 bit support in http request.

diff --git a/Engine/Source/Programs/WebTests/Private/TestHttp.cpp b/Engine/Source/Programs/WebTests/Private/TestHttp.cpp
--- a/Engine/Source/Programs/WebTests/Private/TestHttp.cpp
+++ b/Engine/Source/Programs/WebTests/Private/TestHttp.cpp
@@ -403,6 +403,35 @@ TEST_CASE_METHOD(FWaitUntilCompleteHttpFixture, "Can run parallel stream downloa
 	}
 }
 
+TEST_CASE_METHOD(FWaitUntilCompleteHttpFixture, "Can download big file exceeds 32 bits", HTTP_TAG)
+{
+	// 5 * 1024 * 1024 * 1024 BYTES = 5368709120 BYTES = 5 GB
+	uint64 Chunks = 5 * 1024;
+	uint64 ChunkSize = 1024 * 1024;
+
+	TSharedRef<IHttpRequest> HttpRequest = HttpModule->CreateRequest();
+	HttpRequest->SetURL(UrlStreamDownload(Chunks, ChunkSize));
+	HttpRequest->SetVerb(TEXT("GET"));
+
+	TSharedRef<int64> TotalBytesReceived = MakeShared<int64>(0);
+	FHttpRequestStreamDelegate Delegate;
+	Delegate.BindLambda([TotalBytesReceived](void* Ptr, int64 Length) {
+		*TotalBytesReceived += Length;
+		return true;
+	});
+	HttpRequest->SetResponseBodyReceiveStreamDelegate(Delegate);
+
+	HttpRequest->OnProcessRequestComplete().BindLambda([Chunks, ChunkSize, TotalBytesReceived](FHttpRequestPtr HttpRequest, FHttpResponsePtr HttpResponse, bool bSucceeded) {
+		CHECK(bSucceeded);
+		REQUIRE(HttpResponse != nullptr);
+		CHECK(HttpResponse->GetContentLength() == Chunks * ChunkSize);
+		CHECK(HttpResponse->GetContent().IsEmpty());
+		CHECK(*TotalBytesReceived == Chunks * ChunkSize);
+		CHECK(HttpResponse->GetResponseCode() == 200);
+	});
+	HttpRequest->ProcessRequest();
+}
+
 TEST_CASE_METHOD(FWaitUntilCompleteHttpFixture, "Streaming http upload from memory", HTTP_TAG)
 {
 	TSharedRef<IHttpRequest> HttpRequest = HttpModule->CreateRequest();
@@ -446,6 +475,68 @@ TEST_CASE_METHOD(FWaitUntilCompleteHttpFixture, "Streaming http upload from memo
 	HttpRequest->ProcessRequest();
 }
 
+class FTestHttpUploadStream final : public FArchive
+{
+public:
+	FTestHttpUploadStream(uint64 InTotalSize)
+		: FakeTotalSize(InTotalSize)
+	{
+	}
+
+	virtual void Serialize(void* V, int64 Length) override
+	{
+		for (int64 i = 0; i < Length; ++i)
+		{
+			((char*)V)[i] = 'd';
+		}
+
+		CurrentPos += Length;
+	}
+
+	virtual int64 TotalSize() override
+	{
+		return FakeTotalSize;
+	}
+
+	virtual void Seek(int64 InPos) 
+	{ 
+		CurrentPos = InPos;
+	}
+
+	virtual int64 Tell() override
+	{
+		return CurrentPos;
+	}
+
+	uint64 FakeTotalSize;
+	uint64 CurrentPos = 0;
+};
+
+TEST_CASE_METHOD(FWaitUntilCompleteHttpFixture, "Can upload big file exceeds 32 bits", HTTP_TAG)
+{
+	// TODO: Back to check later. xCurl 2206.4.0.0 doesn't work with file bigger than 32 bits
+	// 5 * 1024 * 1024 * 1024 BYTES = 5368709120 BYTES = 5 GB
+	//const uint64 TotalSize = 5368709120;
+	//const uint64 TotalSize = 4294967296;
+	//const uint64 TotalSize = 4294967295;
+	//const uint64 TotalSize = 2147483648;
+	const uint64 TotalSize = 2147483647;
+	TSharedRef<FTestHttpUploadStream> Stream = MakeShared<FTestHttpUploadStream>(TotalSize);
+
+	TSharedRef<IHttpRequest> HttpRequest = HttpModule->CreateRequest();
+	HttpRequest->SetURL(FString::Format(TEXT("{0}/streaming_upload_put"), { *UrlHttpTests() }));
+	HttpRequest->SetVerb(TEXT("PUT"));
+	HttpRequest->SetContentFromStream(Stream);
+	HttpRequest->SetHeader(TEXT("Content-Disposition"), TEXT("attachment;filename=TestStreamUpload.dat"));
+	HttpRequest->OnProcessRequestComplete().BindLambda([Stream, TotalSize](FHttpRequestPtr HttpRequest, FHttpResponsePtr HttpResponse, bool bSucceeded) {
+		CHECK(bSucceeded);
+		REQUIRE(HttpResponse != nullptr);
+		CHECK(HttpResponse->GetResponseCode() == 200);
+		CHECK(Stream->CurrentPos == TotalSize);
+	});
+	HttpRequest->ProcessRequest();
+}
+
 TEST_CASE_METHOD(FWaitUntilCompleteHttpFixture, "Streaming http upload from file by PUT can work well", HTTP_TAG)
 {
 	FString Filename = FString(FPlatformProcess::UserSettingsDir()) / TEXT("TestStreamUpload.dat");
@@ -515,6 +606,8 @@ public:
 
 TEST_CASE_METHOD(FWaitThreadedHttpFixture, "Http streaming download request can work in non game thread", HTTP_TAG)
 {
+	// TODO: Block main thread here to verify
+
 	ThreadCallback.BindLambda([this]() {
 		TSharedRef<IHttpRequest> HttpRequest = HttpModule->CreateRequest();
 		HttpRequest->SetURL(UrlStreamDownload(3/*Chunks*/, 1024/*ChunkSize*/));
diff --git a/Engine/Source/Runtime/NetworkReplayStreaming/HttpNetworkReplayStreaming/Private/HttpNetworkReplayStreaming.cpp b/Engine/Source/Runtime/NetworkReplayStreaming/HttpNetworkReplayStreaming/Private/HttpNetworkReplayStreaming.cpp
--- a/Engine/Source/Runtime/NetworkReplayStreaming/HttpNetworkReplayStreaming/Private/HttpNetworkReplayStreaming.cpp
+++ b/Engine/Source/Runtime/NetworkReplayStreaming/HttpNetworkReplayStreaming/Private/HttpNetworkReplayStreaming.cpp
@@ -682,7 +682,7 @@ void FHttpNetworkReplayStreamer::FlushStream()
 
 	TotalUploadBytes += HttpRequest->GetContentLength();
 
-	HttpRequest->SetURL( FString::Printf( TEXT( "%sreplay/%s/file/stream.%i?numChunks=%i&time=%i&mTime1=%i&mTime2=%i&absSize=%i" ), *ServerURL, *SessionName, StreamChunkIndex, StreamChunkIndex + 1, TotalDemoTimeInMS, StreamTimeRangeStart, StreamTimeRangeEnd, TotalUploadBytes ) );
+	HttpRequest->SetURL( FString::Printf( TEXT( "%sreplay/%s/file/stream.%i?numChunks=%i&time=%i&mTime1=%i&mTime2=%i&absSize=%llu" ), *ServerURL, *SessionName, StreamChunkIndex, StreamChunkIndex + 1, TotalDemoTimeInMS, StreamTimeRangeStart, StreamTimeRangeEnd, TotalUploadBytes ) );
 	HttpRequest->SetVerb( TEXT( "POST" ) );
 	HttpRequest->SetHeader( TEXT( "Content-Type" ), TEXT( "application/octet-stream" ) );
 
@@ -721,7 +721,7 @@ void FHttpNetworkReplayStreamer::StopUploading()
 
 	HttpRequest->OnProcessRequestComplete().BindRaw( this, &FHttpNetworkReplayStreamer::HttpStopUploadingFinished );
 
-	HttpRequest->SetURL( FString::Printf( TEXT( "%sreplay/%s/stopUploading?numChunks=%i&time=%i&absSize=%i" ), *ServerURL, *SessionName, StreamChunkIndex, TotalDemoTimeInMS, TotalUploadBytes ) );
+	HttpRequest->SetURL( FString::Printf( TEXT( "%sreplay/%s/stopUploading?numChunks=%i&time=%i&absSize=%llu" ), *ServerURL, *SessionName, StreamChunkIndex, TotalDemoTimeInMS, TotalUploadBytes ) );
 	HttpRequest->SetVerb( TEXT( "POST" ) );
 	HttpRequest->SetHeader( TEXT( "Content-Type" ), TEXT( "application/octet-stream" ) );
 
@@ -1800,7 +1800,7 @@ void FHttpNetworkReplayStreamer::HttpHeaderUploadFinished( FHttpRequestPtr HttpR
 			TotalUploadBytes += HttpRequest->GetContentLength();
 		}
 
-		UE_LOG( LogHttpReplay, Verbose, TEXT( "FHttpNetworkReplayStreamer::HttpHeaderUploadFinished. TotalUploadBytes: %i" ), TotalUploadBytes );
+		UE_LOG( LogHttpReplay, Verbose, TEXT( "FHttpNetworkReplayStreamer::HttpHeaderUploadFinished. TotalUploadBytes: %llu" ), TotalUploadBytes );
 		Result.Result = EStreamingOperationResult::Success;
 	}
 	else
@@ -1853,7 +1853,7 @@ void FHttpNetworkReplayStreamer::HttpUploadCheckpointFinished( FHttpRequestPtr H
 			TotalUploadBytes += HttpRequest->GetContentLength();
 		}
 
-		UE_LOG( LogHttpReplay, Verbose, TEXT( "FHttpNetworkReplayStreamer::HttpUploadCheckpointFinished. TotalUploadBytes: %i" ), TotalUploadBytes );
+		UE_LOG( LogHttpReplay, Verbose, TEXT( "FHttpNetworkReplayStreamer::HttpUploadCheckpointFinished. TotalUploadBytes: %llu" ), TotalUploadBytes );
 	}
 	else
 	{
@@ -1878,7 +1878,7 @@ void FHttpNetworkReplayStreamer::HttpUploadCustomEventFinished(FHttpRequestPtr H
 			TotalUploadBytes += HttpRequest->GetContentLength();
 		}
 
-		UE_LOG( LogHttpReplay, Verbose, TEXT( "FHttpNetworkReplayStreamer::HttpUploadCustomEventFinished. TotalUploadBytes: %i" ), TotalUploadBytes );
+		UE_LOG( LogHttpReplay, Verbose, TEXT( "FHttpNetworkReplayStreamer::HttpUploadCustomEventFinished. TotalUploadBytes: %llu" ), TotalUploadBytes );
 	}
 	else
 	{
diff --git a/Engine/Source/Runtime/NetworkReplayStreaming/HttpNetworkReplayStreaming/Public/HttpNetworkReplayStreaming.h b/Engine/Source/Runtime/NetworkReplayStreaming/HttpNetworkReplayStreaming/Public/HttpNetworkReplayStreaming.h
--- a/Engine/Source/Runtime/NetworkReplayStreaming/HttpNetworkReplayStreaming/Public/HttpNetworkReplayStreaming.h
+++ b/Engine/Source/Runtime/NetworkReplayStreaming/HttpNetworkReplayStreaming/Public/HttpNetworkReplayStreaming.h
@@ -388,7 +388,7 @@ public:
 
 	TSet< FString >						EventGroupSet;
 
-	uint32								TotalUploadBytes;
+	uint64								TotalUploadBytes;
 
 	TMap< FString, FCachedResponse >	ResponseCache;
 
diff --git a/Engine/Source/Runtime/Online/BuildPatchServices/Private/Tests/Fake/HttpManager.fake.h b/Engine/Source/Runtime/Online/BuildPatchServices/Private/Tests/Fake/HttpManager.fake.h
--- a/Engine/Source/Runtime/Online/BuildPatchServices/Private/Tests/Fake/HttpManager.fake.h
+++ b/Engine/Source/Runtime/Online/BuildPatchServices/Private/Tests/Fake/HttpManager.fake.h
@@ -27,7 +27,7 @@ namespace BuildPatchServices
 		virtual int32 GetResponseCode() const override;
 		virtual const TArray<uint8>& GetContent() const override;
 		virtual FString GetContentAsString() const override;
-		virtual int32 GetContentLength() const override;
+		virtual uint64 GetContentLength() const override;
 
 	public:
 		TArray<uint8> Data;
@@ -73,7 +73,7 @@ namespace BuildPatchServices
 		return UTF8_TO_TCHAR(ZeroTerminatedPayload.GetData());
 	}
 
-	int32 FFakeHttpResponse::GetContentLength() const
+	uint64 FFakeHttpResponse::GetContentLength() const
 	{
 		return Data.Num();
 	}
diff --git a/Engine/Source/Runtime/Online/BuildPatchServices/Private/Tests/Mock/HttpRequest.mock.h b/Engine/Source/Runtime/Online/BuildPatchServices/Private/Tests/Mock/HttpRequest.mock.h
--- a/Engine/Source/Runtime/Online/BuildPatchServices/Private/Tests/Mock/HttpRequest.mock.h
+++ b/Engine/Source/Runtime/Online/BuildPatchServices/Private/Tests/Mock/HttpRequest.mock.h
@@ -52,10 +52,10 @@ namespace BuildPatchServices
 			return FString();
 		}
 
-		virtual int32 GetContentLength() const override
+		virtual uint64 GetContentLength() const override
 		{
 			MOCK_FUNC_NOT_IMPLEMENTED("FMockHttpRequest::GetContentLength");
-			return int32();
+			return uint64();
 		}
 
 		virtual const TArray<uint8>& GetContent() const override
diff --git a/Engine/Source/Runtime/Online/BuildPatchServices/Private/Tests/Mock/HttpResponse.mock.h b/Engine/Source/Runtime/Online/BuildPatchServices/Private/Tests/Mock/HttpResponse.mock.h
--- a/Engine/Source/Runtime/Online/BuildPatchServices/Private/Tests/Mock/HttpResponse.mock.h
+++ b/Engine/Source/Runtime/Online/BuildPatchServices/Private/Tests/Mock/HttpResponse.mock.h
@@ -58,10 +58,10 @@ namespace BuildPatchServices
 			return FString();
 		}
 
-		virtual int32 GetContentLength() const override
+		virtual uint64 GetContentLength() const override
 		{
 			MOCK_FUNC_NOT_IMPLEMENTED("FMockHttpRequest::GetContentLength");
-			return int32();
+			return uint64();
 		}
 
 		virtual const TArray<uint8>& GetContent() const override
diff --git a/Engine/Source/Runtime/Online/HTTP/Internal/WinHttp/Support/WinHttpConnectionHttp.h b/Engine/Source/Runtime/Online/HTTP/Internal/WinHttp/Support/WinHttpConnectionHttp.h
--- a/Engine/Source/Runtime/Online/HTTP/Internal/WinHttp/Support/WinHttpConnectionHttp.h
+++ b/Engine/Source/Runtime/Online/HTTP/Internal/WinHttp/Support/WinHttpConnectionHttp.h
@@ -61,7 +61,7 @@ public:
 	 */
 	TArray<uint8>& GetLastChunk() { return GameThreadChunk; }
 	/** Provides the current known response content length. */
-	int32 GetResponseContentLength() const { return ResponseContentLength; }
+	uint64 GetResponseContentLength() const { return ResponseContentLength; }
 	/** Provides the current known headers received. */
 	const TMap<FString, FString>& GetHeadersReceived() const { return HeadersReceived; }
 
@@ -139,13 +139,13 @@ protected:
 	 *
 	 * @param AmountSent The amount of bytes to increment our counters by
 	 */
-	void IncrementSentByteCounts(const int32 AmountSent);
+	void IncrementSentByteCounts(const uint64 AmountSent);
 	/**
 	 * Increment our bytes-received counters
 	 *
 	 * @param AmountReceived The amount of bytes to increment our counters by
 	 */
-	void IncrementReceivedByteCounts(const int32 AmountReceived);
+	void IncrementReceivedByteCounts(const uint64 AmountReceived);
 	/** 
 	 * Check if we have more data to send
 	 *
@@ -198,15 +198,15 @@ protected:
 	/** Handles validating request now that we have server information (certificates, etc) */
 	void HandleSendingRequest();
 	/** Handles SOMETHING not sure yet*/
-	void HandleWriteComplete(const uint32 NumBytesSent);
+	void HandleWriteComplete(const uint64 NumBytesSent);
 	/** Handles completing the request so we may read the headers and body of the response */
 	void HandleSendRequestComplete();
 	/** Handles reading the response headers */
 	virtual void HandleHeadersAvailable();
 	/** Handles writing the http response body to our response body */
-	void HandleDataAvailable(const uint32 NumBytesAvailable);
+	void HandleDataAvailable(const uint64 NumBytesAvailable);
 	/** Handles querying the next chunk of our http response body */
-	void HandleReadComplete(const uint32 NumBytesRead);
+	void HandleReadComplete(const uint64 NumBytesRead);
 	/** Handles shutting down the request when there was an error */
 	void HandleRequestError(const uint32 ErrorApiId, const uint32 ErrorCode);
 	/** Handles our request object closing, ending this request */
@@ -302,7 +302,7 @@ private:
 	/** Payload body*/
 	TSharedPtr<FRequestPayload, ESPMode::ThreadSafe> Payload;
 	/** Amount of bytes from payload that has been sent successfully */
-	int32 NumBytesSuccessfullySent = 0;
+	uint64 NumBytesSuccessfullySent = 0;
 	/** Data from the payload currently trying to be written */
 	TArray<uint8> PayloadBuffer;
 
@@ -317,20 +317,20 @@ private:
 	/** The current chunk being built for consumption on the game thread. Only access on the game thread. */
 	TArray<uint8> GameThreadChunk;
 	/** The content length from the response headers, or 0. */
-	int32 ResponseContentLength = 0;
+	uint64 ResponseContentLength = 0;
 	/** */
-	int32 BytesWrittenToGameThreadChunk = 0;
+	uint64 BytesWrittenToGameThreadChunk = 0;
 	/** The amount of bytes that are available to read from WinHttpReadData now */
-	TOptional<int32> ResponseBytesAvailable;
+	TOptional<uint64> ResponseBytesAvailable;
 	
 	/// Members for reporting updates to main thread
 
 	/** Called on the game thread during tick if we have sent or received data since last tick */
 	FWinHttpConnectionHttpOnDataTransferred OnDataTransferredHandler;
 	/** Data to report to the game thread about how many bytes have been sent since last tick (not total received) */
-	TOptional<int32> BytesToReportSent;
+	TOptional<uint64> BytesToReportSent;
 	/** Data to report to the game thread about how many bytes have been received since last tick (not total received) */
-	TOptional<int32> BytesToReportReceived;
+	TOptional<uint64> BytesToReportReceived;
 
 	/** Called on the game thread during tick if we have received header data since last tick */
 	FWinHttpConnectionHttpOnHeaderReceived OnHeaderReceivedHandler;
diff --git a/Engine/Source/Runtime/Online/HTTP/Private/Apple/AppleHTTPNSUrlConnection.cpp b/Engine/Source/Runtime/Online/HTTP/Private/Apple/AppleHTTPNSUrlConnection.cpp
--- a/Engine/Source/Runtime/Online/HTTP/Private/Apple/AppleHTTPNSUrlConnection.cpp
+++ b/Engine/Source/Runtime/Online/HTTP/Private/Apple/AppleHTTPNSUrlConnection.cpp
@@ -228,7 +228,7 @@ FString FAppleHttpNSUrlConnectionRequest::GetContentType() const
 }
 
 
-int32 FAppleHttpNSUrlConnectionRequest::GetContentLength() const
+uint64 FAppleHttpNSUrlConnectionRequest::GetContentLength() const
 {
 	UE_LOG(LogHttp, Verbose, TEXT("FAppleHttpNSUrlConnectionRequest::GetContentLength() - %i"), RequestPayloadByteLength);
 	return RequestPayloadByteLength;
@@ -273,7 +273,7 @@ bool FAppleHttpNSUrlConnectionRequest::SetContentAsStreamedFile(const FString& F
 	struct stat FileAttrs = { 0 };
 	if (stat(PlatformFilename.fileSystemRepresentation, &FileAttrs) == 0)
 	{
-		UE_LOG(LogHttp, VeryVerbose, TEXT("FAppleHttpNSUrlConnectionRequest::SetContentAsStreamedFile succeeded in getting the file size - %d"), FileAttrs.st_size);
+		UE_LOG(LogHttp, VeryVerbose, TEXT("FAppleHttpNSUrlConnectionRequest::SetContentAsStreamedFile succeeded in getting the file size - %llu"), FileAttrs.st_size);
 		// Under the hood, the Foundation framework unsets HTTPBody, and takes over as the stream delegate.
 		// The stream itself should be unopened when passed to setHTTPBodyStream.
 		Request.HTTPBodyStream = [NSInputStream inputStreamWithFileAtPath: PlatformFilename];
@@ -369,7 +369,7 @@ bool FAppleHttpNSUrlConnectionRequest::StartRequest()
 	// set the content-length and user-agent
 	if(GetContentLength() > 0)
 	{
-		[Request setValue:[NSString stringWithFormat:@"%d", GetContentLength()] forHTTPHeaderField:@"Content-Length"];
+		[Request setValue:[NSString stringWithFormat:@"%llu", GetContentLength()] forHTTPHeaderField:@"Content-Length"];
 	}
 
 	const FString UserAgent = GetHeader("User-Agent");
@@ -510,8 +510,8 @@ void FAppleHttpNSUrlConnectionRequest::Tick(float DeltaSeconds)
 	{
 		if (OnRequestProgress().IsBound())
 		{
-			const int32 BytesWritten = Response->GetNumBytesWritten();
-			const int32 BytesRead = Response->GetNumBytesReceived();
+			const uint64 BytesWritten = Response->GetNumBytesWritten();
+			const uint64 BytesRead = Response->GetNumBytesReceived();
 			if (BytesWritten != LastReportedBytesWritten || BytesRead != LastReportedBytesRead)
 			{
 				OnRequestProgress().ExecuteIfBound(SharedThis(this), BytesWritten, BytesRead);
@@ -565,7 +565,7 @@ float FAppleHttpNSUrlConnectionRequest::GetElapsedTime() const
 {
 	UE_LOG(LogHttp, Verbose, TEXT("didSendBodyData:(NSInteger)bytesWritten totalBytes:Written:(NSInteger)totalBytesWritten totalBytesExpectedToWrite:(NSInteger)totalBytesExpectedToWrite"));
 	self.BytesWritten = totalBytesWritten;
-	UE_LOG(LogHttp, Verbose, TEXT("didSendBodyData: totalBytesWritten = %d, totalBytesExpectedToWrite = %d: %p"), totalBytesWritten, totalBytesExpectedToWrite, self);
+	UE_LOG(LogHttp, Verbose, TEXT("didSendBodyData: totalBytesWritten = %llu, totalBytesExpectedToWrite = %llu: %p"), totalBytesWritten, totalBytesExpectedToWrite, self);
 }
 
 -(void) connection:(NSURLConnection *)connection didReceiveResponse:(NSURLResponse *)response
@@ -575,14 +575,14 @@ float FAppleHttpNSUrlConnectionRequest::GetElapsedTime() const
 	
 	// presize the payload container if possible
 	Payload.Empty([response expectedContentLength] != NSURLResponseUnknownLength ? [response expectedContentLength] : 0);
-	UE_LOG(LogHttp, Verbose, TEXT("didReceiveResponse: expectedContentLength = %d. Length = %d: %p"), [response expectedContentLength], Payload.Max(), self);
+	UE_LOG(LogHttp, Verbose, TEXT("didReceiveResponse: expectedContentLength = %llu. Length = %llu: %p"), [response expectedContentLength], Payload.Max(), self);
 }
 
 
 -(void) connection:(NSURLConnection *)connection didReceiveData:(NSData *)data
 {
 	Payload.Append((const uint8*)[data bytes], [data length]);
-	UE_LOG(LogHttp, Verbose, TEXT("didReceiveData with %d bytes. After Append, Payload Length = %d: %p"), [data length], Payload.Num(), self);
+	UE_LOG(LogHttp, Verbose, TEXT("didReceiveData with %llu bytes. After Append, Payload Length = %llu: %p"), [data length], Payload.Num(), self);
 }
 
 
@@ -833,7 +833,7 @@ float FAppleHttpNSUrlConnectionRequest::GetElapsedTime() const
 	return Payload;
 }
 
--(int32)getBytesWritten
+-(uint64)getBytesWritten
 {
 	return self.BytesWritten;
 }
@@ -925,7 +925,7 @@ FString FAppleHttpNSUrlConnectionResponse::GetContentType() const
 }
 
 
-int32 FAppleHttpNSUrlConnectionResponse::GetContentLength() const
+uint64 FAppleHttpNSUrlConnectionResponse::GetContentLength() const
 {
 	UE_LOG(LogHttp, Verbose, TEXT("FAppleHttpNSUrlConnectionResponse::GetContentLength()"));
 
@@ -1004,13 +1004,13 @@ bool FAppleHttpNSUrlConnectionResponse::HadError() const
 	return bHadError;
 }
 
-const int32 FAppleHttpNSUrlConnectionResponse::GetNumBytesReceived() const
+const uint64 FAppleHttpNSUrlConnectionResponse::GetNumBytesReceived() const
 {
 	return [ResponseWrapper getPayload].Num();
 }
 
-const int32 FAppleHttpNSUrlConnectionResponse::GetNumBytesWritten() const
+const uint64 FAppleHttpNSUrlConnectionResponse::GetNumBytesWritten() const
 {
-    int32 NumBytesWritten = [ResponseWrapper getBytesWritten];
+    uint64 NumBytesWritten = [ResponseWrapper getBytesWritten];
     return NumBytesWritten;
 }
diff --git a/Engine/Source/Runtime/Online/HTTP/Private/Apple/AppleHTTPNSUrlConnection.h b/Engine/Source/Runtime/Online/HTTP/Private/Apple/AppleHTTPNSUrlConnection.h
--- a/Engine/Source/Runtime/Online/HTTP/Private/Apple/AppleHTTPNSUrlConnection.h
+++ b/Engine/Source/Runtime/Online/HTTP/Private/Apple/AppleHTTPNSUrlConnection.h
@@ -29,7 +29,7 @@ public:
 	virtual FString GetHeader(const FString& HeaderName) const override;
 	virtual TArray<FString> GetAllHeaders() const override;	
 	virtual FString GetContentType() const override;
-	virtual int32 GetContentLength() const override;
+	virtual uint64 GetContentLength() const override;
 	virtual const TArray<uint8>& GetContent() const override;
 	//~ End IHttpBase Interface
 
@@ -102,7 +102,7 @@ private:
 	bool bIsPayloadFile;
 
 	/** The request payload length in bytes. This must be tracked separately for a file stream */
-	int32 RequestPayloadByteLength;
+	uint64 RequestPayloadByteLength;
 
 	/** The response object which we will use to pair with this request */
 	TSharedPtr<class FAppleHttpNSUrlConnectionResponse,ESPMode::ThreadSafe> Response;
@@ -114,7 +114,7 @@ private:
 	TArray<uint8> ContentData;
 
 	/** Number of bytes sent to progress update */
-	int32 ProgressBytesSent;
+	uint64 ProgressBytesSent;
 
 	/** Start of the request */
 	double StartRequestTime;
@@ -147,7 +147,7 @@ private:
 /** When the response is complete, indicates whether the response failed with an error specific to connecting to the host. */
 @property BOOL bIsHostConnectionFailure;
 /** The total number of bytes written out during the request/response */
-@property int32 BytesWritten;
+@property uint64 BytesWritten;
 
 /** Delegate called when we send data. See Apple docs for when/how this should be used. */
 -(void) connection:(NSURLConnection *)connection didSendBodyData:(NSInteger)bytesWritten totalBytesWritten:(NSInteger)totalBytesWritten totalBytesExpectedToWrite:(NSInteger)totalBytesExpectedToWrite;
@@ -166,7 +166,7 @@ private:
 #endif
 
 - (TArray<uint8>&)getPayload;
-- (int32)getBytesWritten;
+- (uint64)getBytesWritten;
 @end
 
 
@@ -194,7 +194,7 @@ public:
 	virtual FString GetHeader(const FString& HeaderName) const override;
 	virtual TArray<FString> GetAllHeaders() const override;	
 	virtual FString GetContentType() const override;
-	virtual int32 GetContentLength() const override;
+	virtual uint64 GetContentLength() const override;
 	virtual const TArray<uint8>& GetContent() const override;
 	//~ End IHttpBase Interface
 
@@ -218,12 +218,12 @@ public:
 	/**
 	 * Get the number of bytes received so far
 	 */
-	const int32 GetNumBytesReceived() const;
+	const uint64 GetNumBytesReceived() const;
 
 	/**
 	* Get the number of bytes sent so far
 	*/
-	const int32 GetNumBytesWritten() const;
+	const uint64 GetNumBytesWritten() const;
 
 	/**
 	 * Constructor
diff --git a/Engine/Source/Runtime/Online/HTTP/Private/Apple/AppleHTTPNSUrlSession.cpp b/Engine/Source/Runtime/Online/HTTP/Private/Apple/AppleHTTPNSUrlSession.cpp
--- a/Engine/Source/Runtime/Online/HTTP/Private/Apple/AppleHTTPNSUrlSession.cpp
+++ b/Engine/Source/Runtime/Online/HTTP/Private/Apple/AppleHTTPNSUrlSession.cpp
@@ -45,9 +45,9 @@ enum class EAppleHttpRequestResponseState: uint8
 /** A handle for the response */
 @property(retain) NSHTTPURLResponse* Response;
 /** The total number of bytes written out during the request/response */
-@property int32 BytesWritten;
+@property uint64 BytesWritten;
 /** The total number of bytes received out during the request/response */
-@property int32 BytesReceived;
+@property uint64 BytesReceived;
 /** Response state */
 @property EAppleHttpRequestResponseState ResponseState;
 
@@ -102,7 +102,7 @@ enum class EAppleHttpRequestResponseState: uint8
 - (void)URLSession:(NSURLSession *)session task:(NSURLSessionTask *)task didSendBodyData:(int64_t)bytesSent totalBytesSent:(int64_t)totalBytesSent totalBytesExpectedToSend:(int64_t)totalBytesExpectedToSend
 {
 	UE_LOG(LogHttp, Verbose, TEXT("URLSession:task:didSendBodyData:totalBytesSent:totalBytesExpectedToSend: totalBytesSent = %lld, totalBytesSent = %lld: %p"), totalBytesSent, totalBytesExpectedToSend, self);
-	self.BytesWritten = (int32)totalBytesSent;
+	self.BytesWritten = totalBytesSent;
 }
 
 - (void)URLSession:(NSURLSession *)session dataTask:(NSURLSessionDataTask *)dataTask didReceiveResponse:(NSURLResponse *)response completionHandler:(void (^)(NSURLSessionResponseDisposition disposition))completionHandler
@@ -110,12 +110,12 @@ enum class EAppleHttpRequestResponseState: uint8
 	UE_LOG(LogHttp, Verbose, TEXT("URLSession:dataTask:didReceiveResponse:completionHandler"));
 	
 	self.Response = (NSHTTPURLResponse*)response;
-	int32 ExpectedResponseLength = (int32)response.expectedContentLength;
+	uint64 ExpectedResponseLength = response.expectedContentLength;
 	if(!bInitializedWithValidStream && ExpectedResponseLength != NSURLResponseUnknownLength)
 	{
 		Payload.Empty(ExpectedResponseLength);
 	}
-	UE_LOG(LogHttp, Verbose, TEXT("URLSession:dataTask:didReceiveResponse:completionHandler: expectedContentLength = %lld. Length = %d: %p"), ExpectedResponseLength, Payload.Max(), self);
+	UE_LOG(LogHttp, Verbose, TEXT("URLSession:dataTask:didReceiveResponse:completionHandler: expectedContentLength = %lld. Length = %llu: %p"), ExpectedResponseLength, Payload.Max(), self);
 	completionHandler(NSURLSessionResponseAllow);
 }
 
@@ -151,7 +151,7 @@ enum class EAppleHttpRequestResponseState: uint8
 	// Keep BytesReceived as a separated value to avoid concurrent accesses to Payload
 	self.BytesReceived += NewBytesReceived;
 	
-	UE_LOG(LogHttp, Verbose, TEXT("URLSession:dataTask:didReceiveData with %d bytes. After Append, Payload Length = %d: %p"), NewBytesReceived, self.BytesReceived, self);
+	UE_LOG(LogHttp, Verbose, TEXT("URLSession:dataTask:didReceiveData with %llu bytes. After Append, Payload Length = %llu: %p"), NewBytesReceived, self.BytesReceived, self);
 }
 
 - (void)URLSession:(NSURLSession *)session task:(NSURLSessionTask *)task didCompleteWithError:(nullable NSError *)error
@@ -416,7 +416,7 @@ FString FAppleHttpNSUrlSessionRequest::GetContentType() const
 	return ContentType;
 }
 
-int32 FAppleHttpNSUrlSessionRequest::GetContentLength() const
+uint64 FAppleHttpNSUrlSessionRequest::GetContentLength() const
 {
 	UE_LOG(LogHttp, Verbose, TEXT("FAppleHttpNSUrlSessionRequest::GetContentLength() - %i"), ContentBytesLength);
 	return ContentBytesLength;
@@ -539,7 +539,7 @@ bool FAppleHttpNSUrlSessionRequest::StartRequest()
 	// set the content-length and user-agent (it is possible that the OS ignores this value)
 	if(GetContentLength() > 0)
 	{
-		[Request setValue:[NSString stringWithFormat:@"%d", GetContentLength()] forHTTPHeaderField:@"Content-Length"];
+		[Request setValue:[NSString stringWithFormat:@"%llu", GetContentLength()] forHTTPHeaderField:@"Content-Length"];
 	}
 
 	const FString UserAgent = GetHeader("User-Agent");
@@ -793,7 +793,7 @@ FString FAppleHttpNSUrlSessionResponse::GetContentType() const
 	return GetHeader( TEXT( "Content-Type" ) );
 }
 
-int32 FAppleHttpNSUrlSessionResponse::GetContentLength() const
+uint64 FAppleHttpNSUrlSessionResponse::GetContentLength() const
 {
 	UE_LOG(LogHttp, Verbose, TEXT("FAppleHttpNSUrlSessionResponse::GetContentLength()"));
 	
@@ -856,12 +856,12 @@ bool FAppleHttpNSUrlSessionResponse::HadConnectionError() const
 	return (ResponseDelegate.ResponseState == EAppleHttpRequestResponseState::ConnectionError);
 }
 
-const int32 FAppleHttpNSUrlSessionResponse::GetNumBytesReceived() const
+const uint64 FAppleHttpNSUrlSessionResponse::GetNumBytesReceived() const
 {
 	return ResponseDelegate.BytesReceived;
 }
 
-const int32 FAppleHttpNSUrlSessionResponse::GetNumBytesWritten() const
+const uint64 FAppleHttpNSUrlSessionResponse::GetNumBytesWritten() const
 {
 	return ResponseDelegate.BytesWritten;
 }
diff --git a/Engine/Source/Runtime/Online/HTTP/Private/Apple/AppleHTTPNSUrlSession.h b/Engine/Source/Runtime/Online/HTTP/Private/Apple/AppleHTTPNSUrlSession.h
--- a/Engine/Source/Runtime/Online/HTTP/Private/Apple/AppleHTTPNSUrlSession.h
+++ b/Engine/Source/Runtime/Online/HTTP/Private/Apple/AppleHTTPNSUrlSession.h
@@ -23,7 +23,7 @@ public:
 	virtual FString GetHeader(const FString& HeaderName) const override;
 	virtual TArray<FString> GetAllHeaders() const override;	
 	virtual FString GetContentType() const override;
-	virtual int32 GetContentLength() const override;
+	virtual uint64 GetContentLength() const override;
 	virtual const TArray<uint8>& GetContent() const override;
 	//~ End IHttpBase Interface
 
@@ -97,7 +97,7 @@ private:
 	bool bIsPayloadFile;
 
 	/** The request payload length in bytes. This must be tracked separately for a file stream */
-	int32 ContentBytesLength;
+	uint64 ContentBytesLength;
 
 	/** The response object which we will use to pair with this request */
 	TSharedPtr<class FAppleHttpNSUrlSessionResponse,ESPMode::ThreadSafe> Response;
@@ -147,7 +147,7 @@ public:
 	virtual FString GetHeader(const FString& HeaderName) const override;
 	virtual TArray<FString> GetAllHeaders() const override;	
 	virtual FString GetContentType() const override;
-	virtual int32 GetContentLength() const override;
+	virtual uint64 GetContentLength() const override;
 	virtual const TArray<uint8>& GetContent() const override;
 	//~ End IHttpBase Interface
 
@@ -179,12 +179,12 @@ public:
 	/**
 	 * Get the number of bytes received so far
 	 */
-	const int32 GetNumBytesReceived() const;
+	const uint64 GetNumBytesReceived() const;
 
 	/**
 	* Get the number of bytes sent so far
 	*/
-	const int32 GetNumBytesWritten() const;
+	const uint64 GetNumBytesWritten() const;
 
 	/**
 	 * Cleans internal shared objects between request and response
diff --git a/Engine/Source/Runtime/Online/HTTP/Private/Curl/CurlHttp.cpp b/Engine/Source/Runtime/Online/HTTP/Private/Curl/CurlHttp.cpp
--- a/Engine/Source/Runtime/Online/HTTP/Private/Curl/CurlHttp.cpp
+++ b/Engine/Source/Runtime/Online/HTTP/Private/Curl/CurlHttp.cpp
@@ -283,7 +283,7 @@ FString FCurlHttpRequest::GetContentType() const
 	return GetHeader(TEXT( "Content-Type" ));
 }
 
-int32 FCurlHttpRequest::GetContentLength() const
+uint64 FCurlHttpRequest::GetContentLength() const
 {
 	return RequestPayload.IsValid() ? RequestPayload->GetContentLength() : 0;
 }
@@ -343,7 +343,7 @@ void FCurlHttpRequest::SetContentAsString(const FString& ContentString)
 		return;
 	}
 
-	int32 Utf8Length = FPlatformString::ConvertedLength<UTF8CHAR>(*ContentString, ContentString.Len());
+	uint64 Utf8Length = FPlatformString::ConvertedLength<UTF8CHAR>(*ContentString, ContentString.Len());
 	TArray<uint8> Buffer;
 	Buffer.SetNumUninitialized(Utf8Length);
 	FPlatformString::Convert((UTF8CHAR*)Buffer.GetData(), Buffer.Num(), *ContentString, ContentString.Len());
@@ -523,7 +523,7 @@ size_t FCurlHttpRequest::ReceiveResponseHeaderCallback(void* Ptr, size_t SizeInB
 				//Store the content length so OnRequestProgress() delegates have something to work with
 				if (HeaderKey == TEXT("Content-Length"))
 				{
-					Response->ContentLength = FCString::Atoi(*HeaderValue);
+					Response->ContentLength = FCString::Atoi64(*HeaderValue);
 				}
 				Response->NewlyReceivedHeaders.Enqueue(TPair<FString, FString>(MoveTemp(HeaderKey), MoveTemp(HeaderValue)));
 			}
@@ -562,13 +562,11 @@ size_t FCurlHttpRequest::ReceiveResponseBodyCallback(void* Ptr, size_t SizeInBlo
 	// to get aborted and the libcurl function used will return CURLE_WRITE_ERROR.
 	size_t NumberOfBytesProcessed = 0;
 	  
-	uint32 SizeToDownload = SizeInBlocks * BlockSizeInBytes;
+	uint64 SizeToDownload = SizeInBlocks * BlockSizeInBytes;
 
-	UE_LOG(LogHttp, Verbose, TEXT("%p: ReceiveResponseBodyCallback: %d bytes out of %d received. (SizeInBlocks=%d, BlockSizeInBytes=%d, Response->TotalBytesRead=%d, Response->GetContentLength()=%d, SizeToDownload=%d (<-this will get returned from the callback))"),
-		this,
-		static_cast<int32>(Response->TotalBytesRead.GetValue() + SizeToDownload), Response->GetContentLength(),
-		static_cast<int32>(SizeInBlocks), static_cast<int32>(BlockSizeInBytes), Response->TotalBytesRead.GetValue(), Response->GetContentLength(), static_cast<int32>(SizeToDownload)
-		);
+	UE_LOG(LogHttp, Verbose, TEXT("%p: ReceiveResponseBodyCallback: %llu bytes out of %llu received. (SizeInBlocks=%llu, BlockSizeInBytes=%llu, Response->TotalBytesRead=%llu, Response->GetContentLength()=%llu, SizeToDownload=%llu (<-this will get returned from the callback))"),
+		this, Response->TotalBytesRead.GetValue() + SizeToDownload, Response->GetContentLength(),
+		SizeInBlocks, BlockSizeInBytes, Response->TotalBytesRead.GetValue(), Response->GetContentLength(), SizeToDownload);
 
 	// note that we can be passed 0 bytes if file transmitted has 0 length
 	if (SizeToDownload == 0)
@@ -603,20 +601,13 @@ size_t FCurlHttpRequest::UploadCallback(void* Ptr, size_t SizeInBlocks, size_t B
 	TimeSinceLastResponse = 0.0f;
 
 	size_t MaxBufferSize = SizeInBlocks * BlockSizeInBytes;
-	size_t SizeAlreadySent = static_cast<size_t>(BytesSent.GetValue());
+	size_t SizeAlreadySent = BytesSent.GetValue();
 	size_t SizeSentThisTime = RequestPayload->FillOutputBuffer(Ptr, MaxBufferSize, SizeAlreadySent);
 	BytesSent.Add(SizeSentThisTime);
 	TotalBytesSent.Add(SizeSentThisTime);
 
-	UE_LOG(LogHttp, Verbose, TEXT("%p: UploadCallback: %d bytes out of %d sent (%d bytes total sent). (SizeInBlocks=%d, BlockSizeInBytes=%d, SizeToSendThisTime=%d (<-this will get returned from the callback))"),
-		this,
-		static_cast< int32 >(BytesSent.GetValue()),
-		RequestPayload->GetContentLength(),
-		static_cast< int32 >(TotalBytesSent.GetValue()),
-		static_cast< int32 >(SizeInBlocks),
-		static_cast< int32 >(BlockSizeInBytes),
-		static_cast< int32 >(SizeSentThisTime)
-		);
+	UE_LOG(LogHttp, Verbose, TEXT("%p: UploadCallback: %llu bytes out of %llu sent (%llu bytes total sent). (SizeInBlocks=%llu, BlockSizeInBytes=%llu, SizeToSendThisTime=%llu (<-this will get returned from the callback))"),
+		this, BytesSent.GetValue(), RequestPayload->GetContentLength(), TotalBytesSent.GetValue(), SizeInBlocks, BlockSizeInBytes, SizeSentThisTime);
 
 	return SizeSentThisTime;
 }
@@ -626,9 +617,7 @@ int FCurlHttpRequest::SeekCallback(curl_off_t Offset, int Origin)
 	// Only support seeking to the very beginning
 	if (bIsRequestPayloadSeekable && Origin == SEEK_SET && Offset == 0)
 	{
-		UE_LOG(LogHttp, Log, TEXT("%p: SeekCallback: Resetting to the beginning. We had uploaded %d bytes"),
-			this,
-			static_cast<int32>(BytesSent.GetValue()));
+		UE_LOG(LogHttp, Log, TEXT("%p: SeekCallback: Resetting to the beginning. We had uploaded %llu bytes"), this, BytesSent.GetValue());
 		BytesSent.Reset();
 		bIsRequestPayloadSeekable = false; // Do not attempt to re-seek
 		return CURL_SEEKFUNC_OK;
@@ -809,7 +798,7 @@ bool FCurlHttpRequest::SetupRequest()
 	UE_LOG(LogHttp, Verbose, TEXT("%p: URL='%s'"), this, *URL);
 	UE_LOG(LogHttp, Verbose, TEXT("%p: Verb='%s'"), this, *Verb);
 	UE_LOG(LogHttp, Verbose, TEXT("%p: Custom headers are %s"), this, Headers.Num() ? TEXT("present") : TEXT("NOT present"));
-	UE_LOG(LogHttp, Verbose, TEXT("%p: Payload size=%d"), this, RequestPayload->GetContentLength());
+	UE_LOG(LogHttp, Verbose, TEXT("%p: Payload size=%llu"), this, RequestPayload->GetContentLength());
 
 	if (GetHeader(TEXT("User-Agent")).IsEmpty())
 	{
@@ -819,7 +808,7 @@ bool FCurlHttpRequest::SetupRequest()
 	// content-length should be present http://www.w3.org/Protocols/rfc2616/rfc2616-sec4.html#sec4.4
 	if (GetHeader(TEXT("Content-Length")).IsEmpty())
 	{
-		SetHeader(TEXT("Content-Length"), FString::Printf(TEXT("%d"), RequestPayload->GetContentLength()));
+		SetHeader(TEXT("Content-Length"), FString::Printf(TEXT("%llu"), RequestPayload->GetContentLength()));
 	}
 
 	// Remove "Expect: 100-continue" since this is supposed to cause problematic behavior on Amazon ELB (and WinInet doesn't send that either)
@@ -863,6 +852,7 @@ bool FCurlHttpRequest::SetupRequestHttpThread()
 			curl_easy_setopt(EasyHandle, CURLOPT_POST, 1L);
 			curl_easy_setopt(EasyHandle, CURLOPT_POSTFIELDS, NULL);
 #if WITH_CURL_XCURL
+			checkf(RequestPayload->GetContentLength() <= TNumericLimits<int32>::Max(), TEXT("xCurl 2206.4.0.0 doesn't support uploading file with length more than 32 bits"));
 			curl_easy_setopt(EasyHandle, CURLOPT_INFILESIZE, RequestPayload->GetContentLength());
 #else
 			curl_easy_setopt(EasyHandle, CURLOPT_POSTFIELDSIZE, RequestPayload->GetContentLength());
@@ -872,6 +862,9 @@ bool FCurlHttpRequest::SetupRequestHttpThread()
 		else if (Verb == TEXT("PUT") || Verb == TEXT("PATCH"))
 		{
 			curl_easy_setopt(EasyHandle, CURLOPT_UPLOAD, 1L);
+#if WITH_CURL_XCURL
+			checkf(RequestPayload->GetContentLength() <= TNumericLimits<int32>::Max(), TEXT("xCurl 2206.4.0.0 doesn't support uploading file with length more than 32 Bits"));
+#endif
 			curl_easy_setopt(EasyHandle, CURLOPT_INFILESIZE, RequestPayload->GetContentLength());
 			if (Verb != TEXT("PUT"))
 			{
@@ -1115,8 +1108,8 @@ void FCurlHttpRequest::Tick(float DeltaSeconds)
 void FCurlHttpRequest::CheckProgressDelegate()
 {
 	QUICK_SCOPE_CYCLE_COUNTER(STAT_FCurlHttpRequest_CheckProgressDelegate);
-	const int32 CurrentBytesRead = Response.IsValid() ? Response->TotalBytesRead.GetValue() : 0;
-	const int32 CurrentBytesSent = BytesSent.GetValue();
+	const uint64 CurrentBytesRead = Response.IsValid() ? Response->TotalBytesRead.GetValue() : 0;
+	const uint64 CurrentBytesSent = BytesSent.GetValue();
 
 	const bool bProcessing = CompletionStatus == EHttpRequestStatus::Processing;
 	const bool bBytesSentChanged = (CurrentBytesSent != LastReportedBytesSent);
@@ -1186,20 +1179,24 @@ void FCurlHttpRequest::FinishRequest()
 				Response->HttpCode = HttpCode;
 			}
 
-			double ContentLengthDownload = 0.0;
-			if (CURLE_OK == curl_easy_getinfo(EasyHandle, CURLINFO_CONTENT_LENGTH_DOWNLOAD, &ContentLengthDownload) && ContentLengthDownload >= 0.0)
-			{
-				Response->ContentLength = static_cast< int32 >(ContentLengthDownload);
-			}
-			else
+			// If content length wasn't received through response header 
+			if (Response->ContentLength == 0)
 			{
-				// If curl did not know how much we downloaded, or we were missing a Content-Length header (Chunked request), set our ContentLength as the amount we downloaded
-				Response->ContentLength = Response->TotalBytesRead.GetValue();
+				double ContentLengthDownload = 0.0;
+				if (CURLE_OK == curl_easy_getinfo(EasyHandle, CURLINFO_CONTENT_LENGTH_DOWNLOAD, &ContentLengthDownload) && ContentLengthDownload > 0.0)
+				{
+					Response->ContentLength = static_cast< uint64 >(ContentLengthDownload);
+				}
+				else
+				{
+					// If curl did not know how much we downloaded, or we were missing a Content-Length header (Chunked request), set our ContentLength as the amount we downloaded
+					Response->ContentLength = Response->TotalBytesRead.GetValue();
+				}
 			}
 
 			if (Response->HttpCode <= 0 && URL.StartsWith(TEXT("Http"), ESearchCase::IgnoreCase))
 			{
-				UE_LOG(LogHttp, Warning, TEXT("%p: invalid HTTP response code received. URL: %s, HTTP code: %d, content length: %d, actual payload size: %d"),
+				UE_LOG(LogHttp, Warning, TEXT("%p: invalid HTTP response code received. URL: %s, HTTP code: %d, content length: %llu, actual payload size: %llu"),
 					this, *GetURL(), Response->HttpCode, Response->ContentLength, Response->TotalBytesRead.GetValue());
 				Response->bSucceeded = false;
 			}
@@ -1224,12 +1221,12 @@ void FCurlHttpRequest::FinishRequest()
 		{
 			if (bDebugServerResponse)
 			{
-				UE_LOG(LogHttp, Warning, TEXT("%p: request has been successfully processed. URL: %s, HTTP code: %d, content length: %d, actual payload size: %d, elapsed: %.2fs"),
+				UE_LOG(LogHttp, Warning, TEXT("%p: request has been successfully processed. URL: %s, HTTP code: %d, content length: %llu, actual payload size: %llu, elapsed: %.2fs"),
 					this, *GetURL(), Response->HttpCode, Response->ContentLength, Response->TotalBytesRead.GetValue(), ElapsedTime);
 			}
 			else
 			{
-				UE_LOG(LogHttp, Log, TEXT("%p: request has been successfully processed. URL: %s, HTTP code: %d, content length: %d, actual payload size: %d, elapsed: %.2fs"),
+				UE_LOG(LogHttp, Log, TEXT("%p: request has been successfully processed. URL: %s, HTTP code: %d, content length: %llu, actual payload size: %llu, elapsed: %.2fs"),
 					this, *GetURL(), Response->HttpCode, Response->ContentLength, Response->TotalBytesRead.GetValue(), ElapsedTime);
 			}
 
@@ -1398,7 +1395,7 @@ FString FCurlHttpResponse::GetContentType() const
 	return GetHeader(TEXT("Content-Type"));
 }
 
-int32 FCurlHttpResponse::GetContentLength() const
+uint64 FCurlHttpResponse::GetContentLength() const
 {
 	return ContentLength;
 }
diff --git a/Engine/Source/Runtime/Online/HTTP/Private/Curl/CurlHttp.h b/Engine/Source/Runtime/Online/HTTP/Private/Curl/CurlHttp.h
--- a/Engine/Source/Runtime/Online/HTTP/Private/Curl/CurlHttp.h
+++ b/Engine/Source/Runtime/Online/HTTP/Private/Curl/CurlHttp.h
@@ -158,7 +158,7 @@ public:
 	virtual FString GetHeader(const FString& HeaderName) const override;
 	virtual TArray<FString> GetAllHeaders() const override;
 	virtual FString GetContentType() const override;
-	virtual int32 GetContentLength() const override;
+	virtual uint64 GetContentLength() const override;
 	virtual const TArray<uint8>& GetContent() const override;
 	//~ End IHttpBase Interface
 
@@ -401,13 +401,13 @@ private:
 	/** Have we had any HTTP activity with the host? Sending headers, SSL handshake, etc */
 	bool bAnyHttpActivity;
 	/** Number of bytes sent already */
-	FThreadSafeCounter BytesSent;
+	FThreadSafeCounter64 BytesSent;
 	/** Total number of bytes sent already (includes data re-sent by seek attempts) */
-	FThreadSafeCounter TotalBytesSent;
+	FThreadSafeCounter64 TotalBytesSent;
 	/** Last bytes read reported to progress delegate */
-	int32 LastReportedBytesRead;
+	uint64 LastReportedBytesRead;
 	/** Last bytes sent reported to progress delegate */
-	int32 LastReportedBytesSent;
+	uint64 LastReportedBytesSent;
 	/** Number of info channel messages to cache */
 	static const constexpr int32 NumberOfInfoMessagesToCache = 50;
 	/** Index of least recently cached message */
@@ -439,7 +439,7 @@ public:
 	virtual FString GetHeader(const FString& HeaderName) const override;
 	virtual TArray<FString> GetAllHeaders() const override;	
 	virtual FString GetContentType() const override;
-	virtual int32 GetContentLength() const override;
+	virtual uint64 GetContentLength() const override;
 	virtual const TArray<uint8>& GetContent() const override;
 	//~ End IHttpBase Interface
 
@@ -467,7 +467,7 @@ private:
 	/** The stream to receive response body */
 	TSharedPtr<FArchive> ResponseBodyReceiveStream;
 	/** Caches how many bytes of the response we've read so far */
-	FThreadSafeCounter TotalBytesRead;
+	FThreadSafeCounter64 TotalBytesRead;
 	/** Cached key/value header pairs. Parsed once request completes. Only accessible on the game thread. */
 	TMap<FString, FString> Headers;
 	/** Newly received headers we need to inform listeners about */
@@ -475,7 +475,7 @@ private:
 	/** Cached code from completed response */
 	int32 HttpCode;
 	/** Cached content length from completed response */
-	int32 ContentLength;
+	uint64 ContentLength;
 	/** True when the response has finished async processing */
 	int32 volatile bIsReady;
 	/** True if the response was successfully received/processed */
diff --git a/Engine/Source/Runtime/Online/HTTP/Private/GenericPlatform/GenericPlatformHttp.cpp b/Engine/Source/Runtime/Online/HTTP/Private/GenericPlatform/GenericPlatformHttp.cpp
--- a/Engine/Source/Runtime/Online/HTTP/Private/GenericPlatform/GenericPlatformHttp.cpp
+++ b/Engine/Source/Runtime/Online/HTTP/Private/GenericPlatform/GenericPlatformHttp.cpp
@@ -87,7 +87,7 @@ public:
 	virtual FString GetHeader(const FString& HeaderName) const override { return TEXT(""); }
 	virtual TArray<FString> GetAllHeaders() const override { return TArray<FString>(); }
 	virtual FString GetContentType() const override { return TEXT(""); }
-	virtual int32 GetContentLength() const override { return 0; }
+	virtual uint64 GetContentLength() const override { return 0; }
 	virtual const TArray<uint8>& GetContent() const override { static TArray<uint8> Temp; return Temp; }
 
 	// IHttpRequest
diff --git a/Engine/Source/Runtime/Online/HTTP/Private/GenericPlatform/HttpRequestPayload.cpp b/Engine/Source/Runtime/Online/HTTP/Private/GenericPlatform/HttpRequestPayload.cpp
--- a/Engine/Source/Runtime/Online/HTTP/Private/GenericPlatform/HttpRequestPayload.cpp
+++ b/Engine/Source/Runtime/Online/HTTP/Private/GenericPlatform/HttpRequestPayload.cpp
@@ -41,9 +41,9 @@ FRequestPayloadInFileStream::~FRequestPayloadInFileStream()
 {
 }
 
-int32 FRequestPayloadInFileStream::GetContentLength() const
+uint64 FRequestPayloadInFileStream::GetContentLength() const
 {
-	return static_cast<int32>(File->TotalSize());
+	return File->TotalSize();
 }
 
 const TArray<uint8>& FRequestPayloadInFileStream::GetContent() const
@@ -67,7 +67,7 @@ size_t FRequestPayloadInFileStream::FillOutputBuffer(void* OutputBuffer, size_t
 
 size_t FRequestPayloadInFileStream::FillOutputBuffer(TArrayView<uint8> OutputBuffer, size_t SizeAlreadySent)
 {
-	const size_t ContentLength = static_cast<size_t>(GetContentLength());
+	const size_t ContentLength = GetContentLength();
 	check(SizeAlreadySent <= ContentLength);
 	const size_t SizeToSend = ContentLength - SizeAlreadySent;
 	const size_t SizeToSendThisTime = FMath::Min(SizeToSend, static_cast<size_t>(OutputBuffer.Num()));
@@ -94,7 +94,7 @@ FRequestPayloadInMemory::~FRequestPayloadInMemory()
 {
 }
 
-int32 FRequestPayloadInMemory::GetContentLength() const
+uint64 FRequestPayloadInMemory::GetContentLength() const
 {
 	return Buffer.Num();
 }
diff --git a/Engine/Source/Runtime/Online/HTTP/Private/IXML/HttpIXML.h b/Engine/Source/Runtime/Online/HTTP/Private/IXML/HttpIXML.h
--- a/Engine/Source/Runtime/Online/HTTP/Private/IXML/HttpIXML.h
+++ b/Engine/Source/Runtime/Online/HTTP/Private/IXML/HttpIXML.h
@@ -28,7 +28,7 @@ public:
 	virtual FString GetHeader(const FString& HeaderName) const override;
 	virtual TArray<FString> GetAllHeaders() const override;
 	virtual FString GetContentType() const override;
-	virtual int32 GetContentLength() const override;
+	virtual uint64 GetContentLength() const override;
 	virtual const TArray<uint8>& GetContent() const override;
 
 	// IHttpRequest
@@ -103,7 +103,7 @@ public:
 	virtual FString GetHeader(const FString& HeaderName) const override;
 	virtual TArray<FString> GetAllHeaders() const override;
 	virtual FString GetContentType() const override;
-	virtual int32 GetContentLength() const override;
+	virtual uint64 GetContentLength() const override;
 	virtual const TArray<uint8>& GetContent() const override;
 
 	// IHttpResponse
diff --git a/Engine/Source/Runtime/Online/HTTP/Private/IXML/HttpIXMLSupport.cpp b/Engine/Source/Runtime/Online/HTTP/Private/IXML/HttpIXMLSupport.cpp
--- a/Engine/Source/Runtime/Online/HTTP/Private/IXML/HttpIXMLSupport.cpp
+++ b/Engine/Source/Runtime/Online/HTTP/Private/IXML/HttpIXMLSupport.cpp
@@ -292,7 +292,7 @@ RequestStream::~RequestStream()
 //  Name: RequestStream::Open
 //  Desc: Opens the buffer populated with the supplied data
 //--------------------------------------------------------------------------------------
-STDMETHODIMP RequestStream::Open( LPCSTR InBuffer, ULONG BufferSize )
+STDMETHODIMP RequestStream::Open( LPCSTR InBuffer, size_t BufferSize )
 {
 	HRESULT hr = S_OK;
 
diff --git a/Engine/Source/Runtime/Online/HTTP/Private/IXML/HttpIXMLSupport.h b/Engine/Source/Runtime/Online/HTTP/Private/IXML/HttpIXMLSupport.h
--- a/Engine/Source/Runtime/Online/HTTP/Private/IXML/HttpIXMLSupport.h
+++ b/Engine/Source/Runtime/Online/HTTP/Private/IXML/HttpIXMLSupport.h
@@ -76,7 +76,7 @@ public:
 	~RequestStream();
 
 	// ISequentialStream
-	STDMETHODIMP Open( LPCSTR Buffer, ULONG BufferSize );
+	STDMETHODIMP Open( LPCSTR Buffer, size_t BufferSize );
 	STDMETHODIMP Read( void *Data, ULONG ByteCount, ULONG *NumReadBytes );
 	STDMETHODIMP Write( const void *Data,  ULONG ByteCount, ULONG *Written );
 
diff --git a/Engine/Source/Runtime/Online/HTTP/Private/IXML/HttpIXml.cpp b/Engine/Source/Runtime/Online/HTTP/Private/IXML/HttpIXml.cpp
--- a/Engine/Source/Runtime/Online/HTTP/Private/IXML/HttpIXml.cpp
+++ b/Engine/Source/Runtime/Online/HTTP/Private/IXML/HttpIXml.cpp
@@ -77,7 +77,7 @@ FString FHttpRequestIXML::GetContentType() const
 //-----------------------------------------------------------------------------
 //	
 //-----------------------------------------------------------------------------
-int32 FHttpRequestIXML::GetContentLength() const
+uint64 FHttpRequestIXML::GetContentLength() const
 {
 	return Payload->GetContentLength();
 }
@@ -348,12 +348,12 @@ uint32 FHttpRequestIXML::SendRequest()
 
 	if( Payload )
 	{
-		uint32 SizeInBytes = Payload->GetContentLength();
+		size_t SizeInBytes = Payload->GetContentLength();
 
 		// Create and open a new runtime class
 		SendStream = Make<RequestStream>();
 		LPCSTR PayloadChars = (LPCSTR)Payload->GetContent().GetData();
-		SendStream->Open( PayloadChars, (ULONG) SizeInBytes );
+		SendStream->Open( PayloadChars, SizeInBytes );
 
 		hr = XHR->Send(
 			SendStream.Get(),        // body message as an ISequentialStream*
@@ -589,7 +589,7 @@ FString FHttpResponseIXML::GetContentType() const
 //	
 //-----------------------------------------------------------------------------
 
-int32 FHttpResponseIXML::GetContentLength() const
+int64 FHttpResponseIXML::GetContentLength() const
 {
 	check ( HttpCB );
 	
diff --git a/Engine/Source/Runtime/Online/HTTP/Private/NullHttp.cpp b/Engine/Source/Runtime/Online/HTTP/Private/NullHttp.cpp
--- a/Engine/Source/Runtime/Online/HTTP/Private/NullHttp.cpp
+++ b/Engine/Source/Runtime/Online/HTTP/Private/NullHttp.cpp
@@ -42,7 +42,7 @@ FString FNullHttpRequest::GetContentType() const
 	return GetHeader(TEXT("Content-Type"));
 }
 
-int32 FNullHttpRequest::GetContentLength() const
+uint64 FNullHttpRequest::GetContentLength() const
 {
 	return Payload.Num();
 }
@@ -227,7 +227,7 @@ FString FNullHttpResponse::GetContentType() const
 	return FString();
 }
 
-int32 FNullHttpResponse::GetContentLength() const
+uint64 FNullHttpResponse::GetContentLength() const
 {
 	return 0;
 }
diff --git a/Engine/Source/Runtime/Online/HTTP/Private/NullHttp.h b/Engine/Source/Runtime/Online/HTTP/Private/NullHttp.h
--- a/Engine/Source/Runtime/Online/HTTP/Private/NullHttp.h
+++ b/Engine/Source/Runtime/Online/HTTP/Private/NullHttp.h
@@ -19,7 +19,7 @@ public:
 	virtual FString GetHeader(const FString& HeaderName) const override;
 	virtual TArray<FString> GetAllHeaders() const override;	
 	virtual FString GetContentType() const override;
-	virtual int32 GetContentLength() const override;
+	virtual uint64 GetContentLength() const override;
 	virtual const TArray<uint8>& GetContent() const override;
 	// IHttpRequest 
 	virtual FString GetVerb() const override;
@@ -70,7 +70,7 @@ class FNullHttpResponse : public IHttpResponse
 	virtual FString GetHeader(const FString& HeaderName) const override;
 	virtual TArray<FString> GetAllHeaders() const override;	
 	virtual FString GetContentType() const override;
-	virtual int32 GetContentLength() const override;
+	virtual uint64 GetContentLength() const override;
 	virtual const TArray<uint8>& GetContent() const override;
 	//~ Begin IHttpResponse Interface
 	virtual int32 GetResponseCode() const override;
diff --git a/Engine/Source/Runtime/Online/HTTP/Private/WinHttp/Support/WinHttpConnectionHttp.cpp b/Engine/Source/Runtime/Online/HTTP/Private/WinHttp/Support/WinHttpConnectionHttp.cpp
--- a/Engine/Source/Runtime/Online/HTTP/Private/WinHttp/Support/WinHttpConnectionHttp.cpp
+++ b/Engine/Source/Runtime/Online/HTTP/Private/WinHttp/Support/WinHttpConnectionHttp.cpp
@@ -144,10 +144,10 @@ bool FWinHttpConnectionHttp::StartRequest()
 	// Setup our buffer
 	if (Payload.IsValid())
 	{
-		const int32 NumBytesToWriteNow = FMath::Min(UE_WINHTTP_WRITE_BUFFER_BYTES, Payload->GetContentLength());
+		const uint64 NumBytesToWriteNow = FMath::Min((uint64)(UE_WINHTTP_WRITE_BUFFER_BYTES), Payload->GetContentLength());
 		PayloadBuffer.SetNumUninitialized(NumBytesToWriteNow, false);
 
-		const int32 BufferSize = Payload->FillOutputBuffer(MakeArrayView(PayloadBuffer), 0);
+		const uint64 BufferSize = Payload->FillOutputBuffer(MakeArrayView(PayloadBuffer), 0);
 		PayloadBuffer.SetNumUninitialized(BufferSize, false);
 	}
 
@@ -233,7 +233,7 @@ void FWinHttpConnectionHttp::PumpMessages()
 		{
 			GameThreadChunk.Append(CurrentChunk);
 		}
-		const int32 ReserveChunkSize = ResponseContentLength >= BytesWrittenToGameThreadChunk 
+		const uint64 ReserveChunkSize = ResponseContentLength >= BytesWrittenToGameThreadChunk 
 			? (ResponseContentLength - BytesWrittenToGameThreadChunk) 
 			: UE_WINHTTP_READ_BUFFER_BYTES;
 		CurrentChunk.Reset(ReserveChunkSize);
@@ -242,10 +242,10 @@ void FWinHttpConnectionHttp::PumpMessages()
 	// Process Data Transfer callbacks
 	if (BytesToReportSent.IsSet() || BytesToReportReceived.IsSet())
 	{
-		const int32 BytesSent = BytesToReportSent.Get(0);
+		const uint64 BytesSent = BytesToReportSent.Get(0);
 		BytesToReportSent.Reset();
 
-		const int32 BytesReceived = BytesToReportReceived.Get(0);
+		const uint64 BytesReceived = BytesToReportReceived.Get(0);
 		BytesToReportReceived.Reset();
 
 		OnDataTransferredHandler.ExecuteIfBound(BytesSent, BytesReceived);
@@ -424,9 +424,9 @@ FWinHttpConnectionHttp::FWinHttpConnectionHttp(
 	: RequestUrl(InUrl)
 {
 	const uint32 LogPort = InPort.Get(bInIsSecure ? 443 : 80);
-	const int32 LogPayloadSize = InPayload.IsValid() ? InPayload->GetContentLength() : 0;
+	const uint64 LogPayloadSize = InPayload.IsValid() ? InPayload->GetContentLength() : 0;
 
-	UE_LOG(LogWinHttp, Verbose, TEXT("WinHttp Http[%p]: Creating request. InVerb=[%s] bIsSecure=[%d] Domain=[%s] Port=[%u] Path=[%s] PaylodSize=[%d]"), this, *InVerb, bInIsSecure, *InDomain, LogPort, *InPathAndQuery, LogPayloadSize);
+	UE_LOG(LogWinHttp, Verbose, TEXT("WinHttp Http[%p]: Creating request. InVerb=[%s] bIsSecure=[%d] Domain=[%s] Port=[%u] Path=[%s] PaylodSize=[%llu]"), this, *InVerb, bInIsSecure, *InDomain, LogPort, *InPathAndQuery, LogPayloadSize);
 
 	// Note: Microsoft say not to reuse Connection Handles for multiple requests, despite what the API would suggest!  Microsoft say the Session handle
 	// is to be reused amongst requests with the same Security Protocol. If the same Session is used, backing sockets for connections will be reused
@@ -646,7 +646,7 @@ bool FWinHttpConnectionHttp::SendRequest()
 	return true;
 }
 
-void FWinHttpConnectionHttp::IncrementSentByteCounts(const int32 AmountSent)
+void FWinHttpConnectionHttp::IncrementSentByteCounts(const uint64 AmountSent)
 {
 	check(CurrentAction == EState::WaitForSendComplete);
 
@@ -664,7 +664,7 @@ void FWinHttpConnectionHttp::IncrementSentByteCounts(const int32 AmountSent)
 	NumBytesSuccessfullySent += AmountSent;
 }
 
-void FWinHttpConnectionHttp::IncrementReceivedByteCounts(const int32 AmountReceived)
+void FWinHttpConnectionHttp::IncrementReceivedByteCounts(const uint64 AmountReceived)
 {
 	check(CurrentAction == EState::WaitForNextResponseBodyChunkData);
 
@@ -696,15 +696,15 @@ bool FWinHttpConnectionHttp::SendAdditionalRequestBody()
 	check(HasRequestBodyToSend());
 	check(Payload.IsValid());
 
-	const int32 TotalBytesLeftToWrite = Payload->GetContentLength() - NumBytesSuccessfullySent;
+	const int64 TotalBytesLeftToWrite = Payload->GetContentLength() - NumBytesSuccessfullySent;
 	check(TotalBytesLeftToWrite > 0);
 	
 	// Resize buffer to max amount of data we can write
-	const int32 OptimalAmountToWrite = FMath::Min(TotalBytesLeftToWrite, UE_WINHTTP_WRITE_BUFFER_BYTES);
+	const int64 OptimalAmountToWrite = FMath::Min(TotalBytesLeftToWrite, UE_WINHTTP_WRITE_BUFFER_BYTES);
 	PayloadBuffer.SetNumUninitialized(OptimalAmountToWrite, false);
 
 	// Read data into our buffer if possible
-	const int32 ActualDataSize = Payload->FillOutputBuffer(MakeArrayView(PayloadBuffer), NumBytesSuccessfullySent);
+	const int64 ActualDataSize = Payload->FillOutputBuffer(MakeArrayView(PayloadBuffer), NumBytesSuccessfullySent);
 	if (ActualDataSize < 1)
 	{
 		// Set our buffer to be empty since we didn't write anything into it
@@ -715,7 +715,7 @@ bool FWinHttpConnectionHttp::SendAdditionalRequestBody()
 	// Resize our buffer based on how much was actually written to it
 	PayloadBuffer.SetNumUninitialized(ActualDataSize, false);
 
-	UE_LOG(LogWinHttp, VeryVerbose, TEXT("WinHttp Http[%p]: Writing Data. NumBytes=[%d] TotalBytesWritten=[%d] TotalBytes=[%d]"), this, PayloadBuffer.Num(), NumBytesSuccessfullySent, Payload->GetContentLength())
+	UE_LOG(LogWinHttp, VeryVerbose, TEXT("WinHttp Http[%p]: Writing Data. NumBytes=[%llu] TotalBytesWritten=[%llu] TotalBytes=[%llu]"), this, PayloadBuffer.Num(), NumBytesSuccessfullySent, Payload->GetContentLength())
 
 	CurrentAction = EState::WaitForSendComplete;
 	if (!WinHttpWriteData(RequestHandle.Get(), PayloadBuffer.GetData(), PayloadBuffer.Num(), NULL))
@@ -893,7 +893,7 @@ bool FWinHttpConnectionHttp::RequestNextResponseBodyChunkData()
 	check(CurrentAction == EState::RequestNextResponseBodyChunkData);
 	check(ResponseBytesAvailable.IsSet());
 
-	const int32 NumBytesAvailable = ResponseBytesAvailable.GetValue();
+	const uint64 NumBytesAvailable = ResponseBytesAvailable.GetValue();
 	ResponseBytesAvailable.Reset();
 
 	if (NumBytesAvailable == 0)
@@ -904,7 +904,7 @@ bool FWinHttpConnectionHttp::RequestNextResponseBodyChunkData()
 		return true;
 	}
 
-	int32 ResponseBytesWritten = CurrentChunk.Num();
+	uint64 ResponseBytesWritten = CurrentChunk.Num();
 	CurrentChunk.AddUninitialized(NumBytesAvailable);
 
 	CurrentAction = EState::WaitForNextResponseBodyChunkData;
@@ -980,9 +980,9 @@ void FWinHttpConnectionHttp::HandleSendingRequest()
 	}
 }
 
-void FWinHttpConnectionHttp::HandleWriteComplete(const uint32 NumBytesSent)
+void FWinHttpConnectionHttp::HandleWriteComplete(const uint64 NumBytesSent)
 {
-	UE_LOG(LogWinHttp, VeryVerbose, TEXT("WinHttp Http[%p]: Callback Status=[WRITE_COMPLETE] NumBytesSent=[%d]"), this, NumBytesSent);
+	UE_LOG(LogWinHttp, VeryVerbose, TEXT("WinHttp Http[%p]: Callback Status=[WRITE_COMPLETE] NumBytesSent=[%llu]"), this, NumBytesSent);
 
 	FScopeLock ScopeLock(&SyncObject);
 
@@ -1060,9 +1060,9 @@ void FWinHttpConnectionHttp::HandleHeadersAvailable()
 	ReleasePayloadData();
 }
 
-void FWinHttpConnectionHttp::HandleDataAvailable(const uint32 NumBytesAvailable)
+void FWinHttpConnectionHttp::HandleDataAvailable(const uint64 NumBytesAvailable)
 {
-	UE_LOG(LogWinHttp, VeryVerbose, TEXT("WinHttp Http[%p]: Callback Status=[DATA_AVAILABLE] NumBytesAvailable=[%u]"), this, NumBytesAvailable);
+	UE_LOG(LogWinHttp, VeryVerbose, TEXT("WinHttp Http[%p]: Callback Status=[DATA_AVAILABLE] NumBytesAvailable=[%llu]"), this, NumBytesAvailable);
 
 	FScopeLock ScopeLock(&SyncObject);
 
@@ -1075,9 +1075,9 @@ void FWinHttpConnectionHttp::HandleDataAvailable(const uint32 NumBytesAvailable)
 	CurrentAction = EState::RequestNextResponseBodyChunkData;
 }
 
-void FWinHttpConnectionHttp::HandleReadComplete(const uint32 NumBytesRead)
+void FWinHttpConnectionHttp::HandleReadComplete(const uint64 NumBytesRead)
 {
-	UE_LOG(LogWinHttp, VeryVerbose, TEXT("WinHttp Http[%p]: Callback Status=[READ_COMPLETE] NumBytesRead=[%u]"), this, NumBytesRead);
+	UE_LOG(LogWinHttp, VeryVerbose, TEXT("WinHttp Http[%p]: Callback Status=[READ_COMPLETE] NumBytesRead=[%llu]"), this, NumBytesRead);
 
 	FScopeLock ScopeLock(&SyncObject);
 	
@@ -1270,7 +1270,7 @@ void FWinHttpConnectionHttp::HandleHttpStatusCallback(HINTERNET ResourceHandle,
 			check(StatusInformationLength == sizeof(DWORD));
 			check(StatusInformation != nullptr);
 			const DWORD NumBytesSent = *static_cast<DWORD*>(StatusInformation);
-			UE_LOG(LogWinHttp, Verbose, TEXT("WinHttp Http[%p]: Callback Status=[REQUEST_SENT] NumBytesSent=[%d]"), this, static_cast<int32>(NumBytesSent));
+			UE_LOG(LogWinHttp, Verbose, TEXT("WinHttp Http[%p]: Callback Status=[REQUEST_SENT] NumBytesSent=[%llu]"), this, NumBytesSent);
 			return;
 		}
 		case EWinHttpCallbackStatus::ConnectingToServer:
diff --git a/Engine/Source/Runtime/Online/HTTP/Private/WinHttp/WinHttpHttpRequest.cpp b/Engine/Source/Runtime/Online/HTTP/Private/WinHttp/WinHttpHttpRequest.cpp
--- a/Engine/Source/Runtime/Online/HTTP/Private/WinHttp/WinHttpHttpRequest.cpp
+++ b/Engine/Source/Runtime/Online/HTTP/Private/WinHttp/WinHttpHttpRequest.cpp
@@ -65,7 +65,7 @@ FString FWinHttpHttpRequest::GetContentType() const
 	return GetHeader(TEXT("Content-Type"));
 }
 
-int32 FWinHttpHttpRequest::GetContentLength() const
+uint64 FWinHttpHttpRequest::GetContentLength() const
 {
 	return RequestData.Payload.IsValid() ? RequestData.Payload->GetContentLength() : 0;
 }
diff --git a/Engine/Source/Runtime/Online/HTTP/Private/WinHttp/WinHttpHttpRequest.h b/Engine/Source/Runtime/Online/HTTP/Private/WinHttp/WinHttpHttpRequest.h
--- a/Engine/Source/Runtime/Online/HTTP/Private/WinHttp/WinHttpHttpRequest.h
+++ b/Engine/Source/Runtime/Online/HTTP/Private/WinHttp/WinHttpHttpRequest.h
@@ -28,7 +28,7 @@ public:
 	virtual FString GetHeader(const FString& HeaderName) const override;
 	virtual TArray<FString> GetAllHeaders() const override;	
 	virtual FString GetContentType() const override;
-	virtual int32 GetContentLength() const override;
+	virtual uint64 GetContentLength() const override;
 	virtual const TArray<uint8>& GetContent() const override;
 	//~ End IHttpBase Interface
 
diff --git a/Engine/Source/Runtime/Online/HTTP/Private/WinHttp/WinHttpHttpResponse.cpp b/Engine/Source/Runtime/Online/HTTP/Private/WinHttp/WinHttpHttpResponse.cpp
--- a/Engine/Source/Runtime/Online/HTTP/Private/WinHttp/WinHttpHttpResponse.cpp
+++ b/Engine/Source/Runtime/Online/HTTP/Private/WinHttp/WinHttpHttpResponse.cpp
@@ -54,7 +54,7 @@ FString FWinHttpHttpResponse::GetContentType() const
 	return GetHeader(TEXT("Content-Type"));
 }
 
-int32 FWinHttpHttpResponse::GetContentLength() const
+uint64 FWinHttpHttpResponse::GetContentLength() const
 {
 	return Payload.Num();
 }
diff --git a/Engine/Source/Runtime/Online/HTTP/Private/WinHttp/WinHttpHttpResponse.h b/Engine/Source/Runtime/Online/HTTP/Private/WinHttp/WinHttpHttpResponse.h
--- a/Engine/Source/Runtime/Online/HTTP/Private/WinHttp/WinHttpHttpResponse.h
+++ b/Engine/Source/Runtime/Online/HTTP/Private/WinHttp/WinHttpHttpResponse.h
@@ -23,7 +23,7 @@ public:
 	virtual FString GetHeader(const FString& HeaderName) const override;
 	virtual TArray<FString> GetAllHeaders() const override;	
 	virtual FString GetContentType() const override;
-	virtual int32 GetContentLength() const override;
+	virtual uint64 GetContentLength() const override;
 	virtual const TArray<uint8>& GetContent() const override;
 	//~ End IHttpBase Interface
 
diff --git a/Engine/Source/Runtime/Online/HTTP/Public/GenericPlatform/HttpRequestPayload.h b/Engine/Source/Runtime/Online/HTTP/Public/GenericPlatform/HttpRequestPayload.h
--- a/Engine/Source/Runtime/Online/HTTP/Public/GenericPlatform/HttpRequestPayload.h
+++ b/Engine/Source/Runtime/Online/HTTP/Public/GenericPlatform/HttpRequestPayload.h
@@ -18,7 +18,7 @@ class FRequestPayload
 public:
 	virtual ~FRequestPayload() {}
 	/** Get the total content length of the request payload in bytes */
-	virtual int32 GetContentLength() const = 0;
+	virtual uint64 GetContentLength() const = 0;
 	/** Return a reference to the underlying memory buffer. Only valid for in-memory request payloads */
 	virtual const TArray<uint8>& GetContent() const = 0;
 	/** Check if the request payload is URL encoded. This check is only performed for in-memory request payloads */
@@ -47,7 +47,7 @@ class FRequestPayloadInFileStream : public FRequestPayload
 public:
 	FRequestPayloadInFileStream(TSharedRef<FArchive, ESPMode::ThreadSafe> InFile);
 	virtual ~FRequestPayloadInFileStream();
-	virtual int32 GetContentLength() const override;
+	virtual uint64 GetContentLength() const override;
 	virtual const TArray<uint8>& GetContent() const override;
 	virtual bool IsURLEncoded() const override;
 	virtual size_t FillOutputBuffer(void* OutputBuffer, size_t MaxOutputBufferSize, size_t SizeAlreadySent) override;
@@ -62,7 +62,7 @@ public:
 	FRequestPayloadInMemory(const TArray<uint8>& Array);
 	FRequestPayloadInMemory(TArray<uint8>&& Array);
 	virtual ~FRequestPayloadInMemory();
-	virtual int32 GetContentLength() const override;
+	virtual uint64 GetContentLength() const override;
 	virtual const TArray<uint8>& GetContent() const override;
 	virtual bool IsURLEncoded() const override;
 	virtual size_t FillOutputBuffer(void* OutputBuffer, size_t MaxOutputBufferSize, size_t SizeAlreadySent) override;
diff --git a/Engine/Source/Runtime/Online/HTTP/Public/HttpRequestAdapter.h b/Engine/Source/Runtime/Online/HTTP/Public/HttpRequestAdapter.h
--- a/Engine/Source/Runtime/Online/HTTP/Public/HttpRequestAdapter.h
+++ b/Engine/Source/Runtime/Online/HTTP/Public/HttpRequestAdapter.h
@@ -23,7 +23,7 @@ public:
 	virtual FString                       GetHeader(const FString& HeaderName) const override                      { return HttpRequest->GetHeader(HeaderName); }
 	virtual TArray<FString>               GetAllHeaders() const override                                           { return HttpRequest->GetAllHeaders(); }
 	virtual FString                       GetContentType() const override                                          { return HttpRequest->GetContentType(); }
-	virtual int32                         GetContentLength() const override                                        { return HttpRequest->GetContentLength(); }
+	virtual uint64                        GetContentLength() const override                                        { return HttpRequest->GetContentLength(); }
 	virtual const TArray<uint8>&          GetContent() const override                                              { return HttpRequest->GetContent(); }
 	virtual FString                       GetVerb() const override                                                 { return HttpRequest->GetVerb(); }
 	virtual void                          SetVerb(const FString& Verb) override                                    { HttpRequest->SetVerb(Verb); }
diff --git a/Engine/Source/Runtime/Online/HTTP/Public/Interfaces/IHttpBase.h b/Engine/Source/Runtime/Online/HTTP/Public/Interfaces/IHttpBase.h
--- a/Engine/Source/Runtime/Online/HTTP/Public/Interfaces/IHttpBase.h
+++ b/Engine/Source/Runtime/Online/HTTP/Public/Interfaces/IHttpBase.h
@@ -55,7 +55,7 @@ public:
 	 *
 	 * @return the content length (if available)
 	 */
-	virtual int32 GetContentLength() const = 0;
+	virtual uint64 GetContentLength() const = 0;
 
 	/**
 	 * Get the content payload of the request or response.
