[Engine]

diff --git a/Engine/Source/Runtime/Online/BackgroundHTTP/Public/PlatformWithModularFeature/ModularFeaturePlatformBackgroundHttp.h b/Engine/Source/Runtime/Online/BackgroundHTTP/Public/PlatformWithModularFeature/ModularFeaturePlatformBackgroundHttp.h
--- a/Engine/Source/Runtime/Online/BackgroundHTTP/Public/PlatformWithModularFeature/ModularFeaturePlatformBackgroundHttp.h
+++ b/Engine/Source/Runtime/Online/BackgroundHTTP/Public/PlatformWithModularFeature/ModularFeaturePlatformBackgroundHttp.h
@@ -17,24 +17,24 @@ class IBackgroundHttpModularFeature;
  * This version of BackgroundHttp is designed to be used by any platform that relies on a ModularFeature to override the BackgroundHttp behavior. 
  * If no modular feature is found, we fall back on the GenericPlatform implementation to provide functionality
  */
-class BACKGROUNDHTTP_API FModularFeaturePlatformBackgroundHttp
+class FModularFeaturePlatformBackgroundHttp
 {
 public:
 	/**
 	 * Platform initialization step
 	 */
-	static void Initialize();
+	static BACKGROUNDHTTP_API void Initialize();
 
 	/**
 	 * Platform shutdown step
 	 */
-	static void Shutdown();
+	static BACKGROUNDHTTP_API void Shutdown();
 
 	/**
 	 * Creates a platform-specific Background HTTP manager.
 	 * Un-implemented platforms should create a FGenericPlatformBackgroundHttpManager
 	 */
-	static FBackgroundHttpManagerPtr CreatePlatformBackgroundHttpManager();
+	static BACKGROUNDHTTP_API FBackgroundHttpManagerPtr CreatePlatformBackgroundHttpManager();
 
 	/**
 	 * Creates a new Background Http request instance for the current platform
@@ -42,7 +42,7 @@ public:
 	 *
 	 * @return request object
 	 */
-	static FBackgroundHttpRequestPtr ConstructBackgroundRequest();
+	static BACKGROUNDHTTP_API FBackgroundHttpRequestPtr ConstructBackgroundRequest();
 
 	/**
 	 * Creates a new Background Http Response instance for the current platform
@@ -50,13 +50,13 @@ public:
 	 *
 	 * @return response object
 	 */
-	static FBackgroundHttpResponsePtr ConstructBackgroundResponse(int32 ResponseCode, const FString& TempFilePath);
+	static BACKGROUNDHTTP_API FBackgroundHttpResponsePtr ConstructBackgroundResponse(int32 ResponseCode, const FString& TempFilePath);
 
 public:
-	static FName GetModularFeatureName();
+	static BACKGROUNDHTTP_API FName GetModularFeatureName();
 
 protected:
-	static void CacheModularFeature();
+	static BACKGROUNDHTTP_API void CacheModularFeature();
 
 private:
 	static IBackgroundHttpModularFeature* CachedModularFeature;
diff --git a/Engine/Source/Runtime/Online/BackgroundHTTPFileHash/Public/BackgroundHttpFileHashHelper.h b/Engine/Source/Runtime/Online/BackgroundHTTPFileHash/Public/BackgroundHttpFileHashHelper.h
--- a/Engine/Source/Runtime/Online/BackgroundHTTPFileHash/Public/BackgroundHttpFileHashHelper.h
+++ b/Engine/Source/Runtime/Online/BackgroundHTTPFileHash/Public/BackgroundHttpFileHashHelper.h
@@ -41,7 +41,7 @@ public:
 *NOTE: This is in it's own separate module as some platforms (IOS, etc.) need to include this functionality in ApplicationCore
 *and thus fully adding a dependency to those modules on BackgroundHTTP would not work
 */
-class BACKGROUNDHTTPFILEHASH_API FBackgroundHttpFileHashHelper
+class FBackgroundHttpFileHashHelper
 {
 public:
 	FBackgroundHttpFileHashHelper()
@@ -52,55 +52,55 @@ public:
 	}
 	
 	//Handles loading URL Mapping data from disk. No operation if data has previously been loaded.
-	void LoadData();
+	BACKGROUNDHTTPFILEHASH_API void LoadData();
 	
 	//Handles saving URL Mapping data to disk. No operation if data isn't dirty
-	void SaveData();
+	BACKGROUNDHTTPFILEHASH_API void SaveData();
 	
 	//Creates an FString representation of this FBackgroundHttpFileHashHelper's URL mappings. Used for serializing to/from disk
-	FString ToString() const;
+	BACKGROUNDHTTPFILEHASH_API FString ToString() const;
 
 	//Sets the values of this FBackgroundHttpFileHashHelper from the string representation
 	//Returns if the parse was successful or not
-	bool InitFromString(const FString& StringIn);
+	BACKGROUNDHTTPFILEHASH_API bool InitFromString(const FString& StringIn);
 
 	//Looks for a temp filename mapping for the given URL. If one isn't found returns nullptr and does NOT generate one.
-	const FString* FindTempFilenameMappingForURL(const FString& URL) const;
+	BACKGROUNDHTTPFILEHASH_API const FString* FindTempFilenameMappingForURL(const FString& URL) const;
 	
 	//Looks for a temp filename mapping for the given URL. If one isn't found a new one is generated and returned.
-	const FString& FindOrAddTempFilenameMappingForURL(const FString& URL);
+	BACKGROUNDHTTPFILEHASH_API const FString& FindOrAddTempFilenameMappingForURL(const FString& URL);
 
 	//Removes URL mapping
-	void RemoveURLMapping(const FString& URL);
+	BACKGROUNDHTTPFILEHASH_API void RemoveURLMapping(const FString& URL);
 	
 	//Looks for a URL mapped to the supplied TempFilename. Nullptr returned if one isn't found
-	const FString* FindMappedURLForTempFilename(const FString& TempFilename) const;
+	BACKGROUNDHTTPFILEHASH_API const FString* FindMappedURLForTempFilename(const FString& TempFilename) const;
 	
 	//Deletes any URLMapping that doesn't have a corresponding Temp file actually on disk.
-	void DeleteURLMappingsWithoutTempFiles();
+	BACKGROUNDHTTPFILEHASH_API void DeleteURLMappingsWithoutTempFiles();
 	
 	//Get our base directory used to store Temp files
-	static const FString& GetTemporaryRootPath();
+	static BACKGROUNDHTTPFILEHASH_API const FString& GetTemporaryRootPath();
 	
 	//Helper function that returns the file extension used by our BackgroundHTTP temp files
-	static const FString& GetTempFileExtension();
+	static BACKGROUNDHTTPFILEHASH_API const FString& GetTempFileExtension();
 	
 	//Helper function that returns a full path representation of hte given TempFilename
-	static FString GetFullPathOfTempFilename(const FString& TempFilename);
+	static BACKGROUNDHTTPFILEHASH_API FString GetFullPathOfTempFilename(const FString& TempFilename);
 	
 private:
 	
 	//Helper function that just returns the filepath we use to save/load our URL mapping data
-	static const FString& GetURLMappingFilePath();
+	static BACKGROUNDHTTPFILEHASH_API const FString& GetURLMappingFilePath();
 	
 	//Looks for a temp filename mapping for the given URL. If one isn't found returns nullptr an does NOT generate one.
-	const FURLTempFileMapping* FindMappingForURL(const FString& URL) const;
+	BACKGROUNDHTTPFILEHASH_API const FURLTempFileMapping* FindMappingForURL(const FString& URL) const;
 	
 	//Finds the hash that should be used to refer to this URL. Either finds the hash used in the map to reference the
 	//URLFileMappings entry for this URL or returns the hash we should use (the next valid hash entry) without inserting it into the map.
-	FString FindValidFilenameHashForURL(const FString& URL) const;
+	BACKGROUNDHTTPFILEHASH_API FString FindValidFilenameHashForURL(const FString& URL) const;
 	//Creates a hashed Filename with a modifier for the given collision count.
-	FString GenerateHashedFilenameForURL(const FString& URL, uint32 CollisionCount) const;
+	BACKGROUNDHTTPFILEHASH_API FString GenerateHashedFilenameForURL(const FString& URL, uint32 CollisionCount) const;
 
 	//Used to make sure that we have called LoadData before any operations that depend on it
 	bool bHasLoadedURLData;
diff --git a/Engine/Source/Runtime/Online/BuildPatchServices/Public/BuildPatchServicesSingleton.h b/Engine/Source/Runtime/Online/BuildPatchServices/Public/BuildPatchServicesSingleton.h
--- a/Engine/Source/Runtime/Online/BuildPatchServices/Public/BuildPatchServicesSingleton.h
+++ b/Engine/Source/Runtime/Online/BuildPatchServices/Public/BuildPatchServicesSingleton.h
@@ -7,7 +7,7 @@
 #include "Modules/ModuleManager.h"
 #include "UObject/NameTypes.h"
 
-class BUILDPATCHSERVICES_API FBuildPatchServices
+class FBuildPatchServices
 {
 public:
 	static bool IsAvailable()
@@ -48,6 +48,6 @@ public:
 
 private:
 
-	static FName ModuleName;
-	static BuildPatchServices::FBuildPatchServicesInitSettings InitSettings;
+	static BUILDPATCHSERVICES_API FName ModuleName;
+	static BUILDPATCHSERVICES_API BuildPatchServices::FBuildPatchServicesInitSettings InitSettings;
 };
diff --git a/Engine/Source/Runtime/Online/BuildPatchServices/Public/BuildPatchSettings.h b/Engine/Source/Runtime/Online/BuildPatchServices/Public/BuildPatchSettings.h
--- a/Engine/Source/Runtime/Online/BuildPatchServices/Public/BuildPatchSettings.h
+++ b/Engine/Source/Runtime/Online/BuildPatchServices/Public/BuildPatchSettings.h
@@ -25,13 +25,13 @@ namespace BuildPatchServices
 	/**
 	 * Defines a list of all build patch services initialization settings, can be used to override default init behaviors.
 	 */
-	struct BUILDPATCHSERVICES_API FBuildPatchServicesInitSettings
+	struct FBuildPatchServicesInitSettings
 	{
 	public:
 		/**
 		 * Default constructor. Initializes all members with default behavior values.
 		 */
-		FBuildPatchServicesInitSettings();
+		BUILDPATCHSERVICES_API FBuildPatchServicesInitSettings();
 
 	public:
 		// The application settings directory.
@@ -42,7 +42,7 @@ namespace BuildPatchServices
 		FString LocalMachineConfigFileName;
 	};
 
-	struct BUILDPATCHSERVICES_API FInstallerAction
+	struct FInstallerAction
 	{
 	public:
 
@@ -52,7 +52,7 @@ namespace BuildPatchServices
 		 * @param InstallTags       The install tags to use if selectively installing files. If empty set, all files will be installed.
 		 * @return the action setup for performing an installation.
 		 */
-		static FInstallerAction MakeInstall(const IBuildManifestRef& Manifest, TSet<FString> InstallTags = TSet<FString>());
+		static BUILDPATCHSERVICES_API FInstallerAction MakeInstall(const IBuildManifestRef& Manifest, TSet<FString> InstallTags = TSet<FString>());
 
 		/**
 		 * Creates an update action.
@@ -61,7 +61,7 @@ namespace BuildPatchServices
 		 * @param InstallTags       The install tags to use if selectively installing files. If empty set, all files will be updated, or added if missing.
 		 * @return the action setup for performing an update.
 		 */
-		static FInstallerAction MakeUpdate(const IBuildManifestRef& CurrentManifest, const IBuildManifestRef& InstallManifest, TSet<FString> InstallTags = TSet<FString>());
+		static BUILDPATCHSERVICES_API FInstallerAction MakeUpdate(const IBuildManifestRef& CurrentManifest, const IBuildManifestRef& InstallManifest, TSet<FString> InstallTags = TSet<FString>());
 
 		/**
 		 * Creates an install action.
@@ -69,7 +69,7 @@ namespace BuildPatchServices
 		 * @param InstallTags       The install tags to use if selectively installing files. If empty set, all files will be repaired.
 		 * @return the action setup forcing an SHA check, and repair of all tagged files.
 		 */
-		static FInstallerAction MakeRepair(const IBuildManifestRef& Manifest, TSet<FString> InstallTags = TSet<FString>());
+		static BUILDPATCHSERVICES_API FInstallerAction MakeRepair(const IBuildManifestRef& Manifest, TSet<FString> InstallTags = TSet<FString>());
 
 
 		/**
@@ -77,7 +77,7 @@ namespace BuildPatchServices
 		 * @param Manifest          The manifest for the build currently installed.
 		 * @return the action setup for performing an uninstall, deleting all files referenced by the manifest.
 		 */
-		static FInstallerAction MakeUninstall(const IBuildManifestRef& Manifest);
+		static BUILDPATCHSERVICES_API FInstallerAction MakeUninstall(const IBuildManifestRef& Manifest);
 
 		/**
 		 * Helper for creating an install action or update action based on validity of CurrentManifest.
@@ -101,48 +101,48 @@ namespace BuildPatchServices
 		/**
 		 * Copy constructor.
 		 */
-		FInstallerAction(const FInstallerAction& CopyFrom);
+		BUILDPATCHSERVICES_API FInstallerAction(const FInstallerAction& CopyFrom);
 
 		/**
 		 * RValue constructor to allow move semantics.
 		 */
-		FInstallerAction(FInstallerAction&& MoveFrom);
+		BUILDPATCHSERVICES_API FInstallerAction(FInstallerAction&& MoveFrom);
 
 	public:
 		/**
 		 * @return true if this action intent is to perform a fresh installation.
 		 */
-		bool IsInstall() const;
+		BUILDPATCHSERVICES_API bool IsInstall() const;
 
 		/**
 		 * @return true if this action intent is to update an existing installation.
 		 */
-		bool IsUpdate() const;
+		BUILDPATCHSERVICES_API bool IsUpdate() const;
 
 		/**
 		 * @return true if this action intent is to repair an existing installation.
 		 */
-		bool IsRepair() const;
+		BUILDPATCHSERVICES_API bool IsRepair() const;
 
 		/**
 		 * @return true if this action intent is to uninstall an installation.
 		 */
-		bool IsUninstall() const;
+		BUILDPATCHSERVICES_API bool IsUninstall() const;
 
 		/**
 		 * @return the install tags for the action.
 		 */
-		const TSet<FString>& GetInstallTags() const;
+		BUILDPATCHSERVICES_API const TSet<FString>& GetInstallTags() const;
 
 		/**
 		 * @return the manifest for the current installation, this will runtime assert if called invalidly (see TryGetCurrentManifest).
 		 */
-		IBuildManifestRef GetCurrentManifest() const;
+		BUILDPATCHSERVICES_API IBuildManifestRef GetCurrentManifest() const;
 
 		/**
 		 * @return the manifest for the desired installation, this will runtime assert if called invalidly (see TryGetInstallManifest).
 		 */
-		IBuildManifestRef GetInstallManifest() const;
+		BUILDPATCHSERVICES_API IBuildManifestRef GetInstallManifest() const;
 
 	public:
 		/**
@@ -190,7 +190,7 @@ namespace BuildPatchServices
 		}
 
 	private:
-		FInstallerAction();
+		BUILDPATCHSERVICES_API FInstallerAction();
 		IBuildManifestPtr CurrentManifest;
 		IBuildManifestPtr InstallManifest;
 		TSet<FString> InstallTags;
@@ -200,11 +200,11 @@ namespace BuildPatchServices
 	/**
 	 * DEPRECATED STRUCT. Please use FBuildInstallerConfiguration.
 	 */
-	struct BUILDPATCHSERVICES_API FInstallerConfiguration
+	struct FInstallerConfiguration
 	{
-		FInstallerConfiguration(const IBuildManifestRef& InInstallManifest);
-		FInstallerConfiguration(const FInstallerConfiguration& CopyFrom);
-		FInstallerConfiguration(FInstallerConfiguration&& MoveFrom);
+		BUILDPATCHSERVICES_API FInstallerConfiguration(const IBuildManifestRef& InInstallManifest);
+		BUILDPATCHSERVICES_API FInstallerConfiguration(const FInstallerConfiguration& CopyFrom);
+		BUILDPATCHSERVICES_API FInstallerConfiguration(FInstallerConfiguration&& MoveFrom);
 
 	public:
 		IBuildManifestPtr CurrentManifest;
@@ -226,22 +226,22 @@ namespace BuildPatchServices
 	/**
 	 * Defines a list of all the options of an installation task.
 	 */
-	struct BUILDPATCHSERVICES_API FBuildInstallerConfiguration
+	struct FBuildInstallerConfiguration
 	{
 		/**
 		 * Construct with an array of action objects
 		 */
-		FBuildInstallerConfiguration(TArray<FInstallerAction> InstallerActions);
+		BUILDPATCHSERVICES_API FBuildInstallerConfiguration(TArray<FInstallerAction> InstallerActions);
 
 		/**
 		 * Copy constructor.
 		 */
-		FBuildInstallerConfiguration(const FBuildInstallerConfiguration& CopyFrom);
+		BUILDPATCHSERVICES_API FBuildInstallerConfiguration(const FBuildInstallerConfiguration& CopyFrom);
 
 		/**
 		 * RValue constructor to allow move semantics.
 		 */
-		FBuildInstallerConfiguration(FBuildInstallerConfiguration&& MoveFrom);
+		BUILDPATCHSERVICES_API FBuildInstallerConfiguration(FBuildInstallerConfiguration&& MoveFrom);
 
 	public:
 		// The array of intended actions to perform.
@@ -271,7 +271,7 @@ namespace BuildPatchServices
 	/**
 	 * Defines a list of all the options of a build streamer class.
 	 */
-	struct BUILDPATCHSERVICES_API FBuildInstallStreamerConfiguration
+	struct FBuildInstallStreamerConfiguration
 	{
 	public:
 		IBuildManifestPtr Manifest;
@@ -288,7 +288,7 @@ namespace BuildPatchServices
 	/**
 	 * The collection of statistics gathered by the build streamer class.
 	 */
-	struct BUILDPATCHSERVICES_API FBuildInstallStreamerStats
+	struct FBuildInstallStreamerStats
 	{
 		//float FileMegaBytesDownloaded;
 		uint64 FileRequestsCompleted;
@@ -313,13 +313,13 @@ namespace BuildPatchServices
 	/**
 	 * Defines a list of all options for the build chunking task.
 	 */
-	struct BUILDPATCHSERVICES_API FChunkBuildConfiguration
+	struct FChunkBuildConfiguration
 	{
 	public:
 		/**
 		 * Default constructor
 		 */
-		FChunkBuildConfiguration();
+		BUILDPATCHSERVICES_API FChunkBuildConfiguration();
 
 	public:
 		// The client feature level to output data for.
@@ -374,13 +374,13 @@ namespace BuildPatchServices
 	/**
 	 * Defines a list of all options for the chunk delta optimisation task.
 	 */
-	struct BUILDPATCHSERVICES_API FChunkDeltaOptimiserConfiguration
+	struct FChunkDeltaOptimiserConfiguration
 	{
 	public:
 		/**
 		 * Default constructor
 		 */
-		FChunkDeltaOptimiserConfiguration();
+		BUILDPATCHSERVICES_API FChunkDeltaOptimiserConfiguration();
 
 	public:
 		// A full file or http path for the manifest to be used as the source build.
@@ -400,13 +400,13 @@ namespace BuildPatchServices
 	/**
 	 * Defines a list of all options for the patch data enumeration task.
 	 */
-	struct BUILDPATCHSERVICES_API FPatchDataEnumerationConfiguration
+	struct FPatchDataEnumerationConfiguration
 	{
 	public:
 		/**
 		 * Default constructor
 		 */
-		FPatchDataEnumerationConfiguration();
+		BUILDPATCHSERVICES_API FPatchDataEnumerationConfiguration();
 
 	public:
 		// A full file path for the manifest or chunkdb to enumerate referenced data for.
@@ -420,13 +420,13 @@ namespace BuildPatchServices
 	/**
 	 * Defines a list of all options for the diff manifests task.
 	 */
-	struct BUILDPATCHSERVICES_API FDiffManifestsConfiguration
+	struct FDiffManifestsConfiguration
 	{
 	public:
 		/**
 		 * Default constructor
 		 */
-		FDiffManifestsConfiguration();
+		BUILDPATCHSERVICES_API FDiffManifestsConfiguration();
 
 	public:
 		// A full file or http path for the manifest to be used as the source build.
@@ -447,13 +447,13 @@ namespace BuildPatchServices
 	/**
 	 * Defines a list of all options for the cloud directory compactifier task.
 	 */
-	struct BUILDPATCHSERVICES_API FCompactifyConfiguration
+	struct FCompactifyConfiguration
 	{
 	public:
 		/**
 		 * Default constructor
 		 */
-		FCompactifyConfiguration();
+		BUILDPATCHSERVICES_API FCompactifyConfiguration();
 
 	public:
 		// The path to the directory to compactify.
@@ -469,13 +469,13 @@ namespace BuildPatchServices
 	/**
 	 * Defines a list of all options for the chunk packaging task.
 	 */
-	struct BUILDPATCHSERVICES_API FPackageChunksConfiguration
+	struct FPackageChunksConfiguration
 	{
 	public:
 		/**
 		 * Default constructor
 		 */
-		FPackageChunksConfiguration();
+		BUILDPATCHSERVICES_API FPackageChunksConfiguration();
 
 	public:
 		// The client feature level to output data for.
diff --git a/Engine/Source/Runtime/Online/Experimental/EventLoop/Public/EventLoop/BSDSocket/EventLoopIOManagerBSDSocketSelect.h b/Engine/Source/Runtime/Online/Experimental/EventLoop/Public/EventLoop/BSDSocket/EventLoopIOManagerBSDSocketSelect.h
--- a/Engine/Source/Runtime/Online/Experimental/EventLoop/Public/EventLoop/BSDSocket/EventLoopIOManagerBSDSocketSelect.h
+++ b/Engine/Source/Runtime/Online/Experimental/EventLoop/Public/EventLoop/BSDSocket/EventLoopIOManagerBSDSocketSelect.h
@@ -14,7 +14,7 @@
 
 namespace UE::EventLoop {
 
-class EVENTLOOP_API FIOManagerBSDSocketSelect final : public IIOManager
+class FIOManagerBSDSocketSelect final : public IIOManager
 {
 public:
 	using FIOAccess = FIOAccessBSDSocket;
@@ -27,14 +27,14 @@ public:
 	{
 	};
 
-	FIOManagerBSDSocketSelect(IEventLoop& EventLoop, FParams&& Params);
+	EVENTLOOP_API FIOManagerBSDSocketSelect(IEventLoop& EventLoop, FParams&& Params);
 	virtual ~FIOManagerBSDSocketSelect() = default;
-	virtual bool Init() override;
-	virtual void Shutdown() override;
-	virtual void Notify() override;
-	virtual void Poll(FTimespan WaitTime) override;
+	EVENTLOOP_API virtual bool Init() override;
+	EVENTLOOP_API virtual void Shutdown() override;
+	EVENTLOOP_API virtual void Notify() override;
+	EVENTLOOP_API virtual void Poll(FTimespan WaitTime) override;
 
-	FIOAccess& GetIOAccess();
+	EVENTLOOP_API FIOAccess& GetIOAccess();
 
 private:
 	TSharedRef<class FIOManagerBSDSocketSelectImpl> Impl;
diff --git a/Engine/Source/Runtime/Online/HTTP/Internal/WinHttp/Support/WinHttpConnection.h b/Engine/Source/Runtime/Online/HTTP/Internal/WinHttp/Support/WinHttpConnection.h
--- a/Engine/Source/Runtime/Online/HTTP/Internal/WinHttp/Support/WinHttpConnection.h
+++ b/Engine/Source/Runtime/Online/HTTP/Internal/WinHttp/Support/WinHttpConnection.h
@@ -10,7 +10,7 @@
 using HINTERNET = void*;
 using LPVOID = void*;
 
-class HTTP_API IWinHttpConnection
+class IWinHttpConnection
 	: public TSharedFromThis<IWinHttpConnection, ESPMode::ThreadSafe>
 {
 public:
diff --git a/Engine/Source/Runtime/Online/HTTP/Internal/WinHttp/Support/WinHttpConnectionHttp.h b/Engine/Source/Runtime/Online/HTTP/Internal/WinHttp/Support/WinHttpConnectionHttp.h
--- a/Engine/Source/Runtime/Online/HTTP/Internal/WinHttp/Support/WinHttpConnectionHttp.h
+++ b/Engine/Source/Runtime/Online/HTTP/Internal/WinHttp/Support/WinHttpConnectionHttp.h
@@ -25,32 +25,32 @@ DECLARE_DELEGATE_TwoParams(FWinHttpConnectionHttpOnDataTransferred, int32 /*Byte
 DECLARE_DELEGATE_TwoParams(FWinHttpConnectionHttpOnHeaderReceived, const FString& /*HeaderKey*/, const FString& /*HeaderValue*/);
 DECLARE_DELEGATE_OneParam(FWinHttpConnectionHttpOnRequestComplete, EHttpRequestStatus::Type /*CompletionStatus*/);
 
-class HTTP_API FWinHttpConnectionHttp
+class FWinHttpConnectionHttp
 	: public IWinHttpConnection
 {
 public:
-	static TSharedPtr<FWinHttpConnectionHttp, ESPMode::ThreadSafe> CreateHttpConnection(
+	static HTTP_API TSharedPtr<FWinHttpConnectionHttp, ESPMode::ThreadSafe> CreateHttpConnection(
 		FWinHttpSession& Session,
 		const FString& Verb,
 		const FString& Url,
 		const TMap<FString, FString>& Headers,
 		const TSharedPtr<FRequestPayload, ESPMode::ThreadSafe>& Payload);
 
-	virtual ~FWinHttpConnectionHttp();
+	HTTP_API virtual ~FWinHttpConnectionHttp();
 	FWinHttpConnectionHttp(const FWinHttpConnectionHttp& Other) = delete;
 	FWinHttpConnectionHttp(FWinHttpConnectionHttp&& Other) = delete;
 	FWinHttpConnectionHttp& operator=(const FWinHttpConnectionHttp& Other) = delete;
 	FWinHttpConnectionHttp& operator=(FWinHttpConnectionHttp&& Other) = delete;
 
 	//~ Begin IWinHttpConnection Interface
-	virtual bool IsValid() const override;
-	virtual const FString& GetRequestUrl() const override;
-	virtual void* GetHandle() override;
-	virtual bool StartRequest() override;
-	virtual bool CancelRequest() override;
-	virtual bool IsComplete() const override;
-	virtual void PumpMessages() override;
-	virtual void PumpStates() override;
+	HTTP_API virtual bool IsValid() const override;
+	HTTP_API virtual const FString& GetRequestUrl() const override;
+	HTTP_API virtual void* GetHandle() override;
+	HTTP_API virtual bool StartRequest() override;
+	HTTP_API virtual bool CancelRequest() override;
+	HTTP_API virtual bool IsComplete() const override;
+	HTTP_API virtual void PumpMessages() override;
+	HTTP_API virtual void PumpStates() override;
 	//~ End IWinHttpConnection Interface
 
 	/** Provides the current known response code. */
@@ -72,22 +72,22 @@ public:
 	 *
 	 * @param Handler The Handler to call when we receive or send data
 	 */
-	void SetDataTransferredHandler(FWinHttpConnectionHttpOnDataTransferred&& Handler);
+	HTTP_API void SetDataTransferredHandler(FWinHttpConnectionHttpOnDataTransferred&& Handler);
 	/**
 	 * Bind a handler that is called each tick we receive a new header
 	 *
 	 * @param Handler The Handler to call when we receive a new header
 	 */
-	void SetHeaderReceivedHandler(FWinHttpConnectionHttpOnHeaderReceived&& Handler);
+	HTTP_API void SetHeaderReceivedHandler(FWinHttpConnectionHttpOnHeaderReceived&& Handler);
 	/**
 	 * Bind a handler that is called when the request completes
 	 *
 	 * @param Handler The Handler to call when the request completes
 	 */
-	void SetRequestCompletedHandler(FWinHttpConnectionHttpOnRequestComplete&& Handler);
+	HTTP_API void SetRequestCompletedHandler(FWinHttpConnectionHttpOnRequestComplete&& Handler);
 
 protected:
-	FWinHttpConnectionHttp(
+	HTTP_API FWinHttpConnectionHttp(
 		FWinHttpSession& Session,
 		const FString& Url,
 		const FString& Verb,
@@ -107,7 +107,7 @@ protected:
 	 * @param Headers Key/Value Map of headers to add to our request
 	 * @return True if we could set the headers, false otherwise
 	 */
-	bool SetHeaders(const TMap<FString, FString>& Headers);
+	HTTP_API bool SetHeaders(const TMap<FString, FString>& Headers);
 
 	/**
 	 * Add to our request's headers. Theis new header will be added to the existing headers set, and will
@@ -117,14 +117,14 @@ protected:
 	 * @param Value Value of header to set on our request
 	 * @return True if we could set the header, false otherwise
 	 */
-	bool SetHeader(const FString& Key, const FString& Value);
+	HTTP_API bool SetHeader(const FString& Key, const FString& Value);
 	/**
 	 * Set the payload object of our request. This may only be called before StartRequest is called.
 	 *
 	 * @param NewPayload the Payload object to set (it should not be modified during the request!)
 	 * @return True if we were able to set the payload, false otherwise
 	 */
-	bool SetPayload(const TSharedRef<FRequestPayload, ESPMode::ThreadSafe>& NewPayload);
+	HTTP_API bool SetPayload(const TSharedRef<FRequestPayload, ESPMode::ThreadSafe>& NewPayload);
 
 	/// Sending Request
 
@@ -133,31 +133,31 @@ protected:
 	 *
 	 * @return True if we could start the request, or false otherwise
 	 */
-	bool SendRequest();
+	HTTP_API bool SendRequest();
 	/**
 	 * Increment our bytes-sent counters
 	 *
 	 * @param AmountSent The amount of bytes to increment our counters by
 	 */
-	void IncrementSentByteCounts(const uint64 AmountSent);
+	HTTP_API void IncrementSentByteCounts(const uint64 AmountSent);
 	/**
 	 * Increment our bytes-received counters
 	 *
 	 * @param AmountReceived The amount of bytes to increment our counters by
 	 */
-	void IncrementReceivedByteCounts(const uint64 AmountReceived);
+	HTTP_API void IncrementReceivedByteCounts(const uint64 AmountReceived);
 	/** 
 	 * Check if we have more data to send
 	 *
 	 * @return True if we have data to send, or false otherwise
 	 */
-	bool HasRequestBodyToSend() const;
+	HTTP_API bool HasRequestBodyToSend() const;
 	/**
 	 * Gives WinHttp the next check of data.
 	 *
 	 * @param True if we asked WinHttp to send the data, or false if we could not send the next chunk
 	 */
-	bool SendAdditionalRequestBody();
+	HTTP_API bool SendAdditionalRequestBody();
 
 	/// Request Response
 
@@ -166,51 +166,51 @@ protected:
 	 *
 	 * @return True if the action completed successfully, false otherwise
 	 */
-	bool RequestResponse();
+	HTTP_API bool RequestResponse();
 	/**
 	 * Read our HTTP response headers and store the results
 	 *
 	 * @return True if the action completed successfully, false otherwise
 	 */
-	bool ProcessResponseHeaders();
+	HTTP_API bool ProcessResponseHeaders();
 	/**
 	 * Ask WinHttp to asynchronously tell us how large the next chunk will be. If it's 0, we're done.
 	 *
 	 * @return True if the action completed successfully, false otherwise
 	 */
-	bool RequestNextResponseBodyChunkSize();
+	HTTP_API bool RequestNextResponseBodyChunkSize();
 	/**
 	 * Ask WinHttp to asynchronously write the next chunk directly into our response object
 	 *
 	 * @return True if the action completed successfully, false otherwise
 	 */
-	bool RequestNextResponseBodyChunkData();
+	HTTP_API bool RequestNextResponseBodyChunkData();
 
 	/// Request State
 
 	/** Set our request as finished and set state to the specified state. This also releases the WinHttp handles. */
-	virtual bool FinishRequest(const EHttpRequestStatus::Type FinalState);
+	HTTP_API virtual bool FinishRequest(const EHttpRequestStatus::Type FinalState);
 
 	/// Callback handling
 
 	/** Handles storing the server's IP Address if needed */
-	virtual void HandleConnectedToServer(const wchar_t* ServerIP);
+	HTTP_API virtual void HandleConnectedToServer(const wchar_t* ServerIP);
 	/** Handles validating request now that we have server information (certificates, etc) */
-	void HandleSendingRequest();
+	HTTP_API void HandleSendingRequest();
 	/** Handles SOMETHING not sure yet*/
-	void HandleWriteComplete(const uint64 NumBytesSent);
+	HTTP_API void HandleWriteComplete(const uint64 NumBytesSent);
 	/** Handles completing the request so we may read the headers and body of the response */
-	void HandleSendRequestComplete();
+	HTTP_API void HandleSendRequestComplete();
 	/** Handles reading the response headers */
-	virtual void HandleHeadersAvailable();
+	HTTP_API virtual void HandleHeadersAvailable();
 	/** Handles writing the http response body to our response body */
-	void HandleDataAvailable(const uint64 NumBytesAvailable);
+	HTTP_API void HandleDataAvailable(const uint64 NumBytesAvailable);
 	/** Handles querying the next chunk of our http response body */
-	void HandleReadComplete(const uint64 NumBytesRead);
+	HTTP_API void HandleReadComplete(const uint64 NumBytesRead);
 	/** Handles shutting down the request when there was an error */
-	void HandleRequestError(const uint32 ErrorApiId, const uint32 ErrorCode);
+	HTTP_API void HandleRequestError(const uint32 ErrorApiId, const uint32 ErrorCode);
 	/** Handles our request object closing, ending this request */
-	virtual void HandleHandleClosing();
+	HTTP_API virtual void HandleHandleClosing();
 
 private:
 	/**
@@ -221,13 +221,13 @@ private:
 	 * @param StatusInformation Pointer to optional data for this event, if there is any
 	 * @param StatusInformationLength The length of data available to read in StatusInformation or 0 if not used
 	 */
-	void HandleHttpStatusCallback(HINTERNET ResourceHandle, EWinHttpCallbackStatus Status, void* StatusInformation, uint32 StatusInformationLength);
+	HTTP_API void HandleHttpStatusCallback(HINTERNET ResourceHandle, EWinHttpCallbackStatus Status, void* StatusInformation, uint32 StatusInformationLength);
 	
 	/** Mark our callback as a friend so they can call the above status callback function */
 	friend void CALLBACK UE_WinHttpStatusHttpCallback(HINTERNET, DWORD_PTR, DWORD, LPVOID, DWORD);
 
 	/** Reset both PayloadBuffer and Payload when we don't need them anymore, releasing the memory */
-	void ReleasePayloadData();
+	HTTP_API void ReleasePayloadData();
 
 protected:
 	/** Critical section used to synchronize access to the below properties across multiple threads */
diff --git a/Engine/Source/Runtime/Online/HTTP/Internal/WinHttp/Support/WinHttpErrorHelper.h b/Engine/Source/Runtime/Online/HTTP/Internal/WinHttp/Support/WinHttpErrorHelper.h
--- a/Engine/Source/Runtime/Online/HTTP/Internal/WinHttp/Support/WinHttpErrorHelper.h
+++ b/Engine/Source/Runtime/Online/HTTP/Internal/WinHttp/Support/WinHttpErrorHelper.h
@@ -6,7 +6,7 @@
 
 #include "CoreMinimal.h"
 
-class HTTP_API FWinHttpErrorHelper
+class FWinHttpErrorHelper
 {
 public:
 	/**
@@ -14,98 +14,98 @@ public:
 	 *
 	 * @param ErrorCode The error that occured
 	 */
-	static void LogWinHttpOpenFailure(const uint32 ErrorCode);
+	static HTTP_API void LogWinHttpOpenFailure(const uint32 ErrorCode);
 
 	/**
 	 * Log that an WinHttpCloseHandle call failed with the specified ErrorCode
 	 *
 	 * @param ErrorCode The error that occured
 	 */
-	static void LogWinHttpCloseHandleFailure(const uint32 ErrorCode);
+	static HTTP_API void LogWinHttpCloseHandleFailure(const uint32 ErrorCode);
 
 	/**
 	 * Log that an WinHttpSetTimeouts call failed with the specified ErrorCode
 	 *
 	 * @param ErrorCode The error that occured
 	 */
-	static void LogWinHttpSetTimeoutsFailure(const uint32 ErrorCode);
+	static HTTP_API void LogWinHttpSetTimeoutsFailure(const uint32 ErrorCode);
 
 	/**
 	 * Log that an WinHttpConnect call failed with the specified ErrorCode
 	 *
 	 * @param ErrorCode The error that occured
 	 */
-	static void LogWinHttpConnectFailure(const uint32 ErrorCode);
+	static HTTP_API void LogWinHttpConnectFailure(const uint32 ErrorCode);
 
 	/**
 	 * Log that an WinHttpOpenRequest call failed with the specified ErrorCode
 	 *
 	 * @param ErrorCode The error that occured
 	 */
-	static void LogWinHttpOpenRequestFailure(const uint32 ErrorCode);
+	static HTTP_API void LogWinHttpOpenRequestFailure(const uint32 ErrorCode);
 
 	/**
 	 * Log that an WinHttpSetOption call failed with the specified ErrorCode
 	 *
 	 * @param ErrorCode The error that occured
 	 */
-	static void LogWinHttpSetOptionFailure(const uint32 ErrorCode);
+	static HTTP_API void LogWinHttpSetOptionFailure(const uint32 ErrorCode);
 
 	/**
 	 * Log that an WinHttpAddRequestHeaders call failed with the specified ErrorCode
 	 *
 	 * @param ErrorCode The error that occured
 	 */
-	static void LogWinHttpAddRequestHeadersFailure(const uint32 ErrorCode);
+	static HTTP_API void LogWinHttpAddRequestHeadersFailure(const uint32 ErrorCode);
 
 	/**
 	 * Log that an WinHttpSetStatusCallback call failed with the specified ErrorCode
 	 *
 	 * @param ErrorCode The error that occured
 	 */
-	static void LogWinHttpSetStatusCallbackFailure(const uint32 ErrorCode);
+	static HTTP_API void LogWinHttpSetStatusCallbackFailure(const uint32 ErrorCode);
 
 	/**
 	 * Log that an WinHttpSendRequest call failed with the specified ErrorCode
 	 *
 	 * @param ErrorCode The error that occured
 	 */
-	static void LogWinHttpSendRequestFailure(const uint32 ErrorCode);
+	static HTTP_API void LogWinHttpSendRequestFailure(const uint32 ErrorCode);
 
 	/**
 	 * Log that an WinHttpReceiveResponse call failed with the specified ErrorCode
 	 *
 	 * @param ErrorCode The error that occured
 	 */
-	static void LogWinHttpReceiveResponseFailure(const uint32 ErrorCode);
+	static HTTP_API void LogWinHttpReceiveResponseFailure(const uint32 ErrorCode);
 
 	/**
 	 * Log that an WinHttpQueryHeaders call failed with the specified ErrorCode
 	 *
 	 * @param ErrorCode The error that occured
 	 */
-	static void LogWinHttpQueryHeadersFailure(const uint32 ErrorCode);
+	static HTTP_API void LogWinHttpQueryHeadersFailure(const uint32 ErrorCode);
 
 	/**
 	 * Log that an WinHttpQueryDataAvailable call failed with the specified ErrorCode
 	 *
 	 * @param ErrorCode The error that occured
 	 */
-	static void LogWinHttpQueryDataAvailableFailure(const uint32 ErrorCode);
+	static HTTP_API void LogWinHttpQueryDataAvailableFailure(const uint32 ErrorCode);
 
 	/**
 	 * Log that an WinHttpReadData call failed with the specified ErrorCode
 	 *
 	 * @param ErrorCode The error that occured
 	 */
-	static void LogWinHttpReadDataFailure(const uint32 ErrorCode);
+	static HTTP_API void LogWinHttpReadDataFailure(const uint32 ErrorCode);
 
 	/**
 	 * Log that an WinHttpWriteData call failed with the specified ErrorCode
 	 *
 	 * @param ErrorCode The error that occured
 	 */
-	static void LogWinHttpWriteDataFailure(const uint32 ErrorCode);
+	static HTTP_API void LogWinHttpWriteDataFailure(const uint32 ErrorCode);
 
 private:
 	/** This class should not be instantiated */
diff --git a/Engine/Source/Runtime/Online/HTTP/Internal/WinHttp/Support/WinHttpHandle.h b/Engine/Source/Runtime/Online/HTTP/Internal/WinHttp/Support/WinHttpHandle.h
--- a/Engine/Source/Runtime/Online/HTTP/Internal/WinHttp/Support/WinHttpHandle.h
+++ b/Engine/Source/Runtime/Online/HTTP/Internal/WinHttp/Support/WinHttpHandle.h
@@ -8,7 +8,7 @@
 
 using HINTERNET = void*;
 
-class HTTP_API FWinHttpHandle
+class FWinHttpHandle
 {
 public:
 	/**
@@ -19,18 +19,18 @@ public:
 	/**
 	 * Wrap a new HINTERNET Handle
 	 */
-	explicit FWinHttpHandle(HINTERNET NewHandle);
+	HTTP_API explicit FWinHttpHandle(HINTERNET NewHandle);
 
 	/**
 	 * Destroy any currently held object
 	 */
-	~FWinHttpHandle();
+	HTTP_API ~FWinHttpHandle();
 
 	// Copy/Move constructors
 	FWinHttpHandle(const FWinHttpHandle& Other) = delete;
-	FWinHttpHandle(FWinHttpHandle&& Other);
+	HTTP_API FWinHttpHandle(FWinHttpHandle&& Other);
 	FWinHttpHandle& operator=(const FWinHttpHandle& Other) = delete;
-	FWinHttpHandle& operator=(FWinHttpHandle&& Other);
+	HTTP_API FWinHttpHandle& operator=(FWinHttpHandle&& Other);
 
 	/**
 	 * Wrap a new handle (destroying any previously held object)
@@ -38,33 +38,33 @@ public:
 	 * @param NewHandle The new handle to wrap (must not be wrapped by anything else)
 	 * @return A reference to this object that now wraps NewHandle
 	 */
-	FWinHttpHandle& operator=(HINTERNET NewHandle);
+	HTTP_API FWinHttpHandle& operator=(HINTERNET NewHandle);
 
 	/**
 	 * Destroy our current handle and reset our state to holding nothing
 	 */
-	void Reset();
+	HTTP_API void Reset();
 
 	/**
 	 * Do we contain a valid handle?
 	 *
 	 * @return True if we have a valid handle, false otherwise
 	 */
-	explicit operator bool() const;
+	HTTP_API explicit operator bool() const;
 
 	/**
 	 * Do we contain a valid handle?
 	 *
 	 * @return True if we have a valid handle, false otherwise
 	 */
-	bool IsValid() const;
+	HTTP_API bool IsValid() const;
 
 	/**
 	 * Get the underlying handle for use
 	 *
 	 * @return The HINTERNET handle we're wrapping
 	 */
-	HINTERNET Get() const;
+	HTTP_API HINTERNET Get() const;
 
 protected:
 	/**
diff --git a/Engine/Source/Runtime/Online/HTTP/Internal/WinHttp/Support/WinHttpSession.h b/Engine/Source/Runtime/Online/HTTP/Internal/WinHttp/Support/WinHttpSession.h
--- a/Engine/Source/Runtime/Online/HTTP/Internal/WinHttp/Support/WinHttpSession.h
+++ b/Engine/Source/Runtime/Online/HTTP/Internal/WinHttp/Support/WinHttpSession.h
@@ -11,13 +11,13 @@ class IWinHttpConnection;
 
 using HINTERNET = void*;
 
-class HTTP_API FWinHttpSession
+class FWinHttpSession
 {
 public:
 	/**
 	 * Construct a new WinHttp session with the specified security protocols flags
 	 */
-	FWinHttpSession(const uint32 SecurityProtocolFlags, const bool bForceSecureConnections);
+	HTTP_API FWinHttpSession(const uint32 SecurityProtocolFlags, const bool bForceSecureConnections);
 
 	/**
 	 * FWinHttpSession is move-only
@@ -30,21 +30,21 @@ public:
 	/**
 	 * Did this session initialize successfully?
 	 */
-	bool IsValid() const;
+	HTTP_API bool IsValid() const;
 
 	/**
 	 * Get the underlying session handle
 	 *
 	 * @return the HINTERNET for this session
 	 */
-	HINTERNET Get() const;
+	HTTP_API HINTERNET Get() const;
 
 	/**
 	 * Are we only allowed to make secure connection requests (HTTPS, etc)
 	 *
 	 * @return True if the platform does not allow for insecure messages
 	 */
-	bool AreOnlySecureConnectionsAllowed() const;
+	HTTP_API bool AreOnlySecureConnectionsAllowed() const;
 
 private:
 	/** The handle for our session that we are wrapping */
diff --git a/Engine/Source/Runtime/Online/HTTP/Internal/WinHttp/WinHttpHttpManager.h b/Engine/Source/Runtime/Online/HTTP/Internal/WinHttp/WinHttpHttpManager.h
--- a/Engine/Source/Runtime/Online/HTTP/Internal/WinHttp/WinHttpHttpManager.h
+++ b/Engine/Source/Runtime/Online/HTTP/Internal/WinHttp/WinHttpHttpManager.h
@@ -13,14 +13,14 @@ DECLARE_DELEGATE_OneParam(FWinHttpQuerySessionComplete, FWinHttpSession* /*HttpS
 
 class IWinHttpConnection;
 
-class HTTP_API FWinHttpHttpManager
+class FWinHttpHttpManager
 	: public FHttpManager
 {
 public:
-	static FWinHttpHttpManager* GetManager();
+	static HTTP_API FWinHttpHttpManager* GetManager();
 
-	FWinHttpHttpManager();
-	virtual ~FWinHttpHttpManager();
+	HTTP_API FWinHttpHttpManager();
+	HTTP_API virtual ~FWinHttpHttpManager();
 
 	/**
 	 * Asynchronously finds an existing WinHttp session for the provided URL, or creates a new one for it.
@@ -28,7 +28,7 @@ public:
 	 * @param Url The URL to find or create a WinHttp session for
 	 * @param Delegate The delegate that is called with the WinHttp session pointer if successful, or null otherwise
 	 */
-	virtual void QuerySessionForUrl(const FString& Url, FWinHttpQuerySessionComplete&& Delegate);
+	HTTP_API virtual void QuerySessionForUrl(const FString& Url, FWinHttpQuerySessionComplete&& Delegate);
 
 	/**
 	 * Validate the provided connection before we start sending our request.
@@ -38,7 +38,7 @@ public:
 	 * @param Connection the connection to validate
 	 * @return True if the request validated successfully, false if not
 	 */
-	virtual bool ValidateRequestCertificates(IWinHttpConnection& Connection);
+	HTTP_API virtual bool ValidateRequestCertificates(IWinHttpConnection& Connection);
 
 	/**
 	 * Release any resources that were pinned for this request.
@@ -47,24 +47,24 @@ public:
 	 *
 	 * @param Connection The connection to release resources for
 	 */
-	virtual void ReleaseRequestResources(IWinHttpConnection& Connection);
+	HTTP_API virtual void ReleaseRequestResources(IWinHttpConnection& Connection);
 
 	//~ Begin FHttpManager Interface
-	virtual void OnBeforeFork() override;
+	HTTP_API virtual void OnBeforeFork() override;
 	//~ End FHttpManager Interface
 
 	/**
 	 * Handle the application preparing to suspend
 	 */
-	virtual void HandleApplicationSuspending();
+	HTTP_API virtual void HandleApplicationSuspending();
 
 	/**
 	 * Handle the application preparing to suspend
 	 */
-	virtual void HandleApplicationResuming();
+	HTTP_API virtual void HandleApplicationResuming();
 
 protected:
-	FWinHttpSession* FindOrCreateSession(const uint32 SecurityProtocols);
+	HTTP_API FWinHttpSession* FindOrCreateSession(const uint32 SecurityProtocols);
 
 protected:
 	bool bPlatformForcesSecureConnections = false;
diff --git a/Engine/Source/Runtime/Online/HTTP/Public/Android/AndroidPlatformHttp.h b/Engine/Source/Runtime/Online/HTTP/Public/Android/AndroidPlatformHttp.h
--- a/Engine/Source/Runtime/Online/HTTP/Public/Android/AndroidPlatformHttp.h
+++ b/Engine/Source/Runtime/Online/HTTP/Public/Android/AndroidPlatformHttp.h
@@ -6,40 +6,40 @@
 /**
  * Platform specific Http implementations
  */
-class HTTP_API FAndroidPlatformHttp : public FGenericPlatformHttp
+class FAndroidPlatformHttp : public FGenericPlatformHttp
 {
 public:
 
 	/**
 	 * Platform initialization step
 	 */
-	static void Init();
+	static HTTP_API void Init();
 
 	/**
 	 * Creates a platform-specific HTTP manager.
 	 *
 	 * @return NULL if default implementation is to be used
 	 */
-	static FHttpManager* CreatePlatformHttpManager();
+	static HTTP_API FHttpManager* CreatePlatformHttpManager();
 
 	/**
 	 * Platform shutdown step
 	 */
-	static void Shutdown();
+	static HTTP_API void Shutdown();
 
 	/**
 	 * Creates a new Http request instance for the current platform
 	 *
 	 * @return request object
 	 */
-	static IHttpRequest* ConstructRequest();
+	static HTTP_API IHttpRequest* ConstructRequest();
 
 	/**
 	 * Get the proxy address specified by the operating system
 	 *
 	 * @return optional FString: If unset: we are unable to get information from the operating system. If set: the proxy address set by the operating system (may be blank)
 	 */
-	static TOptional<FString> GetOperatingSystemProxyAddress();
+	static HTTP_API TOptional<FString> GetOperatingSystemProxyAddress();
 
 	/**
 	 * Check if getting proxy information from the current operating system is supported
@@ -48,8 +48,8 @@ public:
 	 * 
 	 * @return true if we are able to get proxy information from the current operating system, false if not
 	 */
-	static bool IsOperatingSystemProxyInformationSupported();
+	static HTTP_API bool IsOperatingSystemProxyInformationSupported();
 };
 
 
-typedef FAndroidPlatformHttp FPlatformHttp;
\ No newline at end of file
+typedef FAndroidPlatformHttp FPlatformHttp;
diff --git a/Engine/Source/Runtime/Online/HTTP/Public/GenericPlatform/GenericPlatformHttp.h b/Engine/Source/Runtime/Online/HTTP/Public/GenericPlatform/GenericPlatformHttp.h
--- a/Engine/Source/Runtime/Online/HTTP/Public/GenericPlatform/GenericPlatformHttp.h
+++ b/Engine/Source/Runtime/Online/HTTP/Public/GenericPlatform/GenericPlatformHttp.h
@@ -18,14 +18,14 @@ class IHttpRequest;
  * Platform specific Http implementations
  * Intended usage is to use FPlatformHttp instead of FGenericPlatformHttp
  */
-class HTTP_API FGenericPlatformHttp
+class FGenericPlatformHttp
 {
 public:
 
 	/**
 	 * Platform initialization step
 	 */
-	static void Init();
+	static HTTP_API void Init();
 
 	/**
 	 * Creates a platform-specific HTTP manager.
@@ -40,21 +40,21 @@ public:
 	/**
 	 * Platform shutdown step
 	 */
-	static void Shutdown();
+	static HTTP_API void Shutdown();
 
 	/**
 	 * Creates a new Http request instance for the current platform
 	 *
 	 * @return request object
 	 */
-	static IHttpRequest* ConstructRequest();
+	static HTTP_API IHttpRequest* ConstructRequest();
 
 	/**
 	 * Check if a platform uses the HTTP thread
 	 *
 	 * @return true if the platform uses threaded HTTP, false if not
 	 */
-	static bool UsesThreadedHttp();
+	static HTTP_API bool UsesThreadedHttp();
 
 	/**
 	 * Returns a percent-encoded version of the passed in string
@@ -62,7 +62,7 @@ public:
 	 * @param UnencodedString The unencoded string to convert to percent-encoding
 	 * @return The percent-encoded string
 	 */
-	static FString UrlEncode(const FStringView UnencodedString);
+	static HTTP_API FString UrlEncode(const FStringView UnencodedString);
 
 	/**
 	 * Returns a decoded version of the percent-encoded passed in string
@@ -70,7 +70,7 @@ public:
 	 * @param EncodedString The percent encoded string to convert to string
 	 * @return The decoded string
 	 */
-	static FString UrlDecode(const FStringView EncodedString);
+	static HTTP_API FString UrlDecode(const FStringView EncodedString);
 
 	/**
 	 * Returns the &lt; &gt...etc encoding for strings between HTML elements.
@@ -78,27 +78,27 @@ public:
 	 * @param UnencodedString The unencoded string to convert to html encoding.
 	 * @return The html encoded string
 	 */
-	static FString HtmlEncode(const FStringView UnencodedString);
+	static HTTP_API FString HtmlEncode(const FStringView UnencodedString);
 
 	/** 
 	* Returns the domain and port portion of the URL, e.g., "a.b.c:d" of "http://a.b.c:d/e"
 	* @param Url the URL to return the domain and port of
 	* @return the domain and port of the specified URL
 	*/
-	static FString GetUrlDomainAndPort(const FStringView Url);
+	static HTTP_API FString GetUrlDomainAndPort(const FStringView Url);
 
 	/** 
 	 * Returns the domain portion of the URL, e.g., "a.b.c" of "http://a.b.c:d/e"
 	 * @param Url the URL to return the domain of
 	 * @return the domain of the specified URL
 	 */
-	static FString GetUrlDomain(const FStringView Url);
+	static HTTP_API FString GetUrlDomain(const FStringView Url);
 
 	/**
 	 * Get the mime type for the file
 	 * @return the mime type for the file.
 	 */
-	static FString GetMimeType(const FString& FilePath);
+	static HTTP_API FString GetMimeType(const FString& FilePath);
 
 	/**
 	 * Returns the default User-Agent string to use in HTTP requests.
@@ -106,15 +106,15 @@ public:
 	 *
 	 * @return the default User-Agent string that requests should use.
 	 */
-	static FString GetDefaultUserAgent();
-	static FString EscapeUserAgentString(const FString& UnescapedString);
+	static HTTP_API FString GetDefaultUserAgent();
+	static HTTP_API FString EscapeUserAgentString(const FString& UnescapedString);
 
 	/**
 	 * Get the proxy address specified by the operating system
 	 *
 	 * @return optional FString: If unset: we are unable to get information from the operating system. If set: the proxy address set by the operating system (may be blank)
 	 */
-	static TOptional<FString> GetOperatingSystemProxyAddress();
+	static HTTP_API TOptional<FString> GetOperatingSystemProxyAddress();
 
 	/**
 	 * Check if getting proxy information from the current operating system is supported
@@ -123,12 +123,12 @@ public:
 	 * 
 	 * @return true if we are able to get proxy information from the current operating system, false if not
 	 */
-	static bool IsOperatingSystemProxyInformationSupported();
+	static HTTP_API bool IsOperatingSystemProxyInformationSupported();
 
 	/**
 	 * Helper function for checking if a byte array is in URL encoded format.
 	 */
-	static bool IsURLEncoded(const TArray<uint8>& Payload);
+	static HTTP_API bool IsURLEncoded(const TArray<uint8>& Payload);
 
 	/**
 	 * Extract the URL-Decoded value of the specified ParameterName from Url. An unset return means the parameter was not present in Url, while an empty value means it was present, but had no value.
@@ -137,14 +137,14 @@ public:
 	 * @param ParameterName The parameter name to look for
 	 * @return If ParameterName was found, the string value of its value (in URL Decoded format), otherwise the return value is unset
 	 */
-	static TOptional<FString> GetUrlParameter(const FStringView Url, const FStringView ParameterName);
+	static HTTP_API TOptional<FString> GetUrlParameter(const FStringView Url, const FStringView ParameterName);
 
 	/**
 	 * Extract the Port part of a URL, or an unset object if there was non specified. Example: "http://example.org:23/" would return 23, while "https://example.org/" would return an unset object
 	 * 
 	 * @param Url The URL to parse for a Port
 	 */
-	static TOptional<uint16> GetUrlPort(const FStringView Url);
+	static HTTP_API TOptional<uint16> GetUrlPort(const FStringView Url);
 
 	/**
 	 * Extract the Path portion of a URL, optionally including the Query String, and optionally including the Fragment (when also including the query string.)
@@ -155,7 +155,7 @@ public:
 	 * @param bIncludeQueryString include the URL's query string in the return value (if one is found)
 	 * @param bIncludeFragment include the URL's fragement in the return value (if one is found)
 	 */
-	static FString GetUrlPath(const FStringView Url, const bool bIncludeQueryString = false, const bool bIncludeFragment = false);
+	static HTTP_API FString GetUrlPath(const FStringView Url, const bool bIncludeQueryString = false, const bool bIncludeFragment = false);
 
 	/**
 	 * Check the protocol of the provided URL to determine if this is for a secure connection (HTTPS, WSS, etc)
@@ -163,5 +163,5 @@ public:
 	 * @param The URL to check for a scheme
 	 * @return True if secure, false if insecure, unset if unknown
 	 */
-	static TOptional<bool> IsSecureProtocol(const FStringView Url);
+	static HTTP_API TOptional<bool> IsSecureProtocol(const FStringView Url);
 };
diff --git a/Engine/Source/Runtime/Online/HTTP/Public/GenericPlatform/HttpRequestCommon.h b/Engine/Source/Runtime/Online/HTTP/Public/GenericPlatform/HttpRequestCommon.h
--- a/Engine/Source/Runtime/Online/HTTP/Public/GenericPlatform/HttpRequestCommon.h
+++ b/Engine/Source/Runtime/Online/HTTP/Public/GenericPlatform/HttpRequestCommon.h
@@ -7,19 +7,19 @@
 /**
  * Contains implementation of some common functions that don't vary between implementations of different platforms
  */
-class HTTP_API FHttpRequestCommon : public FHttpRequestImpl
+class FHttpRequestCommon : public FHttpRequestImpl
 {
 public:
 	// IHttpRequest
-	virtual EHttpRequestStatus::Type GetStatus() const override;
-	virtual void SetDelegateThreadPolicy(EHttpRequestDelegateThreadPolicy InDelegateThreadPolicy) override;
-	virtual EHttpRequestDelegateThreadPolicy GetDelegateThreadPolicy() const override;
+	HTTP_API virtual EHttpRequestStatus::Type GetStatus() const override;
+	HTTP_API virtual void SetDelegateThreadPolicy(EHttpRequestDelegateThreadPolicy InDelegateThreadPolicy) override;
+	HTTP_API virtual EHttpRequestDelegateThreadPolicy GetDelegateThreadPolicy() const override;
 
 protected:
 	/**
 	 * Check if this request is valid or allowed, before actually process the request
 	 */
-	bool PreCheck() const;
+	HTTP_API bool PreCheck() const;
 
 protected:
 	/** Current status of request being processed */
diff --git a/Engine/Source/Runtime/Online/HTTP/Public/GenericPlatform/HttpRequestImpl.h b/Engine/Source/Runtime/Online/HTTP/Public/GenericPlatform/HttpRequestImpl.h
--- a/Engine/Source/Runtime/Online/HTTP/Public/GenericPlatform/HttpRequestImpl.h
+++ b/Engine/Source/Runtime/Online/HTTP/Public/GenericPlatform/HttpRequestImpl.h
@@ -8,27 +8,27 @@
 /**
  * Contains implementation of some common functions that don't vary between implementation
  */
-class HTTP_API FHttpRequestImpl : public IHttpRequest
+class FHttpRequestImpl : public IHttpRequest
 {
 public:
 	// IHttpRequest
-	virtual FHttpRequestCompleteDelegate& OnProcessRequestComplete() override;
-	virtual FHttpRequestProgressDelegate& OnRequestProgress() override;
-	virtual FHttpRequestHeaderReceivedDelegate& OnHeaderReceived() override;
-	virtual FHttpRequestWillRetryDelegate& OnRequestWillRetry() override;
+	HTTP_API virtual FHttpRequestCompleteDelegate& OnProcessRequestComplete() override;
+	HTTP_API virtual FHttpRequestProgressDelegate& OnRequestProgress() override;
+	HTTP_API virtual FHttpRequestHeaderReceivedDelegate& OnHeaderReceived() override;
+	HTTP_API virtual FHttpRequestWillRetryDelegate& OnRequestWillRetry() override;
 
-	virtual void SetTimeout(float InTimeoutSecs) override;
-	virtual void ClearTimeout() override;
-	virtual TOptional<float> GetTimeout() const override;
+	HTTP_API virtual void SetTimeout(float InTimeoutSecs) override;
+	HTTP_API virtual void ClearTimeout() override;
+	HTTP_API virtual TOptional<float> GetTimeout() const override;
 
-	float GetTimeoutOrDefault() const;
+	HTTP_API float GetTimeoutOrDefault() const;
 
 protected:
 	/** 
 	 * Broadcast all of our response's headers as having been received
 	 * Used when we don't know when we receive headers in our HTTP implementation
 	 */
-	void BroadcastResponseHeadersReceived();
+	HTTP_API void BroadcastResponseHeadersReceived();
 
 protected:
 	/** Delegate that will get called once request completes or on any error */
diff --git a/Engine/Source/Runtime/Online/HTTP/Public/HttpManager.h b/Engine/Source/Runtime/Online/HTTP/Public/HttpManager.h
--- a/Engine/Source/Runtime/Online/HTTP/Public/HttpManager.h
+++ b/Engine/Source/Runtime/Online/HTTP/Public/HttpManager.h
@@ -52,7 +52,7 @@ DECLARE_DELEGATE_OneParam(FHttpManagerRequestCompletedDelegate, const FHttpReque
 /**
  * Manages Http request that are currently being processed
  */
-class HTTP_API FHttpManager
+class FHttpManager
 	: public FTSTickerObjectBase
 {
 public:
@@ -62,22 +62,22 @@ public:
 	/**
 	 * Constructor
 	 */
-	FHttpManager();
+	HTTP_API FHttpManager();
 
 	/**
 	 * Destructor
 	 */
-	virtual ~FHttpManager();
+	HTTP_API virtual ~FHttpManager();
 
 	/**
 	 * Initialize
 	 */
-	void Initialize();
+	HTTP_API void Initialize();
 
 	/**
 	 * Shutdown logic should be called before quiting
 	 */
-	void Shutdown();
+	HTTP_API void Shutdown();
 
 	/**
 	 * Adds an Http request instance to the manager for tracking/ticking
@@ -85,10 +85,10 @@ public:
 	 *
 	 * @param Request - the request object to add
 	 */
-	void AddRequest(const FHttpRequestRef& Request);
+	HTTP_API void AddRequest(const FHttpRequestRef& Request);
 
-	void SetRequestAddedDelegate(const FHttpManagerRequestAddedDelegate& Delegate);
-	void SetRequestCompletedDelegate(const FHttpManagerRequestCompletedDelegate& Delegate);
+	HTTP_API void SetRequestAddedDelegate(const FHttpManagerRequestAddedDelegate& Delegate);
+	HTTP_API void SetRequestCompletedDelegate(const FHttpManagerRequestCompletedDelegate& Delegate);
 
 	/**
 	 * Removes an Http request instance from the manager
@@ -96,7 +96,7 @@ public:
 	 *
 	 * @param Request - the request object to remove
 	 */
-	void RemoveRequest(const FHttpRequestRef& Request);
+	HTTP_API void RemoveRequest(const FHttpRequestRef& Request);
 
 	/**
 	* Find an Http request in the lists of current valid requests
@@ -105,14 +105,14 @@ public:
 	*
 	* @return true if the request is being tracked, false if not
 	*/
-	bool IsValidRequest(const IHttpRequest* RequestPtr) const;
+	HTTP_API bool IsValidRequest(const IHttpRequest* RequestPtr) const;
 
 	/**
 	 * Block until all pending requests are finished processing
 	 *
 	 * @param FlushReason the flush reason will influence times waited to cancel ongoing http requests
 	 */
-	void Flush(EHttpFlushReason FlushReason);
+	HTTP_API void Flush(EHttpFlushReason FlushReason);
 
 	/**
 	 * FTSTicker callback
@@ -121,42 +121,42 @@ public:
 	 *
 	 * @return false if no longer needs ticking
 	 */
-	bool Tick(float DeltaSeconds) override;
+	HTTP_API bool Tick(float DeltaSeconds) override;
 
 	/**
 	 * Tick called during Flush
 	 *
 	 * @param DeltaSeconds - time in seconds since the last tick
 	 */
-	virtual void FlushTick(float DeltaSeconds);
+	HTTP_API virtual void FlushTick(float DeltaSeconds);
 
 	/** 
 	 * Add a http request to be executed on the http thread
 	 *
 	 * @param Request - the request object to add
 	 */
-	void AddThreadedRequest(const TSharedRef<IHttpThreadedRequest, ESPMode::ThreadSafe>& Request);
+	HTTP_API void AddThreadedRequest(const TSharedRef<IHttpThreadedRequest, ESPMode::ThreadSafe>& Request);
 
 	/**
 	 * Mark a threaded http request as cancelled to be removed from the http thread
 	 *
 	 * @param Request - the request object to cancel
 	 */
-	void CancelThreadedRequest(const TSharedRef<IHttpThreadedRequest, ESPMode::ThreadSafe>& Request);
+	HTTP_API void CancelThreadedRequest(const TSharedRef<IHttpThreadedRequest, ESPMode::ThreadSafe>& Request);
 
 	/**
 	 * List all of the Http requests currently being processed
 	 *
 	 * @param Ar - output device to log with
 	 */
-	void DumpRequests(FOutputDevice& Ar) const;
+	HTTP_API void DumpRequests(FOutputDevice& Ar) const;
 
 	/**
 	 * Method to check dynamic proxy setting support.
 	 *
 	 * @returns Whether this http implementation supports dynamic proxy setting.
 	 */
-	virtual bool SupportsDynamicProxy() const;
+	HTTP_API virtual bool SupportsDynamicProxy() const;
 
 	/**
 	 * Set the method used to set a Correlation id on each request, if one is not already specified.
@@ -165,14 +165,14 @@ public:
 	 *
 	 * @param InCorrelationIdMethod The method to use when sending a request, if no Correlation id is already set
 	 */
-	void SetCorrelationIdMethod(TFunction<FString()> InCorrelationIdMethod);
+	HTTP_API void SetCorrelationIdMethod(TFunction<FString()> InCorrelationIdMethod);
 
 	/**
 	 * Create a new correlation id for a request
 	 *
 	 * @return The new correlationid string
 	 */
-	FString CreateCorrelationId() const;
+	HTTP_API FString CreateCorrelationId() const;
 
 	/**
 	 * Determine if the domain is allowed to be accessed
@@ -181,49 +181,49 @@ public:
 	 *
 	 * @return true if domain is allowed
 	 */
-	bool IsDomainAllowed(const FString& Url) const;
+	HTTP_API bool IsDomainAllowed(const FString& Url) const;
 
 	/**
 	 * Get the default method for creating new correlation ids for a request
 	 *
 	 * @return The default correlationid creation method
 	 */
-	static TFunction<FString()> GetDefaultCorrelationIdMethod();
+	static HTTP_API TFunction<FString()> GetDefaultCorrelationIdMethod();
 
 	/**
 	 * Inform that HTTP Manager that we are about to fork(). Will block to flush all outstanding http requests
 	 */
-	virtual void OnBeforeFork();
+	HTTP_API virtual void OnBeforeFork();
 
 	/**
 	 * Inform that HTTP Manager that we have completed a fork(). Must be called in both the client and parent process
 	 */
-	virtual void OnAfterFork();
+	HTTP_API virtual void OnAfterFork();
 
 	/**
 	 * Inform the HTTP Manager that we finished ticking right after forking. Only called on the forked process
 	 */
-	virtual void OnEndFramePostFork();
+	HTTP_API virtual void OnEndFramePostFork();
 
 
 	/**
 	 * Update configuration. Called when config has been updated and we need to apply any changes.
 	 */
-	virtual void UpdateConfigs();
+	HTTP_API virtual void UpdateConfigs();
 
 	/**
 	 * Add task to be ran on the game thread next tick
 	 *
 	 * @param Task The task to be ran next tick
 	 */
-	void AddGameThreadTask(TFunction<void()>&& Task);
+	HTTP_API void AddGameThreadTask(TFunction<void()>&& Task);
 
 	/**
 	 * Add task to be ran on the http thread next tick
 	 *
 	 * @param Task The task to be ran next tick
 	 */
-	void AddHttpThreadTask(TFunction<void()>&& Task);
+	HTTP_API void AddHttpThreadTask(TFunction<void()>&& Task);
 
 	/**
 	 * Set url request filter through code, instead of setting it through config.
@@ -238,11 +238,11 @@ protected:
 	 *
 	 * @return the HTTP thread object
 	 */
-	virtual FHttpThreadBase* CreateHttpThread();
+	HTTP_API virtual FHttpThreadBase* CreateHttpThread();
 
-	void ReloadFlushTimeLimits();
+	HTTP_API void ReloadFlushTimeLimits();
 
-	bool HasAnyBoundDelegate() const;
+	HTTP_API bool HasAnyBoundDelegate() const;
 
 protected:
 	/** List of Http requests that are actively being processed */
@@ -303,5 +303,5 @@ PACKAGE_SCOPE:
 	 * Called automatically internally for threaded requests.
 	 * Called explicitly by non-threaded requests
 	 */
-	void BroadcastHttpRequestCompleted(const FHttpRequestRef& Request);
+	HTTP_API void BroadcastHttpRequestCompleted(const FHttpRequestRef& Request);
 };
diff --git a/Engine/Source/Runtime/Online/HTTP/Public/HttpRetrySystem.h b/Engine/Source/Runtime/Online/HTTP/Public/HttpRetrySystem.h
--- a/Engine/Source/Runtime/Online/HTTP/Public/HttpRetrySystem.h
+++ b/Engine/Source/Runtime/Online/HTTP/Public/HttpRetrySystem.h
@@ -168,8 +168,8 @@ namespace FHttpRetrySystem
          * @return                true if there are no failures or retries
          */
         HTTP_API bool Update(uint32* FileCount = NULL, uint32* FailingCount = NULL, uint32* FailedCount = NULL, uint32* CompletedCount = NULL);
-		HTTP_API void SetRandomFailureRate(float Value) { RandomFailureRate = FRandomFailureRateSetting(Value); }
-		HTTP_API void SetDefaultRetryLimit(uint32 Value) { RetryLimitCountDefault = FRetryLimitCountSetting(Value); }
+		void SetRandomFailureRate(float Value) { RandomFailureRate = FRandomFailureRateSetting(Value); }
+		void SetDefaultRetryLimit(uint32 Value) { RetryLimitCountDefault = FRetryLimitCountSetting(Value); }
 		
 		// @return Block the current process until all requests are flushed, or timeout has elapsed
 		HTTP_API void BlockUntilFlushed(float TimeoutSec);
diff --git a/Engine/Source/Runtime/Online/HTTP/Public/Unix/UnixPlatformHttp.h b/Engine/Source/Runtime/Online/HTTP/Public/Unix/UnixPlatformHttp.h
--- a/Engine/Source/Runtime/Online/HTTP/Public/Unix/UnixPlatformHttp.h
+++ b/Engine/Source/Runtime/Online/HTTP/Public/Unix/UnixPlatformHttp.h
@@ -11,33 +11,33 @@ class IHttpRequest;
 /**
  * Platform specific Http implementations
  */
-class HTTP_API FUnixPlatformHttp : public FGenericPlatformHttp
+class FUnixPlatformHttp : public FGenericPlatformHttp
 {
 public:
 
 	/**
 	 * Platform initialization step
 	 */
-	static void Init();
+	static HTTP_API void Init();
 
 	/**
 	 * Creates a platform-specific HTTP manager.
 	 *
 	 * @return NULL if default implementation is to be used
 	 */
-	static FHttpManager* CreatePlatformHttpManager();
+	static HTTP_API FHttpManager* CreatePlatformHttpManager();
 
 	/**
 	 * Platform shutdown step
 	 */
-	static void Shutdown();
+	static HTTP_API void Shutdown();
 
 	/**
 	 * Creates a new Http request instance for the current platform
 	 *
 	 * @return request object
 	 */
-	static IHttpRequest* ConstructRequest();
+	static HTTP_API IHttpRequest* ConstructRequest();
 };
 
 
diff --git a/Engine/Source/Runtime/Online/HTTP/Public/Windows/WindowsPlatformHttp.h b/Engine/Source/Runtime/Online/HTTP/Public/Windows/WindowsPlatformHttp.h
--- a/Engine/Source/Runtime/Online/HTTP/Public/Windows/WindowsPlatformHttp.h
+++ b/Engine/Source/Runtime/Online/HTTP/Public/Windows/WindowsPlatformHttp.h
@@ -8,41 +8,41 @@
 /**
  * Platform specific HTTP implementations.
  */
-class HTTP_API FWindowsPlatformHttp : public FGenericPlatformHttp
+class FWindowsPlatformHttp : public FGenericPlatformHttp
 {
 public:
 
 	/** Platform initialization step. */
-	static void Init();
+	static HTTP_API void Init();
 
 	/**
 	 * Creates a platform-specific HTTP manager.
 	 *
 	 * @return nullptr if default implementation is to be used.
 	 */
-	static FHttpManager* CreatePlatformHttpManager();
+	static HTTP_API FHttpManager* CreatePlatformHttpManager();
 
 	/** Platform shutdown step. */
-	static void Shutdown();
+	static HTTP_API void Shutdown();
 
 	/**
 	 * Creates a new HTTP request instance for the current platform.
 	 *
 	 * @return The request object.
 	 */
-	static IHttpRequest* ConstructRequest();
+	static HTTP_API IHttpRequest* ConstructRequest();
 
 	/**
 	 * @return the mime type for the file.
 	 */
-	static FString GetMimeType(const FString& FilePath);
+	static HTTP_API FString GetMimeType(const FString& FilePath);
 
 	/**
 	 * Get the proxy address specified by the operating system
 	 *
 	 * @return optional FString: If unset: we are unable to get information from the operating system. If set: the proxy address set by the operating system (may be blank)
 	 */
-	static TOptional<FString> GetOperatingSystemProxyAddress();
+	static HTTP_API TOptional<FString> GetOperatingSystemProxyAddress();
 
 	/**
 	 * Check if getting proxy information from the current operating system is supported
@@ -51,14 +51,14 @@ public:
 	 * 
 	 * @return true if we are able to get proxy information from the current operating system, false if not
 	 */
-	static bool IsOperatingSystemProxyInformationSupported();
+	static HTTP_API bool IsOperatingSystemProxyInformationSupported();
 	
 	/**
 	 * Verify Peer Ssl Certificate
 	 *
 	 * @return optional bool: the previous value
 	 */
-	static bool VerifyPeerSslCertificate(bool verify);
+	static HTTP_API bool VerifyPeerSslCertificate(bool verify);
 };
 
 #if WINDOWS_USE_FEATURE_PLATFORMHTTP_CLASS
diff --git a/Engine/Source/Runtime/Online/HTTPServer/Public/HttpPath.h b/Engine/Source/Runtime/Online/HTTPServer/Public/HttpPath.h
--- a/Engine/Source/Runtime/Online/HTTPServer/Public/HttpPath.h
+++ b/Engine/Source/Runtime/Online/HTTPServer/Public/HttpPath.h
@@ -64,7 +64,7 @@ public:
 	HTTPSERVER_API void MakeRelative(const FString& OtherPath);
 
 	// TMap<> Comparer
-	HTTPSERVER_API bool operator==(const FHttpPath& Other) const
+	bool operator==(const FHttpPath& Other) const
 	{
 		bool bEqual = (0 == Path.Compare(Other.Path, ESearchCase::CaseSensitive));
 		return bEqual;
diff --git a/Engine/Source/Runtime/Online/ICMP/Public/Icmp.h b/Engine/Source/Runtime/Online/ICMP/Public/Icmp.h
--- a/Engine/Source/Runtime/Online/ICMP/Public/Icmp.h
+++ b/Engine/Source/Runtime/Online/ICMP/Public/Icmp.h
@@ -112,7 +112,7 @@ DECLARE_DELEGATE_OneParam(FIcmpEchoManyCompleteDelegate, FIcmpEchoManyCompleteRe
 
 
 // Simple ping interface that sends an ICMP packet to the given address and returns timing info for the reply if reachable
-class ICMP_API FIcmp
+class FIcmp
 {
 	public:
 
@@ -130,7 +130,7 @@ class ICMP_API FIcmp
 	 * @param Timeout max time to wait for a reply
 	 * @param HandleResult a callback function that will be called when the result is ready
 	 */
-	static void IcmpEcho(const FString& TargetAddress, float Timeout, FIcmpEchoResultCallback HandleResult);
+	static ICMP_API void IcmpEcho(const FString& TargetAddress, float Timeout, FIcmpEchoResultCallback HandleResult);
 
 	/** Send an ICMP echo packet and wait for a reply.
 	 *
@@ -150,7 +150,7 @@ class ICMP_API FIcmp
 };
 
 // Simple ping interface that sends an ICMP packet over UDP to the given address and returns timing info for the reply if reachable
-class ICMP_API FUDPPing
+class FUDPPing
 {
 public:
 
@@ -168,7 +168,7 @@ public:
 	 * @param Timeout max time to wait for a reply
 	 * @param HandleResult a callback function that will be called when the result is ready
 	 */
-	static void UDPEcho(const FString& TargetAddress, float Timeout, FIcmpEchoResultCallback HandleResult);
+	static ICMP_API void UDPEcho(const FString& TargetAddress, float Timeout, FIcmpEchoResultCallback HandleResult);
 
 	/** Send an ICMP echo packet and wait for a reply.
 	 *
@@ -202,7 +202,7 @@ public:
 	 * @param Timeout max time to wait for a replies
 	 * @param CompletionCallback callback function that is invoked when the final results are ready
 	 */
-	static void UDPEchoMany(const TArray<FIcmpTarget>& Targets, float Timeout, FIcmpEchoManyCompleteCallback CompletionCallback);
+	static ICMP_API void UDPEchoMany(const TArray<FIcmpTarget>& Targets, float Timeout, FIcmpEchoManyCompleteCallback CompletionCallback);
 
 	/** Send multiple ICMP echo packets and wait for replies.
 	 *
@@ -220,7 +220,7 @@ public:
 	 * @param Timeout max time to wait for a replies
 	 * @param CompletionDelegate delegate that is invoked when the final results are ready
 	 */
-	static void UDPEchoMany(const TArray<FIcmpTarget>& Targets, float Timeout, FIcmpEchoManyCompleteDelegate CompletionDelegate);
+	static ICMP_API void UDPEchoMany(const TArray<FIcmpTarget>& Targets, float Timeout, FIcmpEchoManyCompleteDelegate CompletionDelegate);
 };
 
 
diff --git a/Engine/Source/Runtime/Online/ImageDownload/Public/WebImage.h b/Engine/Source/Runtime/Online/ImageDownload/Public/WebImage.h
--- a/Engine/Source/Runtime/Online/ImageDownload/Public/WebImage.h
+++ b/Engine/Source/Runtime/Online/ImageDownload/Public/WebImage.h
@@ -39,12 +39,12 @@ typedef TSharedPtr<class IHttpResponse, ESPMode::ThreadSafe> FHttpResponsePtr;
  *   SNew(SImage)
  *   .Image(ImageCache.Download(Url)->Attr())
  */
-class IMAGEDOWNLOAD_API FWebImage 
+class FWebImage 
 	: public TSharedFromThis<FWebImage>
 {
 public:
-	FWebImage();
-	~FWebImage();
+	IMAGEDOWNLOAD_API FWebImage();
+	IMAGEDOWNLOAD_API ~FWebImage();
 
 	/**
 	 * Fired when the image finishes downloading or is canceled.
@@ -59,7 +59,7 @@ public:
 	void SetRGBFormat(const ERGBFormat InRGBFormat) { RGBFormat = InRGBFormat; }
 
 	/** Begin downloading an image. This will automatically set the current brush to the downloaded image when it completes (if successful) */
-	bool BeginDownload(const FString& InUrl, const TOptional<FString>& StandInETag = TOptional<FString>(), const FOnImageDownloaded& DownloadCallback = FOnImageDownloaded());
+	IMAGEDOWNLOAD_API bool BeginDownload(const FString& InUrl, const TOptional<FString>& StandInETag = TOptional<FString>(), const FOnImageDownloaded& DownloadCallback = FOnImageDownloaded());
 
 	/** Begin downloading an image. */
 	FORCEINLINE bool BeginDownload(const FString& InUrl, const FOnImageDownloaded& DownloadCallback)
@@ -68,12 +68,12 @@ public:
 	}
 
 	/** Cancel any download in progress */
-	void CancelDownload();
+	IMAGEDOWNLOAD_API void CancelDownload();
 
 public:
 
 	/** Use .Attr() to pass this brush into a slate attribute */
-	TAttribute< const FSlateBrush* > Attr() const;
+	IMAGEDOWNLOAD_API TAttribute< const FSlateBrush* > Attr() const;
 
 	/** Get the current brush displayed (will automatically change when download completes) */
 	FORCEINLINE const FSlateBrush* GetBrush() const { return DownloadedBrush.IsValid() ? DownloadedBrush.Get() : StandInBrush.Get(); }
@@ -98,8 +98,8 @@ public:
 
 private:
 	/** request complete callback */
-	void HttpRequestComplete(FHttpRequestPtr HttpRequest, FHttpResponsePtr HttpResponse, bool bSucceeded);
-	bool ProcessHttpResponse(const FString& RequestUrl, FHttpResponsePtr HttpResponse);
+	IMAGEDOWNLOAD_API void HttpRequestComplete(FHttpRequestPtr HttpRequest, FHttpResponsePtr HttpResponse, bool bSucceeded);
+	IMAGEDOWNLOAD_API bool ProcessHttpResponse(const FString& RequestUrl, FHttpResponsePtr HttpResponse);
 
 private:
 	/** The Url being downloaded */
diff --git a/Engine/Source/Runtime/Online/ImageDownload/Public/WebImageCache.h b/Engine/Source/Runtime/Online/ImageDownload/Public/WebImageCache.h
--- a/Engine/Source/Runtime/Online/ImageDownload/Public/WebImageCache.h
+++ b/Engine/Source/Runtime/Online/ImageDownload/Public/WebImageCache.h
@@ -14,19 +14,19 @@ struct FSlateBrush;
  *
  * See WebImage.h for example usage.
  */
-class IMAGEDOWNLOAD_API FWebImageCache
+class FWebImageCache
 {
 public:
-	FWebImageCache();
+	IMAGEDOWNLOAD_API FWebImageCache();
 
 	/** Signifies the module is being unloaded and to perform any actions that depend on other modules which may be unloaded as well */
-	void PreUnload();
+	IMAGEDOWNLOAD_API void PreUnload();
 
 	/** Removes all cached images */
-	void Empty();
+	IMAGEDOWNLOAD_API void Empty();
 
 	/** Find or create a WebImage object for this URL (you probably just want to call ->Attr() on this) */
-	TSharedRef<const FWebImage> Download(const FString& Url, const TOptional<FString>& DefaultImageUrl = TOptional<FString>());
+	IMAGEDOWNLOAD_API TSharedRef<const FWebImage> Download(const FString& Url, const TOptional<FString>& DefaultImageUrl = TOptional<FString>());
 
 	/** Set the brush that will be returned until the download completes (only affects future downloads). */
 	FORCEINLINE void SetDefaultStandInBrush(TAttribute<const FSlateBrush*> StandInBrushIn) { DefaultStandInBrush = StandInBrushIn; }
@@ -35,7 +35,7 @@ public:
 	 * once downloaded, an image is cached forever. This allows us to release images that are not currently
 	 * being displayed (those would have Strong pointers existing external to this class) to be released.
 	 */
-	void RelinquishUnusedImages();
+	IMAGEDOWNLOAD_API void RelinquishUnusedImages();
 
 private:
 	/** Map of canonical URL to web images (weak pointer so we don't affect lifetime) */
diff --git a/Engine/Source/Runtime/Online/SSL/Private/SslManager.h b/Engine/Source/Runtime/Online/SSL/Private/SslManager.h
--- a/Engine/Source/Runtime/Online/SSL/Private/SslManager.h
+++ b/Engine/Source/Runtime/Online/SSL/Private/SslManager.h
@@ -12,20 +12,20 @@
 /**
  * Manager of the ssl library
  */
-class SSL_API FSslManager : ISslManager
+class FSslManager : ISslManager
 {
 public:
 
 	//~ Begin ISslManager Interface
-	virtual bool InitializeSsl() override;
-	virtual void ShutdownSsl() override;
-	virtual SSL_CTX* CreateSslContext(const FSslContextCreateOptions& CreateOptions) override;
-	virtual void DestroySslContext(SSL_CTX* SslContext) override;
+	SSL_API virtual bool InitializeSsl() override;
+	SSL_API virtual void ShutdownSsl() override;
+	SSL_API virtual SSL_CTX* CreateSslContext(const FSslContextCreateOptions& CreateOptions) override;
+	SSL_API virtual void DestroySslContext(SSL_CTX* SslContext) override;
 	//~ End ISslManager Interface
 
 protected:
 	/** Default constructor */
-	FSslManager();
+	SSL_API FSslManager();
 	/** Disable copies */
 	UE_NONCOPYABLE(FSslManager)
 	/** SSL ref count */
diff --git a/Engine/Source/Runtime/Online/SSL/Public/Interfaces/ISslCertificateManager.h b/Engine/Source/Runtime/Online/SSL/Public/Interfaces/ISslCertificateManager.h
--- a/Engine/Source/Runtime/Online/SSL/Public/Interfaces/ISslCertificateManager.h
+++ b/Engine/Source/Runtime/Online/SSL/Public/Interfaces/ISslCertificateManager.h
@@ -76,7 +76,7 @@ public:
 	virtual bool VerifySslCertificates(TArray<TArray<uint8, TFixedAllocator<PUBLIC_KEY_DIGEST_SIZE>>>& Digests, const FString& Domain) const = 0;
 };
 
-class SSL_API FSslCertificateDelegates
+class FSslCertificateDelegates
 {
 public:
 	struct FCertInfo
@@ -90,5 +90,5 @@ public:
 	};
 
 	DECLARE_DELEGATE_RetVal_TwoParams(bool, FVerifySslCertificates, const FString&, const TArray<FCertInfo>&);
-	static FVerifySslCertificates VerifySslCertificates;
+	static SSL_API FVerifySslCertificates VerifySslCertificates;
 };
diff --git a/Engine/Source/Runtime/Online/SSL/Public/Interfaces/ISslManager.h b/Engine/Source/Runtime/Online/SSL/Public/Interfaces/ISslManager.h
--- a/Engine/Source/Runtime/Online/SSL/Public/Interfaces/ISslManager.h
+++ b/Engine/Source/Runtime/Online/SSL/Public/Interfaces/ISslManager.h
@@ -50,7 +50,7 @@ struct FSslContextCreateOptions
 /**
  * Manager of the ssl library
  */
-class SSL_API ISslManager
+class ISslManager
 {
 public:
 	/** Destructor */
diff --git a/Engine/Source/Runtime/Online/Stomp/Public/StompModule.h b/Engine/Source/Runtime/Online/Stomp/Public/StompModule.h
--- a/Engine/Source/Runtime/Online/Stomp/Public/StompModule.h
+++ b/Engine/Source/Runtime/Online/Stomp/Public/StompModule.h
@@ -12,7 +12,7 @@ class IStompClient;
 /**
  * Module for Stomp over WebSockets
  */
-class STOMP_API FStompModule :
+class FStompModule :
 	public IModuleInterface
 {
 
@@ -26,7 +26,7 @@ public:
 	 *
 	 * @return Returns singleton instance
 	 */
-	static FStompModule& Get();
+	static STOMP_API FStompModule& Get();
 
 #if WITH_STOMP
 	/**
@@ -35,7 +35,7 @@ public:
 	 * @param Url The URL to which to connect; this should be the URL to which the WebSocket server will respond with Stomp protocol data.
 	 * @return new IStompClient instance
 	 */
-	TSharedRef<IStompClient> CreateClient(const FString& Url, const FString& OptAuthToken = FString());
+	STOMP_API TSharedRef<IStompClient> CreateClient(const FString& Url, const FString& OptAuthToken = FString());
 #endif // #if WITH_STOMP
 
 private:
@@ -46,13 +46,13 @@ private:
 	 * Called when Stomp module is loaded
 	 * Initialize implementation specific parts of Stomp handling
 	 */
-	virtual void StartupModule() override;
+	STOMP_API virtual void StartupModule() override;
 
 	/**
 	 * Called when Stomp module is unloaded
 	 * Shutdown implementation specific parts of Stomp handling
 	 */
-	virtual void ShutdownModule() override;
+	STOMP_API virtual void ShutdownModule() override;
 
 
 	/** singleton for the module while loaded and available */
diff --git a/Engine/Source/Runtime/Online/WebSockets/Public/WebSocketsModule.h b/Engine/Source/Runtime/Online/WebSockets/Public/WebSocketsModule.h
--- a/Engine/Source/Runtime/Online/WebSockets/Public/WebSocketsModule.h
+++ b/Engine/Source/Runtime/Online/WebSockets/Public/WebSocketsModule.h
@@ -18,7 +18,7 @@ class IWebSocketsManager;
 /**
  * Module for web socket implementations
  */
-class WEBSOCKETS_API FWebSocketsModule :
+class FWebSocketsModule :
 	public IModuleInterface
 {
 
@@ -49,7 +49,7 @@ public:
 	 *
 	 * @return Returns singleton instance, loading the module on demand if needed
 	 */
-	static FWebSocketsModule& Get();
+	static WEBSOCKETS_API FWebSocketsModule& Get();
 
 #if WITH_WEBSOCKETS
 	/**
@@ -59,7 +59,7 @@ public:
 	 * @param Protocols a list of protocols the client will handle.
 	 * @return new IWebSocket instance
 	 */
-	virtual TSharedRef<IWebSocket> CreateWebSocket(const FString& Url, const TArray<FString>& Protocols, const TMap<FString, FString>& UpgradeHeaders = TMap<FString, FString>());
+	WEBSOCKETS_API virtual TSharedRef<IWebSocket> CreateWebSocket(const FString& Url, const TArray<FString>& Protocols, const TMap<FString, FString>& UpgradeHeaders = TMap<FString, FString>());
 
 
 	/**
@@ -69,11 +69,11 @@ public:
 	 * @param Protocol an optional sub-protocol. If missing, an empty string is assumed.
 	 * @return new IWebSocket instance
 	 */
-	virtual TSharedRef<IWebSocket> CreateWebSocket(const FString& Url, const FString& Protocol = FString(), const TMap<FString, FString>& UpgradeHeaders = TMap<FString, FString>());
+	WEBSOCKETS_API virtual TSharedRef<IWebSocket> CreateWebSocket(const FString& Url, const FString& Protocol = FString(), const TMap<FString, FString>& UpgradeHeaders = TMap<FString, FString>());
 #endif // #if WITH_WEBSOCKETS
 
 private:
-	static FString BuildUpgradeHeader(const TMap<FString, FString>& Headers);
+	static WEBSOCKETS_API FString BuildUpgradeHeader(const TMap<FString, FString>& Headers);
 
 	// IModuleInterface
 
@@ -81,13 +81,13 @@ private:
 	 * Called when WebSockets module is loaded
 	 * Initialize implementation specific parts of WebSockets handling
 	 */
-	virtual void StartupModule() override;
+	WEBSOCKETS_API virtual void StartupModule() override;
 	
 	/**
 	 * Called when WebSockets module is unloaded
 	 * Shutdown implementation specific parts of WebSockets handling
 	 */
-	virtual void ShutdownModule() override;
+	WEBSOCKETS_API virtual void ShutdownModule() override;
 
 #if WITH_WEBSOCKETS
 	/** Manages active web sockets */
@@ -97,5 +97,5 @@ private:
 #endif // #if WITH_WEBSOCKETS
 
 	/** singleton for the module while loaded and available */
-	static FWebSocketsModule* Singleton;
+	static WEBSOCKETS_API FWebSocketsModule* Singleton;
 };
diff --git a/Engine/Source/Runtime/Online/XMPP/Public/XmppConnection.h b/Engine/Source/Runtime/Online/XMPP/Public/XmppConnection.h
--- a/Engine/Source/Runtime/Online/XMPP/Public/XmppConnection.h
+++ b/Engine/Source/Runtime/Online/XMPP/Public/XmppConnection.h
@@ -97,7 +97,7 @@ public:
 /**
  * Jid for identifying user on the current connection as well as roster members
  */
-class XMPP_API FXmppUserJid
+class FXmppUserJid
 {
 public:
 
@@ -119,7 +119,7 @@ public:
 	 * @param JidString the formatted JID
 	 * @return FXmppUserJid populated from the formatted JID
 	 */
-	static FXmppUserJid FromFullJid(const FString& JidString);
+	static XMPP_API FXmppUserJid FromFullJid(const FString& JidString);
 
 	/** unique id for the user */
 	FString Id;
@@ -138,9 +138,9 @@ public:
 	 *
 	 * @return Whether the Resource was successfully parsed or not
 	 */
-	static bool ParseResource(const FString& InResource, FString& OutAppId, FString& OutPlatform, FString& OutPlatformUserId);
+	static XMPP_API bool ParseResource(const FString& InResource, FString& OutAppId, FString& OutPlatform, FString& OutPlatformUserId);
 
-	static FString CreateResource(const FString& AppId, const FString& Platform, const FString& PlatformUserId);
+	static XMPP_API FString CreateResource(const FString& AppId, const FString& Platform, const FString& PlatformUserId);
 
 	/** 
 	 * Get the components that comprise the resource
@@ -160,7 +160,7 @@ public:
 	 * Separate the MUC half of the resource (nickname:userid) from the UserJid resource portion (Vx:AppId:Platform:etc)
 	 * This is highly dependent on FChatRoomMemberMcp::BuildMemberJidResourceString
 	 */
-	static FString ParseMucUserResource(const FString& InResource);
+	static XMPP_API FString ParseMucUserResource(const FString& InResource);
 
 	bool operator==(const FXmppUserJid& Other) const
 	{
@@ -204,7 +204,7 @@ public:
 		return !Id.IsEmpty() && !Domain.IsEmpty();
 	}
 
-	FString ToDebugString() const;
+	XMPP_API FString ToDebugString() const;
 };
 
 typedef TSharedPtr<IXmppPresence, ESPMode::ThreadSafe> IXmppPresencePtr;
diff --git a/Engine/Source/Runtime/Online/XMPP/Public/XmppModule.h b/Engine/Source/Runtime/Online/XMPP/Public/XmppModule.h
--- a/Engine/Source/Runtime/Online/XMPP/Public/XmppModule.h
+++ b/Engine/Source/Runtime/Online/XMPP/Public/XmppModule.h
@@ -25,7 +25,7 @@ class UWorld;
  * Module for Xmpp connections
  * Use CreateConnection to create a new Xmpp connection
  */
-class XMPP_API FXmppModule :
+class FXmppModule :
 	public IModuleInterface, public FSelfRegisteringExec, public FTSTickerObjectBase
 {
 
@@ -42,14 +42,14 @@ protected:
 	 *
 	 * @return true if the handler consumed the input, false to continue searching handlers
 	 */
-	virtual bool Exec_Runtime(UWorld* InWorld, const TCHAR* Cmd, FOutputDevice& Ar) override;
+	XMPP_API virtual bool Exec_Runtime(UWorld* InWorld, const TCHAR* Cmd, FOutputDevice& Ar) override;
 
 public:
 
 	/** 
 	 * Exec command handlers
 	 */
-	bool HandleXmppCommand( const TCHAR* Cmd, FOutputDevice& Ar );
+	XMPP_API bool HandleXmppCommand( const TCHAR* Cmd, FOutputDevice& Ar );
 
 	// FXmppModule
 
@@ -59,47 +59,47 @@ public:
 	 *
 	 * @return Returns singleton instance, loading the module on demand if needed
 	 */
-	static FXmppModule& Get();
+	static XMPP_API FXmppModule& Get();
 
 	/**
 	 * Checks to see if this module is loaded and ready.
 	 *
 	 * @return True if the module is loaded and ready to use
 	 */
-	static bool IsAvailable();
+	static XMPP_API bool IsAvailable();
 
 	/**
 	 * Creates a new Xmpp connection for the current platform and associated it with the user
 	 *
 	 * @return new Xmpp connection instance
 	 */
-	TSharedRef<class IXmppConnection> CreateConnection(const FString& UserId);
+	XMPP_API TSharedRef<class IXmppConnection> CreateConnection(const FString& UserId);
 
 	/**
 	 * Return an existing Xmpp connection associated with a user
 	 *
 	 * @return new Xmpp connection instance
 	 */
-	TSharedPtr<class IXmppConnection> GetConnection(const FString& UserId) const;
+	XMPP_API TSharedPtr<class IXmppConnection> GetConnection(const FString& UserId) const;
 
 	/**
 	 * Remove an existing Xmpp connection associated with a user
 	 *
 	 * @param UserId user to find connection for
 	 */
-	void RemoveConnection(const FString& UserId);
+	XMPP_API void RemoveConnection(const FString& UserId);
 
 	/**
 	 * Clean up any pending connection removals
 	 */
-	void ProcessPendingRemovals();
+	XMPP_API void ProcessPendingRemovals();
 
 	/**
 	 * Remove an existing Xmpp connection
 	 *
 	 * @param Connection reference to find/remove
 	 */
-	void RemoveConnection(const TSharedRef<class IXmppConnection>& Connection);
+	XMPP_API void RemoveConnection(const TSharedRef<class IXmppConnection>& Connection);
 
 	/**
 	 * @return true if Xmpp requests are globally enabled
@@ -110,7 +110,7 @@ public:
 	}
 
 	// FTSTickerObjectBase
-	virtual bool Tick(float DeltaTime) override;
+	XMPP_API virtual bool Tick(float DeltaTime) override;
 
 	/**
 	 * Delegate callback when a system acquires ownership over an XMPP connection
@@ -144,30 +144,30 @@ private:
 
 	// IModuleInterface
 
-	void OnXmppRoomCreated(const TSharedRef<IXmppConnection>& Connection, bool bSuccess, const FXmppRoomId& RoomId, const FString& Error);
-	void OnXmppRoomConfigured(const TSharedRef<IXmppConnection>& Connection, bool bSuccess, const FXmppRoomId& RoomId, const FString& Error);
+	XMPP_API void OnXmppRoomCreated(const TSharedRef<IXmppConnection>& Connection, bool bSuccess, const FXmppRoomId& RoomId, const FString& Error);
+	XMPP_API void OnXmppRoomConfigured(const TSharedRef<IXmppConnection>& Connection, bool bSuccess, const FXmppRoomId& RoomId, const FString& Error);
 
 	/**
 	 * Called when Xmpp module is loaded
 	 * Initialize platform specific parts of Xmpp handling
 	 */
-	virtual void StartupModule() override;
+	XMPP_API virtual void StartupModule() override;
 	
 	/**
 	 * Called when Xmpp module is unloaded
 	 * Shutdown platform specific parts of Xmpp handling
 	 */
-	virtual void ShutdownModule() override;
+	XMPP_API virtual void ShutdownModule() override;
 
 	/**
 	 * Connection cleanup before removal
 	 */
-	void CleanupConnection(const TSharedRef<class IXmppConnection>& Connection);
+	XMPP_API void CleanupConnection(const TSharedRef<class IXmppConnection>& Connection);
 
 	/** toggles Xmpp requests */
 	bool bEnabled;
 	/** singleton for the module while loaded and available */
-	static FXmppModule* Singleton;
+	static XMPP_API FXmppModule* Singleton;
 
 	/** Active Xmpp server connections mapped by user id */
 	TMap<FString, TSharedRef<class IXmppConnection>> ActiveConnections;
diff --git a/Engine/Source/Runtime/Overlay/Public/BasicOverlays.h b/Engine/Source/Runtime/Overlay/Public/BasicOverlays.h
--- a/Engine/Source/Runtime/Overlay/Public/BasicOverlays.h
+++ b/Engine/Source/Runtime/Overlay/Public/BasicOverlays.h
@@ -21,8 +21,8 @@ class UObject;
  * Implements an asset that contains a set of overlay data (which includes timing, text, and position) to be displayed for any
  * given source (including, but not limited to, audio, dialog, and movies)
  */
-UCLASS(BlueprintType, hidecategories = (Object))
-class OVERLAY_API UBasicOverlays
+UCLASS(BlueprintType, hidecategories = (Object), MinimalAPI)
+class UBasicOverlays
 	: public UOverlays
 {
 	GENERATED_BODY()
@@ -45,17 +45,17 @@ public:
 
 	//~ UOverlays interface
 
-	virtual TArray<FOverlayItem> GetAllOverlays() const override;
-	virtual void GetOverlaysForTime(const FTimespan& Time, TArray<FOverlayItem>& OutOverlays) const override;
+	OVERLAY_API virtual TArray<FOverlayItem> GetAllOverlays() const override;
+	OVERLAY_API virtual void GetOverlaysForTime(const FTimespan& Time, TArray<FOverlayItem>& OutOverlays) const override;
 
 public:
 	
 	//~ UObject interface
 	
-	virtual void PostInitProperties() override;
-	virtual void GetAssetRegistryTags(TArray<FAssetRegistryTag>& OutTags) const override;
+	OVERLAY_API virtual void PostInitProperties() override;
+	OVERLAY_API virtual void GetAssetRegistryTags(TArray<FAssetRegistryTag>& OutTags) const override;
 	// End UObject interface
 
 private:
 
-};
\ No newline at end of file
+};
diff --git a/Engine/Source/Runtime/Overlay/Public/LocalizedOverlays.h b/Engine/Source/Runtime/Overlay/Public/LocalizedOverlays.h
--- a/Engine/Source/Runtime/Overlay/Public/LocalizedOverlays.h
+++ b/Engine/Source/Runtime/Overlay/Public/LocalizedOverlays.h
@@ -25,8 +25,8 @@ class UObject;
  * Implements an asset that contains a set of Basic Overlays that will be displayed in accordance with
  * the current locale, or a default set if an appropriate locale is not found
  */
-UCLASS(BlueprintType, hidecategories = (Object))
-class OVERLAY_API ULocalizedOverlays
+UCLASS(BlueprintType, hidecategories = (Object), MinimalAPI)
+class ULocalizedOverlays
 	: public UOverlays
 {
 	GENERATED_BODY()
@@ -59,16 +59,16 @@ public:
 
 	//~ UOverlays interface
 
-	virtual TArray<FOverlayItem> GetAllOverlays() const override;
-	virtual void GetOverlaysForTime(const FTimespan& Time, TArray<FOverlayItem>& OutOverlays) const override;
+	OVERLAY_API virtual TArray<FOverlayItem> GetAllOverlays() const override;
+	OVERLAY_API virtual void GetOverlaysForTime(const FTimespan& Time, TArray<FOverlayItem>& OutOverlays) const override;
 
 
 public:
 
 	//~ UObject interface
 
-	virtual void PostInitProperties() override;
-	virtual void GetAssetRegistryTags(TArray<FAssetRegistryTag>& OutTags) const override;
+	OVERLAY_API virtual void PostInitProperties() override;
+	OVERLAY_API virtual void GetAssetRegistryTags(TArray<FAssetRegistryTag>& OutTags) const override;
 
 private:
 	
@@ -76,4 +76,4 @@ private:
 	 * Retrieves the overlays object for the current locale
 	 */
 	UBasicOverlays* GetCurrentLocaleOverlays() const;
-};
\ No newline at end of file
+};
diff --git a/Engine/Source/Runtime/Overlay/Public/Overlays.h b/Engine/Source/Runtime/Overlay/Public/Overlays.h
--- a/Engine/Source/Runtime/Overlay/Public/Overlays.h
+++ b/Engine/Source/Runtime/Overlay/Public/Overlays.h
@@ -11,7 +11,7 @@
 
 
 USTRUCT(BlueprintType)
-struct OVERLAY_API FOverlayItem
+struct FOverlayItem
 {
 	GENERATED_BODY()
 
@@ -43,8 +43,8 @@ struct OVERLAY_API FOverlayItem
 /**
  * An interface class for creating overlay data assets
  */
-UCLASS(Abstract)
-class OVERLAY_API UOverlays
+UCLASS(Abstract, MinimalAPI)
+class UOverlays
 	: public UObject
 {
 	GENERATED_BODY()
@@ -54,7 +54,7 @@ public:
 	/**
 	 * Retrieves the set of overlays that should be used by this object (must be implemented by child classes)
 	 */
-	virtual TArray<FOverlayItem> GetAllOverlays() const PURE_VIRTUAL(UOverlays::GetAllOverlays, return TArray<FOverlayItem>(););
+	OVERLAY_API virtual TArray<FOverlayItem> GetAllOverlays() const PURE_VIRTUAL(UOverlays::GetAllOverlays, return TArray<FOverlayItem>(););
 
 	/**
 	 * Retrieves the set of overlays associated with this object for the given timespan (must be implemented by child classes)
@@ -62,5 +62,5 @@ public:
 	 * @param	Time			Determines what overlays should be displayed
 	 * @param	OutOverlays		All overlays that should be displayed for the given Time. Expect this to be emptied.
 	 */
-	virtual void GetOverlaysForTime(const FTimespan& Time, TArray<FOverlayItem>& OutOverlays) const PURE_VIRTUAL(UOverlays::GetAllOverlays, return;);
+	OVERLAY_API virtual void GetOverlaysForTime(const FTimespan& Time, TArray<FOverlayItem>& OutOverlays) const PURE_VIRTUAL(UOverlays::GetAllOverlays, return;);
 };
diff --git a/Engine/Source/Runtime/Overlay/Public/OverlaysImporter.h b/Engine/Source/Runtime/Overlay/Public/OverlaysImporter.h
--- a/Engine/Source/Runtime/Overlay/Public/OverlaysImporter.h
+++ b/Engine/Source/Runtime/Overlay/Public/OverlaysImporter.h
@@ -14,11 +14,11 @@ enum class EOverlaysFileType
 	SubRipSubtitles,
 };
 
-class OVERLAY_API FOverlaysImporter
+class FOverlaysImporter
 {
 public:
-	FOverlaysImporter();
-	~FOverlaysImporter();
+	OVERLAY_API FOverlaysImporter();
+	OVERLAY_API ~FOverlaysImporter();
 
 	/**
 	 * Opens the file and preparses it for import
@@ -26,7 +26,7 @@ public:
 	 * @param	Filename	The file to open
 	 * @return	True if the file was opened successfully, false if the file could not be opened or is not an overlay file.
 	 */
-	bool OpenFile(const FString& FilePath);
+	OVERLAY_API bool OpenFile(const FString& FilePath);
 
 	/**
 	 * Parses the supplied import file for basic overlay data
@@ -34,12 +34,12 @@ public:
 	 * @param	OutSubtitles	The output array where new overlays are stored. This is emptied when the import begins
 	 * @return	True if the file was parsed successfully
 	 */
-	bool ImportBasic(TArray<FOverlayItem>& OutOverlays) const;
+	OVERLAY_API bool ImportBasic(TArray<FOverlayItem>& OutOverlays) const;
 
 	/**
 	 * Resets the importer to a default state
 	 */
-	void Reset();
+	OVERLAY_API void Reset();
 
 private:
 	bool ParseSubRipSubtitles(TArray<FOverlayItem>& OutSubtitles) const;
@@ -49,4 +49,4 @@ private:
 	FString FileContents;
 
 	EOverlaysFileType FileType;
-};
\ No newline at end of file
+};
diff --git a/Engine/Source/Runtime/PacketHandlers/PacketHandler/Classes/HandlerComponentFactory.h b/Engine/Source/Runtime/PacketHandlers/PacketHandler/Classes/HandlerComponentFactory.h
--- a/Engine/Source/Runtime/PacketHandlers/PacketHandler/Classes/HandlerComponentFactory.h
+++ b/Engine/Source/Runtime/PacketHandlers/PacketHandler/Classes/HandlerComponentFactory.h
@@ -13,12 +13,12 @@ class HandlerComponent;
 /**
  * A UObject alternative for loading HandlerComponents without strict module dependency
  */
-UCLASS(abstract)
-class PACKETHANDLER_API UHandlerComponentFactory : public UObject
+UCLASS(abstract, MinimalAPI)
+class UHandlerComponentFactory : public UObject
 {
 	GENERATED_UCLASS_BODY()
 
 public:
-	virtual TSharedPtr<HandlerComponent> CreateComponentInstance(FString& Options)
+	PACKETHANDLER_API virtual TSharedPtr<HandlerComponent> CreateComponentInstance(FString& Options)
 		PURE_VIRTUAL(UHandlerComponentFactory::CreateComponentInstance, return nullptr;);
 };
diff --git a/Engine/Source/Runtime/PacketHandlers/PacketHandler/Public/EncryptionComponent.h b/Engine/Source/Runtime/PacketHandlers/PacketHandler/Public/EncryptionComponent.h
--- a/Engine/Source/Runtime/PacketHandlers/PacketHandler/Public/EncryptionComponent.h
+++ b/Engine/Source/Runtime/PacketHandlers/PacketHandler/Public/EncryptionComponent.h
@@ -10,7 +10,7 @@
 /**
  * IEncryptionComponentInterface
  */
-class PACKETHANDLER_API FEncryptionComponent : public HandlerComponent
+class FEncryptionComponent : public HandlerComponent
 {
 public:
 	/**
