Fix type TCHAR* -> const TCHAR*

#rnx

diff --git a/Engine/Source/Runtime/Online/Experimental/EventLoopTests/Tests/EventLoop/EventLoopManagedStorageTests.cpp b/Engine/Source/Runtime/Online/Experimental/EventLoopTests/Tests/EventLoop/EventLoopManagedStorageTests.cpp
--- a/Engine/Source/Runtime/Online/Experimental/EventLoopTests/Tests/EventLoop/EventLoopManagedStorageTests.cpp
+++ b/Engine/Source/Runtime/Online/Experimental/EventLoopTests/Tests/EventLoop/EventLoopManagedStorageTests.cpp
@@ -8,7 +8,7 @@ namespace UE::EventLoop
 {
 	struct FMockManagedStorageHandleTraits
 	{
-		static constexpr TCHAR* Name = TEXT("MockManagedStorageHandle");
+		static constexpr const TCHAR* Name = TEXT("MockManagedStorageHandle");
 	};
 
 	struct FMockEventLoopManagedStorageTraits
diff --git a/Engine/Source/Runtime/Online/Experimental/EventLoopTests/Tests/EventLoop/EventLoopTests.cpp b/Engine/Source/Runtime/Online/Experimental/EventLoopTests/Tests/EventLoop/EventLoopTests.cpp
--- a/Engine/Source/Runtime/Online/Experimental/EventLoopTests/Tests/EventLoop/EventLoopTests.cpp
+++ b/Engine/Source/Runtime/Online/Experimental/EventLoopTests/Tests/EventLoop/EventLoopTests.cpp
@@ -9,7 +9,7 @@ namespace UE::EventLoop
 {
 	struct FMockTimerHandleTraits
 	{
-		static constexpr TCHAR* Name = TEXT("MockTimerHandle");
+		static constexpr const TCHAR* Name = TEXT("MockTimerHandle");
 	};
 
 	struct FEventLoopMockTraits
diff --git a/Engine/Source/Runtime/Online/Experimental/EventLoopTests/Tests/EventLoop/EventLoopTimerTests.cpp b/Engine/Source/Runtime/Online/Experimental/EventLoopTests/Tests/EventLoop/EventLoopTimerTests.cpp
--- a/Engine/Source/Runtime/Online/Experimental/EventLoopTests/Tests/EventLoop/EventLoopTimerTests.cpp
+++ b/Engine/Source/Runtime/Online/Experimental/EventLoopTests/Tests/EventLoop/EventLoopTimerTests.cpp
@@ -8,7 +8,7 @@ namespace UE::EventLoop
 {
 	struct FMockTimerHandleTraits
 	{
-		static constexpr TCHAR* Name = TEXT("MockTimerHandle");
+		static constexpr const TCHAR* Name = TEXT("MockTimerHandle");
 	};
 
 	struct FEventLoopTimerManagerMockTraits : public FTimerManagerTraitsBase
