Simplify cancellation for Apple requests. Once started Task always goes through the Task delegate

diff --git a/Engine/Source/Runtime/Online/HTTP/Private/Apple/AppleHTTPNSUrlSession.cpp b/Engine/Source/Runtime/Online/HTTP/Private/Apple/AppleHTTPNSUrlSession.cpp
--- a/Engine/Source/Runtime/Online/HTTP/Private/Apple/AppleHTTPNSUrlSession.cpp
+++ b/Engine/Source/Runtime/Online/HTTP/Private/Apple/AppleHTTPNSUrlSession.cpp
@@ -229,7 +229,6 @@ FAppleHttpNSUrlSessionRequest::FAppleHttpNSUrlSessionRequest(NSURLSession* InSes
 ,   Task(nil)
 ,	bIsPayloadFile(false)
 ,	ContentBytesLength(0)
-,   bCanceled(false)
 ,	ElapsedTime(0.0f)
 ,	LastReportedBytesWritten(0)
 ,	LastReportedBytesRead(0)
@@ -466,7 +465,7 @@ bool FAppleHttpNSUrlSessionRequest::SetContentAsStreamedFile(const FString& File
 	}
 	else
 	{
-		UE_LOG(LogHttp, VeryVerbose, TEXT("FAppleHttpNSUrlSessionRequest::SetContentAsStreamedFile failed to get file size"));
+		UE_LOG(LogHttp, Warning, TEXT("FAppleHttpNSUrlSessionRequest::SetContentAsStreamedFile failed to get file size"));
 		Request.HTTPBodyStream = nil;
 		ContentBytesLength = 0;
 		bIsPayloadFile = false;
@@ -554,7 +553,6 @@ bool FAppleHttpNSUrlSessionRequest::StartRequest()
 	LastReportedBytesWritten = 0;
 	LastReportedBytesRead = 0;
 	ElapsedTime = 0.0f;
-	bCanceled = false;
 	Response = nullptr;
 
 	Task = [Session dataTaskWithRequest: Request];
@@ -594,7 +592,7 @@ void FAppleHttpNSUrlSessionRequest::FinishRequest()
 	CleanupRequest();
 
 	bool bSuccess = false;
-	if (!bCanceled && Response.IsValid() && Response->IsReady() && !Response->HadError())
+	if (Response.IsValid() && Response->IsReady() && !Response->HadError())
 	{
 		bSuccess = true;
 		UE_LOG(LogHttp, Verbose, TEXT("Request succeeded"));
@@ -644,28 +642,7 @@ void FAppleHttpNSUrlSessionRequest::CancelRequest()
 
 	if (Task != nil)
 	{
-		if (bCanceled)
-		{
-			return;
-		}
-		
-		bCanceled = true;
 		[Task cancel];
-		
-		if (Response != nullptr)
-		{
-			Response->CleanSharedObjects();
-		}
-
-		FHttpManager& HttpManager = FHttpModule::Get().GetHttpManager();
-		if (HttpManager.IsValidRequest(this))
-		{
-			HttpManager.CancelThreadedRequest(SharedThis(this));
-		}
-		else
-		{
-			FinishRequestNotInHttpManager();
-		}
 	}
 }
 
diff --git a/Engine/Source/Runtime/Online/HTTP/Private/Apple/AppleHTTPNSUrlSession.h b/Engine/Source/Runtime/Online/HTTP/Private/Apple/AppleHTTPNSUrlSession.h
--- a/Engine/Source/Runtime/Online/HTTP/Private/Apple/AppleHTTPNSUrlSession.h
+++ b/Engine/Source/Runtime/Online/HTTP/Private/Apple/AppleHTTPNSUrlSession.h
@@ -108,9 +108,6 @@ private:
 	/** The stream to receive response body */
 	TSharedPtr<FArchive> ResponseBodyReceiveStream;
 
-	/** flag to mark wether we tried canceled this request */
-	bool bCanceled;
-	
 	/** Time taken to complete/cancel the request. */
 	float ElapsedTime;
 
