Event based http manager for Apple platforms

diff --git a/Engine/Source/Runtime/Online/HTTP/Private/Apple/AppleEventLoopHttpThread.cpp b/Engine/Source/Runtime/Online/HTTP/Private/Apple/AppleEventLoopHttpThread.cpp
new file mode 100644
--- /dev/null
+++ b/Engine/Source/Runtime/Online/HTTP/Private/Apple/AppleEventLoopHttpThread.cpp
@@ -0,0 +1,51 @@
+// Copyright Epic Games, Inc. All Rights Reserved.
+
+#include "Apple/AppleEventLoopHttpThread.h"
+
+#include "Apple/AppleHTTPNSUrlSession.h" 
+
+bool FAppleEventLoopHttpThread::StartThreadedRequest(IHttpThreadedRequest* Request)
+{
+	FHttpResponsePtr Response = Request->GetResponse();
+	auto AppleResponse = StaticCastSharedPtr<FAppleHttpNSUrlSessionResponse>(Response);
+	AppleResponse->SetInternalTaskCompleteDelegate(FTaskCompleteDelegate::CreateLambda([IOAccess = EventLoop->GetIOAccess()]() mutable
+	{
+		IOAccess.Notify();
+	}));
+	return FEventLoopHttpThread::StartThreadedRequest(Request);
+}
+
+void FAppleEventLoopHttpThread::CompleteThreadedRequest(IHttpThreadedRequest* Request)
+{
+}
+
+void FAppleEventLoopHttpThread::CreateEventLoop()
+{	
+	UE::EventLoop::TEventLoop<FAppleHTTPIOManager>::FParams EventLoopParams;
+	EventLoopParams.IOManagerParams.ProcessRequests = [this]()
+	{ 
+		TArray<IHttpThreadedRequest*> RequestsToCancel;
+		TArray<IHttpThreadedRequest*> RequestsToComplete;
+		Process(RequestsToCancel, RequestsToComplete);
+	};
+	EventLoop.Emplace(MoveTemp(EventLoopParams));
+}
+
+void FAppleEventLoopHttpThread::DestroyEventLoop()
+{
+	EventLoop.Reset();
+}
+
+void FAppleEventLoopHttpThread::UpdateEventLoopConfigs()
+{
+}
+
+UE::EventLoop::IEventLoop* FAppleEventLoopHttpThread::GetEventLoop()
+{
+	return EventLoop.IsSet() ? &*EventLoop : nullptr;
+}
+
+UE::EventLoop::IEventLoop& FAppleEventLoopHttpThread::GetEventLoopChecked()
+{
+	return *EventLoop;
+}
diff --git a/Engine/Source/Runtime/Online/HTTP/Private/Apple/AppleEventLoopHttpThread.h b/Engine/Source/Runtime/Online/HTTP/Private/Apple/AppleEventLoopHttpThread.h
new file mode 100644
--- /dev/null
+++ b/Engine/Source/Runtime/Online/HTTP/Private/Apple/AppleEventLoopHttpThread.h
@@ -0,0 +1,102 @@
+// Copyright Epic Games, Inc. All Rights Reserved.
+
+#pragma once
+
+#include "CoreTypes.h"
+
+#include "Async/ManualResetEvent.h"
+#include "EventLoop/EventLoop.h"
+#include "EventLoop/IEventLoopIOManager.h"
+#include "Misc/Timespan.h"
+#include "Misc/MonotonicTime.h"
+
+#include "EventLoopHttpThread.h"
+
+class FAppleHTTPIOAccess final
+{
+public:
+	explicit FAppleHTTPIOAccess(TSharedRef<UE::FManualResetEvent> Event)
+		:SharedEvent(MoveTemp(Event))
+	{
+	}
+	
+	void Notify()
+	{
+		SharedEvent->Notify();
+	}
+	
+private:
+	TSharedRef<UE::FManualResetEvent> SharedEvent;
+};
+
+class FAppleHTTPIOManager final: public UE::EventLoop::IIOManager
+{
+public:
+	using FIOAccess = FAppleHTTPIOAccess;
+
+	struct FParams
+	{
+		TUniqueFunction<void()> ProcessRequests;
+	};
+
+	FAppleHTTPIOManager(UE::EventLoop::IEventLoop& EventLoop, FParams&& Params)
+		: Event(MakeShared<UE::FManualResetEvent>())
+		, IOAccess(Event)
+		, Params(MoveTemp(Params))
+	{
+	}
+
+	virtual ~FAppleHTTPIOManager() = default;
+
+	virtual bool Init() override
+	{
+		return true;
+	}
+
+	virtual void Shutdown() override
+	{
+	}
+
+	virtual void Notify() override
+	{
+		Event->Notify();
+	}
+
+	virtual void Poll(FTimespan WaitTime) override
+	{
+		if (Event->WaitFor(UE::FMonotonicTimeSpan::FromSeconds(WaitTime.GetTotalSeconds())))
+		{
+			Event->Reset();
+			Params.ProcessRequests();
+		}
+	}
+
+	FIOAccess& GetIOAccess()
+	{
+		return IOAccess;
+	}
+
+private:
+	TSharedRef<UE::FManualResetEvent> Event;
+	FIOAccess IOAccess;
+	FParams Params;
+};
+
+class FAppleEventLoopHttpThread
+	: public FEventLoopHttpThread
+{
+protected:
+	//~ Begin FHttpThread Interface
+	virtual bool StartThreadedRequest(IHttpThreadedRequest* Request) override;
+	virtual void CompleteThreadedRequest(IHttpThreadedRequest* Request) override;
+	//~ End FHttpThread Interface
+protected:
+	virtual void CreateEventLoop() override;
+	virtual void DestroyEventLoop() override;
+	virtual void UpdateEventLoopConfigs() override;
+
+	virtual UE::EventLoop::IEventLoop* GetEventLoop() override;
+	virtual UE::EventLoop::IEventLoop& GetEventLoopChecked() override;
+
+	TOptional<UE::EventLoop::TEventLoop<FAppleHTTPIOManager>> EventLoop;
+};
diff --git a/Engine/Source/Runtime/Online/HTTP/Private/Apple/AppleHTTPNSUrlSession.cpp b/Engine/Source/Runtime/Online/HTTP/Private/Apple/AppleHTTPNSUrlSession.cpp
--- a/Engine/Source/Runtime/Online/HTTP/Private/Apple/AppleHTTPNSUrlSession.cpp
+++ b/Engine/Source/Runtime/Online/HTTP/Private/Apple/AppleHTTPNSUrlSession.cpp
@@ -40,6 +40,9 @@ enum class EAppleHttpRequestResponseState: uint8
 
 	/** flag meant to reduce locking on ResponseStreamLock*/
 	@public BOOL bInitializedWithValidStream;
+
+	/** Delegate invoked after processing URLSession:task:didCompleteWithError:*/
+	@public FTaskCompleteDelegate TaskCompleteDelegate;
 }
 
 /** A handle for the response */
@@ -209,6 +212,7 @@ enum class EAppleHttpRequestResponseState: uint8
 			}
 		}
 	}
+	TaskCompleteDelegate.ExecuteIfBound();
 }
 
 - (void)URLSession:(NSURLSession *)session dataTask:(NSURLSessionDataTask *)dataTask willCacheResponse:(NSCachedURLResponse *)proposedResponse completionHandler:(void (^)(NSCachedURLResponse *cachedResponse))completionHandler
@@ -568,10 +572,11 @@ bool FAppleHttpNSUrlSessionRequest::StartRequest()
 		// Both Task and Response keep a strong reference to the delegate
 		Task.delegate = Response->ResponseDelegate;
 
+		//Setup delegates before starting the request
+		FHttpModule::Get().GetHttpManager().AddThreadedRequest(SharedThis(this));
+
 		[[Task retain] resume];
 		UE_LOG(LogHttp, Verbose, TEXT("[NSURLSessionTask resume]"));
-
-		FHttpModule::Get().GetHttpManager().AddThreadedRequest(SharedThis(this));
 	}
 	else
 	{
@@ -709,6 +714,11 @@ FAppleHttpNSUrlSessionResponse::~FAppleHttpNSUrlSessionResponse()
 	ResponseDelegate = nil;
 }
 
+void FAppleHttpNSUrlSessionResponse::SetInternalTaskCompleteDelegate(FTaskCompleteDelegate&& Delegate)
+{	
+	ResponseDelegate->TaskCompleteDelegate = MoveTemp(Delegate);
+}
+
 void FAppleHttpNSUrlSessionResponse::CleanSharedObjects()
 {
 	[ResponseDelegate ClearResponseStream];
diff --git a/Engine/Source/Runtime/Online/HTTP/Private/Apple/AppleHTTPNSUrlSession.h b/Engine/Source/Runtime/Online/HTTP/Private/Apple/AppleHTTPNSUrlSession.h
--- a/Engine/Source/Runtime/Online/HTTP/Private/Apple/AppleHTTPNSUrlSession.h
+++ b/Engine/Source/Runtime/Online/HTTP/Private/Apple/AppleHTTPNSUrlSession.h
@@ -7,6 +7,13 @@
 #include "Interfaces/IHttpResponse.h"
 #include "PlatformHttp.h"
 
+
+/**
+ * Delegate invoked when in progress NSUrlSessionTask completes. It is invoked in an out of our control thread
+ *
+ */
+DECLARE_DELEGATE(FTaskCompleteDelegate);
+
 /**
  * Apple implementation of an Http request
  */
@@ -187,6 +194,13 @@ public:
 	 * Cleans internal shared objects between request and response
 	 */
 	void CleanSharedObjects();
+
+	/**
+	 * Sets delegate invoked after processing URLSession:task:didCompleteWithError:
+	 * Should be set right before task is started 
+	*/
+	void SetInternalTaskCompleteDelegate(FTaskCompleteDelegate&& Delegate);
+
 	/**
 	 * Constructor
 	 *
diff --git a/Engine/Source/Runtime/Online/HTTP/Private/Apple/AppleHttpManager.cpp b/Engine/Source/Runtime/Online/HTTP/Private/Apple/AppleHttpManager.cpp
new file mode 100644
--- /dev/null
+++ b/Engine/Source/Runtime/Online/HTTP/Private/Apple/AppleHttpManager.cpp
@@ -0,0 +1,29 @@
+// Copyright Epic Games, Inc. All Rights Reserved.
+
+#include "Apple/AppleHttpManager.h"
+
+#include "Apple/AppleEventLoopHttpThread.h"
+#include "HAL/IConsoleManager.h"
+#include "Math/UnrealMathUtility.h"
+#include "HttpModule.h"
+#include "Http.h"
+
+TAutoConsoleVariable<int32> CVarAppleEventLoopEnableChance(
+	TEXT("http.AppleEventLoopEnableChance"),
+	0,
+	TEXT("Enable chance of Apple event loop, from 0 to 100"),
+	ECVF_SaveForNextBoot
+);
+
+FHttpThreadBase* FAppleHttpManager::CreateHttpThread()
+{
+	bool bUseEventLoop = !(FMath::RandRange(0, 99) < CVarAppleEventLoopEnableChance.GetValueOnGameThread());
+	if (bUseEventLoop)
+	{
+		UE_LOG(LogHttp, Log, TEXT("CreateHttpThread using FAppleEventLoopHttpThread"));
+		return new FAppleEventLoopHttpThread();
+	}
+
+	UE_LOG(LogHttp, Log, TEXT("CreateHttpThread using FLegacyHttpThread"));
+	return new FLegacyHttpThread();
+}
diff --git a/Engine/Source/Runtime/Online/HTTP/Private/Apple/AppleHttpManager.h b/Engine/Source/Runtime/Online/HTTP/Private/Apple/AppleHttpManager.h
new file mode 100644
--- /dev/null
+++ b/Engine/Source/Runtime/Online/HTTP/Private/Apple/AppleHttpManager.h
@@ -0,0 +1,15 @@
+// Copyright Epic Games, Inc. All Rights Reserved.
+
+#pragma once
+
+#include "HttpManager.h"
+
+class FHttpThread;
+
+class FAppleHttpManager : public FHttpManager
+{
+	//~ Begin HttpManager Interface
+protected:
+	virtual FHttpThreadBase* CreateHttpThread() override;
+	//~ End HttpManager Interface
+};
diff --git a/Engine/Source/Runtime/Online/HTTP/Private/Apple/ApplePlatformHttp.cpp b/Engine/Source/Runtime/Online/HTTP/Private/Apple/ApplePlatformHttp.cpp
--- a/Engine/Source/Runtime/Online/HTTP/Private/Apple/ApplePlatformHttp.cpp
+++ b/Engine/Source/Runtime/Online/HTTP/Private/Apple/ApplePlatformHttp.cpp
@@ -2,6 +2,7 @@
 
 #include "Apple/ApplePlatformHttp.h"
 #include "Http.h"
+#include "AppleHttpManager.h"
 #include "AppleHTTPNSUrlConnection.h"
 #include "AppleHTTPNSUrlSession.h"
 #include "Apple/CFRef.h"
@@ -273,6 +274,19 @@ void FApplePlatformHttp::ShutdownWithNSUrlSession()
 	Session = nil;
 }
 
+FHttpManager* FApplePlatformHttp::CreatePlatformHttpManager()
+{
+	if(bUseNSUrlSession)
+	{
+		return new FAppleHttpManager();
+	}
+	else
+	{
+		// Event based http manager does not support FAppleHTTPNSURLConnection 
+		return nullptr;
+	}
+}
+
 IHttpRequest* FApplePlatformHttp::ConstructRequest()
 {
 	if(bUseNSUrlSession)
diff --git a/Engine/Source/Runtime/Online/HTTP/Public/Apple/ApplePlatformHttp.h b/Engine/Source/Runtime/Online/HTTP/Public/Apple/ApplePlatformHttp.h
--- a/Engine/Source/Runtime/Online/HTTP/Public/Apple/ApplePlatformHttp.h
+++ b/Engine/Source/Runtime/Online/HTTP/Public/Apple/ApplePlatformHttp.h
@@ -20,10 +20,7 @@ public:
 	 *
 	 * @return NULL if default implementation is to be used
 	 */
-	static FHttpManager* CreatePlatformHttpManager()
-	{
-		return nullptr;
-	}
+	static FHttpManager* CreatePlatformHttpManager();
 
 	/**
 	 * Platform shutdown step
