- Remove ExecuteIfBound from the lock scope

diff --git a/Engine/Source/Runtime/Online/HTTP/Private/HttpManager.cpp b/Engine/Source/Runtime/Online/HTTP/Private/HttpManager.cpp
--- a/Engine/Source/Runtime/Online/HTTP/Private/HttpManager.cpp
+++ b/Engine/Source/Runtime/Online/HTTP/Private/HttpManager.cpp
@@ -466,9 +466,11 @@ void FHttpManager::FlushTick(float DeltaSeconds)
 
 void FHttpManager::AddRequest(const FHttpRequestRef& Request)
 {
-	FScopeLock ScopeLock(&RequestLock);
-	check(!bFlushing);
-	Requests.Add(Request);
+	{
+		FScopeLock ScopeLock(&RequestLock);
+		check(!bFlushing);
+		Requests.Add(Request);
+	}
 	RequestAddedDelegate.ExecuteIfBound(Request);
 }
 
