Implemented FAppleHttpNSUrlSessionRequest::SetContentFromStream
Extra: Invoked more convenient method on shutdown for NSUrlSession

diff --git a/Engine/Source/Runtime/Online/HTTP/Private/Apple/AppleHTTPNSUrlSession.cpp b/Engine/Source/Runtime/Online/HTTP/Private/Apple/AppleHTTPNSUrlSession.cpp
--- a/Engine/Source/Runtime/Online/HTTP/Private/Apple/AppleHTTPNSUrlSession.cpp
+++ b/Engine/Source/Runtime/Online/HTTP/Private/Apple/AppleHTTPNSUrlSession.cpp
@@ -74,7 +74,7 @@ enum class EAppleHttpRequestResponseState: uint8
 @synthesize BytesWritten;
 @synthesize BytesReceived;
 
--(FAppleHttpNSUrlSessionResponseDelegate*) initWithResponseStream:(TSharedPtr<FArchive>)ResponseStream
+- (FAppleHttpNSUrlSessionResponseDelegate*)initWithResponseStream:(TSharedPtr<FArchive>)ResponseStream
 {
 	self = [super init];
 	
@@ -87,7 +87,7 @@ enum class EAppleHttpRequestResponseState: uint8
 	return self;
 }
 
-- (void) ClearResponseStream
+- (void)ClearResponseStream
 {
 	if (bInitializedWithValidStream)
 	{
@@ -96,7 +96,7 @@ enum class EAppleHttpRequestResponseState: uint8
 	}
 }
 
-- (void) dealloc
+- (void)dealloc
 {
 	[Response release];
 	[super dealloc];
@@ -224,6 +224,145 @@ enum class EAppleHttpRequestResponseState: uint8
 }
 @end
 
+
+/**
+ * NSInputStream subclass to send streamed FArchive contents
+ */
+@interface FNSInputStreamFromArchive : NSInputStream<NSStreamDelegate>
+{
+	TSharedPtr<FArchive> Archive;
+	int64 AlreadySentContent;
+	NSStreamStatus StreamStatus;
+	id<NSStreamDelegate> Delegate;
+}
+@end
+
+@implementation FNSInputStreamFromArchive
+
++(FNSInputStreamFromArchive*)initWithArchive:(TSharedRef<FArchive>) Archive
+{
+	FNSInputStreamFromArchive* Ret = [[[FNSInputStreamFromArchive alloc] init] autorelease];
+	Ret->Archive = MoveTemp(Archive);
+	return Ret;
+}
+
+- (id)init
+{
+	self = [super init];
+	if (self)
+	{
+		StreamStatus = NSStreamStatusNotOpen;
+
+		// Docs say it is good practice that streams are it's own delegates by default
+		Delegate = self;
+	}
+	
+	return self;
+}
+
+/** NSStream implementation */
+- (void)dealloc
+{
+	[super dealloc];
+}
+
+- (void)open
+{
+	AlreadySentContent = 0;
+	StreamStatus = NSStreamStatusOpen;
+}
+
+- (void)close
+{
+	StreamStatus = NSStreamStatusClosed;
+}
+
+- (NSStreamStatus)streamStatus
+{
+	return StreamStatus;
+}
+
+- (NSError *)streamError
+{
+	return nil;
+}
+
+- (id<NSStreamDelegate>)delegate
+{
+	return Delegate;
+}
+
+- (void)setDelegate:(id<NSStreamDelegate>)InDelegate
+{
+	if (InDelegate == nil)
+	{
+		InDelegate = self;
+	}
+	else
+	{
+		Delegate = InDelegate;
+	}
+}
+
+- (id)propertyForKey:(NSString *)key
+{
+	return nil;
+}
+
+- (BOOL)setProperty:(id)property forKey:(NSString *)key
+{
+	return NO;
+}
+
+- (void)scheduleInRunLoop:(NSRunLoop *)aRunLoop forMode:(NSString *)mode
+{
+	// There is no need to scheduled anything. Data is always available until end is reached
+}
+
+- (void)removeFromRunLoop:(NSRunLoop *)aRunLoop forMode:(NSString *)mode
+{
+	// There is no need to be descheduled since we didn't schedule
+}
+
+/** NSStreamDelegate implementation */
+- (void)stream:(NSStream *)stream handleEvent:(NSStreamEvent)eventCode
+{
+	// Won't update local data
+}
+
+/** NSInputStream implementation. Those methods are invoked in a worker thread out of our control */
+
+// Reads up to 'len' bytes into 'buffer'. Returns the actual number of bytes read.
+- (NSInteger)read:(uint8_t *)buffer maxLength:(NSUInteger)len
+{
+	const int64 ContentLength = Archive->TotalSize();
+	check(AlreadySentContent <= ContentLength);
+	const int64 SizeToSend = ContentLength - AlreadySentContent;
+	const int64 SizeToSendThisTime = FMath::Min(SizeToSend, static_cast<int64>(len));
+	if (SizeToSendThisTime != 0)
+	{
+		if (Archive->Tell() != AlreadySentContent)
+		{
+			Archive->Seek(AlreadySentContent);
+		}
+		Archive->Serialize((uint8*)buffer, SizeToSendThisTime);
+		AlreadySentContent += SizeToSendThisTime;
+	}
+	return SizeToSendThisTime;
+}
+
+// return NO because getting the internal buffer is not appropriate for this subclass
+- (BOOL)getBuffer:(uint8_t **)buffer length:(NSUInteger *)len
+{
+	return NO;
+}
+
+// returns YES to always force reads
+- (BOOL)hasBytesAvailable
+{
+	return YES;
+}
+@end
 /****************************************************************************
  * FAppleHttpNSUrlSessionRequest implementation
  ***************************************************************************/
@@ -387,6 +526,7 @@ void FAppleHttpNSUrlSessionRequest::SetContent(const TArray<uint8>& ContentPaylo
 	}
 
 	UE_LOG(LogHttp, Verbose, TEXT("FAppleHttpNSUrlSessionRequest::SetContent()"));
+	Request.HTTPBodyStream = nil;
 	Request.HTTPBody = [NSData dataWithBytes:ContentPayload.GetData() length:ContentPayload.Num()];
 	ContentBytesLength = ContentPayload.Num();
 	bIsPayloadFile = false;
@@ -402,6 +542,7 @@ void FAppleHttpNSUrlSessionRequest::SetContent(TArray<uint8>&& ContentPayload)
 
 	UE_LOG(LogHttp, Verbose, TEXT("FAppleHttpNSUrlSessionRequest::SetContent()"));
 
+	Request.HTTPBodyStream = nil;
 	// We cannot use NSData dataWithBytesNoCopy:length:freeWhenDone: and keep the data in this instance because we don't have control
 	// over the lifetime of the request copy that NSURLSessionTask keeps
 	Request.HTTPBody = [NSData dataWithBytes:ContentPayload.GetData() length:ContentPayload.Num()];
@@ -436,6 +577,7 @@ void FAppleHttpNSUrlSessionRequest::SetContentAsString(const FString& ContentStr
 	UE_LOG(LogHttp, Verbose, TEXT("FAppleHttpNSUrlSessionRequest::SetContentAsString() - %s"), *ContentString);
 	FTCHARToUTF8 Converter(*ContentString);
 
+	Request.HTTPBodyStream = nil;
 	// The extra length computation here is unfortunate, but it's technically not safe to assume the length is the same.
 	Request.HTTPBody = [NSData dataWithBytes:(ANSICHAR*)Converter.Get() length:Converter.Length()];
 	ContentBytesLength = Converter.Length();
@@ -480,8 +622,20 @@ bool FAppleHttpNSUrlSessionRequest::SetContentAsStreamedFile(const FString& File
 
 bool FAppleHttpNSUrlSessionRequest::SetContentFromStream(TSharedRef<FArchive, ESPMode::ThreadSafe> Stream)
 {
-	UE_LOG(LogHttp, Warning, TEXT("FAppleHttpNSUrlSessionRequest::SetContentFromStream is not implemented"));
-	return false;
+	SCOPED_AUTORELEASE_POOL;
+	if (CompletionStatus == EHttpRequestStatus::Processing)
+	{
+		UE_LOG(LogHttp, Warning, TEXT("FAppleHttpNSUrlSessionRequest::SetContentFromStream() - attempted to set content on a request that is inflight"));
+		return false;
+	}
+
+	Request.HTTPBody = nil;
+
+	Request.HTTPBodyStream = [FNSInputStreamFromArchive initWithArchive: Stream];
+	ContentBytesLength = Stream->TotalSize();
+	bIsPayloadFile = true;
+
+	return true;
 }
 
 bool FAppleHttpNSUrlSessionRequest::SetResponseBodyReceiveStream(TSharedRef<FArchive> Stream)
diff --git a/Engine/Source/Runtime/Online/HTTP/Private/Apple/ApplePlatformHttp.cpp b/Engine/Source/Runtime/Online/HTTP/Private/Apple/ApplePlatformHttp.cpp
--- a/Engine/Source/Runtime/Online/HTTP/Private/Apple/ApplePlatformHttp.cpp
+++ b/Engine/Source/Runtime/Online/HTTP/Private/Apple/ApplePlatformHttp.cpp
@@ -269,7 +269,7 @@ void FApplePlatformHttp::InitWithNSUrlSession()
 
 void FApplePlatformHttp::ShutdownWithNSUrlSession()
 {
-	[Session finishTasksAndInvalidate];
+	[Session invalidateAndCancel];
 	[Session release];
 	Session = nil;
 }
