Do not use CurlHttp on Windows platforms that have opted in to XCurl - they should always use WinHTTP.

#rnx

diff --git a/Engine/Source/Runtime/Online/HTTP/HTTP.Build.cs b/Engine/Source/Runtime/Online/HTTP/HTTP.Build.cs
--- a/Engine/Source/Runtime/Online/HTTP/HTTP.Build.cs
+++ b/Engine/Source/Runtime/Online/HTTP/HTTP.Build.cs
@@ -21,13 +21,7 @@ public class HTTP : ModuleRules
 				Target.IsInPlatformGroup(UnrealPlatformGroup.Android);
 		}
 	}
-	protected virtual bool bPlatformSupportsXCurl
-	{ 
-		get 
-		{ 
-			return (Target.Platform.IsInGroup(UnrealPlatformGroup.Windows) && Target.WindowsPlatform.bUseXCurl); 
-		} 
-	}
+	protected virtual bool bPlatformSupportsXCurl { get { return false; } }
 
 	protected virtual bool bPlatformSupportsCurlMultiPoll { get { return true; } }
 
