Adding per-match metrics to the cosmetic streaming analytics report.

diff --git a/Engine/Source/Runtime/Online/BuildPatchServices/Private/Installer/BuildInstallStreamer.cpp b/Engine/Source/Runtime/Online/BuildPatchServices/Private/Installer/BuildInstallStreamer.cpp
--- a/Engine/Source/Runtime/Online/BuildPatchServices/Private/Installer/BuildInstallStreamer.cpp
+++ b/Engine/Source/Runtime/Online/BuildPatchServices/Private/Installer/BuildInstallStreamer.cpp
@@ -70,7 +70,8 @@ namespace BuildPatchServices
 		, bIsShuttingDown(false)
 		, RequestTrigger(FPlatformProcess::GetSynchEventFromPool(false))
 		, CloudTrigger(FPlatformProcess::GetSynchEventFromPool(false))
-		, BuildStats(FBuildInstallStreamerStats())
+		, StreamerStats(FBuildInstallStreamerStats())
+		, PerSessionStreamerStats(FBuildInstallStreamerStats())
 		
 	{
 		RequestWorker = Async(EAsyncExecution::ThreadIfForkSafe, [this]()
@@ -134,7 +135,17 @@ namespace BuildPatchServices
 
 	const FBuildInstallStreamerStats& FBuildInstallStreamer::GetInstallStreamerStatistics() const
 	{
-		return BuildStats;
+		return StreamerStats;
+	}
+	
+	const FBuildInstallStreamerStats& FBuildInstallStreamer::GetInstallStreamerSessionStatistics() const
+	{
+		return PerSessionStreamerStats;
+	}
+	
+	void FBuildInstallStreamer::ResetSessionStatistics()
+	{
+		PerSessionStreamerStats = FBuildInstallStreamerStats();
 	}
 
 	void FBuildInstallStreamer::Initialise()
@@ -158,7 +169,8 @@ namespace BuildPatchServices
 		auto CompleteRequest = [this, &Requests, &RequestTime]()
 		{
 			uint64 TotalDownload = DownloadServiceStatistics->GetBytesDownloaded();
-			uint64 PreviousTotalRequestsCompleted = BuildStats.BundleRequestsCompleted + BuildStats.FileRequestsCompleted;
+			uint64 PreviousTotalRequestsCompleted = StreamerStats.BundleRequestsCompleted + StreamerStats.FileRequestsCompleted;
+			uint64 PreviousTotalRequestsCompletedForSession = PerSessionStreamerStats.BundleRequestsCompleted + PerSessionStreamerStats.FileRequestsCompleted;
 			for (FStreamRequest& Request : Requests)
 			{
 				const bool bIsTagRequest = Request.Get<1>();
@@ -169,21 +181,34 @@ namespace BuildPatchServices
 				// TODO: Not accurate.. How to fix this
 				Result.TotalDownloaded = TotalDownload;
 				
-				// if (bIsTagRequest) BuildStats.BundleMegaBytesDownloaded += (Result.TotalDownloaded / 1024.0f / 1024.0f);
-				// else BuildStats.FileMegaBytesDownloaded += (Result.TotalDownloaded / 1024.0f / 1024.0f);
+				// if (bIsTagRequest) StreamerStats.BundleMegaBytesDownloaded += (Result.TotalDownloaded / 1024.0f / 1024.0f);
+				// else StreamerStats.FileMegaBytesDownloaded += (Result.TotalDownloaded / 1024.0f / 1024.0f);
 				
-				if (bIsTagRequest) BuildStats.BundleRequestsCompleted++;
-				else BuildStats.FileRequestsCompleted++;
+				if (bIsTagRequest)
+				{
+					StreamerStats.BundleRequestsCompleted++;
+					PerSessionStreamerStats.BundleRequestsCompleted++;
+				}
+				else
+				{
+					StreamerStats.FileRequestsCompleted++;
+					PerSessionStreamerStats.FileRequestsCompleted++;
+				}
 
 				UE_LOG(LogBuildInstallStreamer, Verbose, TEXT("Completed request %s in %s"), *FString::Join(Result.Request, TEXT(",")), *FPlatformTime::PrettyTime(RequestTime));
 				Request.Get<2>().ExecuteIfBound(MoveTemp(Result));
 			}
-			BuildStats.TotalMegaBytesDownloaded += (TotalDownload / 1024.0f / 1024.0f);
+			StreamerStats.TotalMegaBytesDownloaded += (TotalDownload / 1024.0f / 1024.0f);
+			PerSessionStreamerStats.TotalMegaBytesDownloaded += (TotalDownload / 1024.0f / 1024.0f);
+				
+			if (RequestTime > StreamerStats.MaxRequestTime) StreamerStats.MaxRequestTime = RequestTime;
+			if (RequestTime > PerSessionStreamerStats.MaxRequestTime) PerSessionStreamerStats.MaxRequestTime = RequestTime;
 				
-			if (RequestTime > BuildStats.MaxRequestTime) BuildStats.MaxRequestTime = RequestTime;
+			StreamerStats.AverageRequestTime = ((StreamerStats.AverageRequestTime * PreviousTotalRequestsCompleted) + RequestTime) /
+				(StreamerStats.BundleRequestsCompleted + StreamerStats.FileRequestsCompleted);
 				
-			BuildStats.AverageRequestTime = ((BuildStats.AverageRequestTime * PreviousTotalRequestsCompleted) + RequestTime) /
-				(BuildStats.BundleRequestsCompleted + BuildStats.FileRequestsCompleted);
+			PerSessionStreamerStats.AverageRequestTime = ((PerSessionStreamerStats.AverageRequestTime * PreviousTotalRequestsCompletedForSession) + RequestTime) /
+				(PerSessionStreamerStats.BundleRequestsCompleted + PerSessionStreamerStats.FileRequestsCompleted);
 		};
 
 		auto CancelRequest = [this, &Requests]()
@@ -196,8 +221,16 @@ namespace BuildPatchServices
 				Result.ErrorType = InstallerError->GetErrorType();
 				Result.ErrorCode = InstallerError->GetErrorCode();
 				Result.TotalDownloaded = 0;
-				if (bIsTagRequest) BuildStats.BundleRequestsCancelled++;
-				else BuildStats.FileRequestsCancelled++;
+				if (bIsTagRequest)
+				{
+					StreamerStats.BundleRequestsCancelled++;
+					PerSessionStreamerStats.BundleRequestsCancelled++;
+				}
+				else
+				{
+					StreamerStats.FileRequestsCancelled++;
+					PerSessionStreamerStats.FileRequestsCancelled++;
+				}
 				UE_LOG(LogBuildInstallStreamer, Verbose, TEXT("Cancelled request %s"), *FString::Join(Result.Request, TEXT(",")));
 				Request.Get<2>().ExecuteIfBound(MoveTemp(Result));
 			}
@@ -233,8 +266,16 @@ namespace BuildPatchServices
 			{
 				for (const FStreamRequest& Request : Requests)
 				{
-					if (Request.Get<1>()) BuildStats.BundleRequestsMade++;
-					else BuildStats.FileRequestsMade++;
+					if (Request.Get<1>())
+					{
+						StreamerStats.BundleRequestsMade++;
+						PerSessionStreamerStats.BundleRequestsMade++;
+					}
+					else
+					{
+						StreamerStats.FileRequestsMade++;
+						PerSessionStreamerStats.FileRequestsMade++;
+					}
 					UE_LOG(LogBuildInstallStreamer, Verbose, TEXT("Begin request %s"), *FString::Join(Request.Get<0>(), TEXT(",")));
 				}
 
@@ -244,10 +285,12 @@ namespace BuildPatchServices
 				TSharedPtr<IVirtualFileCache, ESPMode::ThreadSafe> VirtualFileCache = IVirtualFileCache::CreateVirtualFileCache();
 				InstallerError->Reset();
 
-				//Used and Total size is calulcated on VFC Start.
+				//Used and Total size is calculated on VFC Start.
 				//Used + Requested write will always result in new total. Doing this to recalculating all blocks.
-				BuildStats.VFCCachedUsedSize = (VirtualFileCache->GetUsedSize() / 1024.0f / 1024.0f);
-				BuildStats.VFCCachedTotalSize = (VirtualFileCache->GetTotalSize()  / 1024.0f / 1024.0f);
+				StreamerStats.VFCCachedUsedSize = (VirtualFileCache->GetUsedSize() / 1024.0f / 1024.0f);
+				StreamerStats.VFCCachedTotalSize = (VirtualFileCache->GetTotalSize()  / 1024.0f / 1024.0f);
+				PerSessionStreamerStats.VFCCachedUsedSize = StreamerStats.VFCCachedUsedSize;
+				PerSessionStreamerStats.VFCCachedTotalSize = StreamerStats.VFCCachedTotalSize;
 
 				// Setup file build list.
 				FVirtualFileConstructorConfiguration VFCConfig;
@@ -274,7 +317,8 @@ namespace BuildPatchServices
 					}
 					else if (FileManifest != nullptr)
 					{
-						BuildStats.VFCRequestedFileWrite += (FileManifest->FileSize / 1024.0f / 1024.0f);
+						StreamerStats.VFCRequestedFileWrite += (FileManifest->FileSize / 1024.0f / 1024.0f);
+						PerSessionStreamerStats.VFCRequestedFileWrite += (FileManifest->FileSize / 1024.0f / 1024.0f);
 					}
 				}
 
@@ -391,24 +435,24 @@ namespace BuildPatchServices
 			TickFunc();
 		}
 
-		CSV_CUSTOM_STAT(CosmeticStreamingCsv, BuildStats.VFCRequestedFileWrite,		BuildStats.VFCRequestedFileWrite, ECsvCustomStatOp::Set);
-		//CSV_CUSTOM_STAT(CosmeticStreamingCsv, BuildStats.VFCActualFileWrite,		BuildStats.VFCActualFileWrite, ECsvCustomStatOp::Set);
-		CSV_CUSTOM_STAT(CosmeticStreamingCsv, BuildStats.VFCCachedUsedSize,			BuildStats.VFCCachedUsedSize, ECsvCustomStatOp::Set);
-		CSV_CUSTOM_STAT(CosmeticStreamingCsv, BuildStats.VFCCachedTotalSize,		BuildStats.VFCCachedTotalSize, ECsvCustomStatOp::Set);
-		CSV_CUSTOM_STAT(CosmeticStreamingCsv, BuildStats.TotalMegaBytesDownloaded,	BuildStats.TotalMegaBytesDownloaded, ECsvCustomStatOp::Set);
+		CSV_CUSTOM_STAT(CosmeticStreamingCsv, StreamerStats.VFCRequestedFileWrite,		StreamerStats.VFCRequestedFileWrite, ECsvCustomStatOp::Set);
+		//CSV_CUSTOM_STAT(CosmeticStreamingCsv, StreamerStats.VFCActualFileWrite,		StreamerStats.VFCActualFileWrite, ECsvCustomStatOp::Set);
+		CSV_CUSTOM_STAT(CosmeticStreamingCsv, StreamerStats.VFCCachedUsedSize,			StreamerStats.VFCCachedUsedSize, ECsvCustomStatOp::Set);
+		CSV_CUSTOM_STAT(CosmeticStreamingCsv, StreamerStats.VFCCachedTotalSize,		StreamerStats.VFCCachedTotalSize, ECsvCustomStatOp::Set);
+		CSV_CUSTOM_STAT(CosmeticStreamingCsv, StreamerStats.TotalMegaBytesDownloaded,	StreamerStats.TotalMegaBytesDownloaded, ECsvCustomStatOp::Set);
 		
-		//CSV_CUSTOM_STAT(CosmeticStreamingCsv, BuildStats.FileMegaBytesDownloaded,	BuildStats.FileMegaBytesDownloaded, ECsvCustomStatOp::Set);
-		CSV_CUSTOM_STAT(CosmeticStreamingCsv, BuildStats.FileRequestsCompleted,	    (double)BuildStats.FileRequestsCompleted, ECsvCustomStatOp::Set);
-		CSV_CUSTOM_STAT(CosmeticStreamingCsv, BuildStats.FileRequestsCancelled,	    (double)BuildStats.FileRequestsCancelled, ECsvCustomStatOp::Set);
-		CSV_CUSTOM_STAT(CosmeticStreamingCsv, BuildStats.FileRequestsMade,		    (double)BuildStats.FileRequestsMade, ECsvCustomStatOp::Set);
+		//CSV_CUSTOM_STAT(CosmeticStreamingCsv, StreamerStats.FileMegaBytesDownloaded,	StreamerStats.FileMegaBytesDownloaded, ECsvCustomStatOp::Set);
+		CSV_CUSTOM_STAT(CosmeticStreamingCsv, StreamerStats.FileRequestsCompleted,	    (double)StreamerStats.FileRequestsCompleted, ECsvCustomStatOp::Set);
+		CSV_CUSTOM_STAT(CosmeticStreamingCsv, StreamerStats.FileRequestsCancelled,	    (double)StreamerStats.FileRequestsCancelled, ECsvCustomStatOp::Set);
+		CSV_CUSTOM_STAT(CosmeticStreamingCsv, StreamerStats.FileRequestsMade,		    (double)StreamerStats.FileRequestsMade, ECsvCustomStatOp::Set);
 		
-		//CSV_CUSTOM_STAT(CosmeticStreamingCsv, BuildStats.BundleMegaBytesDownloaded,	BuildStats.BundleMegaBytesDownloaded, ECsvCustomStatOp::Set);
-		CSV_CUSTOM_STAT(CosmeticStreamingCsv, BuildStats.BundleRequestsCompleted,	(double)BuildStats.BundleRequestsCompleted, ECsvCustomStatOp::Set);
-		CSV_CUSTOM_STAT(CosmeticStreamingCsv, BuildStats.BundleRequestsCancelled,	(double)BuildStats.BundleRequestsCancelled, ECsvCustomStatOp::Set);
-		CSV_CUSTOM_STAT(CosmeticStreamingCsv, BuildStats.BundleRequestsMade,		(double)BuildStats.BundleRequestsMade, ECsvCustomStatOp::Set);
+		//CSV_CUSTOM_STAT(CosmeticStreamingCsv, StreamerStats.BundleMegaBytesDownloaded,	StreamerStats.BundleMegaBytesDownloaded, ECsvCustomStatOp::Set);
+		CSV_CUSTOM_STAT(CosmeticStreamingCsv, StreamerStats.BundleRequestsCompleted,	(double)StreamerStats.BundleRequestsCompleted, ECsvCustomStatOp::Set);
+		CSV_CUSTOM_STAT(CosmeticStreamingCsv, StreamerStats.BundleRequestsCancelled,	(double)StreamerStats.BundleRequestsCancelled, ECsvCustomStatOp::Set);
+		CSV_CUSTOM_STAT(CosmeticStreamingCsv, StreamerStats.BundleRequestsMade,		(double)StreamerStats.BundleRequestsMade, ECsvCustomStatOp::Set);
 		
-		CSV_CUSTOM_STAT(CosmeticStreamingCsv, BuildStats.AverageRequestTime,	BuildStats.AverageRequestTime, ECsvCustomStatOp::Set);
-		CSV_CUSTOM_STAT(CosmeticStreamingCsv, BuildStats.MaxRequestTime,		BuildStats.MaxRequestTime, ECsvCustomStatOp::Set);
+		CSV_CUSTOM_STAT(CosmeticStreamingCsv, StreamerStats.AverageRequestTime,	StreamerStats.AverageRequestTime, ECsvCustomStatOp::Set);
+		CSV_CUSTOM_STAT(CosmeticStreamingCsv, StreamerStats.MaxRequestTime,		StreamerStats.MaxRequestTime, ECsvCustomStatOp::Set);
 
 		return bKeepTicking;
 	}
diff --git a/Engine/Source/Runtime/Online/BuildPatchServices/Private/Installer/BuildInstallStreamer.h b/Engine/Source/Runtime/Online/BuildPatchServices/Private/Installer/BuildInstallStreamer.h
--- a/Engine/Source/Runtime/Online/BuildPatchServices/Private/Installer/BuildInstallStreamer.h
+++ b/Engine/Source/Runtime/Online/BuildPatchServices/Private/Installer/BuildInstallStreamer.h
@@ -50,6 +50,8 @@ namespace BuildPatchServices
 		virtual void UnregisterMessageHandler(FMessageHandler* MessageHandler) override;
 		virtual const FBuildInstallStreamerConfiguration& GetConfiguration() const override;
 		virtual const FBuildInstallStreamerStats& GetInstallStreamerStatistics() const override;
+		virtual const FBuildInstallStreamerStats& GetInstallStreamerSessionStatistics() const override;
+		virtual void ResetSessionStatistics() override;
 		// IBuildInstallStreamer interface end.
 
 		bool Tick();
@@ -96,7 +98,8 @@ namespace BuildPatchServices
 		TQueue<TFunction<void()>, EQueueMode::Spsc> TickQueue;
 		FEvent* RequestTrigger;
 		FEvent* CloudTrigger;
-		FBuildInstallStreamerStats BuildStats;
+		FBuildInstallStreamerStats StreamerStats;
+		FBuildInstallStreamerStats PerSessionStreamerStats;
 	};
 
 	/**
diff --git a/Engine/Source/Runtime/Online/BuildPatchServices/Public/Interfaces/IBuildInstallStreamer.h b/Engine/Source/Runtime/Online/BuildPatchServices/Public/Interfaces/IBuildInstallStreamer.h
--- a/Engine/Source/Runtime/Online/BuildPatchServices/Public/Interfaces/IBuildInstallStreamer.h
+++ b/Engine/Source/Runtime/Online/BuildPatchServices/Public/Interfaces/IBuildInstallStreamer.h
@@ -64,4 +64,16 @@ public:
 	 * @returns a const reference to the statistics.
 	 */
 	virtual const BuildPatchServices::FBuildInstallStreamerStats& GetInstallStreamerStatistics() const = 0;
+	
+	/**
+	 * Get the per-session stats for this install streamer.
+	 * 
+	 * @returns a const reference to the statistics.
+	 */
+	virtual const BuildPatchServices::FBuildInstallStreamerStats& GetInstallStreamerSessionStatistics() const = 0;
+	
+	/**
+	 * Resets tracking on the per-session stats for this install streamer.
+	 */
+	virtual void ResetSessionStatistics() = 0;
 };
\ No newline at end of file
