Fix leak when a one-shot timer expires.

diff --git a/Engine/Source/Runtime/Online/Experimental/EventLoop/Public/EventLoop/EventLoopManagedStorage.h b/Engine/Source/Runtime/Online/Experimental/EventLoop/Public/EventLoop/EventLoopManagedStorage.h
--- a/Engine/Source/Runtime/Online/Experimental/EventLoop/Public/EventLoop/EventLoopManagedStorage.h
+++ b/Engine/Source/Runtime/Online/Experimental/EventLoop/Public/EventLoop/EventLoopManagedStorage.h
@@ -322,6 +322,18 @@ public:
 		RemoveRequests.Enqueue({Handle, MoveTemp(OnRemoveComplete)});
 	}
 
+	/**
+	 * Removes a resource directly using its internal handle.
+	 * 
+	 * NOT thread safe.
+	 *
+	 * @param Handle The handle for the resource to be removed.
+	 */
+	void Remove(const FInternalHandle Handle)
+	{
+		RemoveImpl(Handle.GetExternalHandle(), nullptr);
+	}
+
 public:
 	/*
 	 * STORAGE ACCESS
diff --git a/Engine/Source/Runtime/Online/Experimental/EventLoop/Public/EventLoop/EventLoopTimer.h b/Engine/Source/Runtime/Online/Experimental/EventLoop/Public/EventLoop/EventLoopTimer.h
--- a/Engine/Source/Runtime/Online/Experimental/EventLoop/Public/EventLoop/EventLoopTimer.h
+++ b/Engine/Source/Runtime/Online/Experimental/EventLoop/Public/EventLoop/EventLoopTimer.h
@@ -178,6 +178,18 @@ public:
 		return false;
 	}
 
+	/**
+	 * Returns the total number of timers. Includes active timers and timers waiting to be rescheduled.
+	 *
+	 * NOT thread safe.
+	 * 
+	 * @return the total number of timers.
+	 */
+	FORCEINLINE uint32 GetNumTimers() const
+	{
+		return Storage.Num();
+	}
+
 	/**
 	 * Set a new timer. The timer callback will be triggered from within the call to Tick,
 	 * which may occur on a different thread from the one which set the timer.
@@ -399,6 +411,11 @@ void TTimerManager<Traits>::Tick(FTimespan DeltaTime)
 					TimerData->Status = ETimerStatus::PendingReschedule;
 					PendingRepeatTimers.Add(TimerHandle);
 				}
+				else
+				{
+					// Timer has been fired and is non-repeating. Remove the timer.
+					Storage.Remove(TopHandle);
+				}
 			}
 			else
 			{
diff --git a/Engine/Source/Runtime/Online/Experimental/EventLoopTests/Tests/EventLoop/EventLoopTimerTests.cpp b/Engine/Source/Runtime/Online/Experimental/EventLoopTests/Tests/EventLoop/EventLoopTimerTests.cpp
--- a/Engine/Source/Runtime/Online/Experimental/EventLoopTests/Tests/EventLoop/EventLoopTimerTests.cpp
+++ b/Engine/Source/Runtime/Online/Experimental/EventLoopTests/Tests/EventLoop/EventLoopTimerTests.cpp
@@ -270,6 +270,23 @@ namespace UE::EventLoop
 					// Check that no timer is pending reschedule.
 					CHECK(!TimerManager.HasPendingRepeatTimer());
 				}
+
+				SECTION("Timer storage")
+				{
+					FMockEventLoopTimerHandle Handle = TimerManager.SetTimer(MoveTemp(TimerFunc1), FTimespan::FromSeconds(1));
+					REQUIRE(Handle.IsValid());
+
+					// Tick to add the timer.
+					TimerManager.Tick(FTimespan::Zero());
+					CHECK(TimerRunCount1 == 0);
+					CHECK(TimerManager.GetNumTimers() == 1);
+
+					// Tick to fire the timer.
+					TimerRunCount1 = 0;
+					TimerManager.Tick(FTimespan::FromSeconds(10));
+					CHECK(TimerRunCount1 == 1);
+					CHECK(TimerManager.GetNumTimers() == 0);
+				}
 			}
 
 			SECTION("Repeating")
@@ -415,6 +432,50 @@ namespace UE::EventLoop
 					// Check that timer is pending reschedule.
 					CHECK(TimerManager.HasPendingRepeatTimer());
 				}
+
+				SECTION("Timer storage")
+				{
+					FMockEventLoopTimerHandle Handle = TimerManager.SetTimer(MoveTemp(TimerFunc1), FTimespan::FromSeconds(1), true);
+					REQUIRE(Handle.IsValid());
+
+					// Tick to add the timer.
+					TimerRunCount1 = 0;
+					TimerManager.Tick(FTimespan::Zero());
+					CHECK(TimerRunCount1 == 0);
+					CHECK(TimerManager.GetNumTimers() == 1);
+
+					// Tick to fire the timer.
+					TimerRunCount1 = 0;
+					TimerManager.Tick(FTimespan::FromSeconds(2));
+					CHECK(TimerRunCount1 == 1);
+					CHECK(TimerManager.GetNumTimers() == 1);
+
+					SECTION("Remove before reschedule")
+					{
+						// Tick to clear the timer.
+						TimerRunCount1 = 0;
+						TimerManager.ClearTimer(Handle);
+						TimerManager.Tick(FTimespan::FromSeconds(2));
+						CHECK(TimerRunCount1 == 0);
+						CHECK(TimerManager.GetNumTimers() == 0);
+					}
+
+					SECTION("Remove after reschedule")
+					{
+						// Tick to reschedule the timer.
+						TimerRunCount1 = 0;
+						TimerManager.Tick(FTimespan::Zero());
+						CHECK(TimerRunCount1 == 0);
+						CHECK(TimerManager.GetNumTimers() == 1);
+
+						// Tick to clear the timer.
+						TimerRunCount1 = 0;
+						TimerManager.ClearTimer(Handle);
+						TimerManager.Tick(FTimespan::FromSeconds(2));
+						CHECK(TimerRunCount1 == 0);
+						CHECK(TimerManager.GetNumTimers() == 0);
+					}
+				}
 			}
 
 			SECTION("Primed timer removed before it is processed during tick.")
