Finish out eventloop BSD socket support.

diff --git a/Engine/Source/Runtime/Online/Experimental/EventLoop/Private/EventLoop/BSDSocket/BSDSocketTypesPrivate.h b/Engine/Source/Runtime/Online/Experimental/EventLoop/Private/EventLoop/BSDSocket/BSDSocketTypesPrivate.h
--- a/Engine/Source/Runtime/Online/Experimental/EventLoop/Private/EventLoop/BSDSocket/BSDSocketTypesPrivate.h
+++ b/Engine/Source/Runtime/Online/Experimental/EventLoop/Private/EventLoop/BSDSocket/BSDSocketTypesPrivate.h
@@ -3,6 +3,8 @@
 #pragma once
 
 #include "CoreTypes.h"
+#include "Containers/StringConv.h"
+#include "EventLoop/BSDSocket/BSDSocketTypes.h"
 
 #if PLATFORM_HAS_BSD_SOCKETS
 
@@ -14,6 +16,7 @@
 #include <ws2tcpip.h>
 
 typedef int32 SOCKLEN;
+#define strerror_r(errno,buf,len) strerror_s(buf,len,errno)
 
 #include "Windows/HideWindowsPlatformTypes.h"
 #else // PLATFORM_HAS_BSD_SOCKET_FEATURE_WINSOCKETS
@@ -57,4 +60,16 @@ inline int32 closesocket(SOCKET Socket)
 }
 
 #endif // PLATFORM_HAS_BSD_SOCKET_FEATURE_WINSOCKETS
+
+inline FString GetLastErrorString()
+{
+	constexpr uint32 BufferSize = 256;
+	ANSICHAR Buffer[BufferSize];
+	Buffer[0] = 0;
+
+	strerror_r(errno, Buffer, BufferSize);
+
+	return FString::Printf(TEXT("%s: (%d)"), StringCast<TCHAR>(Buffer).Get(), errno);
+}
+
 #endif // PLATFORM_HAS_BSD_SOCKETS
diff --git a/Engine/Source/Runtime/Online/Experimental/EventLoop/Private/EventLoop/BSDSocket/EventLoopIOManagerBSDSocketSelect.cpp b/Engine/Source/Runtime/Online/Experimental/EventLoop/Private/EventLoop/BSDSocket/EventLoopIOManagerBSDSocketSelect.cpp
--- a/Engine/Source/Runtime/Online/Experimental/EventLoop/Private/EventLoop/BSDSocket/EventLoopIOManagerBSDSocketSelect.cpp
+++ b/Engine/Source/Runtime/Online/Experimental/EventLoop/Private/EventLoop/BSDSocket/EventLoopIOManagerBSDSocketSelect.cpp
@@ -9,7 +9,7 @@
 #include "Math/UnrealMathUtility.h"
 #include "Stats/Stats.h"
 
-#if PLATFORM_HAS_BSD_SOCKETS
+#if PLATFORM_HAS_BSD_SOCKETS && PLATFORM_HAS_BSD_SOCKET_FEATURE_SELECT
 
 namespace UE::EventLoop {
 
@@ -35,10 +35,27 @@ private:
 	using FInternalHandleArryType = FStorageType::FInternalHandleArryType;
 	using FInternalHandle = FStorageType::FInternalHandle;
 
+	struct FSelectData
+	{
+		fd_set SocketReadSet;
+		fd_set SocketWriteSet;
+		fd_set SocketExceptionSet;
+		SOCKET MaxFd = INVALID_SOCKET;
+	};
+
+	void RebuildSelectData(FSelectData& SelectData) const;
+	int32 PollSelectData(const FSelectData& SelectData, FSelectData& OutSignaledData, FTimespan WaitTime) const;
+
+	FSelectData SelectData;
+
 	IEventLoop& EventLoop;
 	FStorageType IORequestStorage;
 	FIOAccessBSDSocket IOAccess;
 	TAtomic<bool> bAsyncSignal;
+
+	// Indices
+	TMap<FInternalHandle, SOCKET> SocketInternalHandleIndex;
+	TMap<SOCKET, FInternalHandle> SocketIndex;
 };
 
 FIOManagerBSDSocketSelectImpl::FIOManagerBSDSocketSelectImpl(IEventLoop& InEventLoop)
@@ -71,32 +88,27 @@ void FIOManagerBSDSocketSelectImpl::Poll(FTimespan WaitTime)
 	// woken by another thread. To mitigate this issue, many smaller calls to PollInternal are done
 	// while no signal is received.
 
-	const FTimespan MaxWaitTime = FTimespan::FromMilliseconds(10);
-	const double StartTime = FPlatformTime::Seconds();
+	const FTimespan MaxIterationWaitTime = FTimespan::FromMilliseconds(10);
+	FTimespan AccumulatedWaitTime;
 
-	for (;;)
+	// If a signal was received before polling, still poll for any ready events using a timeout of 0.
+	if (bAsyncSignal)
 	{
-		const double CurrentTime = FPlatformTime::Seconds();
-		const FTimespan CurrentDuration = FTimespan::FromSeconds(CurrentTime - StartTime);
-		FTimespan CurrentWaitTime = FMath::Min(WaitTime - CurrentDuration, MaxWaitTime);
-
-		// If a signal was received before polling, still poll for any ready events using a timeout of 0.
-		if (bAsyncSignal)
-		{
-			CurrentWaitTime = FTimespan::Zero();
-		}
-
-		if (CurrentWaitTime.GetTicks() < 0)
-		{
-			break;
-		}
+		WaitTime = FTimespan::Zero();
+	}
 
+	// Run at least one iteration. Keep running while WaitTime has not been reached.
+	for (int32 Index = 0; Index == 0 || AccumulatedWaitTime < WaitTime; ++Index)
+	{
+		FTimespan CurrentWaitTime = FMath::Min(WaitTime - AccumulatedWaitTime, MaxIterationWaitTime);
 		PollInternal(CurrentWaitTime);
 
 		if (bAsyncSignal)
 		{
 			break;
 		}
+
+		AccumulatedWaitTime += MaxIterationWaitTime;
 	}
 
 	// Reset signal. Resetting is safe here as setting the async signal indicates that the poll should be exited.
@@ -108,84 +120,87 @@ void FIOManagerBSDSocketSelectImpl::PollInternal(FTimespan WaitTime)
 	QUICK_SCOPE_CYCLE_COUNTER(STAT_EventLoop_FIOManagerBSDSocketSelectImpl_PollInternal);
 
 	// Process queued actions.
-	IORequestStorage.Update();
+	FInternalHandleArryType AddedHandles;
+	FInternalHandleArryType RemovedHandles;
+	uint32 NumChanges = IORequestStorage.Update(&AddedHandles, &RemovedHandles);
 
-	// When there is no work, block the loop here.
-	if (IORequestStorage.IsEmpty())
+	// Remove active requests before adding more to have an accurate count of available resources.
+	for (FInternalHandle RemovedHandle : RemovedHandles)
 	{
-		QUICK_SCOPE_CYCLE_COUNTER(STAT_EventLoop_FIOManagerBSDSocketSelectImpl_PollInternal_IdleSleep);
-		FPlatformProcess::SleepNoStats(WaitTime.GetTotalSeconds());
-		return;
+		SOCKET RemovedSocket;
+		if (SocketInternalHandleIndex.RemoveAndCopyValue(RemovedHandle, RemovedSocket))
+		{
+			UE_LOG(LogEventLoop, VeryVerbose, TEXT("[FIOManagerBSDSocketSelectImpl::PollInternal] Request removed: Socket: %p, Handle: %s"), RemovedSocket, *RemovedHandle.ToString());
+			SocketIndex.Remove(RemovedSocket);
+		}
 	}
 
-	fd_set SocketReadSet;
-	fd_set SocketWriteSet;
-	fd_set SocketExceptionSet;
-	SOCKET MaxFd = INVALID_SOCKET;
-
-	// Build the FD_SETs.
+	for (FInternalHandle AddedHandle : AddedHandles)
 	{
-		QUICK_SCOPE_CYCLE_COUNTER(STAT_EventLoop_FIOManagerBSDSocketSelectImpl_PollInternal_BuildFdSets);
-		FD_ZERO(&SocketReadSet);
-		FD_ZERO(&SocketWriteSet);
-		FD_ZERO(&SocketExceptionSet);
-
-		for (TPair<FInternalHandle, FIORequestBSDSocket&>& IORequestEntry : IORequestStorage)
+		// Request may have been added and removed in the same update.
+		if (FIORequestBSDSocket* NewRequest = IORequestStorage.Find(AddedHandle))
 		{
-			FIORequestBSDSocket& IORequest = IORequestEntry.Value;
-
-			if (MaxFd == INVALID_SOCKET)
+			// Check whether there are enough resources available to handle the request.
+			if (SocketIndex.Num() >= FD_SETSIZE)
 			{
-				MaxFd = IORequest.Socket;
-			}
-			else
-			{
-				MaxFd = FMath::Max(MaxFd, IORequest.Socket);
+				UE_LOG(LogEventLoop, Verbose, TEXT("[FIOManagerBSDSocketSelectImpl::PollInternal] Request failed: Socket: %p, Reason: %s"), NewRequest->Socket, *LexToString(ESocketIoRequestStatus::NoResources));
+				NewRequest->Callback(NewRequest->Socket, ESocketIoRequestStatus::NoResources, EIOFlags::None);
+				IORequestStorage.Remove(AddedHandle);
+				continue;
 			}
 
-			if (EnumHasAnyFlags(IORequest.Flags, EIOFlags::Read))
+			// Only one request may be active for a socket at one time.
+			if (SocketIndex.Contains(NewRequest->Socket))
 			{
-				FD_SET(IORequest.Socket, &SocketReadSet);
+				UE_LOG(LogEventLoop, Verbose, TEXT("[FIOManagerBSDSocketSelectImpl::PollInternal] Request failed: Socket: %p, Reason: %s"), NewRequest->Socket, *LexToString(ESocketIoRequestStatus::Invalid));
+				NewRequest->Callback(NewRequest->Socket, ESocketIoRequestStatus::Invalid, EIOFlags::None);
+				IORequestStorage.Remove(AddedHandle);
+				continue;
 			}
 
-			if (EnumHasAnyFlags(IORequest.Flags, EIOFlags::Write))
-			{
-				FD_SET(IORequest.Socket, &SocketWriteSet);
-				FD_SET(IORequest.Socket, &SocketExceptionSet);
-			}
+			UE_LOG(LogEventLoop, VeryVerbose, TEXT("[FIOManagerBSDSocketSelectImpl::PollInternal] Request added: Socket: %p, Handle: %s, Flags: 0x%08X"), NewRequest->Socket, *AddedHandle.ToString(), NewRequest->Flags);
+
+			SocketInternalHandleIndex.Add(AddedHandle, NewRequest->Socket);
+			SocketIndex.Add(NewRequest->Socket, AddedHandle);
 		}
 	}
 
-#if PLATFORM_HAS_BSD_SOCKET_FEATURE_SELECT
-	// convert WaitTime to a timeval
-	timeval Time;
-	Time.tv_sec = (int32)WaitTime.GetTotalSeconds();
-	Time.tv_usec = WaitTime.GetFractionMicro();
-
-	timeval* TimePointer = WaitTime.GetTicks() >= 0 ? &Time : nullptr;
+	// Check whether the cached FD_SETs need to be rebuilt.
+	if (NumChanges > 0)
+	{
+		RebuildSelectData(SelectData);
+	}
 
-	// Poll for socket status.
-	int32 SelectStatus = 0;
+	// When there is no work, block the loop here.
+	if (IORequestStorage.IsEmpty())
 	{
-		QUICK_SCOPE_CYCLE_COUNTER(STAT_EventLoop_FIOManagerBSDSocketSelectImpl_PollInternal_SocketSelect);
-		SelectStatus = select(IntCastChecked<int>(MaxFd + 1), &SocketReadSet, &SocketWriteSet, NULL, TimePointer);
+		QUICK_SCOPE_CYCLE_COUNTER(STAT_EventLoop_FIOManagerBSDSocketSelectImpl_PollInternal_IdleSleep);
+		FPlatformProcess::SleepNoStats(WaitTime.GetTotalSeconds());
+		return;
 	}
+
+	// Poll for socket status.
+	FSelectData SignaledSelectData;
+	int32 SelectStatus = PollSelectData(SelectData, SignaledSelectData, WaitTime);
+
+	// Check whether any events triggered.
 	if (SelectStatus > 0)
 	{
 		QUICK_SCOPE_CYCLE_COUNTER(STAT_EventLoop_FIOManagerBSDSocketSelectImpl_PollInternal_HandleIOState);
+
 		// Check request state and trigger callbacks if needed.
 		for (TPair<FInternalHandle, FIORequestBSDSocket&>& IORequestEntry : IORequestStorage)
 		{
 			FIORequestBSDSocket& IORequest = IORequestEntry.Value;
 			EIOFlags SignaledFlags = EIOFlags::None;
 
-			if (EnumHasAnyFlags(IORequest.Flags, EIOFlags::Read) && FD_ISSET(IORequest.Socket, &SocketReadSet))
+			if (EnumHasAnyFlags(IORequest.Flags, EIOFlags::Read) && FD_ISSET(IORequest.Socket, &SignaledSelectData.SocketReadSet))
 			{
 				SignaledFlags |= EIOFlags::Read;
 			}
 
 			if (EnumHasAnyFlags(IORequest.Flags, EIOFlags::Write) &&
-				(FD_ISSET(IORequest.Socket, &SocketWriteSet) || FD_ISSET(IORequest.Socket, &SocketExceptionSet)))
+				(FD_ISSET(IORequest.Socket, &SignaledSelectData.SocketWriteSet) || FD_ISSET(IORequest.Socket, &SignaledSelectData.SocketExceptionSet)))
 			{
 				SignaledFlags |= EIOFlags::Write;
 			}
@@ -193,19 +208,68 @@ void FIOManagerBSDSocketSelectImpl::PollInternal(FTimespan WaitTime)
 			if (SignaledFlags != EIOFlags::None)
 			{
 				// Signal event status.
-				IORequest.Callback(SignaledFlags);
+				UE_LOG(LogEventLoop, VeryVerbose, TEXT("[FIOManagerBSDSocketSelectImpl::PollInternal] Request signaled: Socket: %p, Handle: %s, Flags: 0x%08X"), IORequest.Socket, *IORequestEntry.Key.ToString(), SignaledFlags);
+				IORequest.Callback(IORequest.Socket, ESocketIoRequestStatus::Ok, SignaledFlags);
 			}
 		}
 	}
 	else if (SelectStatus < 0)
 	{
-		// Todo: handle error by iterating set to see which socket caused the failure.
-		UE_LOG(LogEventLoop, Error, TEXT("[FIOManagerBSDSocketSelectImpl::PollInternal] Select error"));
-		EventLoop.RequestShutdown();
+		UE_LOG(LogEventLoop, Fatal, TEXT("[FIOManagerBSDSocketSelectImpl::PollInternal] Select error: %s"), *GetLastErrorString());
+	}
+}
+
+void FIOManagerBSDSocketSelectImpl::RebuildSelectData(FSelectData& OutSelectData) const
+{
+	QUICK_SCOPE_CYCLE_COUNTER(STAT_EventLoop_FIOManagerBSDSocketSelectImpl_RebuildSelectData);
+	FD_ZERO(&OutSelectData.SocketReadSet);
+	FD_ZERO(&OutSelectData.SocketWriteSet);
+	FD_ZERO(&OutSelectData.SocketExceptionSet);
+	OutSelectData.MaxFd = INVALID_SOCKET;
+
+	for (const TPair<SOCKET, FInternalHandle>& IndexEntry : SocketIndex)
+	{
+		if (const FIORequestBSDSocket* IORequest = IORequestStorage.Find(IndexEntry.Value))
+		{
+			if (OutSelectData.MaxFd == INVALID_SOCKET)
+			{
+				OutSelectData.MaxFd = IORequest->Socket;
+			}
+			else
+			{
+				OutSelectData.MaxFd = FMath::Max(SelectData.MaxFd, IORequest->Socket);
+			}
+
+			if (EnumHasAnyFlags(IORequest->Flags, EIOFlags::Read))
+			{
+				FD_SET(IORequest->Socket, &OutSelectData.SocketReadSet);
+			}
+
+			if (EnumHasAnyFlags(IORequest->Flags, EIOFlags::Write))
+			{
+				FD_SET(IORequest->Socket, &OutSelectData.SocketWriteSet);
+				FD_SET(IORequest->Socket, &OutSelectData.SocketExceptionSet);
+			}
+		}
 	}
-#else
-	UE_LOG(LogEventLoop, Fatal, TEXT("This platform doesn't support FIOManagerSocketBSDSelectImpl"));
-#endif
+}
+
+int32 FIOManagerBSDSocketSelectImpl::PollSelectData(const FSelectData& InSelectData, FIOManagerBSDSocketSelectImpl::FSelectData& OutSignaledData, FTimespan WaitTime) const
+{
+	QUICK_SCOPE_CYCLE_COUNTER(STAT_EventLoop_FIOManagerBSDSocketSelectImpl_PollSelectData);
+
+	// Copy cached FD sets.
+	OutSignaledData = InSelectData;
+
+	// convert WaitTime to a timeval
+	timeval Time;
+	Time.tv_sec = (int32)WaitTime.GetTotalSeconds();
+	Time.tv_usec = WaitTime.GetFractionMicro();
+
+	timeval* TimePointer = WaitTime.GetTicks() >= 0 ? &Time : nullptr;
+
+	// Poll for socket status.
+	return select(IntCastChecked<int>(OutSignaledData.MaxFd + 1), &OutSignaledData.SocketReadSet, &OutSignaledData.SocketWriteSet, &OutSignaledData.SocketExceptionSet, TimePointer);
 }
 
 //-----------------------
@@ -244,4 +308,5 @@ FIOManagerBSDSocketSelect::FIOAccess& FIOManagerBSDSocketSelect::GetIOAccess()
 
 /* UE::EventLoop */ }
 
-#endif // PLATFORM_HAS_BSD_SOCKETS
+#endif // PLATFORM_HAS_BSD_SOCKETS && PLATFORM_HAS_BSD_SOCKET_FEATURE_SELECT
+
diff --git a/Engine/Source/Runtime/Online/Experimental/EventLoop/Public/EventLoop/BSDSocket/EventLoopIOAccessBSDSocket.h b/Engine/Source/Runtime/Online/Experimental/EventLoop/Public/EventLoop/BSDSocket/EventLoopIOAccessBSDSocket.h
--- a/Engine/Source/Runtime/Online/Experimental/EventLoop/Public/EventLoop/BSDSocket/EventLoopIOAccessBSDSocket.h
+++ b/Engine/Source/Runtime/Online/Experimental/EventLoop/Public/EventLoop/BSDSocket/EventLoopIOAccessBSDSocket.h
@@ -6,6 +6,7 @@
 
 #if PLATFORM_HAS_BSD_SOCKETS
 
+#include "Containers/UnrealString.h"
 #include "EventLoop/BSDSocket/BSDSocketTypes.h"
 #include "EventLoop/IEventLoopIOManager.h"
 #include "EventLoop/EventLoopManagedStorage.h"
@@ -13,13 +14,38 @@
 
 namespace UE::EventLoop {
 
-using FIOCallback = TUniqueFunction<void(EIOFlags SignaledFlags)>;
+enum class ESocketIoRequestStatus : uint8
+{
+	/**
+	 * No error occurred, signaled flags are valid.
+	 * After signaling the handle associated with the request will remain valid until the user
+	 * deletes the associated request.
+	 */
+	Ok,
+
+	/**
+	 * There are no resources available to serve the IO request.
+	 * This will occur if more requests have been made than the system is capable of handling.
+	 * After signaling the handle associated with the request is removed.
+	 */
+	NoResources,
+
+	/**
+	 * The request is invalid.
+	 * A request may be invalid if another request for the same socket is already in progress.
+	 * After signaling the handle associated with the request is removed.
+	 */
+	Invalid,
+};
+
+using FSocketIOCallback = TUniqueFunction<void(SOCKET Socket, ESocketIoRequestStatus Status, EIOFlags SignaledFlags)>;
+using FSocketIORequestDestroyedCallback = FManagedStorageOnRemoveComplete;
 
 struct FIORequestBSDSocket
 {
 	SOCKET Socket = INVALID_SOCKET;
 	EIOFlags Flags = EIOFlags::None;
-	FIOCallback Callback;
+	FSocketIOCallback Callback;
 };
 
 class FIOAccessBSDSocket final : public FNoncopyable
@@ -47,9 +73,9 @@ public:
 		return IORequestStorage.Add({MoveTemp(Request)});
 	}
 
-	void DestroyIORequest(FIORequestHandle& Handle)
+	void DestroyIORequest(FIORequestHandle& Handle, FSocketIORequestDestroyedCallback&& OnDestroyComplete = FSocketIORequestDestroyedCallback())
 	{
-		IORequestStorage.Remove(Handle);
+		IORequestStorage.Remove(Handle, MoveTemp(OnDestroyComplete));
 	}
 
 private:
@@ -58,4 +84,18 @@ private:
 
 /* UE::EventLoop */ }
 
+inline FString LexToString(UE::EventLoop::ESocketIoRequestStatus Status)
+{
+	switch (Status)
+	{
+	case UE::EventLoop::ESocketIoRequestStatus::Ok: return TEXT("Ok");
+	case UE::EventLoop::ESocketIoRequestStatus::NoResources: return TEXT("NoResources");
+	case UE::EventLoop::ESocketIoRequestStatus::Invalid: return TEXT("Invalid");
+
+	default:
+		checkNoEntry();
+		return TEXT("Unknown");
+	}
+}
+
 #endif // PLATFORM_HAS_BSD_SOCKETS
diff --git a/Engine/Source/Runtime/Online/Experimental/EventLoop/Public/EventLoop/BSDSocket/EventLoopIOManagerBSDSocket.h b/Engine/Source/Runtime/Online/Experimental/EventLoop/Public/EventLoop/BSDSocket/EventLoopIOManagerBSDSocket.h
--- a/Engine/Source/Runtime/Online/Experimental/EventLoop/Public/EventLoop/BSDSocket/EventLoopIOManagerBSDSocket.h
+++ b/Engine/Source/Runtime/Online/Experimental/EventLoop/Public/EventLoop/BSDSocket/EventLoopIOManagerBSDSocket.h
@@ -11,8 +11,10 @@ namespace UE::EventLoop {
 #if HAS_EVENTLOOP_PLATFORM_SOCKET_IMPLEMENTATION
 	#include COMPILED_PLATFORM_HEADER(EventLoopIOManagerBSDSocket.h)
 #else
+#if PLATFORM_HAS_BSD_SOCKET_FEATURE_SELECT
 	// Default implementation.
 	using FIOManagerBSDSocket = FIOManagerBSDSocketSelect;
+#endif // PLATFORM_HAS_BSD_SOCKET_FEATURE_SELECT
 #endif // EVENTLOOP_PLATFORM_IMPLEMENTATION
 #endif // PLATFORM_HAS_BSD_SOCKETS
 
diff --git a/Engine/Source/Runtime/Online/Experimental/EventLoop/Public/EventLoop/BSDSocket/EventLoopIOManagerBSDSocketSelect.h b/Engine/Source/Runtime/Online/Experimental/EventLoop/Public/EventLoop/BSDSocket/EventLoopIOManagerBSDSocketSelect.h
--- a/Engine/Source/Runtime/Online/Experimental/EventLoop/Public/EventLoop/BSDSocket/EventLoopIOManagerBSDSocketSelect.h
+++ b/Engine/Source/Runtime/Online/Experimental/EventLoop/Public/EventLoop/BSDSocket/EventLoopIOManagerBSDSocketSelect.h
@@ -4,7 +4,7 @@
 
 #include "CoreTypes.h"
 
-#if PLATFORM_HAS_BSD_SOCKETS
+#if PLATFORM_HAS_BSD_SOCKETS && PLATFORM_HAS_BSD_SOCKET_FEATURE_SELECT
 
 #include "EventLoop/BSDSocket/BSDSocketTypes.h"
 #include "EventLoop/BSDSocket/EventLoopIOAccessBSDSocket.h"
@@ -42,4 +42,4 @@ private:
 
 /* UE::EventLoop */ }
 
-#endif // PLATFORM_HAS_BSD_SOCKETS
+#endif // PLATFORM_HAS_BSD_SOCKETS && PLATFORM_HAS_BSD_SOCKET_FEATURE_SELECT
diff --git a/Engine/Source/Runtime/Online/Experimental/EventLoopTests/EventLoopUnitTests.Build.cs b/Engine/Source/Runtime/Online/Experimental/EventLoopTests/EventLoopUnitTests.Build.cs
--- a/Engine/Source/Runtime/Online/Experimental/EventLoopTests/EventLoopUnitTests.Build.cs
+++ b/Engine/Source/Runtime/Online/Experimental/EventLoopTests/EventLoopUnitTests.Build.cs
@@ -10,7 +10,8 @@ public class EventLoopUnitTests : TestModuleRules
 		PrivateDependencyModuleNames.AddRange(
 			new string[] {
 				"Core",
-				"EventLoop"
+				"EventLoop",
+				"Sockets"
 			});
 
 		UpdateBuildGraphPropertiesFile(new Metadata() { TestName = "EventLoop", TestShortName = "EventLoop" });
diff --git a/Engine/Source/Runtime/Online/Experimental/EventLoopTests/Tests/EventLoop/BSDSocketTypesPrivate.h b/Engine/Source/Runtime/Online/Experimental/EventLoopTests/Tests/EventLoop/BSDSocketTypesPrivate.h
new file mode 100644
--- /dev/null
+++ b/Engine/Source/Runtime/Online/Experimental/EventLoopTests/Tests/EventLoop/BSDSocketTypesPrivate.h
@@ -0,0 +1,61 @@
+// Copyright Epic Games, Inc. All Rights Reserved.
+
+#pragma once
+
+#include "CoreTypes.h"
+#include "EventLoop/BSDSocket/BSDSocketTypes.h"
+
+#if PLATFORM_HAS_BSD_SOCKETS
+
+#if PLATFORM_HAS_BSD_SOCKET_FEATURE_WINSOCKETS
+#include "Windows/WindowsHWrapper.h"
+#include "Windows/AllowWindowsPlatformTypes.h"
+
+#include <winsock2.h>
+#include <ws2tcpip.h>
+
+typedef int32 SOCKLEN;
+
+#include "Windows/HideWindowsPlatformTypes.h"
+#else // PLATFORM_HAS_BSD_SOCKET_FEATURE_WINSOCKETS
+#if HAS_EVENTLOOP_PLATFORM_BSD_SOCKET_HEADER
+#include COMPILED_PLATFORM_HEADER(BSDSocketTypesPrivate.h)
+#else // HAS_EVENTLOOP_PLATFORM_BSD_SOCKET_HEADER
+#include <unistd.h>
+#include <sys/socket.h>
+#if PLATFORM_HAS_BSD_SOCKET_FEATURE_IOCTL
+#include <fcntl.h>
+#include <sys/types.h>
+#include <sys/ioctl.h>
+#endif // PLATFORM_HAS_BSD_SOCKET_FEATURE_IOCTL
+#if PLATFORM_HAS_BSD_SOCKET_FEATURE_POLL
+#include <poll.h>
+#endif // PLATFORM_HAS_BSD_SOCKET_FEATURE_POLL
+#include <netinet/in.h>
+#include <arpa/inet.h>
+#if PLATFORM_HAS_BSD_SOCKET_FEATURE_GETHOSTNAME
+#include <netdb.h>
+#endif // PLATFORM_HAS_BSD_SOCKET_FEATURE_GETHOSTNAME
+#if PLATFORM_HAS_BSD_SOCKET_FEATURE_NODELAY
+#include <netinet/tcp.h>
+#endif // PLATFORM_HAS_BSD_SOCKET_FEATURE_NODELAY
+
+#define ioctlsocket ioctl
+#endif // HAS_EVENTLOOP_PLATFORM_BSD_SOCKET_HEADER
+
+#define SOCKET_ERROR -1
+#define INVALID_SOCKET -1
+
+typedef socklen_t SOCKLEN;
+typedef int32 SOCKET;
+typedef sockaddr_in SOCKADDR_IN;
+typedef struct timeval TIMEVAL;
+
+inline int32 closesocket(SOCKET Socket)
+{
+	shutdown(Socket, SHUT_RDWR); // gracefully shutdown if connected
+	return close(Socket);
+}
+
+#endif // PLATFORM_HAS_BSD_SOCKET_FEATURE_WINSOCKETS
+#endif // PLATFORM_HAS_BSD_SOCKETS
diff --git a/Engine/Source/Runtime/Online/Experimental/EventLoopTests/Tests/EventLoop/EventLoopIOManagerBSDSocketSelectTests.cpp b/Engine/Source/Runtime/Online/Experimental/EventLoopTests/Tests/EventLoop/EventLoopIOManagerBSDSocketSelectTests.cpp
new file mode 100644
--- /dev/null
+++ b/Engine/Source/Runtime/Online/Experimental/EventLoopTests/Tests/EventLoop/EventLoopIOManagerBSDSocketSelectTests.cpp
@@ -0,0 +1,851 @@
+// Copyright Epic Games, Inc. All Rights Reserved.
+
+#include "CoreTypes.h"
+
+#if PLATFORM_HAS_BSD_SOCKETS
+
+#include "BSDSocketTypesPrivate.h"
+#include "EventLoop/IEventLoop.h"
+#include "EventLoop/BSDSocket/EventLoopIOManagerBSDSocketSelect.h"
+#include "TestHarness.h"
+
+#include "SocketSubsystem.h"
+
+namespace UE::EventLoop
+{
+	struct FBSDSocketFuncs
+	{
+		static SOCKET Open(int32 Family, int32 Type, int32 Protocol)
+		{
+			return socket(Family, Type, Protocol);
+		}
+
+		static void Close(SOCKET Socket)
+		{
+			closesocket(Socket);
+		}
+
+		static bool SetSendBufferSize(SOCKET Socket, int32 Size)
+		{
+			SOCKLEN SizeSize = sizeof(int32);
+			bool bOk = setsockopt(Socket, SOL_SOCKET, SO_SNDBUF, (char*)&Size, sizeof(int32)) == 0;
+
+			// Read the value back in case the size was modified
+			int32 NewSize;
+			getsockopt(Socket, SOL_SOCKET, SO_SNDBUF, (char*)&NewSize, &SizeSize);
+
+			return bOk && Size == NewSize;
+		}
+
+		static bool SetRecvBufferSize(SOCKET Socket, int32 Size)
+		{
+			SOCKLEN SizeSize = sizeof(int32);
+			bool bOk = setsockopt(Socket, SOL_SOCKET, SO_RCVBUF, (char*)&Size, sizeof(int32)) == 0;
+
+			// Read the value back in case the size was modified
+			int32 NewSize;
+			getsockopt(Socket, SOL_SOCKET, SO_RCVBUF, (char*)&NewSize, &SizeSize);
+
+			return bOk && Size == NewSize;
+		}
+
+		static bool SetNonBlocking(SOCKET Socket, bool bIsNonBlocking)
+		{
+	#if PLATFORM_HAS_BSD_SOCKET_FEATURE_WINSOCKETS
+			u_long Value = bIsNonBlocking ? true : false;
+			return ioctlsocket(Socket, FIONBIO, &Value) == 0;
+	#else
+			int Flags = fcntl(Socket, F_GETFL, 0);
+			//Set the flag or clear it, without destroying the other flags.
+			Flags = bIsNonBlocking ? Flags | O_NONBLOCK : Flags ^ (Flags & O_NONBLOCK);
+			int err = fcntl(Socket, F_SETFL, Flags);
+			return (err == 0 ? true : false);
+	#endif
+		}
+
+		static bool SetReusePort(SOCKET Socket, bool bIsReusePort)
+		{
+			int32 Resuse = bIsReusePort ? true : false;
+			bool bOk = setsockopt(Socket, SOL_SOCKET, SO_BROADCAST, (char*)&Resuse, sizeof(int32)) == 0;
+			bOk &= setsockopt(Socket, SOL_SOCKET, SO_REUSEADDR, (char*)&Resuse, sizeof(int32)) == 0;
+			return bOk;
+		}
+
+		static bool BindLoopback(SOCKET Socket)
+		{
+			sockaddr_in Addr;
+			FMemory::Memzero(&Addr, sizeof(Addr));
+			Addr.sin_addr.s_addr = htonl(INADDR_LOOPBACK);
+			Addr.sin_family = AF_INET;
+
+			return Bind(Socket, Addr);
+		}
+
+		static bool Bind(SOCKET Socket, const sockaddr_in& Addr)
+		{
+			return bind(Socket, (const sockaddr*)&Addr, sizeof(sockaddr_in)) == 0;
+		}
+
+		static bool Listen(SOCKET Socket, int32 MaxBacklog)
+		{
+			return listen(Socket, MaxBacklog) == 0;
+		}
+
+		static bool GetAddress(SOCKET Socket, sockaddr_in& OutAddr)
+		{
+			SOCKLEN Size = sizeof(sockaddr_in);
+
+			// Figure out what ip/port we are bound to
+			return getsockname(Socket, (sockaddr*)&OutAddr, &Size) == 0;
+		}
+
+		static bool Connect(SOCKET Socket, const sockaddr_in& Addr)
+		{
+			SOCKLEN Size = sizeof(sockaddr_in);
+			int32 Return = connect(Socket, (const sockaddr*)&Addr, Size);
+
+			ISocketSubsystem* SocketSubsystem = ISocketSubsystem::Get(PLATFORM_SOCKETSUBSYSTEM);
+			check(SocketSubsystem);
+			ESocketErrors Error = SocketSubsystem->TranslateErrorCode(Return);
+
+			// EWOULDBLOCK is not an error, and EINPROGRESS is fine on initial connection as it may still be creating for nonblocking sockets
+			return ((Error == SE_NO_ERROR) || (Error == SE_EWOULDBLOCK) || (Error == SE_EINPROGRESS));
+		}
+
+		static SOCKET Accept(SOCKET Socket)
+		{
+			return accept(Socket, NULL, NULL);
+		}
+
+		static bool Send(SOCKET Socket, const uint8* Data, int32 Count, int32& BytesSent)
+		{
+			BytesSent = send(Socket,(const char*)Data,Count,0);
+			return BytesSent >= 0;
+		}
+
+		static bool Recv(SOCKET Socket, uint8* Data, int32 BufferSize, int32& BytesRead)
+		{
+			int SocketType = 0;
+			SOCKLEN SocketTypeLength = sizeof(int32);
+			if (getsockopt(Socket, SOL_SOCKET, SO_TYPE, (char*)&SocketType, &SocketTypeLength) == SOCKET_ERROR)
+			{
+				return false;
+			}
+
+			bool bSuccess = false;
+			const bool bStreamSocket = (SocketType == SOCK_STREAM);
+			BytesRead = recv(Socket, (char*)Data, BufferSize, 0);
+
+			if (BytesRead >= 0)
+			{
+				// For Streaming sockets, 0 indicates a graceful failure
+				bSuccess = !bStreamSocket || (BytesRead > 0);
+			}
+			else
+			{
+				// For Streaming sockets, don't treat SE_EWOULDBLOCK as an error
+				ISocketSubsystem* SocketSubsystem = ISocketSubsystem::Get(PLATFORM_SOCKETSUBSYSTEM);
+				check(SocketSubsystem);
+				bSuccess = bStreamSocket && (SocketSubsystem->TranslateErrorCode(BytesRead) == SE_EWOULDBLOCK);
+				BytesRead = 0;
+			}
+
+			return bSuccess;
+		}
+
+		static bool SendTo(SOCKET Socket, const uint8* Data, int32 Count, int32& BytesSent, const sockaddr_in& Addr)
+		{
+			SOCKLEN Size = sizeof(sockaddr_in);
+			// Write the data and see how much was written
+			BytesSent = sendto(Socket, (const char*)Data, Count, 0, (const sockaddr*)&(Addr), Size);
+			return BytesSent >= 0;
+		}
+
+		static bool RecvFrom(SOCKET Socket, uint8* Data, int32 BufferSize, int32& BytesRead, sockaddr_in& Addr)
+		{
+			int SocketType = 0;
+			SOCKLEN SocketTypeLength = sizeof(int32);
+			if (getsockopt(Socket, SOL_SOCKET, SO_TYPE, (char*)&SocketType, &SocketTypeLength) == SOCKET_ERROR)
+			{
+				return false;
+			}
+
+			bool bSuccess = false;
+			const bool bStreamSocket = (SocketType == SOCK_STREAM);
+			SOCKLEN Size = sizeof(sockaddr_in);
+
+			// Read into the buffer and set the source address
+			BytesRead = recvfrom(Socket, (char*)Data, BufferSize, 0, (sockaddr*)&Addr, &Size);
+
+			if (BytesRead >= 0)
+			{
+				// For Streaming sockets, 0 indicates a graceful failure
+				bSuccess = !bStreamSocket || (BytesRead > 0);
+			}
+			else
+			{
+				ISocketSubsystem* SocketSubsystem = ISocketSubsystem::Get(PLATFORM_SOCKETSUBSYSTEM);
+				check(SocketSubsystem);
+				// For Streaming sockets, don't treat SE_EWOULDBLOCK as an error
+				bSuccess = bStreamSocket && (SocketSubsystem->TranslateErrorCode(BytesRead) == SE_EWOULDBLOCK);
+				BytesRead = 0;
+			}
+
+			return bSuccess;
+		}
+	};
+
+	// Todo - remove eventloop parameter from IOManager.
+	class FMockEventLoop : public IEventLoop
+	{
+	public:
+		FMockEventLoop() = default;
+		virtual ~FMockEventLoop() = default;
+		virtual bool Init() override { return false; }
+		virtual void RequestShutdown(FOnShutdownComplete&& OnShutdownComplete) override {}
+		virtual FTimerHandle SetTimer(FTimerCallback&& Callback, FTimespan InRate, bool InbRepeat, TOptional<FTimespan> InFirstDelay) override { return FTimerHandle(); }
+		virtual void ClearTimer(FTimerHandle& InHandle, FOnTimerCleared&& OnTimerCleared) override {}
+		virtual void PostAsyncTask(FAsyncTask&& Task) override {}
+		virtual void Run() override {}
+		virtual bool RunOnce(FTimespan WaitTime) override { return false; }
+		virtual FTimespan GetLoopTime() const override { return FTimespan(); }
+	};
+
+	TEST_CASE("IOManagerBSDSocketSelect", "[Online][EventLoop][IOManager][Smoke]")
+	{
+		FMockEventLoop MockEventLoop;
+		FIOManagerBSDSocketSelect IOManager(MockEventLoop, FIOManagerBSDSocketSelect::FParams());
+
+		REQUIRE(IOManager.Init());
+
+		ISocketSubsystem* SocketSubsystem = ISocketSubsystem::Get(PLATFORM_SOCKETSUBSYSTEM);
+		check(SocketSubsystem);
+
+		SECTION("Notify tests")
+		{
+			double StartTime = 0;
+			FTimespan Duration;
+
+			// 1. Waiting while no events are received will wait for at least the requested duration.
+			StartTime = FPlatformTime::Seconds();
+			IOManager.Poll(FTimespan::FromMilliseconds(50));
+			Duration = FTimespan::FromSeconds(FPlatformTime::Seconds() - StartTime);
+			CHECK(Duration > FTimespan::FromMilliseconds(50));
+
+			// 2. Notifying the IO manager will wake the loop up early.
+			StartTime = FPlatformTime::Seconds();
+			IOManager.Notify();
+			IOManager.Poll(FTimespan::FromMilliseconds(50));
+			Duration = FTimespan::FromSeconds(FPlatformTime::Seconds() - StartTime);
+			CHECK(Duration < FTimespan::FromMilliseconds(1));
+		}
+
+		SECTION("Single socket tests")
+		{
+			// Size of buffer used for sending test data.
+			constexpr uint64 WriteBufferSize = 128;
+			TArray<uint8> WriteBuffer;
+			WriteBuffer.AddDefaulted(WriteBufferSize);
+
+			// Size of buffer used for receiving test data.
+			constexpr uint64 ReadBufferSize = 128;
+			TArray<uint8> ReadBuffer;
+			ReadBuffer.AddDefaulted(ReadBufferSize);
+
+			// Socket buffer sizes for send / recv
+			constexpr int32 SocketSendBufferSize = 64;
+			constexpr int32 SocketReceiveBufferSize = 64;
+
+			// Create listen socket.
+			SOCKET ListenSocket = FBSDSocketFuncs::Open(AF_INET, SOCK_STREAM, IPPROTO_TCP);
+			REQUIRE(ListenSocket != INVALID_SOCKET);
+			REQUIRE(FBSDSocketFuncs::SetSendBufferSize(ListenSocket, SocketSendBufferSize));
+			REQUIRE(FBSDSocketFuncs::SetRecvBufferSize(ListenSocket, SocketReceiveBufferSize));
+			REQUIRE(FBSDSocketFuncs::SetNonBlocking(ListenSocket, true));
+			REQUIRE(FBSDSocketFuncs::BindLoopback(ListenSocket));
+			REQUIRE(FBSDSocketFuncs::Listen(ListenSocket, 1));
+			sockaddr_in HostAddr;
+			FMemory::Memzero(&HostAddr, sizeof(HostAddr));
+			REQUIRE(FBSDSocketFuncs::GetAddress(ListenSocket, HostAddr));
+
+			// Create remote socket.
+			SOCKET RemoteSocket1 = FBSDSocketFuncs::Open(AF_INET, SOCK_STREAM, IPPROTO_TCP);
+			REQUIRE(RemoteSocket1 != INVALID_SOCKET);
+			REQUIRE(FBSDSocketFuncs::SetSendBufferSize(RemoteSocket1, SocketSendBufferSize));
+			REQUIRE(FBSDSocketFuncs::SetRecvBufferSize(RemoteSocket1, SocketReceiveBufferSize));
+			REQUIRE(FBSDSocketFuncs::SetNonBlocking(RemoteSocket1, true));
+			REQUIRE(FBSDSocketFuncs::Connect(RemoteSocket1, HostAddr));
+
+			// Create local socket.
+			ESocketErrors ErrorCode1 = SE_NO_ERROR;
+			SOCKET LocalSocket1 = INVALID_SOCKET;
+			do
+			{
+				LocalSocket1 = FBSDSocketFuncs::Accept(ListenSocket);
+				ErrorCode1 = SocketSubsystem->GetLastErrorCode();
+			}
+			while (ErrorCode1 == SE_EWOULDBLOCK);
+			REQUIRE(LocalSocket1 != INVALID_SOCKET);
+
+			SECTION("Single waiting socket tests")
+			{
+				SECTION("Wait on socket readable")
+				{
+					SOCKET TriggeredSocket = INVALID_SOCKET;
+					ESocketIoRequestStatus TriggeredStatus = ESocketIoRequestStatus::Ok;
+					EIOFlags TriggeredFlags = EIOFlags::None;
+					uint32 CallbackTriggerCount = 0;
+
+					FIORequestBSDSocket Request;
+					Request.Socket = LocalSocket1;
+					Request.Flags = EIOFlags::Read;
+					Request.Callback = [&](SOCKET Socket, ESocketIoRequestStatus Status, EIOFlags SignaledFlags){
+						TriggeredSocket = Socket;
+						TriggeredStatus = Status;
+						TriggeredFlags = SignaledFlags;
+						++CallbackTriggerCount;
+					};
+
+					FIORequestHandle RequestHandle = IOManager.GetIOAccess().CreateSocketIORequest(MoveTemp(Request));
+					CHECK(RequestHandle.IsValid());
+
+					// 1. No events triggered when no data has arrived at the socket.
+					IOManager.Poll(FTimespan::Zero());
+					CHECK(CallbackTriggerCount == 0);
+
+					// 2. Send data to the socket.
+					int32 BytesSent = 0;
+					CHECK(FBSDSocketFuncs::Send(RemoteSocket1, WriteBuffer.GetData(), 1, BytesSent));
+					CHECK(BytesSent == 1);
+
+					// 3. See that socket has a read event.
+					TriggeredFlags = EIOFlags::None;
+					IOManager.Poll(FTimespan::Zero());
+					CHECK(CallbackTriggerCount == 1);
+					CHECK(TriggeredSocket == LocalSocket1);
+					CHECK(TriggeredStatus == ESocketIoRequestStatus::Ok);
+					CHECK(TriggeredFlags == EIOFlags::Read);
+
+					// 4. See that socket still has a read event.
+					IOManager.Poll(FTimespan::Zero());
+					CHECK(CallbackTriggerCount == 2);
+					CHECK(TriggeredSocket == LocalSocket1);
+					CHECK(TriggeredStatus == ESocketIoRequestStatus::Ok);
+					CHECK(TriggeredFlags == EIOFlags::Read);
+
+					// 5. Read the waiting data from the socket.
+					int32 BytesRead = 0;
+					CHECK(FBSDSocketFuncs::Recv(LocalSocket1, ReadBuffer.GetData(), ReadBufferSize, BytesRead));
+					CHECK(BytesRead == 1);
+
+					// 6. See that socket no longer has a read event.
+					TriggeredFlags = EIOFlags::None;
+					IOManager.Poll(FTimespan::Zero());
+					CHECK(CallbackTriggerCount == 2);
+				}
+
+				SECTION("Wait on socket writeable")
+				{
+					SOCKET TriggeredSocket = INVALID_SOCKET;
+					ESocketIoRequestStatus TriggeredStatus = ESocketIoRequestStatus::Ok;
+					EIOFlags TriggeredFlags = EIOFlags::None;
+					uint32 CallbackTriggerCount = 0;
+
+					FIORequestBSDSocket Request;
+					Request.Socket = LocalSocket1;
+					Request.Flags = EIOFlags::Write;
+					Request.Callback = [&](SOCKET Socket, ESocketIoRequestStatus Status, EIOFlags SignaledFlags){
+						TriggeredSocket = Socket;
+						TriggeredStatus = Status;
+						TriggeredFlags = SignaledFlags;
+						++CallbackTriggerCount;
+					};
+
+					FIORequestHandle RequestHandle = IOManager.GetIOAccess().CreateSocketIORequest(MoveTemp(Request));
+					CHECK(RequestHandle.IsValid());
+
+					// 1. See that socket has a write event.
+					IOManager.Poll(FTimespan::Zero());
+					CHECK(CallbackTriggerCount == 1);
+					CHECK(TriggeredSocket == LocalSocket1);
+					CHECK(TriggeredStatus == ESocketIoRequestStatus::Ok);
+					CHECK(TriggeredFlags == EIOFlags::Write);
+
+					// 2. Write data to the socket until it returns an error indicating that it would block.
+					int32 BytesSent = 0;
+					int32 TotalBytesSent = 0;
+					while (FBSDSocketFuncs::Send(LocalSocket1, WriteBuffer.GetData(), WriteBufferSize, BytesSent));
+
+					// 3. See that socket no longer has a write event.
+					IOManager.Poll(FTimespan::Zero());
+					CHECK(CallbackTriggerCount == 1);
+
+					// 4. Read the data from the remote end of the connection.
+					int32 BytesRead = 0;
+					while (FBSDSocketFuncs::Recv(RemoteSocket1, ReadBuffer.GetData(), ReadBufferSize, BytesRead))
+					{
+						if (BytesRead == 0)
+						{
+							break;
+						}
+					}
+
+					// 5. See that socket now has a write event.
+					IOManager.Poll(FTimespan::FromMilliseconds(10));
+					CHECK(CallbackTriggerCount == 2);
+					CHECK(TriggeredSocket == LocalSocket1);
+					CHECK(TriggeredStatus == ESocketIoRequestStatus::Ok);
+					CHECK(TriggeredFlags == EIOFlags::Write);
+				}
+
+				SECTION("Receive both read and write events with one request")
+				{
+					SOCKET TriggeredSocket = INVALID_SOCKET;
+					ESocketIoRequestStatus TriggeredStatus = ESocketIoRequestStatus::Ok;
+					EIOFlags TriggeredFlags = EIOFlags::None;
+					uint32 CallbackTriggerCount = 0;
+
+					FIORequestBSDSocket Request;
+					Request.Socket = LocalSocket1;
+					Request.Flags = EIOFlags::Read | EIOFlags::Write;
+					Request.Callback = [&](SOCKET Socket, ESocketIoRequestStatus Status, EIOFlags SignaledFlags){
+						TriggeredSocket = Socket;
+						TriggeredStatus = Status;
+						TriggeredFlags = SignaledFlags;
+						++CallbackTriggerCount;
+					};
+
+					FIORequestHandle RequestHandle = IOManager.GetIOAccess().CreateSocketIORequest(MoveTemp(Request));
+					CHECK(RequestHandle.IsValid());
+
+					// 1. No events triggered when no data has arrived at the socket.
+					IOManager.Poll(FTimespan::Zero());
+					CHECK(CallbackTriggerCount == 1);
+					CHECK(TriggeredSocket == LocalSocket1);
+					CHECK(TriggeredStatus == ESocketIoRequestStatus::Ok);
+					CHECK(EnumHasAnyFlags(TriggeredFlags, EIOFlags::Write));
+					CHECK(!EnumHasAnyFlags(TriggeredFlags, EIOFlags::Read));
+
+					// 2. Send data to the socket.
+					int32 BytesSent = 0;
+					CHECK(FBSDSocketFuncs::Send(RemoteSocket1, WriteBuffer.GetData(), 1, BytesSent));
+					CHECK(BytesSent == 1);
+
+					// 3. See that socket has a read event.
+					IOManager.Poll(FTimespan::Zero());
+					CHECK(CallbackTriggerCount == 2);
+					CHECK(TriggeredSocket == LocalSocket1);
+					CHECK(TriggeredStatus == ESocketIoRequestStatus::Ok);
+					CHECK(EnumHasAnyFlags(TriggeredFlags, EIOFlags::Write));
+					CHECK(EnumHasAnyFlags(TriggeredFlags, EIOFlags::Read));
+				}
+
+				SECTION("Canceling a request stops signaling events")
+				{
+					SOCKET TriggeredSocket = INVALID_SOCKET;
+					ESocketIoRequestStatus TriggeredStatus = ESocketIoRequestStatus::Ok;
+					EIOFlags TriggeredFlags = EIOFlags::None;
+					uint32 CallbackTriggerCount = 0;
+
+					FIORequestBSDSocket Request;
+					Request.Socket = LocalSocket1;
+					Request.Flags = EIOFlags::Read;
+					Request.Callback = [&](SOCKET Socket, ESocketIoRequestStatus Status, EIOFlags SignaledFlags){
+						TriggeredSocket = Socket;
+						TriggeredStatus = Status;
+						TriggeredFlags = SignaledFlags;
+						++CallbackTriggerCount;
+					};
+
+					FIORequestHandle RequestHandle = IOManager.GetIOAccess().CreateSocketIORequest(MoveTemp(Request));
+					CHECK(RequestHandle.IsValid());
+
+					// 1. No events triggered when no data has arrived at the socket.
+					IOManager.Poll(FTimespan::Zero());
+					CHECK(CallbackTriggerCount == 0);
+
+					// 2. Send data to the socket.
+					int32 BytesSent = 0;
+					CHECK(FBSDSocketFuncs::Send(RemoteSocket1, WriteBuffer.GetData(), 1, BytesSent));
+					CHECK(BytesSent == 1);
+
+					// 3. See that socket has a read event.
+					IOManager.Poll(FTimespan::Zero());
+					CHECK(CallbackTriggerCount == 1);
+					CHECK(TriggeredSocket == LocalSocket1);
+					CHECK(TriggeredStatus == ESocketIoRequestStatus::Ok);
+					CHECK(TriggeredFlags == EIOFlags::Read);
+
+					// 4. Cancel the request.
+					IOManager.GetIOAccess().DestroyIORequest(RequestHandle);
+
+					// 5. See that event is not triggered.
+					IOManager.Poll(FTimespan::Zero());
+					CHECK(CallbackTriggerCount == 1);
+				}
+
+				SECTION("TCP socket remote end closing triggers writeable wakeup")
+				{
+					SOCKET TriggeredSocket = INVALID_SOCKET;
+					ESocketIoRequestStatus TriggeredStatus = ESocketIoRequestStatus::Ok;
+					EIOFlags TriggeredFlags = EIOFlags::None;
+					uint32 CallbackTriggerCount = 0;
+
+					FIORequestBSDSocket Request;
+					Request.Socket = LocalSocket1;
+					Request.Flags = EIOFlags::Write;
+					Request.Callback = [&](SOCKET Socket, ESocketIoRequestStatus Status, EIOFlags SignaledFlags){
+						TriggeredSocket = Socket;
+						TriggeredStatus = Status;
+						TriggeredFlags = SignaledFlags;
+						++CallbackTriggerCount;
+					};
+
+					FIORequestHandle RequestHandle = IOManager.GetIOAccess().CreateSocketIORequest(MoveTemp(Request));
+					CHECK(RequestHandle.IsValid());
+
+					// 1. See that socket has a write event.
+					IOManager.Poll(FTimespan::Zero());
+					CHECK(CallbackTriggerCount == 1);
+					CHECK(TriggeredSocket == LocalSocket1);
+					CHECK(TriggeredStatus == ESocketIoRequestStatus::Ok);
+					CHECK(TriggeredFlags == EIOFlags::Write);
+
+					// 2. Write data to the socket until it returns an error indicating that it would block.
+					int32 BytesSent = 0;
+					while (FBSDSocketFuncs::Send(LocalSocket1, WriteBuffer.GetData(), WriteBufferSize, BytesSent));
+
+					// 3. See that socket no longer has a write event.
+					IOManager.Poll(FTimespan::Zero());
+					CHECK(CallbackTriggerCount == 1);
+
+					// 4. Close the remote end of the connection.
+					FBSDSocketFuncs::Close(RemoteSocket1);
+
+					// 5. See that socket now has a write event.
+					IOManager.Poll(FTimespan::Zero());
+					CHECK(CallbackTriggerCount == 2);
+					CHECK(TriggeredSocket == LocalSocket1);
+					CHECK(TriggeredStatus == ESocketIoRequestStatus::Ok);
+					CHECK(TriggeredFlags == EIOFlags::Write);
+				}
+
+				SECTION("More than one request for the same socket")
+				{
+					SOCKET TriggeredSocket = INVALID_SOCKET;
+					ESocketIoRequestStatus TriggeredStatus = ESocketIoRequestStatus::Ok;
+					EIOFlags TriggeredFlags = EIOFlags::None;
+					uint32 CallbackTriggerCount1 = 0;
+					uint32 CallbackTriggerCount2 = 0;
+
+					FIORequestBSDSocket Request1;
+					Request1.Socket = LocalSocket1;
+					Request1.Flags = EIOFlags::Read;
+					Request1.Callback = [&](SOCKET Socket, ESocketIoRequestStatus Status, EIOFlags SignaledFlags){
+						TriggeredSocket = Socket;
+						TriggeredStatus = Status;
+						TriggeredFlags = SignaledFlags;
+						++CallbackTriggerCount1;
+					};
+
+					FIORequestHandle RequestHandle1 = IOManager.GetIOAccess().CreateSocketIORequest(MoveTemp(Request1));
+					CHECK(RequestHandle1.IsValid());
+
+					FIORequestBSDSocket Request2;
+					Request2.Socket = LocalSocket1;
+					Request2.Flags = EIOFlags::Write;
+					Request2.Callback = [&](SOCKET Socket, ESocketIoRequestStatus Status, EIOFlags SignaledFlags){
+						TriggeredSocket = Socket;
+						TriggeredStatus = Status;
+						TriggeredFlags = SignaledFlags;
+						++CallbackTriggerCount2;
+					};
+
+					FIORequestHandle RequestHandle2 = IOManager.GetIOAccess().CreateSocketIORequest(MoveTemp(Request2));
+					CHECK(RequestHandle2.IsValid());
+
+					IOManager.Poll(FTimespan::Zero());
+					CHECK(CallbackTriggerCount1 == 0);
+					CHECK(CallbackTriggerCount2 == 1);
+					CHECK(TriggeredStatus == ESocketIoRequestStatus::Invalid);
+				}
+			}
+
+			SECTION("Multiple waiting TCP socket tests")
+			{
+				// Create remote socket.
+				SOCKET RemoteSocket2 = FBSDSocketFuncs::Open(AF_INET, SOCK_STREAM, IPPROTO_TCP);
+				REQUIRE(RemoteSocket2 != INVALID_SOCKET);
+				REQUIRE(FBSDSocketFuncs::SetSendBufferSize(RemoteSocket2, SocketSendBufferSize));
+				REQUIRE(FBSDSocketFuncs::SetRecvBufferSize(RemoteSocket2, SocketReceiveBufferSize));
+				REQUIRE(FBSDSocketFuncs::SetNonBlocking(RemoteSocket2, true));
+				REQUIRE(FBSDSocketFuncs::Connect(RemoteSocket2, HostAddr));
+
+				// Create local socket.
+				ESocketErrors ErrorCode2 = SE_NO_ERROR;
+				SOCKET LocalSocket2 = INVALID_SOCKET;
+				do
+				{
+					LocalSocket2 = FBSDSocketFuncs::Accept(ListenSocket);
+					ErrorCode2 = SocketSubsystem->GetLastErrorCode();
+				}
+				while (ErrorCode2 == SE_EWOULDBLOCK);
+				REQUIRE(LocalSocket2 != INVALID_SOCKET);
+
+				SECTION("Triggering LocalSocket1 does not trigger LocalSocket2")
+				{
+					SOCKET TriggeredSocket1 = INVALID_SOCKET;
+					ESocketIoRequestStatus TriggeredStatus1 = ESocketIoRequestStatus::Ok;
+					EIOFlags TriggeredFlags1 = EIOFlags::None;
+					uint32 CallbackTriggerCount1 = 0;
+					FIORequestBSDSocket Request1;
+					Request1.Socket = LocalSocket1;
+					Request1.Flags = EIOFlags::Read;
+					Request1.Callback = [&](SOCKET Socket, ESocketIoRequestStatus Status, EIOFlags SignaledFlags){
+						TriggeredSocket1 = Socket;
+						TriggeredStatus1 = Status;
+						TriggeredFlags1 = SignaledFlags;
+						++CallbackTriggerCount1;
+					};
+					FIORequestHandle RequestHandle1 = IOManager.GetIOAccess().CreateSocketIORequest(MoveTemp(Request1));
+					CHECK(RequestHandle1.IsValid());
+
+					SOCKET TriggeredSocket2 = INVALID_SOCKET;
+					ESocketIoRequestStatus TriggeredStatus2 = ESocketIoRequestStatus::Ok;
+					EIOFlags TriggeredFlags2 = EIOFlags::None;
+					uint32 CallbackTriggerCount2 = 0;
+					FIORequestBSDSocket Request2;
+					Request2.Socket = LocalSocket2;
+					Request2.Flags = EIOFlags::Read;
+					Request2.Callback = [&](SOCKET Socket, ESocketIoRequestStatus Status, EIOFlags SignaledFlags){
+						TriggeredSocket2 = Socket;
+						TriggeredStatus2 = Status;
+						TriggeredFlags2 = SignaledFlags;
+						++CallbackTriggerCount2;
+					};
+					FIORequestHandle RequestHandle2 = IOManager.GetIOAccess().CreateSocketIORequest(MoveTemp(Request2));
+					CHECK(RequestHandle2.IsValid());
+
+					// 1. No events triggered when no data has arrived at the socket.
+					IOManager.Poll(FTimespan::Zero());
+					CHECK(CallbackTriggerCount1 == 0);
+					CHECK(CallbackTriggerCount2 == 0);
+
+					// 2. Send data to the socket.
+					int32 BytesSent = 0;
+					CHECK(FBSDSocketFuncs::Send(RemoteSocket1, WriteBuffer.GetData(), 1, BytesSent));
+					CHECK(BytesSent == 1);
+
+					// 3. See that socket has a read event.
+					TriggeredFlags1 = EIOFlags::None;
+					IOManager.Poll(FTimespan::Zero());
+					CHECK(CallbackTriggerCount1 == 1);
+					CHECK(CallbackTriggerCount2 == 0);
+					CHECK(TriggeredSocket1 == LocalSocket1);
+					CHECK(TriggeredStatus1 == ESocketIoRequestStatus::Ok);
+					CHECK(TriggeredFlags1 == EIOFlags::Read);
+
+					// 4. See that socket still has a read event.
+					IOManager.Poll(FTimespan::Zero());
+					CHECK(CallbackTriggerCount1 == 2);
+					CHECK(CallbackTriggerCount2 == 0);
+					CHECK(TriggeredSocket1 == LocalSocket1);
+					CHECK(TriggeredStatus1 == ESocketIoRequestStatus::Ok);
+					CHECK(TriggeredFlags1 == EIOFlags::Read);
+
+					// 5. Read the waiting data from the socket.
+					int32 BytesRead = 0;
+					CHECK(FBSDSocketFuncs::Recv(LocalSocket1, ReadBuffer.GetData(), ReadBufferSize, BytesRead));
+					CHECK(BytesRead == 1);
+
+					// 6. See that socket no longer has a read event.
+					TriggeredFlags1 = EIOFlags::None;
+					IOManager.Poll(FTimespan::Zero());
+					CHECK(CallbackTriggerCount1 == 2);
+					CHECK(CallbackTriggerCount2 == 0);
+				}
+
+				FBSDSocketFuncs::Close(RemoteSocket2);
+				FBSDSocketFuncs::Close(LocalSocket2);
+			}
+
+			SECTION("Waiting on a TCP and UDP socket")
+			{
+				// Create local socket.
+				SOCKET LocalSocket2 = FBSDSocketFuncs::Open(AF_INET, SOCK_DGRAM, IPPROTO_UDP);
+				REQUIRE(LocalSocket2 != INVALID_SOCKET);
+
+				REQUIRE(FBSDSocketFuncs::SetNonBlocking(LocalSocket2, true));
+				REQUIRE(FBSDSocketFuncs::BindLoopback(LocalSocket2));
+				sockaddr_in HostAddr2;
+				FMemory::Memzero(&HostAddr2, sizeof(HostAddr2));
+				REQUIRE(FBSDSocketFuncs::GetAddress(LocalSocket2, HostAddr2));
+
+				// Create remote socket.
+				SOCKET RemoteSocket2 = FBSDSocketFuncs::Open(AF_INET, SOCK_DGRAM, IPPROTO_UDP);
+				REQUIRE(RemoteSocket2 != INVALID_SOCKET);
+				REQUIRE(FBSDSocketFuncs::SetSendBufferSize(RemoteSocket2, SocketSendBufferSize));
+				REQUIRE(FBSDSocketFuncs::SetRecvBufferSize(RemoteSocket2, SocketReceiveBufferSize));
+				REQUIRE(FBSDSocketFuncs::SetNonBlocking(RemoteSocket2, true));
+
+				SECTION("Recevie events for TCP and UDP socket types.")
+				{
+					SOCKET TriggeredSocket1 = INVALID_SOCKET;
+					ESocketIoRequestStatus TriggeredStatus1 = ESocketIoRequestStatus::Ok;
+					EIOFlags TriggeredFlags1 = EIOFlags::None;
+					uint32 CallbackTriggerCount1 = 0;
+					FIORequestBSDSocket Request1;
+					Request1.Socket = LocalSocket1;
+					Request1.Flags = EIOFlags::Read;
+					Request1.Callback = [&](SOCKET Socket, ESocketIoRequestStatus Status, EIOFlags SignaledFlags){
+						TriggeredSocket1 = Socket;
+						TriggeredStatus1 = Status;
+						TriggeredFlags1 = SignaledFlags;
+						++CallbackTriggerCount1;
+					};
+					FIORequestHandle RequestHandle1 = IOManager.GetIOAccess().CreateSocketIORequest(MoveTemp(Request1));
+					CHECK(RequestHandle1.IsValid());
+
+					SOCKET TriggeredSocket2 = INVALID_SOCKET;
+					ESocketIoRequestStatus TriggeredStatus2 = ESocketIoRequestStatus::Ok;
+					EIOFlags TriggeredFlags2 = EIOFlags::None;
+					uint32 CallbackTriggerCount2 = 0;
+					FIORequestBSDSocket Request2;
+					Request2.Socket = LocalSocket2;
+					Request2.Flags = EIOFlags::Read;
+					Request2.Callback = [&](SOCKET Socket, ESocketIoRequestStatus Status, EIOFlags SignaledFlags){
+						TriggeredSocket2 = Socket;
+						TriggeredStatus2 = Status;
+						TriggeredFlags2 = SignaledFlags;
+						++CallbackTriggerCount2;
+					};
+					FIORequestHandle RequestHandle2 = IOManager.GetIOAccess().CreateSocketIORequest(MoveTemp(Request2));
+					CHECK(RequestHandle2.IsValid());
+
+					// 1. No events triggered when no data has arrived at the socket.
+					IOManager.Poll(FTimespan::Zero());
+					CHECK(CallbackTriggerCount1 == 0);
+					CHECK(CallbackTriggerCount2 == 0);
+
+					// 2. Send data to the socket.
+					int32 BytesSent = 0;
+					CHECK(FBSDSocketFuncs::Send(RemoteSocket1, WriteBuffer.GetData(), 1, BytesSent));
+					CHECK(BytesSent == 1);
+
+					// 3. See that socket has a read event.
+					IOManager.Poll(FTimespan::Zero());
+					CHECK(CallbackTriggerCount1 == 1);
+					CHECK(CallbackTriggerCount2 == 0);
+					CHECK(TriggeredSocket1 == LocalSocket1);
+					CHECK(TriggeredStatus1 == ESocketIoRequestStatus::Ok);
+					CHECK(TriggeredFlags1 == EIOFlags::Read);
+
+					// 4. Send data to the UDP socket.
+					BytesSent = 0;
+					CHECK(FBSDSocketFuncs::SendTo(RemoteSocket2, WriteBuffer.GetData(), 1, BytesSent, HostAddr2));
+					CHECK(BytesSent == 1);
+
+					// 5. See that both sockets have a read event.
+					IOManager.Poll(FTimespan::Zero());
+					CHECK(CallbackTriggerCount1 == 2);
+					CHECK(CallbackTriggerCount2 == 1);
+					CHECK(TriggeredSocket1 == LocalSocket1);
+					CHECK(TriggeredStatus1 == ESocketIoRequestStatus::Ok);
+					CHECK(TriggeredFlags1 == EIOFlags::Read);
+					CHECK(TriggeredSocket2 == LocalSocket2);
+					CHECK(TriggeredStatus2 == ESocketIoRequestStatus::Ok);
+					CHECK(TriggeredFlags2 == EIOFlags::Read);
+
+					// 5. Read the waiting data from the TCP socket.
+					int32 BytesRead = 0;
+					CHECK(FBSDSocketFuncs::Recv(LocalSocket1, ReadBuffer.GetData(), ReadBufferSize, BytesRead));
+					CHECK(BytesRead == 1);
+
+					// 6. See that the UDP socket still has a read event.
+					IOManager.Poll(FTimespan::Zero());
+					CHECK(CallbackTriggerCount1 == 2);
+					CHECK(CallbackTriggerCount2 == 2);
+					CHECK(TriggeredSocket2 == LocalSocket2);
+					CHECK(TriggeredStatus2 == ESocketIoRequestStatus::Ok);
+					CHECK(TriggeredFlags2 == EIOFlags::Read);
+				}
+
+				FBSDSocketFuncs::Close(RemoteSocket2);
+				FBSDSocketFuncs::Close(LocalSocket2);
+			}
+
+			FBSDSocketFuncs::Close(ListenSocket);
+			FBSDSocketFuncs::Close(RemoteSocket1);
+			FBSDSocketFuncs::Close(LocalSocket1);
+		}
+
+		SECTION("Wait on too many sockets")
+		{
+			struct FCallbackData
+			{
+				SOCKET Socket = INVALID_SOCKET;
+				ESocketIoRequestStatus Status = ESocketIoRequestStatus::Ok;
+				EIOFlags SignaledFlags = EIOFlags::None;
+			};
+
+			TArray<SOCKET> Sockets;
+			TArray<FCallbackData> CallbackData;
+
+			SOCKET BaseSocket = FBSDSocketFuncs::Open(AF_INET, SOCK_DGRAM, IPPROTO_UDP);
+			REQUIRE(BaseSocket != INVALID_SOCKET);
+			REQUIRE(FBSDSocketFuncs::SetNonBlocking(BaseSocket, true));
+			REQUIRE(FBSDSocketFuncs::SetReusePort(BaseSocket, true));
+			REQUIRE(FBSDSocketFuncs::BindLoopback(BaseSocket));
+
+			sockaddr_in BaseAddr;
+			FMemory::Memzero(&BaseAddr, sizeof(BaseAddr));
+			REQUIRE(FBSDSocketFuncs::GetAddress(BaseSocket, BaseAddr));
+
+			const int32 MaxTestSockets = 10000;
+			for (int32 index = 0; index < MaxTestSockets; ++index)
+			{
+				SOCKET Socket = FBSDSocketFuncs::Open(AF_INET, SOCK_DGRAM, IPPROTO_UDP);
+				REQUIRE(Socket != INVALID_SOCKET);
+				REQUIRE(FBSDSocketFuncs::SetNonBlocking(Socket, true));
+				REQUIRE(FBSDSocketFuncs::SetReusePort(Socket, true));
+				REQUIRE(FBSDSocketFuncs::Bind(Socket, BaseAddr));
+				Sockets.Add(Socket);
+
+				FIORequestBSDSocket Request;
+				Request.Socket = Socket;
+				Request.Flags = EIOFlags::Read;
+				Request.Callback = [&CallbackData](SOCKET Socket, ESocketIoRequestStatus Status, EIOFlags SignaledFlags){
+					CallbackData.Add({Socket, Status, SignaledFlags});
+				};
+				FIORequestHandle RequestHandle = IOManager.GetIOAccess().CreateSocketIORequest(MoveTemp(Request));
+				CHECK(RequestHandle.IsValid());
+
+				IOManager.Poll(FTimespan::Zero());
+
+				if (!CallbackData.IsEmpty())
+				{
+					CHECK(CallbackData.Num() == 1);
+					CHECK(CallbackData[0].Socket == Socket);
+					CHECK(CallbackData[0].Status == ESocketIoRequestStatus::NoResources);
+					CHECK(CallbackData[0].SignaledFlags == EIOFlags::None);
+					break;
+				}
+			}
+
+			// Check that a failure occurred.
+			CHECK(!CallbackData.IsEmpty());
+
+			for (SOCKET Socket : Sockets)
+			{
+				FBSDSocketFuncs::Close(Socket);
+			}
+
+			FBSDSocketFuncs::Close(BaseSocket);
+		}
+
+		IOManager.Shutdown();
+	}
+
+} // UE::EventLoop
+
+#endif PLATFORM_HAS_BSD_SOCKETS
diff --git a/Engine/Source/Runtime/Online/Experimental/EventLoopTests/Tests/EventLoop/EventLoopLLTModule.cpp b/Engine/Source/Runtime/Online/Experimental/EventLoopTests/Tests/EventLoop/EventLoopLLTModule.cpp
new file mode 100644
--- /dev/null
+++ b/Engine/Source/Runtime/Online/Experimental/EventLoopTests/Tests/EventLoop/EventLoopLLTModule.cpp
@@ -0,0 +1,52 @@
+// Copyright Epic Games, Inc. All Rights Reserved.
+
+#include "CoreTypes.h"
+
+#include "Containers/Array.h"
+#include "Containers/UnrealString.h"
+#include "Modules/ModuleManager.h"
+#include "TestRunner.h"
+
+namespace UE::EventLoop
+{
+
+TArray<FString> GetRequiredModules()
+{
+	static TArray<FString> Modules({
+		FString(TEXT("EventLoop")),
+		FString(TEXT("Sockets"))
+	});
+
+	return Modules;
+}
+
+void InitializeEventLoopTests()
+{
+	for (const FString& ModuleName : GetRequiredModules())
+	{
+		FModuleManager::LoadModulePtr<IModuleInterface>(*ModuleName);
+	}
+}
+
+void CleanupEventLoopTests()
+{
+	for (const FString& ModuleName : GetRequiredModules())
+	{
+		if (IModuleInterface* Module = FModuleManager::Get().GetModule(*ModuleName))
+		{
+			Module->ShutdownModule();
+		}
+	}
+}
+
+class EventLoopTestGlobalSetup
+{
+public:
+	EventLoopTestGlobalSetup()
+	{
+		FTestDelegates::GetGlobalSetup().BindStatic(&InitializeEventLoopTests);
+		FTestDelegates::GetGlobalTeardown().BindStatic(&CleanupEventLoopTests);
+	}
+} GEventLoopTestGlobalSetup;
+
+} // UE::EventLoop
diff --git a/Engine/Source/Runtime/Online/HTTP/Private/Curl/CurlSocketEventLoopHttpThread.cpp b/Engine/Source/Runtime/Online/HTTP/Private/Curl/CurlSocketEventLoopHttpThread.cpp
--- a/Engine/Source/Runtime/Online/HTTP/Private/Curl/CurlSocketEventLoopHttpThread.cpp
+++ b/Engine/Source/Runtime/Online/HTTP/Private/Curl/CurlSocketEventLoopHttpThread.cpp
@@ -145,7 +145,7 @@ int FCurlSocketEventLoopHttpThread::HandleCurlSocketCallback(CURL* CurlE, curl_s
 
 	FCurlSocketData* CurlSocketData = reinterpret_cast<FCurlSocketData*>(SocketData);
 
-	UE_LOG(LogHttp, VeryVerbose, TEXT("[FCurlSocketEventLoopHttpThread::HandleCurlSocketCallback] Socket: %p, EventFlags: %d, SocketData: %p"), Socket, EventFlags, SocketData);
+	UE_LOG(LogHttp, VeryVerbose, TEXT("[FCurlSocketEventLoopHttpThread::HandleCurlSocketCallback] Socket: %p, EventFlags: 0x%08X, SocketData: %p"), Socket, EventFlags, SocketData);
 
 	switch (EventFlags)
 	{
@@ -168,9 +168,9 @@ int FCurlSocketEventLoopHttpThread::HandleCurlSocketCallback(CURL* CurlE, curl_s
 		UE::EventLoop::FIORequestBSDSocket Request;
 		Request.Socket = Socket;
 		Request.Flags = TranslateCurlPollEventFlags(EventFlags);
-		Request.Callback = [this, Socket](UE::EventLoop::EIOFlags Status)
+		Request.Callback = [this](SOCKET Socket, UE::EventLoop::ESocketIoRequestStatus Status, UE::EventLoop::EIOFlags SignaledFlags)
 		{
-			ProcessCurlSocketEvent(Socket, Status);
+			ProcessCurlSocketEvent(Socket, Status, SignaledFlags);
 		};
 
 		CurlSocketData->IORequestHandle = EventLoop->GetIOAccess().CreateSocketIORequest(MoveTemp(Request));
@@ -206,16 +206,10 @@ int FCurlSocketEventLoopHttpThread::HandleCurlTimerCallback(CURLM* CurlM, long T
 	// Start by canceling any previous timer.
 	EventLoop->ClearTimer(RequestTimeoutTimer);
 
-	if (TimeoutMS == 0)
-	{
-		// Run immediately.
-		EventLoop->PostAsyncTask([this]()
-		{
-			ProcessCurlSocketActions(CURL_SOCKET_TIMEOUT, 0);
-		});
-	}
-	else if (TimeoutMS > 0)
+	if (TimeoutMS >= 0)
 	{
+		UE_LOG(LogHttp, VeryVerbose, TEXT("[FCurlSocketEventLoopHttpThread::HandleCurlTimerCallback] Set timer for: %d ms"), TimeoutMS);
+
 		RequestTimeoutTimer = EventLoop->SetTimer([this]()
 		{
 			ProcessCurlSocketActions(CURL_SOCKET_TIMEOUT, 0);
@@ -230,14 +224,24 @@ void FCurlSocketEventLoopHttpThread::ProcessCurlSocketActions(curl_socket_t Sock
 {
 	QUICK_SCOPE_CYCLE_COUNTER(STAT_FCurlSocketEventLoopHttpThread_ProcessCurlSocketActions);
 
+	UE_LOG(LogHttp, VeryVerbose, TEXT("[FCurlSocketEventLoopHttpThread::ProcessCurlSocketActions] Socket: %p CurlEvents: 0x%08X"), Socket, EventFlags);
+
 	int running_handles;
 	curl_multi_socket_action(FCurlHttpManager::GMultiHandle, Socket, EventFlags, &running_handles);
 	ProcessCurlRequests();
 }
 
-void FCurlSocketEventLoopHttpThread::ProcessCurlSocketEvent(curl_socket_t Socket, UE::EventLoop::EIOFlags Flags)
+void FCurlSocketEventLoopHttpThread::ProcessCurlSocketEvent(curl_socket_t Socket, UE::EventLoop::ESocketIoRequestStatus Status, UE::EventLoop::EIOFlags Flags)
 {
-	ProcessCurlSocketActions(Socket, TranslateCurlSocketActionFlags(Flags));
+	if (Status != UE::EventLoop::ESocketIoRequestStatus::Ok)
+	{
+		UE_LOG(LogHttp, Warning, TEXT("Socket event request failed. Socket %p, Status: %s"), Socket, *LexToString(Status));
+		ProcessCurlSocketActions(CURL_SOCKET_TIMEOUT, 0);
+	}
+	else
+	{
+		ProcessCurlSocketActions(Socket, TranslateCurlSocketActionFlags(Flags));
+	}
 }
 
 void FCurlSocketEventLoopHttpThread::ProcessCurlRequests()
diff --git a/Engine/Source/Runtime/Online/HTTP/Private/Curl/CurlSocketEventLoopHttpThread.h b/Engine/Source/Runtime/Online/HTTP/Private/Curl/CurlSocketEventLoopHttpThread.h
--- a/Engine/Source/Runtime/Online/HTTP/Private/Curl/CurlSocketEventLoopHttpThread.h
+++ b/Engine/Source/Runtime/Online/HTTP/Private/Curl/CurlSocketEventLoopHttpThread.h
@@ -62,7 +62,7 @@ protected:
 	int HandleCurlTimerCallback(CURLM* CurlM, long TimeoutMS);
 
 	void ProcessCurlSocketActions(curl_socket_t Socket, int EventFlags);
-	void ProcessCurlSocketEvent(curl_socket_t Socket, UE::EventLoop::EIOFlags Flags);
+	void ProcessCurlSocketEvent(curl_socket_t Socket, UE::EventLoop::ESocketIoRequestStatus Status, UE::EventLoop::EIOFlags Flags);
 	void ProcessCurlRequests();
 
 	struct FCurlSocketData
