[EventLoop]
* BsdSocketTypesPrivate - Fixed non-unity compile error

diff --git a/Engine/Source/Runtime/Online/Experimental/EventLoop/Private/EventLoop/BSDSocket/BSDSocketTypesPrivate.h b/Engine/Source/Runtime/Online/Experimental/EventLoop/Private/EventLoop/BSDSocket/BSDSocketTypesPrivate.h
--- a/Engine/Source/Runtime/Online/Experimental/EventLoop/Private/EventLoop/BSDSocket/BSDSocketTypesPrivate.h
+++ b/Engine/Source/Runtime/Online/Experimental/EventLoop/Private/EventLoop/BSDSocket/BSDSocketTypesPrivate.h
@@ -2,7 +2,7 @@
 
 #pragma once
 
-#include "CoreTypes.h"
+#include "Containers/UnrealString.h"
 #include "Containers/StringConv.h"
 #include "EventLoop/BSDSocket/BSDSocketTypes.h"
 
