Add support for comments in the default http user agent string.

diff --git a/Engine/Source/Runtime/Analytics/AnalyticsET/Private/IAnalyticsProviderET.cpp b/Engine/Source/Runtime/Analytics/AnalyticsET/Private/IAnalyticsProviderET.cpp
--- a/Engine/Source/Runtime/Analytics/AnalyticsET/Private/IAnalyticsProviderET.cpp
+++ b/Engine/Source/Runtime/Analytics/AnalyticsET/Private/IAnalyticsProviderET.cpp
@@ -31,8 +31,79 @@ namespace AnalyticsProviderETCvars
 		PreventMultipleFlushesInOneFrame,
 		TEXT("When true, prevents more than one AnalyticsProviderET instance from flushing in the same frame, allowing the flush and HTTP cost to be amortized.")
 	);
+
+	TAutoConsoleVariable<bool> CVarDefaultUserAgentCommentsEnabled(
+		TEXT("AnalyticsET.UserAgentCommentsEnabled"),
+		true,
+		TEXT("Whether comments are supported in the analytics user agent string"),
+		ECVF_SaveForNextBoot
+	);
 }
 
+// Want to avoid putting the project name into the User-Agent, because for some apps (like the editor), the project name is private info.
+// The analytics User-Agent uses the default User-Agent, but with project name removed.
+class FAnalyticsUserAgentCache
+{
+public:
+	FAnalyticsUserAgentCache()
+		: CachedUserAgent()
+		, CachedAgentVersion(0)
+	{
+	}
+
+	FString GetUserAgent()
+	{
+		if (CachedUserAgent.IsEmpty() || CachedAgentVersion != FPlatformHttp::GetDefaultUserAgentVersion())
+		{
+			UpdateUserAgent();
+		}
+
+		return CachedUserAgent;
+	}
+
+private:
+	void UpdateUserAgent()
+	{
+		static TSet<FString> AllowedProjectComments(GetAllowedProjectComments());
+		static TSet<FString> AllowedPlatformComments(GetAllowedPlatformComments());
+
+		FDefaultUserAgentBuilder Builder = FPlatformHttp::GetDefaultUserAgentBuilder();
+		Builder.SetProjectName(TEXT("PROJECTNAME"));
+		CachedUserAgent = Builder.BuildUserAgentString(&AllowedProjectComments, &AllowedPlatformComments);
+		CachedAgentVersion = Builder.GetAgentVersion();
+	}
+
+	static TSet<FString> GetAllowedProjectComments()
+	{
+		TArray<FString> AllowedProjectComments;
+		if (AnalyticsProviderETCvars::CVarDefaultUserAgentCommentsEnabled.GetValueOnAnyThread())
+		{
+			GConfig->GetArray(TEXT("Analytics"), TEXT("AllowedUserAgentProjectComments"), AllowedProjectComments, GEngineIni);
+		}
+		return TSet<FString>(MoveTemp(AllowedProjectComments));
+	}
+
+	static TSet<FString> GetAllowedPlatformComments()
+	{
+		TArray<FString> AllowedPlatformComments;
+		if (AnalyticsProviderETCvars::CVarDefaultUserAgentCommentsEnabled.GetValueOnAnyThread())
+		{
+			GConfig->GetArray(TEXT("Analytics"), TEXT("AllowedUserAgentPlatformComments"), AllowedPlatformComments, GEngineIni);
+		}
+		return TSet<FString>(MoveTemp(AllowedPlatformComments));
+	}
+
+	static TSet<FString> ParseCommentSet(const FString& CommentBlob)
+	{
+		TArray<FString> Comments;
+		CommentBlob.ParseIntoArray(Comments, TEXT(";"));
+		return TSet<FString>(MoveTemp(Comments));
+	}
+
+	FString CachedUserAgent;
+	uint32 CachedAgentVersion;
+};
+
 /**
  * Implementation of analytics for Epic Telemetry.
  * Supports caching events and flushing them periodically (currently hardcoded limits).
@@ -132,6 +203,8 @@ private:
 	FHttpRetrySystem::FRetryDomainsPtr RetryServers;
 	/** Http headers to add to requests */
 	TMap<FString, FString> HttpHeaders;
+
+	FAnalyticsUserAgentCache UserAgentCache;
 };
 
 TSharedPtr<IAnalyticsProviderET> FAnalyticsET::CreateAnalyticsProvider(const Config& ConfigValues) const
@@ -402,12 +475,8 @@ void FAnalyticsProviderET::FlushEventsOnce()
 			TSharedRef<IHttpRequest, ESPMode::ThreadSafe> HttpRequest = CreateRequest();
 			HttpRequest->SetHeader(TEXT("Content-Type"), TEXT("application/json; charset=utf-8"));
 			// Want to avoid putting the project name into the User-Agent, because for some apps (like the editor), the project name is private info.
-			// Code Pulled from FGenericPlatformHttp::GetDefaultUserAgent, but with project name pulled out
-			HttpRequest->SetHeader(TEXT("User-Agent"), FString::Printf(TEXT("%s/%s %s/%s"),
-				TEXT("PROJECTNAME"),
-				*FPlatformHttp::EscapeUserAgentString(FApp::GetBuildVersion()),
-				*FPlatformHttp::EscapeUserAgentString(FString(FPlatformProperties::IniPlatformName())),
-				*FPlatformHttp::EscapeUserAgentString(FPlatformMisc::GetOSVersion())));
+			// The analytics User-Agent uses the default User-Agent, but with project name removed.
+			HttpRequest->SetHeader(TEXT("User-Agent"), UserAgentCache.GetUserAgent());
 			HttpRequest->SetURL(Config.APIServerET / URLPath);
 			HttpRequest->SetVerb(TEXT("POST"));
 			HttpRequest->SetContent(MoveTemp(Payload));
diff --git a/Engine/Source/Runtime/Online/HTTP/Private/EventLoopHttpThread.cpp b/Engine/Source/Runtime/Online/HTTP/Private/EventLoopHttpThread.cpp
--- a/Engine/Source/Runtime/Online/HTTP/Private/EventLoopHttpThread.cpp
+++ b/Engine/Source/Runtime/Online/HTTP/Private/EventLoopHttpThread.cpp
@@ -5,10 +5,12 @@
 #include "EventLoop/EventLoopIOManagerNull.h"
 #include "IHttpThreadedRequest.h"
 #include "Http.h"
+#include "PlatformHttp.h"
 #include "Stats/Stats.h"
 
 FEventLoopHttpThread::FEventLoopHttpThread()
 {
+	FPlatformHttp::AddDefaultUserAgentProjectComment(TEXT("http-eventloop"));
 }
 
 FEventLoopHttpThread::~FEventLoopHttpThread()
diff --git a/Engine/Source/Runtime/Online/HTTP/Private/GenericPlatform/GenericPlatformHttp.cpp b/Engine/Source/Runtime/Online/HTTP/Private/GenericPlatform/GenericPlatformHttp.cpp
--- a/Engine/Source/Runtime/Online/HTTP/Private/GenericPlatform/GenericPlatformHttp.cpp
+++ b/Engine/Source/Runtime/Online/HTTP/Private/GenericPlatform/GenericPlatformHttp.cpp
@@ -1,12 +1,25 @@
 // Copyright Epic Games, Inc. All Rights Reserved.
 
 #include "GenericPlatform/GenericPlatformHttp.h"
+
+#include "Algo/Transform.h"
+#include "Containers/Array.h"
 #include "GenericPlatform/HttpRequestImpl.h"
+#include "HAL/CriticalSection.h"
+#include "HAL/IConsoleManager.h"
 #include "Misc/Paths.h"
 #include "Misc/EngineVersion.h"
 #include "Misc/App.h"
+#include "Misc/ScopeRWLock.h"
 #include "String/BytesToHex.h"
 
+TAutoConsoleVariable<bool> CVarDefaultUserAgentCommentsEnabled(
+	TEXT("http.DefaultUserAgentCommentsEnabled"),
+	true,
+	TEXT("Whether comments are supported in the defualt user agent string"),
+	ECVF_SaveForNextBoot
+);
+
 namespace
 {
 	/**
@@ -112,8 +125,170 @@ public:
 	virtual EHttpRequestDelegateThreadPolicy GetDelegateThreadPolicy() const override { return EHttpRequestDelegateThreadPolicy::CompleteOnGameThread; }
 };
 
+FDefaultUserAgentBuilder::FDefaultUserAgentBuilder()
+	: ProjectName(FApp::GetProjectName())
+	, ProjectVersion(FApp::GetBuildVersion())
+	, ProjectComments()
+	, PlatformName(FString(FPlatformProperties::IniPlatformName()))
+	, PlatformVersion(FPlatformMisc::GetOSVersion())
+	, PlatformComments()
+	, AgentVersion(1)
+{
+}
+
+FString FDefaultUserAgentBuilder::BuildUserAgentString(const TSet<FString>* AllowedProjectCommentsFilter, const TSet<FString>* AllowedPlatformCommentsFilter) const
+{
+	// strip/escape slashes and whitespace from components
+	return FString::Printf(TEXT("%s/%s%s %s/%s%s"),
+		*FGenericPlatformHttp::EscapeUserAgentString(ProjectName),
+		*FGenericPlatformHttp::EscapeUserAgentString(ProjectVersion),
+		*BuildCommentString(ProjectComments, AllowedProjectCommentsFilter),
+		*FGenericPlatformHttp::EscapeUserAgentString(PlatformName),
+		*FGenericPlatformHttp::EscapeUserAgentString(PlatformVersion),
+		*BuildCommentString(PlatformComments, AllowedPlatformCommentsFilter));
+}
+
+void FDefaultUserAgentBuilder::SetProjectName(const FString& InProjectName)
+{
+	ProjectName = InProjectName;
+	++AgentVersion;
+}
+
+void FDefaultUserAgentBuilder::SetProjectVersion(const FString& InProjectVersion)
+{
+	ProjectVersion = InProjectVersion;
+	++AgentVersion;
+}
+
+void FDefaultUserAgentBuilder::AddProjectComment(const FString& InComment)
+{
+	ProjectComments.AddUnique(InComment);
+	++AgentVersion;
+}
+
+void FDefaultUserAgentBuilder::SetPlatformName(const FString& InPlatformName)
+{
+	PlatformName = InPlatformName;
+	++AgentVersion;
+}
+
+void FDefaultUserAgentBuilder::SetPlatformVersion(const FString& InPlatformVersion)
+{
+	PlatformVersion = InPlatformVersion;
+	++AgentVersion;
+}
+
+void FDefaultUserAgentBuilder::AddPlatformComment(const FString& InComment)
+{
+	PlatformComments.AddUnique(InComment);
+	++AgentVersion;
+}
+
+uint32 FDefaultUserAgentBuilder::GetAgentVersion() const
+{
+	return AgentVersion;
+}
+
+FString FDefaultUserAgentBuilder::BuildCommentString(const TArray<FString>& Comments, const TSet<FString>* AllowedCommentsFilter)
+{
+	TArray<FString> FilteredComments;
+	Algo::TransformIf(Comments, FilteredComments,
+		[AllowedCommentsFilter](const FString& Comment) -> bool
+		{
+			return !AllowedCommentsFilter || AllowedCommentsFilter->Contains(Comment);
+		},
+		[](const FString& UnsanitizedComment)
+		{
+			FString SanitizedComment = UnsanitizedComment;
+			SanitizedComment.ReplaceInline(TEXT("("), TEXT(""));
+			SanitizedComment.ReplaceInline(TEXT(")"), TEXT(""));
+			SanitizedComment.ReplaceInline(TEXT(";"), TEXT(""));
+			return SanitizedComment;
+		}
+	);
+
+	return !FilteredComments.IsEmpty() ? FString::Printf(TEXT(" (%s)"), *FString::Join(FilteredComments, TEXT("; "))) : FString();
+}
+
+class FUserAgentImpl final
+{
+public:
+	static FUserAgentImpl& Get()
+	{
+		static FUserAgentImpl Instance;
+		return Instance;
+	}
+
+	FUserAgentImpl()
+		: RWLock()
+		, bCachedCommentsEnabled(false)
+	{
+		CachedUserAgent = BuildUserAgentString();
+		CVarDefaultUserAgentCommentsEnabled->OnChangedDelegate().AddRaw(this, &FUserAgentImpl::OnCommentsEnabledChanged);
+	}
+
+	void AddProjectComment(const FString& Comment)
+	{
+		FRWScopeLock WriteLock(RWLock, SLT_Write);
+		Builder.AddProjectComment(Comment);
+		CachedUserAgent = BuildUserAgentString();
+	}
+
+	void AddPlatformComment(const FString& Comment)
+	{
+		FRWScopeLock WriteLock(RWLock, SLT_Write);
+		Builder.AddPlatformComment(Comment);
+		CachedUserAgent = BuildUserAgentString();
+	}
+
+	FString GetUserAgent()
+	{
+		FRWScopeLock ReadLock(RWLock, SLT_ReadOnly);
+		return CachedUserAgent;
+	}
+
+	FDefaultUserAgentBuilder GetBuilder()
+	{
+		FRWScopeLock ReadLock(RWLock, SLT_ReadOnly);
+		return Builder;
+	}
+
+	uint32 GetAgentVersion()
+	{
+		return Builder.GetAgentVersion();
+	}
+
+private:
+	FString BuildUserAgentString()
+	{
+		if (CVarDefaultUserAgentCommentsEnabled.GetValueOnAnyThread())
+		{
+			return Builder.BuildUserAgentString();
+		}
+		else
+		{
+			// Build the user agent string without any comments.
+			TSet<FString> AllowedCommentsFilter;
+			return Builder.BuildUserAgentString(&AllowedCommentsFilter, &AllowedCommentsFilter);
+		}
+	}
+
+	void OnCommentsEnabledChanged(IConsoleVariable*)
+	{
+		FRWScopeLock WriteLock(RWLock, SLT_Write);
+		CachedUserAgent = BuildUserAgentString();
+	}
+
+	FRWLock RWLock;
+	FDefaultUserAgentBuilder Builder;
+	FString CachedUserAgent;
+	bool bCachedCommentsEnabled;
+};
+
 void FGenericPlatformHttp::Init()
 {
+	// Call during init to instantiate the user agent cache.
+	FUserAgentImpl::Get();
 }
 
 void FGenericPlatformHttp::Shutdown()
@@ -375,13 +550,27 @@ FString FGenericPlatformHttp::GetMimeType(const FString& FilePath)
 
 FString FGenericPlatformHttp::GetDefaultUserAgent()
 {
-	//** strip/escape slashes and whitespace from components
-	static FString CachedUserAgent = FString::Printf(TEXT("%s/%s %s/%s"),
-		*EscapeUserAgentString(FApp::GetProjectName()),
-		*EscapeUserAgentString(FApp::GetBuildVersion()),
-		*EscapeUserAgentString(FString(FPlatformProperties::IniPlatformName())),
-		*EscapeUserAgentString(FPlatformMisc::GetOSVersion()));
-	return CachedUserAgent;
+	return FUserAgentImpl::Get().GetUserAgent();
+}
+
+void FGenericPlatformHttp::AddDefaultUserAgentProjectComment(const FString& Comment)
+{
+	FUserAgentImpl::Get().AddProjectComment(Comment);
+}
+
+void FGenericPlatformHttp::AddDefaultUserAgentPlatformComment(const FString& Comment)
+{
+	FUserAgentImpl::Get().AddPlatformComment(Comment);
+}
+
+uint32 FGenericPlatformHttp::GetDefaultUserAgentVersion()
+{
+	return FUserAgentImpl::Get().GetAgentVersion();
+}
+
+FDefaultUserAgentBuilder FGenericPlatformHttp::GetDefaultUserAgentBuilder()
+{
+	return FUserAgentImpl::Get().GetBuilder();
 }
 
 FString FGenericPlatformHttp::EscapeUserAgentString(const FString& UnescapedString)
diff --git a/Engine/Source/Runtime/Online/HTTP/Private/HttpThread.cpp b/Engine/Source/Runtime/Online/HTTP/Private/HttpThread.cpp
--- a/Engine/Source/Runtime/Online/HTTP/Private/HttpThread.cpp
+++ b/Engine/Source/Runtime/Online/HTTP/Private/HttpThread.cpp
@@ -11,6 +11,7 @@
 #include "Misc/Parse.h"
 #include "HttpModule.h"
 #include "Http.h"
+#include "PlatformHttp.h"
 #include "Stats/Stats.h"
 
 DECLARE_STATS_GROUP(TEXT("HTTP Thread"), STATGROUP_HTTPThread, STATCAT_Advanced);
@@ -306,6 +307,8 @@ void FHttpThreadBase::Process(TArray<IHttpThreadedRequest*>& RequestsToCancel, T
 
 FLegacyHttpThread::FLegacyHttpThread()
 {
+	FPlatformHttp::AddDefaultUserAgentProjectComment(TEXT("http-legacy"));
+
 	HttpThreadActiveFrameTimeInSeconds = FHttpModule::Get().GetHttpThreadActiveFrameTimeInSeconds();
 	HttpThreadActiveMinimumSleepTimeInSeconds = FHttpModule::Get().GetHttpThreadActiveMinimumSleepTimeInSeconds();
 	HttpThreadIdleFrameTimeInSeconds = FHttpModule::Get().GetHttpThreadIdleFrameTimeInSeconds();
diff --git a/Engine/Source/Runtime/Online/HTTP/Public/GenericPlatform/GenericPlatformHttp.h b/Engine/Source/Runtime/Online/HTTP/Public/GenericPlatform/GenericPlatformHttp.h
--- a/Engine/Source/Runtime/Online/HTTP/Public/GenericPlatform/GenericPlatformHttp.h
+++ b/Engine/Source/Runtime/Online/HTTP/Public/GenericPlatform/GenericPlatformHttp.h
@@ -4,6 +4,7 @@
 #pragma once
 
 #include "Containers/Array.h"
+#include "Containers/Set.h"
 #include "Containers/StringFwd.h"
 #include "Containers/StringView.h"
 #include "Containers/UnrealString.h"
@@ -14,6 +15,41 @@
 class FHttpManager;
 class IHttpRequest;
 
+/**
+ * Utility structure to help with building the HTTP user agent string.
+ */
+class FDefaultUserAgentBuilder
+{
+public:
+	HTTP_API FDefaultUserAgentBuilder();
+
+	// Build the user agent string. Passing a filter will limit the comments to only those allowed.
+	// Passing nullptr for a filter implies that all comments are allowed.
+	HTTP_API FString BuildUserAgentString(const TSet<FString>* AllowedProjectCommentsFilter = nullptr, const TSet<FString>* AllowedPlatformCommentsFilter = nullptr) const;
+
+	HTTP_API void SetProjectName(const FString& InProjectName);
+	HTTP_API void SetProjectVersion(const FString& InProjectVersion);
+	HTTP_API void AddProjectComment(const FString& InComment);
+	HTTP_API void SetPlatformName(const FString& InPlatformName);
+	HTTP_API void SetPlatformVersion(const FString& InPlatformVersion);
+	HTTP_API void AddPlatformComment(const FString& InComment);
+
+	// Get the builder version. The version is used for tracking app changes to the default user
+	// agent and is not included in the user agent string.
+	HTTP_API uint32 GetAgentVersion() const;
+
+private:
+	static HTTP_API FString BuildCommentString(const TArray<FString>& Comments, const TSet<FString>* AllowedCommentsFilter);
+
+	FString ProjectName;
+	FString ProjectVersion;
+	TArray<FString> ProjectComments;
+	FString PlatformName;
+	FString PlatformVersion;
+	TArray<FString> PlatformComments;
+	uint32 AgentVersion;
+};
+
 /**
  * Platform specific Http implementations
  * Intended usage is to use FPlatformHttp instead of FGenericPlatformHttp
@@ -103,12 +139,44 @@ public:
 	/**
 	 * Returns the default User-Agent string to use in HTTP requests.
 	 * Requests that explicitly set the User-Agent header will not use this value.
+	 * 
+	 * The default User-Agent is built in the format "project/version (comments) platform/OS version (comments)"
 	 *
 	 * @return the default User-Agent string that requests should use.
 	 */
 	static HTTP_API FString GetDefaultUserAgent();
 	static HTTP_API FString EscapeUserAgentString(const FString& UnescapedString);
 
+	/**
+	 * Add a comment to be included in the project section of the default User-Agent string.
+	 */
+	static HTTP_API void AddDefaultUserAgentProjectComment(const FString& Comment);
+
+	/**
+	 * Add a comment to be included in the platform section of the default User-Agent string.
+	 */
+	static HTTP_API void AddDefaultUserAgentPlatformComment(const FString& Comment);
+
+	/**
+	 * Get the version of the default user agent. The version is incremented any time the default
+	 * user agent is changed. Used to invalidate any cached user agent value.
+	 * 
+	 * @return the current version of the default user agent.
+	 */
+	static HTTP_API uint32 GetDefaultUserAgentVersion();
+
+	/**
+	 * Gets a copy of the values used to build the default user agent string. Used for modifying
+	 * default user agent values before sending HTTP requests.
+	 */
+	static HTTP_API FDefaultUserAgentBuilder GetDefaultUserAgentBuilder();
+
+	/**
+	 * Utility for escaping a string used in the HTTP user agent. Removes disallowed characters.
+	 * 
+	 * @return the escaped string.
+	 */
+
 	/**
 	 * Get the proxy address specified by the operating system
 	 *
