Add missing include for UnrealString.h

#rnx

diff --git a/Engine/Source/Runtime/Online/Experimental/EventLoop/Private/EventLoop/BSDSocket/BSDSocketTypesPrivate.h b/Engine/Source/Runtime/Online/Experimental/EventLoop/Private/EventLoop/BSDSocket/BSDSocketTypesPrivate.h
--- a/Engine/Source/Runtime/Online/Experimental/EventLoop/Private/EventLoop/BSDSocket/BSDSocketTypesPrivate.h
+++ b/Engine/Source/Runtime/Online/Experimental/EventLoop/Private/EventLoop/BSDSocket/BSDSocketTypesPrivate.h
@@ -4,6 +4,7 @@
 
 #include "Containers/UnrealString.h"
 #include "Containers/StringConv.h"
+#include "Containers/UnrealString.h"
 #include "EventLoop/BSDSocket/BSDSocketTypes.h"
 
 #if PLATFORM_HAS_BSD_SOCKETS
