Update BPS with latest changes

Unshelved from pending changelist '25585543'

diff --git a/Engine/Source/Runtime/Online/BuildPatchServices/Private/BuildPatchFileConstructor.cpp b/Engine/Source/Runtime/Online/BuildPatchServices/Private/BuildPatchFileConstructor.cpp
--- a/Engine/Source/Runtime/Online/BuildPatchServices/Private/BuildPatchFileConstructor.cpp
+++ b/Engine/Source/Runtime/Online/BuildPatchServices/Private/BuildPatchFileConstructor.cpp
@@ -390,7 +390,7 @@ uint32 FBuildPatchFileConstructor::Run()
 			}
 			else
 			{
-				bFileSuccess = ConstructFileFromChunks(*FileManifest, bFilePreviouslyStarted);
+				bFileSuccess = ConstructFileFromChunks(FileToConstruct, *FileManifest, bFilePreviouslyStarted);
 			}
 		}
 		else
@@ -557,12 +557,12 @@ uint64 FBuildPatchFileConstructor::CalculateRequiredDiskSpace(const FFileManifes
 	return FMath::Max<int64>(DiskSpaceDeltaPeak, 0);
 }
 
-bool FBuildPatchFileConstructor::ConstructFileFromChunks(const FFileManifest& FileManifest, bool bResumeExisting)
+bool FBuildPatchFileConstructor::ConstructFileFromChunks(const FString& BuildFilename, const FFileManifest& FileManifest, bool bResumeExisting)
 {
 	bool bSuccess = true;
 	EConstructionError ConstructionError = EConstructionError::None;
 	uint32 LastError = 0;
-	FString NewFilename = Configuration.StagingDirectory / FileManifest.Filename;
+	FString NewFilename = Configuration.StagingDirectory / BuildFilename;
 
 	// Calculate the hash as we write the data
 	FSHA1 HashState;
@@ -623,7 +623,7 @@ bool FBuildPatchFileConstructor::ConstructFileFromChunks(const FFileManifest& Fi
 					// Inform the reference tracker of the chunk part skip
 					bSuccess = ChunkReferenceTracker->PopReference(ChunkPart.Guid) && bSuccess;
 					CountBytesProcessed(ChunkPart.Size);
-					FileConstructorStat->OnFileProgress(FileManifest.Filename, NewFileReader->Tell());
+					FileConstructorStat->OnFileProgress(BuildFilename, NewFileReader->Tell());
 					// Wait if paused
 					FileConstructorHelpers::WaitWhilePaused(bIsPaused, bShouldAbort);
 				}
@@ -695,7 +695,7 @@ bool FBuildPatchFileConstructor::ConstructFileFromChunks(const FFileManifest& Fi
 		{
 			const FChunkPart& ChunkPart = FileManifest.ChunkParts[ChunkPartIdx];
 			bSuccess = InsertChunkData(ChunkPart, *NewFile, HashState, ConstructionError);
-			FileConstructorStat->OnFileProgress(FileManifest.Filename, NewFile->Tell());
+			FileConstructorStat->OnFileProgress(BuildFilename, NewFile->Tell());
 			if (bSuccess)
 			{
 				CountBytesProcessed(ChunkPart.Size);
@@ -707,13 +707,13 @@ bool FBuildPatchFileConstructor::ConstructFileFromChunks(const FFileManifest& Fi
 			{
 				if (ConstructionError == EConstructionError::MissingChunk)
 				{
-					InstallerAnalytics->RecordConstructionError(FileManifest.Filename, INDEX_NONE, TEXT("Missing Chunk"));
-					UE_LOG(LogBuildPatchServices, Error, TEXT("FBuildPatchFileConstructor: Failed %s due to missing chunk %s"), *FileManifest.Filename, *ChunkPart.Guid.ToString());
+					InstallerAnalytics->RecordConstructionError(BuildFilename, INDEX_NONE, TEXT("Missing Chunk"));
+					UE_LOG(LogBuildPatchServices, Error, TEXT("FBuildPatchFileConstructor: Failed %s due to missing chunk %s"), *BuildFilename, *ChunkPart.Guid.ToString());
 				}
 				else if (ConstructionError == EConstructionError::TrackingError)
 				{
-					InstallerAnalytics->RecordConstructionError(FileManifest.Filename, INDEX_NONE, TEXT("Tracking Error"));
-					UE_LOG(LogBuildPatchServices, Error, TEXT("FBuildPatchFileConstructor: Failed %s due to untracked chunk %s"), *FileManifest.Filename, *ChunkPart.Guid.ToString());
+					InstallerAnalytics->RecordConstructionError(BuildFilename, INDEX_NONE, TEXT("Tracking Error"));
+					UE_LOG(LogBuildPatchServices, Error, TEXT("FBuildPatchFileConstructor: Failed %s due to untracked chunk %s"), *BuildFilename, *ChunkPart.Guid.ToString());
 				}
 			}
 		}
@@ -771,8 +771,8 @@ bool FBuildPatchFileConstructor::ConstructFileFromChunks(const FFileManifest& Fi
 			case EConstructionError::CannotCreateFile:
 				if (bReportAnalytic)
 				{
-					InstallerAnalytics->RecordConstructionError(FileManifest.Filename, LastError, TEXT("Could Not Create File"));
-					UE_LOG(LogBuildPatchServices, Error, TEXT("FBuildPatchFileConstructor: Could not create %s"), *FileManifest.Filename);
+					InstallerAnalytics->RecordConstructionError(BuildFilename, LastError, TEXT("Could Not Create File"));
+					UE_LOG(LogBuildPatchServices, Error, TEXT("FBuildPatchFileConstructor: Could not create %s"), *BuildFilename);
 				}
 				InstallerError->SetError(EBuildPatchInstallError::FileConstructionFail, ConstructionErrorCodes::FileCreateFail, LastError);
 				break;
@@ -801,8 +801,8 @@ bool FBuildPatchFileConstructor::ConstructFileFromChunks(const FFileManifest& Fi
 			// Only report or log if the first error
 			if (InstallerError->HasError() == false)
 			{
-				InstallerAnalytics->RecordConstructionError(FileManifest.Filename, INDEX_NONE, TEXT("Serialised Verify Fail"));
-				UE_LOG(LogBuildPatchServices, Error, TEXT("FBuildPatchFileConstructor: Verify failed after constructing %s"), *FileManifest.Filename);
+				InstallerAnalytics->RecordConstructionError(BuildFilename, INDEX_NONE, TEXT("Serialised Verify Fail"));
+				UE_LOG(LogBuildPatchServices, Error, TEXT("FBuildPatchFileConstructor: Verify failed after constructing %s"), *BuildFilename);
 			}
 			// Always set
 			InstallerError->SetError(EBuildPatchInstallError::FileConstructionFail, ConstructionErrorCodes::OutboundCorrupt);
diff --git a/Engine/Source/Runtime/Online/BuildPatchServices/Private/BuildPatchFileConstructor.h b/Engine/Source/Runtime/Online/BuildPatchServices/Private/BuildPatchFileConstructor.h
--- a/Engine/Source/Runtime/Online/BuildPatchServices/Private/BuildPatchFileConstructor.h
+++ b/Engine/Source/Runtime/Online/BuildPatchServices/Private/BuildPatchFileConstructor.h
@@ -172,11 +172,12 @@ namespace BuildPatchServices
 		/**
 		 * Constructs a particular file referenced by the given BuildManifest. The function takes an interface to a class that can provide availability information of chunks so that this
 		 * file construction process can be ran alongside chunk acquisition threads. It will Sleep while waiting for chunks that it needs.
+		 * @param BuildFilename		The relative build filename to use for construction.
 		 * @param FileManifest		The FFileManifest for the file to construct.
 		 * @param bResumeExisting	Whether we should resume from an existing file
 		 * @return	true if no file errors occurred
 		 */
-		bool ConstructFileFromChunks(const FFileManifest& FileManifest, bool bResumeExisting);
+		bool ConstructFileFromChunks(const FString& BuildFilename, const FFileManifest& FileManifest, bool bResumeExisting);
 
 		/**
 		 * Inserts the data data from a chunk into the destination file according to the chunk part info
diff --git a/Engine/Source/Runtime/Online/BuildPatchServices/Private/BuildPatchInstaller.cpp b/Engine/Source/Runtime/Online/BuildPatchServices/Private/BuildPatchInstaller.cpp
--- a/Engine/Source/Runtime/Online/BuildPatchServices/Private/BuildPatchInstaller.cpp
+++ b/Engine/Source/Runtime/Online/BuildPatchServices/Private/BuildPatchInstaller.cpp
@@ -109,6 +109,13 @@ namespace ConfigHelpers
 		FOptimisedDeltaConfiguration OptimisedDeltaConfiguration(InstallerAction.GetSharedInstallOrCurrentManifest());
 		OptimisedDeltaConfiguration.SourceManifest = InstallerAction.TryGetSharedCurrentManifest();
 		OptimisedDeltaConfiguration.CloudDirectories = Config.CloudDirectories;
+		if (!InstallerAction.GetCloudSubdirectory().IsEmpty())
+		{
+			for (FString& CloudDirectory : OptimisedDeltaConfiguration.CloudDirectories)
+			{
+				CloudDirectory /= InstallerAction.GetCloudSubdirectory();
+			}
+		}
 		OptimisedDeltaConfiguration.DeltaPolicy = Config.DeltaPolicy;
 		OptimisedDeltaConfiguration.InstallMode = Config.InstallMode;
 		return OptimisedDeltaConfiguration;
@@ -251,6 +258,14 @@ namespace InstallerHelpers
 			{
 				UE_LOG(LogBPSInstallerConfig, Log, TEXT("Build Config: ValidTags: %s"), *Tag);
 			}
+			if (!InstallerAction.GetInstallSubdirectory().IsEmpty())
+			{
+				UE_LOG(LogBuildPatchServices, Log, TEXT("Build Config: InstallSubdirectory: %s"), *InstallerAction.GetInstallSubdirectory());
+			}
+			if (!InstallerAction.GetCloudSubdirectory().IsEmpty())
+			{
+				UE_LOG(LogBuildPatchServices, Log, TEXT("Build Config: CloudSubdirectory: %s"), *InstallerAction.GetCloudSubdirectory());
+			}
 		}
 
 		UE_LOG(LogBPSInstallerConfig, Log, TEXT("Build Config: InstallDirectory: %s"), *InstallerConfiguration.InstallDirectory);
@@ -267,9 +282,9 @@ namespace InstallerHelpers
 			UE_LOG(LogBPSInstallerConfig, Log, TEXT("Build Config: CloudDirectories: %s"), *CloudDirectory);
 		}
 
-		UE_LOG(LogBPSInstallerConfig, Log, TEXT("Build Config: InstallMode: %s"), *EnumToString(InstallerConfiguration.InstallMode));
-		UE_LOG(LogBPSInstallerConfig, Log, TEXT("Build Config: VerifyMode: %s"), *EnumToString(InstallerConfiguration.VerifyMode));
-		UE_LOG(LogBPSInstallerConfig, Log, TEXT("Build Config: DeltaPolicy: %s"), *EnumToString(InstallerConfiguration.DeltaPolicy));
+		UE_LOG(LogBPSInstallerConfig, Log, TEXT("Build Config: InstallMode: %s"), *LexToString(InstallerConfiguration.InstallMode));
+		UE_LOG(LogBPSInstallerConfig, Log, TEXT("Build Config: VerifyMode: %s"), *LexToString(InstallerConfiguration.VerifyMode));
+		UE_LOG(LogBPSInstallerConfig, Log, TEXT("Build Config: DeltaPolicy: %s"), *LexToString(InstallerConfiguration.DeltaPolicy));
 		UE_LOG(LogBPSInstallerConfig, Log, TEXT("Build Config: bRunRequiredPrereqs: %s"), (InstallerConfiguration.bRunRequiredPrereqs) ? TEXT("true") : TEXT("false"));
 		UE_LOG(LogBPSInstallerConfig, Log, TEXT("Build Config: bAllowConcurrentExecution: %s"), (InstallerConfiguration.bAllowConcurrentExecution) ? TEXT("true") : TEXT("false"));
 	}
@@ -469,7 +484,7 @@ namespace BuildPatchServices
 			// Make sure existing manifests are added to installation info.
 			if (BuildPatchInstallerAction.TryGetCurrentManifest())
 			{
-				InstallationInfo.Add(InstallDirectory, BuildPatchInstallerAction.GetSharedCurrentManifest());
+				InstallationInfo.Add(InstallDirectory / BuildPatchInstallerAction.GetInstallSubdirectory(), BuildPatchInstallerAction.GetSharedCurrentManifest());
 			}
 			// Cache chunk sizes too
 			ChunkDataSizeProvider->AddManifestData(BuildPatchInstallerAction.TryGetSharedCurrentManifest());
diff --git a/Engine/Source/Runtime/Online/BuildPatchServices/Private/BuildPatchManifestSet.cpp b/Engine/Source/Runtime/Online/BuildPatchServices/Private/BuildPatchManifestSet.cpp
--- a/Engine/Source/Runtime/Online/BuildPatchServices/Private/BuildPatchManifestSet.cpp
+++ b/Engine/Source/Runtime/Online/BuildPatchServices/Private/BuildPatchManifestSet.cpp
@@ -2,6 +2,7 @@
 
 #include "IBuildManifestSet.h"
 #include "Algo/Accumulate.h"
+#include "Algo/Transform.h"
 
 #include "BuildPatchSettings.h"
 
@@ -29,10 +30,10 @@ namespace BuildPatchServices
 	class FBuildPatchManifestSet : public IBuildManifestSet
 	{
 		// Lookup for data references.
-		typedef TTuple<FChunkInfo const* const, FBuildPatchAppManifest const * const> FDataReference;
+		typedef TTuple<FChunkInfo const* const, FBuildPatchAppManifest const * const, FBuildPatchInstallerAction const * const> FDataReference;
 		TMap<FGuid, FDataReference> DataLookup;
 		// Lookup for file references.
-		typedef TTuple<FFileManifest const* const, FBuildPatchAppManifest const * const> FFileReference;
+		typedef TTuple<FFileManifest const* const, FBuildPatchAppManifest const * const, FBuildPatchInstallerAction const * const> FFileReference;
 		TMap<FString, FFileReference> CurrentFileLookup;
 		TMap<FString, FFileReference> NewFileLookup;
 	public:
@@ -46,14 +47,14 @@ namespace BuildPatchServices
 				{
 					for (const TPair<FGuid, const BuildPatchServices::FChunkInfo*>& ChunkInfoPair : Manifest.ChunkInfoLookup)
 					{
-						DataLookup.Add(ChunkInfoPair.Key, FDataReference{ ChunkInfoPair.Value, &Manifest });
+						DataLookup.Add(ChunkInfoPair.Key, FDataReference{ ChunkInfoPair.Value, &Manifest, &InstallerAction });
 					}
 				}
 				if (!InstallerAction.IsInstall())
 				{
 					for (const TPair<FStringView, const BuildPatchServices::FFileManifest*>& FileManifestPair : InstallerAction.GetCurrentManifest().FileManifestLookup)
 					{
-						CurrentFileLookup.Add(FString(FileManifestPair.Key), FFileReference{ FileManifestPair.Value, InstallerAction.TryGetCurrentManifest() });
+						CurrentFileLookup.Add(InstallerAction.GetInstallSubdirectory() / FString(FileManifestPair.Key), FFileReference{ FileManifestPair.Value, InstallerAction.TryGetCurrentManifest(), &InstallerAction });
 					}
 				}
 			}
@@ -63,19 +64,30 @@ namespace BuildPatchServices
 			for (int32 Idx = InstallerActions.Num() - 1; Idx >= 0; --Idx)
 			{
 				FBuildPatchInstallerAction& InstallerAction = InstallerActions[Idx];
-				if (!InstallerAction.IsUninstall())
-				{
-					InstallerAction.SetTaggedFiles(InstallerAction.GetTaggedFiles().Difference(VisitedFiles));
-					VisitedFiles.Append(InstallerAction.GetTaggedFiles());
-				}
+				const bool bIsUninstall = InstallerAction.IsUninstall();
+				TSet<FString> FilesToFilter;
 				// Ensure that all tagged files, have their new file manifest in the lookup, and current equiv in the lookup in case of multiple file references.
 				for (const FString& TaggedFile : InstallerAction.GetTaggedFiles())
 				{
-					NewFileLookup.Add(TaggedFile, FFileReference{ InstallerAction.GetInstallManifest().FileManifestLookup[TaggedFile], InstallerAction.TryGetInstallManifest() });
+					FString FullFilename = InstallerAction.GetInstallSubdirectory() / TaggedFile;
+					if (!bIsUninstall)
+					{
+						bool bDuplicateFilename = false;
+						VisitedFiles.Add(FullFilename, &bDuplicateFilename);
+						if (bDuplicateFilename)
+						{
+							FilesToFilter.Add(TaggedFile);
+						}
+					}
 					if (!InstallerAction.IsInstall() && InstallerAction.GetCurrentManifest().FileManifestLookup.Contains(TaggedFile))
 					{
-						CurrentFileLookup.Add(TaggedFile, FFileReference{ InstallerAction.GetCurrentManifest().FileManifestLookup[TaggedFile], InstallerAction.TryGetCurrentManifest() });
+						CurrentFileLookup.Add(FullFilename, FFileReference{ InstallerAction.GetCurrentManifest().FileManifestLookup[TaggedFile], InstallerAction.TryGetCurrentManifest(), &InstallerAction });
 					}
+					NewFileLookup.Add(MoveTemp(FullFilename), FFileReference{ InstallerAction.GetInstallManifest().FileManifestLookup[TaggedFile], InstallerAction.TryGetInstallManifest(), &InstallerAction });
+				}
+				if (FilesToFilter.Num() > 0)
+				{
+					InstallerAction.SetTaggedFiles(InstallerAction.GetTaggedFiles().Difference(FilesToFilter));
 				}
 			}
 		}
@@ -97,17 +109,14 @@ namespace BuildPatchServices
 
 		virtual void GetInstallResumeIdsForFile(const FString& BuildFile, TSet<FString>& ResumeIds, bool bIncludeLegacy) const
 		{
-			for (const FBuildPatchInstallerAction& InstallerAction : InstallerActions)
+			const FFileReference* LookupResult = NewFileLookup.Find(BuildFile);
+			if (LookupResult)
 			{
-				if (InstallerAction.GetTaggedFiles().Contains(BuildFile))
+				if (bIncludeLegacy)
 				{
-					if (bIncludeLegacy)
-					{
-						ResumeIds.Add(InstallerAction.GetInstallManifest().GetAppName() + InstallerAction.GetInstallManifest().GetVersionString());
-					}
-					ResumeIds.Add(InstallerAction.GetInstallManifest().GetBuildId());
-					return;
+					ResumeIds.Add(LookupResult->Get<2>()->GetInstallManifest().GetAppName() + LookupResult->Get<2>()->GetInstallManifest().GetVersionString());
 				}
+				ResumeIds.Add(LookupResult->Get<2>()->GetInstallManifest().GetBuildId());
 			}
 		}
 
@@ -149,12 +158,12 @@ namespace BuildPatchServices
 			return false;
 		}
 
-		virtual FString GetDataFilename(const FString& RootDirectory, const FGuid& DataGuid) const override
+		virtual FString GetDataFilename(const FGuid& DataGuid) const override
 		{
 			FDataReference const * const LookupResult = DataLookup.Find(DataGuid);
 			if (LookupResult)
 			{
-				return FBuildPatchUtils::GetDataFilename(*LookupResult->Get<1>(), RootDirectory, DataGuid);
+				return LookupResult->Get<2>()->GetCloudSubdirectory() / FBuildPatchUtils::GetDataFilename(*LookupResult->Get<1>(), DataGuid);
 			}
 			return FString();
 		}
@@ -175,11 +184,14 @@ namespace BuildPatchServices
 
 		virtual void GetOutdatedFiles(const FString& InstallDirectory, TSet<FString>& OutdatedFiles) const override
 		{
+			TSet<FString> TempFileSet;
 			for (const FBuildPatchInstallerAction& InstallerAction : InstallerActions)
 			{
 				if (!InstallerAction.IsUninstall())
 				{
-					InstallerAction.GetInstallManifest().GetOutdatedFiles(InstallerAction.TryGetCurrentManifest(), InstallDirectory, InstallerAction.GetTaggedFiles(), OutdatedFiles);
+					TempFileSet.Reset();
+					InstallerAction.GetInstallManifest().GetOutdatedFiles(InstallerAction.TryGetCurrentManifest(), InstallDirectory / InstallerAction.GetInstallSubdirectory(), InstallerAction.GetTaggedFiles(), TempFileSet);
+					Algo::Transform(TempFileSet, OutdatedFiles, [&](const FString& Filename) { return InstallerAction.GetInstallSubdirectory() / Filename; });
 				}
 			}
 		}
@@ -194,9 +206,10 @@ namespace BuildPatchServices
 				{
 					for (const FString& TrackedFile : Manifest.GetBuildFileList())
 					{
-						if (!NewFileLookup.Contains(TrackedFile))
+						const FString FullFilename = InstallerAction.GetInstallSubdirectory() / TrackedFile;
+						if (!NewFileLookup.Contains(FullFilename))
 						{
-							FilesToRemove.Add(TrackedFile);
+							FilesToRemove.Add(FullFilename);
 						}
 					}
 				}
@@ -254,7 +267,7 @@ namespace BuildPatchServices
 						PreInfo.AppName = Manifest.GetAppName();
 						PreInfo.Args = Manifest.GetPrereqArgs();
 						PreInfo.Name = Manifest.GetPrereqName();
-						PreInfo.Path = Manifest.GetPrereqPath();
+						PreInfo.Path = InstallerAction.GetInstallSubdirectory() / Manifest.GetPrereqPath();
 						PreInfo.VersionString = Manifest.GetVersionString();
 						PreInfo.bIsRepair = InstallerAction.IsRepair();
 					}
@@ -264,23 +277,21 @@ namespace BuildPatchServices
 
 		virtual void GetFilesTaggedForRepair(TSet<FString>& Filenames) const override
 		{
-			for (const FBuildPatchInstallerAction& InstallerAction : InstallerActions)
+			for (const TPair<FString, FFileReference>& NewFileReference : NewFileLookup)
 			{
-				if (InstallerAction.IsRepair())
+				if (NewFileReference.Value.Get<2>()->IsRepair())
 				{
-					Filenames.Append(InstallerAction.GetTaggedFiles());
+					Filenames.Add(NewFileReference.Key);
 				}
 			}
 		}
 
 		virtual bool IsFileRepairAction(const FString& Filename) const override
 		{
-			for (const FBuildPatchInstallerAction& InstallerAction : InstallerActions)
+			FFileReference const * const LookupResult = NewFileLookup.Find(Filename);
+			if (LookupResult)
 			{
-				if (InstallerAction.IsRepair() && InstallerAction.GetTaggedFiles().Contains(Filename))
-				{
-					return true;
-				}
+				return LookupResult->Get<2>()->IsRepair();
 			}
 			return false;
 		}
diff --git a/Engine/Source/Runtime/Online/BuildPatchServices/Private/BuildPatchSettings.cpp b/Engine/Source/Runtime/Online/BuildPatchServices/Private/BuildPatchSettings.cpp
--- a/Engine/Source/Runtime/Online/BuildPatchServices/Private/BuildPatchSettings.cpp
+++ b/Engine/Source/Runtime/Online/BuildPatchServices/Private/BuildPatchSettings.cpp
@@ -24,39 +24,47 @@ namespace BuildPatchServices
 		Invalid
 	};
 
-	BuildPatchServices::FInstallerAction FInstallerAction::MakeInstall(const IBuildManifestRef& Manifest, TSet<FString> InstallTags /*= TSet<FString>()*/)
+	BuildPatchServices::FInstallerAction FInstallerAction::MakeInstall(const IBuildManifestRef& Manifest, TSet<FString> InstallTags /*= TSet<FString>()*/, FString InstallSubdirectory /*= FString()*/, FString CloudSubdirectory /*= FString()*/)
 	{
 		FInstallerAction InstallerAction;
 		InstallerAction.InstallManifest = Manifest;
 		InstallerAction.InstallTags = MoveTemp(InstallTags);
+		InstallerAction.InstallSubdirectory = MoveTemp(InstallSubdirectory);
+		InstallerAction.CloudSubdirectory = MoveTemp(CloudSubdirectory);
 		InstallerAction.ActionIntent = EInstallActionIntent::Install;
 		return InstallerAction;
 	}
 
-	BuildPatchServices::FInstallerAction FInstallerAction::MakeUpdate(const IBuildManifestRef& CurrentManifest, const IBuildManifestRef& InstallManifest, TSet<FString> InstallTags /*= TSet<FString>()*/)
+	BuildPatchServices::FInstallerAction FInstallerAction::MakeUpdate(const IBuildManifestRef& CurrentManifest, const IBuildManifestRef& InstallManifest, TSet<FString> InstallTags /*= TSet<FString>()*/, FString InstallSubdirectory /*= FString()*/, FString CloudSubdirectory /*= FString()*/)
 	{
 		FInstallerAction InstallerAction;
 		InstallerAction.CurrentManifest = CurrentManifest;
 		InstallerAction.InstallManifest = InstallManifest;
 		InstallerAction.InstallTags = MoveTemp(InstallTags);
+		InstallerAction.InstallSubdirectory = MoveTemp(InstallSubdirectory);
+		InstallerAction.CloudSubdirectory = MoveTemp(CloudSubdirectory);
 		InstallerAction.ActionIntent = EInstallActionIntent::Update;
 		return InstallerAction;
 	}
 
-	BuildPatchServices::FInstallerAction FInstallerAction::MakeRepair(const IBuildManifestRef& Manifest, TSet<FString> InstallTags /*= TSet<FString>()*/)
+	BuildPatchServices::FInstallerAction FInstallerAction::MakeRepair(const IBuildManifestRef& Manifest, TSet<FString> InstallTags /*= TSet<FString>()*/, FString InstallSubdirectory /*= FString()*/, FString CloudSubdirectory /*= FString()*/)
 	{
 		FInstallerAction InstallerAction;
 		InstallerAction.CurrentManifest = Manifest;
 		InstallerAction.InstallManifest = Manifest;
 		InstallerAction.InstallTags = MoveTemp(InstallTags);
+		InstallerAction.InstallSubdirectory = MoveTemp(InstallSubdirectory);
+		InstallerAction.CloudSubdirectory = MoveTemp(CloudSubdirectory);
 		InstallerAction.ActionIntent = EInstallActionIntent::Repair;
 		return InstallerAction;
 	}
 
-	BuildPatchServices::FInstallerAction FInstallerAction::MakeUninstall(const IBuildManifestRef& Manifest)
+	BuildPatchServices::FInstallerAction FInstallerAction::MakeUninstall(const IBuildManifestRef& Manifest, FString InstallSubdirectory /*= FString()*/, FString CloudSubdirectory /*= FString()*/)
 	{
 		FInstallerAction InstallerAction;
 		InstallerAction.CurrentManifest = Manifest;
+		InstallerAction.InstallSubdirectory = MoveTemp(InstallSubdirectory);
+		InstallerAction.CloudSubdirectory = MoveTemp(CloudSubdirectory);
 		InstallerAction.ActionIntent = EInstallActionIntent::Uninstall;
 		return InstallerAction;
 	}
@@ -65,6 +73,8 @@ namespace BuildPatchServices
 		: CurrentManifest(CopyFrom.CurrentManifest)
 		, InstallManifest(CopyFrom.InstallManifest)
 		, InstallTags(CopyFrom.InstallTags)
+		, InstallSubdirectory(CopyFrom.InstallSubdirectory)
+		, CloudSubdirectory(CopyFrom.CloudSubdirectory)
 		, ActionIntent(CopyFrom.ActionIntent)
 	{
 	}
@@ -73,6 +83,8 @@ namespace BuildPatchServices
 		: CurrentManifest(MoveTemp(MoveFrom.CurrentManifest))
 		, InstallManifest(MoveTemp(MoveFrom.InstallManifest))
 		, InstallTags(MoveTemp(MoveFrom.InstallTags))
+		, InstallSubdirectory(MoveTemp(MoveFrom.InstallSubdirectory))
+		, CloudSubdirectory(MoveTemp(MoveFrom.CloudSubdirectory))
 		, ActionIntent(MoveFrom.ActionIntent)
 	{
 	}
@@ -107,6 +119,16 @@ namespace BuildPatchServices
 		return InstallTags;
 	}
 
+	const FString& FInstallerAction::GetInstallSubdirectory() const
+	{
+		return InstallSubdirectory;
+	}
+
+	const FString& FInstallerAction::GetCloudSubdirectory() const
+	{
+		return CloudSubdirectory;
+	}
+
 	IBuildManifestRef FInstallerAction::GetInstallManifest() const
 	{
 		return InstallManifest.ToSharedRef();
diff --git a/Engine/Source/Runtime/Online/BuildPatchServices/Private/BuildPatchState.cpp b/Engine/Source/Runtime/Online/BuildPatchServices/Private/BuildPatchState.cpp
--- a/Engine/Source/Runtime/Online/BuildPatchServices/Private/BuildPatchState.cpp
+++ b/Engine/Source/Runtime/Online/BuildPatchServices/Private/BuildPatchState.cpp
@@ -4,55 +4,7 @@
 
 #define LOCTEXT_NAMESPACE "BuildPatchInstallProgress"
 
-const FString& BuildPatchServices::StateToString(const BuildPatchServices::EBuildPatchState& State)
-{
-	// Static const fixed FString values so that they are not constantly constructed
-	static const FString Queued(TEXT("Queued"));
-	static const FString Initializing(TEXT("Initialising"));
-	static const FString Resuming(TEXT("Resuming"));
-	static const FString Downloading(TEXT("Downloading"));
-	static const FString Recycling(TEXT("Recycling"));
-	static const FString Installing(TEXT("Installing"));
-	static const FString MovingToInstall(TEXT("MovingToInstall"));
-	static const FString SettingAttributes(TEXT("SettingAttributes"));
-	static const FString BuildVerification(TEXT("BuildVerification"));
-	static const FString CleanUp(TEXT("CleanUp"));
-	static const FString PrerequisitesInstall(TEXT("PrerequisitesInstall"));
-	static const FString Completed(TEXT("Completed"));
-	static const FString Paused(TEXT("Paused"));
-	static const FString Default(TEXT("InvalidOrMax"));
-
-	switch (State)
-	{
-	case BuildPatchServices::EBuildPatchState::Queued:
-		return Queued;
-	case BuildPatchServices::EBuildPatchState::Initializing:
-		return Initializing;
-	case BuildPatchServices::EBuildPatchState::Resuming:
-		return Resuming;
-	case BuildPatchServices::EBuildPatchState::Downloading:
-		return Downloading;
-	case BuildPatchServices::EBuildPatchState::Installing:
-		return Installing;
-	case BuildPatchServices::EBuildPatchState::MovingToInstall:
-		return MovingToInstall;
-	case BuildPatchServices::EBuildPatchState::SettingAttributes:
-		return SettingAttributes;
-	case BuildPatchServices::EBuildPatchState::BuildVerification:
-		return BuildVerification;
-	case BuildPatchServices::EBuildPatchState::CleanUp:
-		return CleanUp;
-	case BuildPatchServices::EBuildPatchState::PrerequisitesInstall:
-		return PrerequisitesInstall;
-	case BuildPatchServices::EBuildPatchState::Completed:
-		return Completed;
-	case BuildPatchServices::EBuildPatchState::Paused:
-		return Paused;
-	default:
-		return Default;
-	}
-}
-
+// Deprecated in header
 const FText& BuildPatchServices::StateToText(const BuildPatchServices::EBuildPatchState& State)
 {
 	// Static const fixed FText values so that they are not constantly constructed
diff --git a/Engine/Source/Runtime/Online/BuildPatchServices/Private/BuildPatchUtil.cpp b/Engine/Source/Runtime/Online/BuildPatchServices/Private/BuildPatchUtil.cpp
--- a/Engine/Source/Runtime/Online/BuildPatchServices/Private/BuildPatchUtil.cpp
+++ b/Engine/Source/Runtime/Online/BuildPatchServices/Private/BuildPatchUtil.cpp
@@ -16,26 +16,26 @@ using namespace BuildPatchServices;
 
 /* FBuildPatchUtils implementation
 *****************************************************************************/
-FString FBuildPatchUtils::GetChunkNewFilename(const EFeatureLevel FeatureLevel, const FString& RootDirectory, const FGuid& ChunkGUID, const uint64& ChunkHash)
+FString FBuildPatchUtils::GetChunkNewFilename(const EFeatureLevel FeatureLevel, const FGuid& ChunkGUID, const uint64& ChunkHash)
 {
 	check(ChunkGUID.IsValid());
 	if (FeatureLevel < EFeatureLevel::DataFileRenames)
 	{
-		return GetChunkOldFilename(RootDirectory, ChunkGUID);
+		return GetChunkOldFilename(ChunkGUID);
 	}
-	return FPaths::Combine(*RootDirectory, *FString::Printf(TEXT("%s/%02d/%016llX_%s.chunk"), ManifestVersionHelpers::GetChunkSubdir(FeatureLevel), FCrc::MemCrc32(&ChunkGUID, sizeof(FGuid)) % 100, ChunkHash, *ChunkGUID.ToString()));
+	return FString::Printf(TEXT("%s/%02d/%016llX_%s.chunk"), ManifestVersionHelpers::GetChunkSubdir(FeatureLevel), FCrc::MemCrc32(&ChunkGUID, sizeof(FGuid)) % 100, ChunkHash, *ChunkGUID.ToString());
 }
 
-FString FBuildPatchUtils::GetFileNewFilename(const EFeatureLevel FeatureLevel, const FString& RootDirectory, const FGuid& FileGUID, const FSHAHash& FileHash)
+FString FBuildPatchUtils::GetFileNewFilename(const EFeatureLevel FeatureLevel, const FGuid& FileGUID, const FSHAHash& FileHash)
 {
 	check(FileGUID.IsValid());
-	return FPaths::Combine(*RootDirectory, *FString::Printf(TEXT("%s/%02d/%s_%s.file"), ManifestVersionHelpers::GetFileSubdir(FeatureLevel), FCrc::MemCrc32(&FileGUID, sizeof(FGuid)) % 100, *FileHash.ToString(), *FileGUID.ToString()));
+	return FString::Printf(TEXT("%s/%02d/%s_%s.file"), ManifestVersionHelpers::GetFileSubdir(FeatureLevel), FCrc::MemCrc32(&FileGUID, sizeof(FGuid)) % 100, *FileHash.ToString(), *FileGUID.ToString());
 }
 
-FString FBuildPatchUtils::GetFileNewFilename(const EFeatureLevel FeatureLevel, const FString& RootDirectory, const FGuid& FileGUID, const uint64& FileHash)
+FString FBuildPatchUtils::GetFileNewFilename(const EFeatureLevel FeatureLevel, const FGuid& FileGUID, const uint64& FileHash)
 {
 	check(FileGUID.IsValid());
-	return FPaths::Combine(*RootDirectory, *FString::Printf(TEXT("%s/%02d/%016llX_%s.file"), ManifestVersionHelpers::GetFileSubdir(FeatureLevel), FCrc::MemCrc32(&FileGUID, sizeof(FGuid)) % 100, FileHash, *FileGUID.ToString()));
+	return FString::Printf(TEXT("%s/%02d/%016llX_%s.file"), ManifestVersionHelpers::GetFileSubdir(FeatureLevel), FCrc::MemCrc32(&FileGUID, sizeof(FGuid)) % 100, FileHash, *FileGUID.ToString());
 }
 
 void FBuildPatchUtils::GetChunkDetailFromNewFilename( const FString& ChunkNewFilename, FGuid& ChunkGUID, uint64& ChunkHash )
@@ -64,46 +64,46 @@ void FBuildPatchUtils::GetFileDetailFromNewFilename(const FString& FileNewFilena
 	FGuid::Parse( GuidString, FileGUID );
 }
 
-FString FBuildPatchUtils::GetChunkOldFilename( const FString& RootDirectory, const FGuid& ChunkGUID )
+FString FBuildPatchUtils::GetChunkOldFilename(const FGuid& ChunkGUID)
 {
-	check( ChunkGUID.IsValid() );
-	return FPaths::Combine( *RootDirectory, *FString::Printf( TEXT("Chunks/%02d/%s.chunk"), FCrc::MemCrc_DEPRECATED( &ChunkGUID, sizeof( FGuid ) ) % 100, *ChunkGUID.ToString() ) );
+	check(ChunkGUID.IsValid());
+	return FString::Printf(TEXT("Chunks/%02d/%s.chunk"), FCrc::MemCrc_DEPRECATED(&ChunkGUID, sizeof(FGuid)) % 100, *ChunkGUID.ToString());
 }
 
-FString FBuildPatchUtils::GetFileOldFilename( const FString& RootDirectory, const FGuid& FileGUID )
+FString FBuildPatchUtils::GetFileOldFilename(const FGuid& FileGUID)
 {
-	check( FileGUID.IsValid() );
-	return FPaths::Combine( *RootDirectory, *FString::Printf( TEXT("Files/%02d/%s.file"), FCrc::MemCrc_DEPRECATED( &FileGUID, sizeof( FGuid ) ) % 100, *FileGUID.ToString() ) );
+	check(FileGUID.IsValid());
+	return FString::Printf(TEXT("Files/%02d/%s.file"), FCrc::MemCrc_DEPRECATED(&FileGUID, sizeof(FGuid)) % 100, *FileGUID.ToString());
 }
 
-FString FBuildPatchUtils::GetDataTypeOldFilename( EBuildPatchDataType DataType, const FString& RootDirectory, const FGuid& Guid )
+FString FBuildPatchUtils::GetDataTypeOldFilename(EBuildPatchDataType DataType, const FGuid& Guid)
 {
-	check( Guid.IsValid() );
+	check(Guid.IsValid());
 
-	switch ( DataType )
+	switch (DataType)
 	{
 	case EBuildPatchDataType::ChunkData:
-		return GetChunkOldFilename( RootDirectory, Guid );
+		return GetChunkOldFilename(Guid);
 	case EBuildPatchDataType::FileData:
-		return GetFileOldFilename( RootDirectory, Guid );
+		return GetFileOldFilename(Guid);
 	}
 
 	// Error, didn't case type
-	check( false );
-	return TEXT( "" );
+	check(false);
+	return TEXT("");
 }
 
-FString FBuildPatchUtils::GetDataFilename(const FBuildPatchAppManifestRef& Manifest, const FString& RootDirectory, const FGuid& DataGUID)
+FString FBuildPatchUtils::GetDataFilename(const FBuildPatchAppManifestRef& Manifest, const FGuid& DataGUID)
 {
-	return GetDataFilename(Manifest.Get(), RootDirectory, DataGUID);
+	return GetDataFilename(Manifest.Get(), DataGUID);
 }
 
-FString FBuildPatchUtils::GetDataFilename(const FBuildPatchAppManifest&    Manifest, const FString& RootDirectory, const FGuid& DataGUID)
+FString FBuildPatchUtils::GetDataFilename(const FBuildPatchAppManifest&    Manifest, const FGuid& DataGUID)
 {
 	const EBuildPatchDataType DataType = Manifest.IsFileDataManifest() ? EBuildPatchDataType::FileData : EBuildPatchDataType::ChunkData;
 	if (Manifest.GetFeatureLevel() < EFeatureLevel::DataFileRenames)
 	{
-		return FBuildPatchUtils::GetDataTypeOldFilename(DataType, RootDirectory, DataGUID);
+		return FBuildPatchUtils::GetDataTypeOldFilename(DataType, DataGUID);
 	}
 	else if (!Manifest.IsFileDataManifest())
 	{
@@ -111,7 +111,7 @@ FString FBuildPatchUtils::GetDataFilename(const FBuildPatchAppManifest&    Manif
 		const bool bFound = Manifest.GetChunkHash(DataGUID, ChunkHash);
 		// Should be impossible to not exist
 		check(bFound);
-		return FBuildPatchUtils::GetChunkNewFilename(Manifest.GetFeatureLevel(), RootDirectory, DataGUID, ChunkHash);
+		return FBuildPatchUtils::GetChunkNewFilename(Manifest.GetFeatureLevel(), DataGUID, ChunkHash);
 	}
 	else if (Manifest.GetFeatureLevel() <= EFeatureLevel::StoredAsCompressedUClass)
 	{
@@ -119,7 +119,7 @@ FString FBuildPatchUtils::GetDataFilename(const FBuildPatchAppManifest&    Manif
 		const bool bFound = Manifest.GetFileHash(DataGUID, FileHash);
 		// Should be impossible to not exist
 		check(bFound);
-		return FBuildPatchUtils::GetFileNewFilename(Manifest.GetFeatureLevel(), RootDirectory, DataGUID, FileHash);
+		return FBuildPatchUtils::GetFileNewFilename(Manifest.GetFeatureLevel(), DataGUID, FileHash);
 	}
 	else
 	{
@@ -127,7 +127,7 @@ FString FBuildPatchUtils::GetDataFilename(const FBuildPatchAppManifest&    Manif
 		const bool bFound = Manifest.GetFilePartHash(DataGUID, FileHash);
 		// Should be impossible to not exist
 		check(bFound);
-		return FBuildPatchUtils::GetFileNewFilename(Manifest.GetFeatureLevel(), RootDirectory, DataGUID, FileHash);
+		return FBuildPatchUtils::GetFileNewFilename(Manifest.GetFeatureLevel(), DataGUID, FileHash);
 	}
 	return TEXT("");
 }
diff --git a/Engine/Source/Runtime/Online/BuildPatchServices/Private/BuildPatchUtil.h b/Engine/Source/Runtime/Online/BuildPatchServices/Private/BuildPatchUtil.h
--- a/Engine/Source/Runtime/Online/BuildPatchServices/Private/BuildPatchUtil.h
+++ b/Engine/Source/Runtime/Online/BuildPatchServices/Private/BuildPatchUtil.h
@@ -38,25 +38,23 @@ DECLARE_DELEGATE_RetVal(bool, FBuildPatchBoolRetDelegate);
 struct FBuildPatchUtils
 {
 	/**
-	 * Gets the filename for a chunk generated from it's GUID and Hash, which is the new format.
+	 * Gets the the relative cloud filename for a chunk generated from it's GUID and Hash, which is the new format.
 	 * @param FeatureLevel   The manifest version that references this chunk.
-	 * @param RootDirectory  The root directory.
 	 * @param ChunkGuid      The chunk Guid.
 	 * @param ChunkHash      The chunk rolling hash value.
 	 * @return the chunk path.
 	 */
-	static FString GetChunkNewFilename(BuildPatchServices::EFeatureLevel FeatureLevel, const FString& RootDirectory, const FGuid& ChunkGUID, const uint64& ChunkHash);
+	static FString GetChunkNewFilename(BuildPatchServices::EFeatureLevel FeatureLevel, const FGuid& ChunkGUID, const uint64& ChunkHash);
 
 	/**
-	 * Gets the filename for a file chunk generated from it's GUID and Hash, which is the new format.
+	 * Gets the the relative cloud filename for a file chunk generated from it's GUID and Hash, which is the new format.
 	 * @param FeatureLevel   The manifest version that references this file.
-	 * @param RootDirectory  The root directory.
 	 * @param FileGUID       The file chunk Guid.
 	 * @param FileHash       The file hash value.
 	 * @return the file chunk path.
 	 */
-	static FString GetFileNewFilename(BuildPatchServices::EFeatureLevel FeatureLevel, const FString& RootDirectory, const FGuid& FileGUID, const FSHAHash& FileHash);
-	static FString GetFileNewFilename(BuildPatchServices::EFeatureLevel FeatureLevel, const FString& RootDirectory, const FGuid& FileGUID, const uint64& FilePartHash);
+	static FString GetFileNewFilename(BuildPatchServices::EFeatureLevel FeatureLevel, const FGuid& FileGUID, const FSHAHash& FileHash);
+	static FString GetFileNewFilename(BuildPatchServices::EFeatureLevel FeatureLevel, const FGuid& FileGUID, const uint64& FilePartHash);
 
 	/**
 	 * Gets the chunk GUID and Hash, from the filename, which is the new format.
@@ -75,40 +73,35 @@ struct FBuildPatchUtils
 	static void GetFileDetailFromNewFilename(const FString& FileNewFilename, FGuid& FileGUID, FSHAHash& FileHash);
 
 	/**
-	 * Gets the filename for a chunk generated from it's GUID
-	 * @param RootDirectory		The root directory
+	 * Gets the the relative cloud filename for a chunk generated from it's GUID
 	 * @param ChunkGuid			The chunk Guid
 	 * @return	the chunk path
 	 */
-	static FString GetChunkOldFilename(const FString& RootDirectory, const FGuid& ChunkGUID);
+	static FString GetChunkOldFilename(const FGuid& ChunkGUID);
 
 	/**
-	 * Gets the filename for a file data part generated from it's GUID
-	 * @param RootDirectory		The root directory
+	 * Gets the the relative cloud filename for a file data part generated from it's GUID
 	 * @param FileGUID			The file part Guid
 	 * @return	The file data path
 	 */
-	static FString GetFileOldFilename(const FString& RootDirectory, const FGuid& FileGUID);
+	static FString GetFileOldFilename(const FGuid& FileGUID);
 
 	/**
 	 * Gets the filename for a specific data part type from it's GUID
 	 * @param DataType			The type of data
-	 * @param RootDirectory		The root directory
 	 * @param DataGUID			The data Guid
 	 * @return	the data part path
 	 */
-	static FString GetDataTypeOldFilename(EBuildPatchDataType DataType, const FString& RootDirectory, const FGuid& Guid);
+	static FString GetDataTypeOldFilename(EBuildPatchDataType DataType, const FGuid& Guid);
 
 	/**
-	 * Gets the filename for any data part. Wraps the choice between all of the above
+	 * Gets the relative cloud filename for any data part. Wraps the choice between all of the above
 	 * @param Manifest			The manifest referencing this data
-	 * @param RootDirectory		The root directory
 	 * @param DataGUID			The data Guid
 	 * @return	the data part path
 	 */
-	static FString GetDataFilename(const FBuildPatchAppManifestRef& Manifest, const FString& RootDirectory, const FGuid& DataGUID);
-	static FString GetDataFilename(const FBuildPatchAppManifest&    Manifest, const FString& RootDirectory, const FGuid& DataGUID);
-	//static FString GetDataFilename(IBuildManifestSet* ManifestSet, const FString& RootDirectory, const FGuid& DataGUID);
+	static FString GetDataFilename(const FBuildPatchAppManifestRef& Manifest, const FGuid& DataGUID);
+	static FString GetDataFilename(const FBuildPatchAppManifest&    Manifest, const FGuid& DataGUID);
 
 	/**
 	 * Gets the GUID for a data file according to it's filename (new or old)
diff --git a/Engine/Source/Runtime/Online/BuildPatchServices/Private/Diffing/DiffManifests.cpp b/Engine/Source/Runtime/Online/BuildPatchServices/Private/Diffing/DiffManifests.cpp
--- a/Engine/Source/Runtime/Online/BuildPatchServices/Private/Diffing/DiffManifests.cpp
+++ b/Engine/Source/Runtime/Online/BuildPatchServices/Private/Diffing/DiffManifests.cpp
@@ -451,7 +451,7 @@ namespace BuildPatchServices
 					++NewChunksCount;
 					int32 ChunkFileSize = ManifestB->GetDataSize(ChunkB);
 					TotalChunkSize += ChunkFileSize;
-					NewChunkPaths.Add(FBuildPatchUtils::GetDataFilename(ManifestB.ToSharedRef(), TEXT("."), ChunkB));
+					NewChunkPaths.Add(FBuildPatchUtils::GetDataFilename(ManifestB.ToSharedRef(), ChunkB));
 					UE_LOG(LogDiffManifests, Verbose, TEXT("New chunk discovered: Size: %10lld, Path: %s"), ChunkFileSize, *NewChunkPaths.Last());
 				}
 			}
diff --git a/Engine/Source/Runtime/Online/BuildPatchServices/Private/Enumeration/PatchDataEnumeration.cpp b/Engine/Source/Runtime/Online/BuildPatchServices/Private/Enumeration/PatchDataEnumeration.cpp
--- a/Engine/Source/Runtime/Online/BuildPatchServices/Private/Enumeration/PatchDataEnumeration.cpp
+++ b/Engine/Source/Runtime/Online/BuildPatchServices/Private/Enumeration/PatchDataEnumeration.cpp
@@ -57,7 +57,7 @@ namespace EnumerationHelpers
 			DeDupeSet.Add(DataGuid, &bAlreadyInSet);
 			if (!bAlreadyInSet)
 			{
-				FString OutputLine = FBuildPatchUtils::GetDataFilename(Manifest, FString(), DataGuid);
+				FString OutputLine = FBuildPatchUtils::GetDataFilename(Manifest, DataGuid);
 				if (bIncludeSizes)
 				{
 					uint64 FileSize = Manifest.GetDataSize(DataGuid);
diff --git a/Engine/Source/Runtime/Online/BuildPatchServices/Private/Generation/ChunkDeltaOptimiser.cpp b/Engine/Source/Runtime/Online/BuildPatchServices/Private/Generation/ChunkDeltaOptimiser.cpp
--- a/Engine/Source/Runtime/Online/BuildPatchServices/Private/Generation/ChunkDeltaOptimiser.cpp
+++ b/Engine/Source/Runtime/Online/BuildPatchServices/Private/Generation/ChunkDeltaOptimiser.cpp
@@ -564,7 +564,7 @@ namespace BuildPatchServices
 		~FChunkDeltaOptimiser();
 
 		// IChunkDeltaOptimiser interface begin.
-		virtual	bool Run() override;
+		virtual bool Run() override;
 		// IChunkDeltaOptimiser interface end.
 
 	private:
@@ -1117,7 +1117,7 @@ namespace BuildPatchServices
 					{
 						const FChunkInfo* UpgradeChunkInfo = ManifestA->GetChunkInfo(UpgradeChunk);
 						IChunkDataAccess* UpgradeChunkDataAccess = UpgradeCloudChunkSource->Get(UpgradeChunk);
-						checkf(UpgradeChunkDataAccess != nullptr, TEXT("Failed to download chunk from source %s."), *FBuildPatchUtils::GetDataFilename(*ManifestA.Get(), Configuration.CloudDirectory, UpgradeChunk));
+						checkf(UpgradeChunkDataAccess != nullptr, TEXT("Failed to download chunk from source %s."), *FBuildPatchUtils::GetDataFilename(*ManifestA.Get(), UpgradeChunk));
 						FScopeLockedChunkData LockedChunkData(UpgradeChunkDataAccess);
 						TArray<uint8> ChunkDataArray(LockedChunkData.GetData(), LockedChunkData.GetHeader()->DataSizeUncompressed);
 						ChunkWriter->AddChunkData(MoveTemp(ChunkDataArray), UpgradeChunk, UpgradeChunkInfo->Hash, UpgradeChunkInfo->ShaHash);
diff --git a/Engine/Source/Runtime/Online/BuildPatchServices/Private/Generation/ChunkWriter.cpp b/Engine/Source/Runtime/Online/BuildPatchServices/Private/Generation/ChunkWriter.cpp
--- a/Engine/Source/Runtime/Online/BuildPatchServices/Private/Generation/ChunkWriter.cpp
+++ b/Engine/Source/Runtime/Online/BuildPatchServices/Private/Generation/ChunkWriter.cpp
@@ -165,7 +165,7 @@ namespace BuildPatchServices
 					ChunkHeader.RollingHash = ChunkHash;
 					ChunkHeader.SHAHash = ChunkSha;
 					TUniquePtr<FWriterChunkDataAccess> ChunkDataAccess(new FWriterChunkDataAccess(ChunkData, ChunkHeader));
-					const FString NewChunkFilename = FBuildPatchUtils::GetChunkNewFilename(Config.FeatureLevel, Config.ChunkDirectory, ChunkGuid, ChunkHash);
+					const FString NewChunkFilename = Config.ChunkDirectory / FBuildPatchUtils::GetChunkNewFilename(Config.FeatureLevel, ChunkGuid, ChunkHash);
 					int64 ExistingChunkFilesize = INDEX_NONE;
 					const bool bFileExists = FileSystem->GetFileSize(*NewChunkFilename, ExistingChunkFilesize);
 					if (!bFileExists)
diff --git a/Engine/Source/Runtime/Online/BuildPatchServices/Private/Generation/CloudEnumeration.cpp b/Engine/Source/Runtime/Online/BuildPatchServices/Private/Generation/CloudEnumeration.cpp
--- a/Engine/Source/Runtime/Online/BuildPatchServices/Private/Generation/CloudEnumeration.cpp
+++ b/Engine/Source/Runtime/Online/BuildPatchServices/Private/Generation/CloudEnumeration.cpp
@@ -313,7 +313,7 @@ namespace BuildPatchServices
 			}
 			for (const FGuid& ChunkToCheck : ChunksToCheck)
 			{
-				FString FinalChunkPartFilename = FBuildPatchUtils::GetDataFilename(Manifest, CloudDirectory, ChunkToCheck);
+				FString FinalChunkPartFilename = CloudDirectory / FBuildPatchUtils::GetDataFilename(Manifest, ChunkToCheck);
 				IFileManager& FileManager = IFileManager::Get();
 				TUniquePtr<FArchive> FinalChunkPartFile(FileManager.CreateFileReader(*FinalChunkPartFilename));
 				if (FinalChunkPartFile.IsValid())
diff --git a/Engine/Source/Runtime/Online/BuildPatchServices/Private/IBuildManifestSet.h b/Engine/Source/Runtime/Online/BuildPatchServices/Private/IBuildManifestSet.h
--- a/Engine/Source/Runtime/Online/BuildPatchServices/Private/IBuildManifestSet.h
+++ b/Engine/Source/Runtime/Online/BuildPatchServices/Private/IBuildManifestSet.h
@@ -55,6 +55,8 @@ namespace BuildPatchServices
 		bool IsUninstall() const { return InstallerAction.IsUninstall(); }
 		const TSet<FString>& GetInstallTags() const { return InstallTags; }
 		const TSet<FString>& GetTaggedFiles() const { return TaggedFiles; }
+		const FString& GetInstallSubdirectory() const { return InstallerAction.GetInstallSubdirectory(); }
+		const FString& GetCloudSubdirectory() const { return InstallerAction.GetCloudSubdirectory(); }
 		const FBuildPatchAppManifest& GetCurrentManifest() const { return *CurrentManifest; }
 		const FBuildPatchAppManifest& GetInstallManifest() const { return *InstallManifest; }
 		const FBuildPatchAppManifest* TryGetCurrentManifest() const { return CurrentManifest.Get(); }
@@ -93,7 +95,7 @@ namespace BuildPatchServices
 		virtual uint64 GetDownloadSize(const FGuid& DataGuid) const = 0;
 		virtual uint64 GetDownloadSize(const TSet<FGuid>& DataGuids) const = 0;
 		virtual bool GetChunkShaHash(const FGuid& ChunkGuid, FSHAHash& OutHash) const = 0;
-		virtual FString GetDataFilename(const FString& RootDirectory, const FGuid& DataGuid) const = 0;
+		virtual FString GetDataFilename(const FGuid& DataGuid) const = 0;
 
 		// File enumeration
 		virtual int32 GetNumExpectedFiles() const = 0;
diff --git a/Engine/Source/Runtime/Online/BuildPatchServices/Private/Installer/CloudChunkSource.cpp b/Engine/Source/Runtime/Online/BuildPatchServices/Private/Installer/CloudChunkSource.cpp
--- a/Engine/Source/Runtime/Online/BuildPatchServices/Private/Installer/CloudChunkSource.cpp
+++ b/Engine/Source/Runtime/Online/BuildPatchServices/Private/Installer/CloudChunkSource.cpp
@@ -493,7 +493,7 @@ namespace BuildPatchServices
 
 				const FTaskInfo* Info = TaskInfos.Find(ChunkUriToRequest);
 				ChunkUriRequest.CloudDirectory = GetCloudRoot(Info ? Info->RetryNum : 0 );
-				ChunkUriRequest.RelativePath = ManifestSet->GetDataFilename(TEXT(""), ChunkUriToRequest);
+				ChunkUriRequest.RelativePath = ManifestSet->GetDataFilename(ChunkUriToRequest);
 				ChunkUriRequest.RelativePath.RemoveFromStart(TEXT("/"));
 
 				MessagePump->SendRequest(ChunkUriRequest, [WeakChunkUriResponses, ChunkUriToRequest](FChunkUriResponse Response)
@@ -588,11 +588,13 @@ namespace BuildPatchServices
 					else
 					{
 						CloudChunkSourceStat->OnDownloadCorrupt(DownloadId, TaskInfo.UrlUsed, LoadResult);
+						UE_LOG(LogCloudChunkSource, Error, TEXT("CORRUPT: %s"), *TaskInfo.UrlUsed);
 					}
 				}
 				else
 				{
 					CloudChunkSourceStat->OnDownloadFailed(DownloadId, TaskInfo.UrlUsed);
+					UE_LOG(LogCloudChunkSource, Error, TEXT("FAILED: %s"), *TaskInfo.UrlUsed);
 				}
 
 				// Handle failed
diff --git a/Engine/Source/Runtime/Online/BuildPatchServices/Private/Installer/Prerequisites.cpp b/Engine/Source/Runtime/Online/BuildPatchServices/Private/Installer/Prerequisites.cpp
--- a/Engine/Source/Runtime/Online/BuildPatchServices/Private/Installer/Prerequisites.cpp
+++ b/Engine/Source/Runtime/Online/BuildPatchServices/Private/Installer/Prerequisites.cpp
@@ -75,7 +75,6 @@ namespace BuildPatchServices
 			if (PreReq.Path.IsEmpty())
 			{
 				UE_LOG(LogPrerequisites, Log, TEXT("Skipping prerequisites as manifest does not have prerequisites specified."));
-				BuildProgress.SetStateProgress(EBuildPatchState::PrerequisitesInstall, 1.0f);
 				continue;
 			}
 
@@ -101,8 +100,7 @@ namespace BuildPatchServices
 				if (MissingPrereqs.Num() == 0)
 				{
 					UE_LOG(LogPrerequisites, Log, TEXT("Skipping already installed prerequisites installer"));
-					BuildProgress.SetStateProgress(EBuildPatchState::PrerequisitesInstall, 1.0f);
-					return true;
+					continue;
 				}
 			}
 
@@ -137,7 +135,8 @@ namespace BuildPatchServices
 			{
 				UE_LOG(LogPrerequisites, Error, TEXT("Could not find prerequisites file %s on disk."), *PrereqPath);
 				InstallerError->SetError(EBuildPatchInstallError::PrerequisiteError, PrerequisiteErrorPrefixes::NotFoundCode);
-				return false;
+				bInstallSuccessful = false;
+				continue;
 			}
 
 			const FString PrereqCommandline = PreReq.Args
@@ -172,13 +171,15 @@ namespace BuildPatchServices
 			if (bPrereqInstallSuccessful)
 			{
 				UE_LOG(LogPrerequisites, Log, TEXT("Prerequisite installation successful"));
-				BuildProgress.SetStateProgress(EBuildPatchState::PrerequisitesInstall, 1.0f);
 				InstalledPrereqs.Append(PreReq.IdSet);
 				MachineConfig->SaveInstalledPrereqIds(InstalledPrereqs);
 			}
 			bInstallSuccessful = bInstallSuccessful && bPrereqInstallSuccessful;
 		}
-
+		if (bInstallSuccessful)
+		{
+			BuildProgress.SetStateProgress(EBuildPatchState::PrerequisitesInstall, 1.0f);
+		}
 		return bInstallSuccessful;
 	}
 
diff --git a/Engine/Source/Runtime/Online/BuildPatchServices/Public/BuildPatchDelta.h b/Engine/Source/Runtime/Online/BuildPatchServices/Public/BuildPatchDelta.h
--- a/Engine/Source/Runtime/Online/BuildPatchServices/Public/BuildPatchDelta.h
+++ b/Engine/Source/Runtime/Online/BuildPatchServices/Public/BuildPatchDelta.h
@@ -14,35 +14,16 @@ namespace BuildPatchServices
 	enum class EDeltaPolicy : uint32
 	{
 		// Try to fetch, but continue without if request fail.
-		TryFetchContinueWithout,
+		TryFetchContinueWithout = 0,
 
 		// Expect the delta to exist, hard fail the installation if it could not be retrieved.
 		Expect,
 
 		// Expect the delta to not exist, skipping any attempt to use one.
 		Skip,
-	};
-
-	/**
-	 * Returns the string representation of the EDeltaPolicy value. Used for analytics and logging only.
-	 * @param DeltaPolicy     The value.
-	 * @return the enum's string representation.
-	 */
-	inline const FString& EnumToString(const EDeltaPolicy& DeltaPolicy)
-	{
-		static const FString TryFetchContinueWithout(TEXT("EDeltaPolicy::TryFetchContinueWithout"));
-		static const FString Expect(TEXT("EDeltaPolicy::Expect"));
-		static const FString Skip(TEXT("EDeltaPolicy::Skip"));
-		static const FString InvalidOrMax(TEXT("InvalidOrMax"));
 
-		switch (DeltaPolicy)
-		{
-			case EDeltaPolicy::TryFetchContinueWithout: return TryFetchContinueWithout;
-			case EDeltaPolicy::Expect: return Expect;
-			case EDeltaPolicy::Skip: return Skip;
-			default: return InvalidOrMax;
-		}
-	}
+		InvalidOrMax
+	};
 
 	/**
 	 * Based on the source and destination manifests, get the filename for the delta that optimises patching from source to destination.
@@ -54,3 +35,35 @@ namespace BuildPatchServices
 
 	BUILDPATCHSERVICES_API IBuildManifestPtr MergeDeltaManifest(const IBuildManifestRef& Manifest, const IBuildManifestRef& Delta);
 }
+
+static_assert((uint32)BuildPatchServices::EDeltaPolicy::InvalidOrMax == 3, "Please add support for the extra values to the Lex functions below.");
+
+inline const TCHAR* LexToString(BuildPatchServices::EDeltaPolicy DeltaPolicy)
+{
+#define CASE_ENUM_TO_STR(Value) case BuildPatchServices::EDeltaPolicy::Value: return TEXT(#Value)
+		switch (DeltaPolicy)
+		{
+		CASE_ENUM_TO_STR(TryFetchContinueWithout);
+		CASE_ENUM_TO_STR(Expect);
+		CASE_ENUM_TO_STR(Skip);
+	default: return TEXT("InvalidOrMax");
+		}
+#undef CASE_ENUM_TO_STR
+}
+
+inline void LexFromString(BuildPatchServices::EDeltaPolicy& DeltaPolicy, const TCHAR* Buffer)
+{
+#define RETURN_IF_EQUAL(Value) if (FCString::Stricmp(Buffer, TEXT(#Value)) == 0) { DeltaPolicy = BuildPatchServices::EDeltaPolicy::Value; return; }
+	const TCHAR* const Prefix = TEXT("EDeltaPolicy::");
+	const SIZE_T PrefixLen = FCString::Strlen(Prefix);
+	if (FCString::Strnicmp(Buffer, Prefix, PrefixLen) == 0)
+	{
+		Buffer += PrefixLen;
+	}
+	RETURN_IF_EQUAL(TryFetchContinueWithout);
+	RETURN_IF_EQUAL(Expect);
+	RETURN_IF_EQUAL(Skip);
+	// Did not match
+	DeltaPolicy = BuildPatchServices::EDeltaPolicy::InvalidOrMax;
+#undef RETURN_IF_EQUAL
+}
diff --git a/Engine/Source/Runtime/Online/BuildPatchServices/Public/BuildPatchInstall.h b/Engine/Source/Runtime/Online/BuildPatchServices/Public/BuildPatchInstall.h
--- a/Engine/Source/Runtime/Online/BuildPatchServices/Public/BuildPatchInstall.h
+++ b/Engine/Source/Runtime/Online/BuildPatchServices/Public/BuildPatchInstall.h
@@ -15,7 +15,7 @@ namespace BuildPatchServices
 	enum class EInstallMode : uint32
 	{
 		// Construct all required files, but only stage them ready to be completed later.
-		StageFiles,
+		StageFiles = 0,
 
 		// Full installation, allowing immediate changes to be made to existing files. The installation is unusable until complete.
 		DestructiveInstall,
@@ -24,32 +24,44 @@ namespace BuildPatchServices
 		NonDestructiveInstall,
 
 		// Execute the prerequisite installer only, downloading it first if necessary. If the specified manifest has no prerequisites, this will result in an error.
-		PrereqOnly
+		PrereqOnly,
+
+		InvalidOrMax
 	};
 
-	/**
-	 * Returns the string representation of the EInstallMode value. Used for analytics and logging only.
-	 * @param InstallMode     The value.
-	 * @return the enum's string representation.
-	 */
-	inline const FString& EnumToString(const EInstallMode& InstallMode)
-	{
-		// Const enum strings, special case no error.
-		static const FString StageFiles(TEXT("EInstallMode::StageFiles"));
-		static const FString DestructiveInstall(TEXT("EInstallMode::DestructiveInstall"));
-		static const FString NonDestructiveInstall(TEXT("EInstallMode::NonDestructiveInstall"));
-		static const FString PrereqOnly(TEXT("EInstallMode::PrereqOnly"));
-		static const FString InvalidOrMax(TEXT("InvalidOrMax"));
+	BUILDPATCHSERVICES_API uint64 CalculateRequiredDiskSpace(const IBuildManifestPtr& CurrentManifest, const IBuildManifestRef& BuildManifest, const EInstallMode& InstallMode, const TSet<FString>& InstallTags);
+}
+
+static_assert((uint32)BuildPatchServices::EInstallMode::InvalidOrMax == 4, "Please add support for the extra values to the Lex functions below.");
 
+inline const TCHAR* LexToString(BuildPatchServices::EInstallMode InstallMode)
+{
+#define CASE_ENUM_TO_STR(Value) case BuildPatchServices::EInstallMode::Value: return TEXT(#Value)
 		switch (InstallMode)
 		{
-			case EInstallMode::StageFiles: return StageFiles;
-			case EInstallMode::DestructiveInstall: return DestructiveInstall;
-			case EInstallMode::NonDestructiveInstall: return NonDestructiveInstall;
-			case EInstallMode::PrereqOnly: return PrereqOnly;
-			default: return InvalidOrMax;
+		CASE_ENUM_TO_STR(StageFiles);
+		CASE_ENUM_TO_STR(DestructiveInstall);
+		CASE_ENUM_TO_STR(NonDestructiveInstall);
+		CASE_ENUM_TO_STR(PrereqOnly);
+	default: return TEXT("InvalidOrMax");
 		}
-	}
+#undef CASE_ENUM_TO_STR
+}
 
-	BUILDPATCHSERVICES_API uint64 CalculateRequiredDiskSpace(const IBuildManifestPtr& CurrentManifest, const IBuildManifestRef& BuildManifest, const EInstallMode& InstallMode, const TSet<FString>& InstallTags);
+inline void LexFromString(BuildPatchServices::EInstallMode& InstallMode, const TCHAR* Buffer)
+{
+#define RETURN_IF_EQUAL(Value) if (FCString::Stricmp(Buffer, TEXT(#Value)) == 0) { InstallMode = BuildPatchServices::EInstallMode::Value; return; }
+	const TCHAR* const Prefix = TEXT("EInstallMode::");
+	const SIZE_T PrefixLen = FCString::Strlen(Prefix);
+	if (FCString::Strnicmp(Buffer, Prefix, PrefixLen) == 0)
+	{
+		Buffer += PrefixLen;
+	}
+	RETURN_IF_EQUAL(StageFiles);
+	RETURN_IF_EQUAL(DestructiveInstall);
+	RETURN_IF_EQUAL(NonDestructiveInstall);
+	RETURN_IF_EQUAL(PrereqOnly);
+	// Did not match
+	InstallMode = BuildPatchServices::EInstallMode::InvalidOrMax;
+#undef RETURN_IF_EQUAL
 }
diff --git a/Engine/Source/Runtime/Online/BuildPatchServices/Public/BuildPatchSettings.h b/Engine/Source/Runtime/Online/BuildPatchServices/Public/BuildPatchSettings.h
--- a/Engine/Source/Runtime/Online/BuildPatchServices/Public/BuildPatchSettings.h
+++ b/Engine/Source/Runtime/Online/BuildPatchServices/Public/BuildPatchSettings.h
@@ -50,51 +50,61 @@ namespace BuildPatchServices
 		 * Creates an install action.
 		 * @param Manifest          The manifest for the build to be installed.
 		 * @param InstallTags       The install tags to use if selectively installing files. If empty set, all files will be installed.
+		 * @param InstallSubdirectory  The subdirectory to install this data to inside of the main install directory.
+		 * @param CloudSubdirectory    The subdirectory of the cloud distribution root that this patch data should be sourced.
 		 * @return the action setup for performing an installation.
 		 */
-		static BUILDPATCHSERVICES_API FInstallerAction MakeInstall(const IBuildManifestRef& Manifest, TSet<FString> InstallTags = TSet<FString>());
+		static BUILDPATCHSERVICES_API FInstallerAction MakeInstall(const IBuildManifestRef& Manifest, TSet<FString> InstallTags = TSet<FString>(), FString InstallSubdirectory = FString(), FString CloudSubdirectory = FString());
 
 		/**
 		 * Creates an update action.
 		 * @param CurrentManifest   The manifest for the build currently installed.
 		 * @param InstallManifest   The manifest for the build to be installed.
 		 * @param InstallTags       The install tags to use if selectively installing files. If empty set, all files will be updated, or added if missing.
+		 * @param InstallSubdirectory  The subdirectory to install this data to inside of the main install directory.
+		 * @param CloudSubdirectory    The subdirectory of the cloud distribution root that this patch data should be sourced.
 		 * @return the action setup for performing an update.
 		 */
-		static BUILDPATCHSERVICES_API FInstallerAction MakeUpdate(const IBuildManifestRef& CurrentManifest, const IBuildManifestRef& InstallManifest, TSet<FString> InstallTags = TSet<FString>());
+		static BUILDPATCHSERVICES_API FInstallerAction MakeUpdate(const IBuildManifestRef& CurrentManifest, const IBuildManifestRef& InstallManifest, TSet<FString> InstallTags = TSet<FString>(), FString InstallSubdirectory = FString(), FString CloudSubdirectory = FString());
 
 		/**
 		 * Creates an install action.
 		 * @param Manifest          The manifest for the build to be installed.
 		 * @param InstallTags       The install tags to use if selectively installing files. If empty set, all files will be repaired.
+		 * @param InstallSubdirectory  The subdirectory to install this data to inside of the main install directory.
+		 * @param CloudSubdirectory    The subdirectory of the cloud distribution root that this patch data should be sourced.
 		 * @return the action setup forcing an SHA check, and repair of all tagged files.
 		 */
-		static BUILDPATCHSERVICES_API FInstallerAction MakeRepair(const IBuildManifestRef& Manifest, TSet<FString> InstallTags = TSet<FString>());
+		static BUILDPATCHSERVICES_API FInstallerAction MakeRepair(const IBuildManifestRef& Manifest, TSet<FString> InstallTags = TSet<FString>(), FString InstallSubdirectory = FString(), FString CloudSubdirectory = FString());
 
 
 		/**
 		 * Creates an uninstall action.
 		 * @param Manifest          The manifest for the build currently installed.
+		 * @param InstallSubdirectory  The subdirectory to install this data to inside of the main install directory.
+		 * @param CloudSubdirectory    The subdirectory of the cloud distribution root that this patch data should be sourced.
 		 * @return the action setup for performing an uninstall, deleting all files referenced by the manifest.
 		 */
-		static BUILDPATCHSERVICES_API FInstallerAction MakeUninstall(const IBuildManifestRef& Manifest);
+		static BUILDPATCHSERVICES_API FInstallerAction MakeUninstall(const IBuildManifestRef& Manifest, FString InstallSubdirectory = FString(), FString CloudSubdirectory = FString());
 
 		/**
 		 * Helper for creating an install action or update action based on validity of CurrentManifest.
 		 * @param CurrentManifest   The manifest for the build currently installed, invalid if no build installed.
 		 * @param InstallManifest   The manifest for the build to be installed.
 		 * @param InstallTags       The install tags to use if selectively installing files. If empty set, all files will be installed/updated, or added if missing.
+		 * @param InstallSubdirectory  The subdirectory to install this data to inside of the main install directory.
+		 * @param CloudSubdirectory    The subdirectory of the cloud distribution root that this patch data should be sourced.
 		 * @return the action setup for performing an install or update.
 		 */
-		static FInstallerAction MakeInstallOrUpdate(const IBuildManifestPtr& CurrentManifest, const IBuildManifestRef& InstallManifest, TSet<FString> InstallTags = TSet<FString>())
+		static FInstallerAction MakeInstallOrUpdate(const IBuildManifestPtr& CurrentManifest, const IBuildManifestRef& InstallManifest, TSet<FString> InstallTags = TSet<FString>(), FString InstallSubdirectory = FString(), FString CloudSubdirectory = FString())
 		{
 			if (!CurrentManifest.IsValid())
 			{
-				return MakeInstall(InstallManifest, MoveTemp(InstallTags));
+				return MakeInstall(InstallManifest, MoveTemp(InstallTags), MoveTemp(InstallSubdirectory), MoveTemp(CloudSubdirectory));
 			}
 			else
 			{
-				return MakeUpdate(CurrentManifest.ToSharedRef(), InstallManifest, MoveTemp(InstallTags));
+				return MakeUpdate(CurrentManifest.ToSharedRef(), InstallManifest, MoveTemp(InstallTags), MoveTemp(InstallSubdirectory), MoveTemp(CloudSubdirectory));
 			}
 		}
 
@@ -134,6 +144,16 @@ namespace BuildPatchServices
 		 */
 		BUILDPATCHSERVICES_API const TSet<FString>& GetInstallTags() const;
 
+		/**
+		 * @return the install subdirectory for the action.
+		 */
+		BUILDPATCHSERVICES_API const FString& GetInstallSubdirectory() const;
+
+		/**
+		 * @return the cloud subdirectory for the action.
+		 */
+		BUILDPATCHSERVICES_API const FString& GetCloudSubdirectory() const;
+
 		/**
 		 * @return the manifest for the current installation, this will runtime assert if called invalidly (see TryGetCurrentManifest).
 		 */
@@ -194,6 +214,8 @@ namespace BuildPatchServices
 		IBuildManifestPtr CurrentManifest;
 		IBuildManifestPtr InstallManifest;
 		TSet<FString> InstallTags;
+		FString InstallSubdirectory;
+		FString CloudSubdirectory;
 		EInstallActionIntent ActionIntent;
 	};
 
diff --git a/Engine/Source/Runtime/Online/BuildPatchServices/Public/BuildPatchState.h b/Engine/Source/Runtime/Online/BuildPatchServices/Public/BuildPatchState.h
--- a/Engine/Source/Runtime/Online/BuildPatchServices/Public/BuildPatchState.h
+++ b/Engine/Source/Runtime/Online/BuildPatchServices/Public/BuildPatchState.h
@@ -54,13 +54,6 @@ namespace BuildPatchServices
 		NUM_PROGRESS_STATES,
 	};
 
-	/**
-	 * Returns the string representation of the EBuildPatchState value. Used for analytics and logging only.
-	 * @param State - The value.
-	 * @return The enum's string value.
-	 */
-	BUILDPATCHSERVICES_API const FString& StateToString(const EBuildPatchState& State);
-
 	/**
 	 * Returns the FText representation of the specified EBuildPatchState value. Used for displaying to the user.
 	 * @param State - The error type value.
@@ -69,3 +62,54 @@ namespace BuildPatchServices
 	UE_DEPRECATED(4.21, "BuildPatchServices::StateToText(const EBuildPatchState& State) has been deprecated. It will no longer be supported by BuildPatchServices in the future.")
 	BUILDPATCHSERVICES_API const FText& StateToText(const EBuildPatchState& State);
 }
+
+static_assert((uint32)BuildPatchServices::EBuildPatchState::NUM_PROGRESS_STATES == 12, "Please add support for the extra values to the Lex functions below.");
+
+inline const TCHAR* LexToString(BuildPatchServices::EBuildPatchState State)
+{
+#define CASE_ENUM_TO_STR(Value) case BuildPatchServices::EBuildPatchState::Value: return TEXT(#Value)
+	switch (State)
+	{
+		CASE_ENUM_TO_STR(Queued);
+		CASE_ENUM_TO_STR(Initializing);
+		CASE_ENUM_TO_STR(Resuming);
+		CASE_ENUM_TO_STR(Downloading);
+		CASE_ENUM_TO_STR(Installing);
+		CASE_ENUM_TO_STR(MovingToInstall);
+		CASE_ENUM_TO_STR(SettingAttributes);
+		CASE_ENUM_TO_STR(BuildVerification);
+		CASE_ENUM_TO_STR(CleanUp);
+		CASE_ENUM_TO_STR(PrerequisitesInstall);
+		CASE_ENUM_TO_STR(Completed);
+		CASE_ENUM_TO_STR(Paused);
+	default: return TEXT("InvalidOrMax");
+	}
+#undef CASE_ENUM_TO_STR
+}
+
+inline void LexFromString(BuildPatchServices::EBuildPatchState& Error, const TCHAR* Buffer)
+{
+#define RETURN_IF_EQUAL(Value) if (FCString::Stricmp(Buffer, TEXT(#Value)) == 0) { Error = BuildPatchServices::EBuildPatchState::Value; return; }
+	const TCHAR* const Prefix = TEXT("EBuildPatchState::");
+	const SIZE_T PrefixLen = FCString::Strlen(Prefix);
+	if (FCString::Strnicmp(Buffer, Prefix, PrefixLen) == 0)
+	{
+		Buffer += PrefixLen;
+	}
+	RETURN_IF_EQUAL(Queued);
+	RETURN_IF_EQUAL(Initializing);
+	RETURN_IF_EQUAL(Resuming);
+	RETURN_IF_EQUAL(Downloading);
+	RETURN_IF_EQUAL(Installing);
+	RETURN_IF_EQUAL(MovingToInstall);
+	RETURN_IF_EQUAL(SettingAttributes);
+	RETURN_IF_EQUAL(BuildVerification);
+	RETURN_IF_EQUAL(CleanUp);
+	RETURN_IF_EQUAL(PrerequisitesInstall);
+	RETURN_IF_EQUAL(Completed);
+	RETURN_IF_EQUAL(Paused);
+	// Did not match
+	Error = BuildPatchServices::EBuildPatchState::NUM_PROGRESS_STATES;
+	return;
+#undef RETURN_IF_EQUAL
+}
\ No newline at end of file
diff --git a/Engine/Source/Runtime/Online/BuildPatchServices/Public/BuildPatchVerify.h b/Engine/Source/Runtime/Online/BuildPatchServices/Public/BuildPatchVerify.h
--- a/Engine/Source/Runtime/Online/BuildPatchServices/Public/BuildPatchVerify.h
+++ b/Engine/Source/Runtime/Online/BuildPatchServices/Public/BuildPatchVerify.h
@@ -12,7 +12,7 @@ namespace BuildPatchServices
 	enum class EVerifyMode : uint32
 	{
 		// Fully SHA checks all files in the build.
-		ShaVerifyAllFiles,
+		ShaVerifyAllFiles = 0,
 
 		// Fully SHA checks only files touched by the install/patch process.
 		ShaVerifyTouchedFiles,
@@ -21,7 +21,9 @@ namespace BuildPatchServices
 		FileSizeCheckAllFiles,
 
 		// Checks just the existence and file size of only files touched by the install/patch process.
-		FileSizeCheckTouchedFiles
+		FileSizeCheckTouchedFiles,
+
+		InvalidOrMax
 	};
 
 	/**
@@ -30,7 +32,7 @@ namespace BuildPatchServices
 	enum class EVerifyError : uint32
 	{
 		// The file was not found.
-		FileMissing,
+		FileMissing = 0,
 
 		// The file failed to open.
 		OpenFileFailed,
@@ -39,32 +41,78 @@ namespace BuildPatchServices
 		HashCheckFailed,
 
 		// The file was not the expected size.
-		FileSizeFailed
+		FileSizeFailed,
+
+		InvalidOrMax
 	};
+}
 
-	/**
-	 * Returns the string representation of the EVerifyMode value. Used for analytics and logging only.
-	 * @param  VerifyMode     The value.
-	 * @return the enum's string representation.
-	 */
-	inline const FString& EnumToString(const EVerifyMode& VerifyMode)
+ENUM_RANGE_BY_FIRST_AND_LAST(BuildPatchServices::EVerifyError, BuildPatchServices::EVerifyError::FileMissing, BuildPatchServices::EVerifyError::FileSizeFailed)
+
+static_assert((uint32)BuildPatchServices::EVerifyMode::InvalidOrMax == 4, "Please add support for the extra values to the Lex functions below.");
+
+inline const TCHAR* LexToString(BuildPatchServices::EVerifyMode VerifyMode)
+{
+#define CASE_ENUM_TO_STR(Value) case BuildPatchServices::EVerifyMode::Value: return TEXT(#Value)
+	switch (VerifyMode)
 	{
-		// Const enum strings, special case no error.
-		static const FString ShaVerifyAllFiles(TEXT("EVerifyMode::ShaVerifyAllFiles"));
-		static const FString ShaVerifyTouchedFiles(TEXT("EVerifyMode::ShaVerifyTouchedFiles"));
-		static const FString FileSizeCheckAllFiles(TEXT("EVerifyMode::FileSizeCheckAllFiles"));
-		static const FString FileSizeCheckTouchedFiles(TEXT("EVerifyMode::FileSizeCheckTouchedFiles"));
-		static const FString InvalidOrMax(TEXT("InvalidOrMax"));
-
-		switch (VerifyMode)
-		{
-			case EVerifyMode::ShaVerifyAllFiles: return ShaVerifyAllFiles;
-			case EVerifyMode::ShaVerifyTouchedFiles: return ShaVerifyTouchedFiles;
-			case EVerifyMode::FileSizeCheckAllFiles: return FileSizeCheckAllFiles;
-			case EVerifyMode::FileSizeCheckTouchedFiles: return FileSizeCheckTouchedFiles;
-			default: return InvalidOrMax;
-		}
+		CASE_ENUM_TO_STR(ShaVerifyAllFiles);
+		CASE_ENUM_TO_STR(ShaVerifyTouchedFiles);
+		CASE_ENUM_TO_STR(FileSizeCheckAllFiles);
+		CASE_ENUM_TO_STR(FileSizeCheckTouchedFiles);
+	default: return TEXT("InvalidOrMax");
 	}
+#undef CASE_ENUM_TO_STR
 }
 
-ENUM_RANGE_BY_FIRST_AND_LAST(BuildPatchServices::EVerifyError, BuildPatchServices::EVerifyError::FileMissing, BuildPatchServices::EVerifyError::FileSizeFailed)
\ No newline at end of file
+inline void LexFromString(BuildPatchServices::EVerifyMode& VerifyMode, const TCHAR* Buffer)
+{
+#define RETURN_IF_EQUAL(Value) if (FCString::Stricmp(Buffer, TEXT(#Value)) == 0) { VerifyMode = BuildPatchServices::EVerifyMode::Value; return; }
+	const TCHAR* const Prefix = TEXT("EVerifyMode::");
+	const SIZE_T PrefixLen = FCString::Strlen(Prefix);
+	if (FCString::Strnicmp(Buffer, Prefix, PrefixLen) == 0)
+	{
+		Buffer += PrefixLen;
+	}
+	RETURN_IF_EQUAL(ShaVerifyAllFiles);
+	RETURN_IF_EQUAL(ShaVerifyTouchedFiles);
+	RETURN_IF_EQUAL(FileSizeCheckAllFiles);
+	RETURN_IF_EQUAL(FileSizeCheckTouchedFiles);
+	// Did not match
+	VerifyMode = BuildPatchServices::EVerifyMode::InvalidOrMax;
+#undef RETURN_IF_EQUAL
+}
+
+static_assert((uint32)BuildPatchServices::EVerifyError::InvalidOrMax == 4, "Please add support for the extra values to the Lex functions below.");
+
+inline const TCHAR* LexToString(BuildPatchServices::EVerifyError VerifyError)
+{
+#define CASE_ENUM_TO_STR(Value) case BuildPatchServices::EVerifyError::Value: return TEXT(#Value)
+	switch (VerifyError)
+	{
+		CASE_ENUM_TO_STR(FileMissing);
+		CASE_ENUM_TO_STR(OpenFileFailed);
+		CASE_ENUM_TO_STR(HashCheckFailed);
+		CASE_ENUM_TO_STR(FileSizeFailed);
+	default: return TEXT("InvalidOrMax");
+	}
+#undef CASE_ENUM_TO_STR
+}
+
+inline void LexFromString(BuildPatchServices::EVerifyError& VerifyError, const TCHAR* Buffer)
+{
+#define RETURN_IF_EQUAL(Value) if (FCString::Stricmp(Buffer, TEXT(#Value)) == 0) { VerifyError = BuildPatchServices::EVerifyError::Value; return; }
+	const TCHAR* const Prefix = TEXT("EVerifyError::");
+	const SIZE_T PrefixLen = FCString::Strlen(Prefix);
+	if (FCString::Strnicmp(Buffer, Prefix, PrefixLen) == 0)
+	{
+		Buffer += PrefixLen;
+	}
+	RETURN_IF_EQUAL(FileMissing);
+	RETURN_IF_EQUAL(OpenFileFailed);
+	RETURN_IF_EQUAL(HashCheckFailed);
+	RETURN_IF_EQUAL(FileSizeFailed);
+	// Did not match
+	VerifyError = BuildPatchServices::EVerifyError::InvalidOrMax;
+#undef RETURN_IF_EQUAL
+}
diff --git a/Engine/Source/Runtime/Online/BuildPatchServices/Public/Interfaces/IBuildInstaller.h b/Engine/Source/Runtime/Online/BuildPatchServices/Public/Interfaces/IBuildInstaller.h
--- a/Engine/Source/Runtime/Online/BuildPatchServices/Public/Interfaces/IBuildInstaller.h
+++ b/Engine/Source/Runtime/Online/BuildPatchServices/Public/Interfaces/IBuildInstaller.h
@@ -61,34 +61,6 @@ enum class EBuildPatchInstallError
 	NumInstallErrors
 };
 
-/**
- * @return the stringified version of the EBuildPatchInstallError enum passed in
- */
-inline const TCHAR* LexToString(EBuildPatchInstallError EnumVal)
-{
-	static_assert((int32)EBuildPatchInstallError::NumInstallErrors == 12, "Please add support for the extra values to the function below.");
-#define CASE_ENUM_SET(State) case State: return TEXT(#State); break;
-	switch (EnumVal)
-	{
-		CASE_ENUM_SET(EBuildPatchInstallError::NoError)
-		CASE_ENUM_SET(EBuildPatchInstallError::DownloadError)
-		CASE_ENUM_SET(EBuildPatchInstallError::FileConstructionFail)
-		CASE_ENUM_SET(EBuildPatchInstallError::MoveFileToInstall)
-		CASE_ENUM_SET(EBuildPatchInstallError::BuildVerifyFail)
-		CASE_ENUM_SET(EBuildPatchInstallError::ApplicationClosing)
-		CASE_ENUM_SET(EBuildPatchInstallError::ApplicationError)
-		CASE_ENUM_SET(EBuildPatchInstallError::UserCanceled)
-		CASE_ENUM_SET(EBuildPatchInstallError::PrerequisiteError)
-		CASE_ENUM_SET(EBuildPatchInstallError::InitializationError)
-		CASE_ENUM_SET(EBuildPatchInstallError::PathLengthExceeded)
-		CASE_ENUM_SET(EBuildPatchInstallError::OutOfDiskSpace)
-			
-		default:
-			return TEXT("EBuildPatchInstallError::Unknown");
-	}
-#undef CASE_ENUM_SET
-};
-
 /**
  * Declares the error code prefixes for each error type enum
  */
@@ -495,3 +467,90 @@ public:
 	virtual const BuildPatchServices::FBuildInstallerConfiguration& GetConfiguration() const = 0;
 };
 
+static_assert((uint32)EBuildPatchInstallError::NumInstallErrors == 12, "Please add support for the extra values to the Lex functions below.");
+
+inline const TCHAR* LexToString(EBuildPatchInstallError Error)
+{
+#define CASE_ENUM_TO_STR(Value) case EBuildPatchInstallError::Value: return TEXT(#Value)
+	switch (Error)
+	{
+		CASE_ENUM_TO_STR(NoError);
+		CASE_ENUM_TO_STR(DownloadError);
+		CASE_ENUM_TO_STR(FileConstructionFail);
+		CASE_ENUM_TO_STR(MoveFileToInstall);
+		CASE_ENUM_TO_STR(BuildVerifyFail);
+		CASE_ENUM_TO_STR(ApplicationClosing);
+		CASE_ENUM_TO_STR(ApplicationError);
+		CASE_ENUM_TO_STR(UserCanceled);
+		CASE_ENUM_TO_STR(PrerequisiteError);
+		CASE_ENUM_TO_STR(InitializationError);
+		CASE_ENUM_TO_STR(PathLengthExceeded);
+		CASE_ENUM_TO_STR(OutOfDiskSpace);
+	default: return TEXT("InvalidOrMax");
+	}
+#undef CASE_ENUM_TO_STR
+}
+
+inline void LexFromString(EBuildPatchInstallError& Error, const TCHAR* Buffer)
+{
+#define RETURN_IF_EQUAL(Value) if (FCString::Stricmp(Buffer, TEXT(#Value)) == 0) { Error = EBuildPatchInstallError::Value; return; }
+	const TCHAR* const Prefix = TEXT("EBuildPatchInstallError::");
+	const SIZE_T PrefixLen = FCString::Strlen(Prefix);
+	if (FCString::Strnicmp(Buffer, Prefix, PrefixLen) == 0)
+	{
+		Buffer += PrefixLen;
+	}
+	RETURN_IF_EQUAL(NoError);
+	RETURN_IF_EQUAL(DownloadError);
+	RETURN_IF_EQUAL(FileConstructionFail);
+	RETURN_IF_EQUAL(MoveFileToInstall);
+	RETURN_IF_EQUAL(BuildVerifyFail);
+	RETURN_IF_EQUAL(ApplicationClosing);
+	RETURN_IF_EQUAL(ApplicationError);
+	RETURN_IF_EQUAL(UserCanceled);
+	RETURN_IF_EQUAL(PrerequisiteError);
+	RETURN_IF_EQUAL(InitializationError);
+	RETURN_IF_EQUAL(PathLengthExceeded);
+	RETURN_IF_EQUAL(OutOfDiskSpace);
+	// Did not match
+	Error = EBuildPatchInstallError::NumInstallErrors;
+	return;
+#undef RETURN_IF_EQUAL
+}
+
+static_assert((uint32)EBuildPatchDownloadHealth::NUM_Values == 5, "Please add support for the extra values to the Lex functions below.");
+
+inline const TCHAR* LexToString(EBuildPatchDownloadHealth Error)
+{
+#define CASE_ENUM_TO_STR(Value) case EBuildPatchDownloadHealth::Value: return TEXT(#Value)
+	switch (Error)
+	{
+		CASE_ENUM_TO_STR(Disconnected);
+		CASE_ENUM_TO_STR(Poor);
+		CASE_ENUM_TO_STR(OK);
+		CASE_ENUM_TO_STR(Good);
+		CASE_ENUM_TO_STR(Excellent);
+	default: return TEXT("InvalidOrMax");
+	}
+#undef CASE_ENUM_TO_STR
+}
+
+inline void LexFromString(EBuildPatchDownloadHealth& Error, const TCHAR* Buffer)
+{
+#define RETURN_IF_EQUAL(Value) if (FCString::Stricmp(Buffer, TEXT(#Value)) == 0) { Error = EBuildPatchDownloadHealth::Value; return; }
+	const TCHAR* const Prefix = TEXT("EBuildPatchDownloadHealth::");
+	const SIZE_T PrefixLen = FCString::Strlen(Prefix);
+	if (FCString::Strnicmp(Buffer, Prefix, PrefixLen) == 0)
+	{
+		Buffer += PrefixLen;
+	}
+	RETURN_IF_EQUAL(Disconnected);
+	RETURN_IF_EQUAL(Poor);
+	RETURN_IF_EQUAL(OK);
+	RETURN_IF_EQUAL(Good);
+	RETURN_IF_EQUAL(Excellent);
+	// Did not match
+	Error = EBuildPatchDownloadHealth::NUM_Values;
+	return;
+#undef RETURN_IF_EQUAL
+}
diff --git a/Engine/Source/Runtime/Online/BuildPatchServices/Public/Interfaces/IBuildPatchServicesModule.h b/Engine/Source/Runtime/Online/BuildPatchServices/Public/Interfaces/IBuildPatchServicesModule.h
--- a/Engine/Source/Runtime/Online/BuildPatchServices/Public/Interfaces/IBuildPatchServicesModule.h
+++ b/Engine/Source/Runtime/Online/BuildPatchServices/Public/Interfaces/IBuildPatchServicesModule.h
@@ -48,6 +48,7 @@ public:
 	/**
 	 * Factory providing construction of a build installer class.
 	 * @param Configuration     The configuration for the installer.
+	 * @param OnComplete        The delegate that will be called when the installer completes.
 	 * @return an instance of an IBuildInstaller implementation.
 	 */
 	virtual IBuildInstallerRef CreateBuildInstaller(BuildPatchServices::FBuildInstallerConfiguration Configuration, FBuildPatchInstallerDelegate OnComplete) const = 0;
