[Engine]
* Forced signature of lambda to return FString by const reference instead of value.. removes ~17 million memory allocations when starting game

diff --git a/Engine/Source/Runtime/Online/BuildPatchServices/Private/Data/ManifestData.cpp b/Engine/Source/Runtime/Online/BuildPatchServices/Private/Data/ManifestData.cpp
--- a/Engine/Source/Runtime/Online/BuildPatchServices/Private/Data/ManifestData.cpp
+++ b/Engine/Source/Runtime/Online/BuildPatchServices/Private/Data/ManifestData.cpp
@@ -386,7 +386,7 @@ namespace BuildPatchServices
 
 	void FFileManifestList::OnPostLoad()
 	{
-		Algo::SortBy(FileList, [](const FFileManifest& FileManifest){ return FileManifest.Filename; }, TLess<FString>());
+		Algo::SortBy(FileList, [](const FFileManifest& FileManifest) -> const FString& { return FileManifest.Filename; }, TLess<FString>());
 
 		for (FFileManifest& FileManifest : FileList)
 		{
