Build patch installer
- Remove file constructor thread since the main thread just waits on it.
- Don't use ChunkDB thread if there are no ChunkDBs.

diff --git a/Engine/Source/Runtime/Online/BuildPatchServices/Private/BuildPatchFileConstructor.cpp b/Engine/Source/Runtime/Online/BuildPatchServices/Private/BuildPatchFileConstructor.cpp
--- a/Engine/Source/Runtime/Online/BuildPatchServices/Private/BuildPatchFileConstructor.cpp
+++ b/Engine/Source/Runtime/Online/BuildPatchServices/Private/BuildPatchFileConstructor.cpp
@@ -255,10 +255,6 @@ public:
  *****************************************************************************/
 FBuildPatchFileConstructor::FBuildPatchFileConstructor(FFileConstructorConfig InConfiguration, IFileSystem* InFileSystem, IChunkSource* InChunkSource, IChunkReferenceTracker* InChunkReferenceTracker, IInstallerError* InInstallerError, IInstallerAnalytics* InInstallerAnalytics, IFileConstructorStat* InFileConstructorStat)
 	: Configuration(MoveTemp(InConfiguration))
-	, Thread(nullptr)
-	, bIsRunning(false)
-	, bIsInited(false)
-	, bInitFailed(false)
 	, bIsDownloadStarted(false)
 	, bInitialDiskSizeCheck(false)
 	, bIsPaused(false)
@@ -290,38 +286,14 @@ FBuildPatchFileConstructor::FBuildPatchFileConstructor(FFileConstructorConfig In
 		}
 		ConstructionStack[(ConstructListNum - 1) - ConstructListIdx] = ConstructListElem;
 	}
-
-	// Start thread!
-	const TCHAR* ThreadName = TEXT("FileConstructorThread");
-
-	// Ideally this would check if we were forkable or a forked child process but there is 
-	// currently no in-engine way to check that.  Since BPS does not currently support
-	// FRunnableThread::ThreadType::Fake or FRunnableThread::ThreadType::Forkable 
-	// this check ends up being equivalent for now.  We most likely *never* want support 
-	// forking while an installer is running.
-	if (FPlatformProcess::SupportsMultithreading() || FForkProcessHelper::IsForkedMultithreadInstance())
-	{
-		Thread = FForkProcessHelper::CreateForkableThread(this, ThreadName);
-		check(Thread != nullptr);
-		check(Thread->GetThreadType() == FRunnableThread::ThreadType::Real);
-	}
 }
 
 FBuildPatchFileConstructor::~FBuildPatchFileConstructor()
 {
-	// Wait for and deallocate the thread
-	if( Thread != nullptr )
-	{
-		Thread->WaitForCompletion();
-		delete Thread;
-		Thread = nullptr;
-	}
 }
 
-uint32 FBuildPatchFileConstructor::Run()
+void FBuildPatchFileConstructor::Run()
 {
-	SetRunning(true);
-	SetInited(true);
 	FileConstructorStat->OnTotalRequiredUpdated(TotalJobSize);
 
 	// Check for resume data, we need to also look for a legacy resume file to use instead in case we are resuming from an install of previous code version.
@@ -441,23 +413,6 @@ uint32 FBuildPatchFileConstructor::Run()
 		FileConstructorStat->OnResumeCompleted();
 	}
 	FileConstructorStat->OnConstructionCompleted();
-
-	SetRunning(false);
-	return 0;
-}
-
-void FBuildPatchFileConstructor::Wait()
-{
-	if( Thread != nullptr )
-	{
-		Thread->WaitForCompletion();
-	}
-}
-
-bool FBuildPatchFileConstructor::IsComplete()
-{
-	FScopeLock Lock( &ThreadLock );
-	return ( !bIsRunning && bIsInited ) || bInitFailed;
 }
 
 uint64 FBuildPatchFileConstructor::GetRequiredDiskSpace()
@@ -477,24 +432,6 @@ FBuildPatchFileConstructor::FOnBeforeDeleteFile& FBuildPatchFileConstructor::OnB
 	return BeforeDeleteFileEvent;
 }
 
-void FBuildPatchFileConstructor::SetRunning( bool bRunning )
-{
-	FScopeLock Lock( &ThreadLock );
-	bIsRunning = bRunning;
-}
-
-void FBuildPatchFileConstructor::SetInited( bool bInited )
-{
-	FScopeLock Lock( &ThreadLock );
-	bIsInited = bInited;
-}
-
-void FBuildPatchFileConstructor::SetInitFailed( bool bFailed )
-{
-	FScopeLock Lock( &ThreadLock );
-	bInitFailed = bFailed;
-}
-
 void FBuildPatchFileConstructor::CountBytesProcessed( const int64& ByteCount )
 {
 	ByteProcessed += ByteCount;
diff --git a/Engine/Source/Runtime/Online/BuildPatchServices/Private/BuildPatchFileConstructor.h b/Engine/Source/Runtime/Online/BuildPatchServices/Private/BuildPatchFileConstructor.h
--- a/Engine/Source/Runtime/Online/BuildPatchServices/Private/BuildPatchFileConstructor.h
+++ b/Engine/Source/Runtime/Online/BuildPatchServices/Private/BuildPatchFileConstructor.h
@@ -54,9 +54,7 @@ namespace BuildPatchServices
 	 * FBuildPatchFileConstructor
 	 * This class controls a thread that constructs files from a file list, given install details, and chunk availability notifications
 	 */
-	class FBuildPatchFileConstructor
-		: public FRunnable
-		, public IControllable
+	class FBuildPatchFileConstructor : public IControllable
 	{
 	public:
 
@@ -77,9 +75,7 @@ namespace BuildPatchServices
 		 */
 		~FBuildPatchFileConstructor();
 
-		// FRunnable interface begin.
-		virtual uint32 Run() override;
-		// FRunnable interface end.
+		void Run();
 
 		// IControllable interface begin.
 		virtual void SetPaused(bool bInIsPaused) override
@@ -93,17 +89,6 @@ namespace BuildPatchServices
 		}
 		// IControllable interface end.
 
-		/**
-		 * Blocks the calling thread until this one has completed
-		 */
-		void Wait();
-
-		/**
-		 * Get whether the thread has finished working
-		 * @return	true if the thread completed
-		 */
-		bool IsComplete();
-
 		/**
 		 * Get the disk space that was required to perform the installation
 		 * @return	the disk space required to perform the installation in bytes
@@ -125,24 +110,6 @@ namespace BuildPatchServices
 
 	private:
 
-		/**
-		 * Sets the bIsRunning flag
-		 * @param bRunning	Whether the thread is running
-		 */
-		void SetRunning(bool bRunning);
-
-		/**
-		 * Sets the bIsInited flag
-		 * @param bInited	Whether the thread successfully initialized
-		 */
-		void SetInited(bool bInited);
-
-		/**
-		 * Sets the bInitFailed flag
-		 * @param bFailed	Whether the thread failed on init
-		 */
-		void SetInitFailed(bool bFailed);
-
 		/**
 		 * Count additional bytes processed, and set new install progress value
 		 * @param ByteCount		Number of bytes to increment by
@@ -199,18 +166,6 @@ namespace BuildPatchServices
 		// The configuration for the constructor.
 		const FFileConstructorConfig Configuration;
 
-		// Hold a pointer to my thread for easier deleting.
-		FRunnableThread* Thread;
-
-		// A flag marking that we a running.
-		bool bIsRunning;
-
-		// A flag marking that we initialized correctly.
-		bool bIsInited;
-
-		// A flag marking that our init returned a failure (true means failed).
-		bool bInitFailed;
-
 		// A flag marking that we told the chunk cache to queue required downloads.
 		bool bIsDownloadStarted;
 
diff --git a/Engine/Source/Runtime/Online/BuildPatchServices/Private/BuildPatchInstaller.cpp b/Engine/Source/Runtime/Online/BuildPatchServices/Private/BuildPatchInstaller.cpp
--- a/Engine/Source/Runtime/Online/BuildPatchServices/Private/BuildPatchInstaller.cpp
+++ b/Engine/Source/Runtime/Online/BuildPatchServices/Private/BuildPatchInstaller.cpp
@@ -1410,7 +1410,7 @@ namespace BuildPatchServices
 
 			// Wait for the file constructor to complete
 			ConstructTimer.Start();
-			FileConstructor->Wait();
+			FileConstructor->Run();
 			ConstructTimer.Stop();
 			FileConstructor->OnBeforeDeleteFile().Remove(OnBeforeDeleteFileHandle);
 			UE_LOG(LogBuildPatchServices, Log, TEXT("File construction complete"));
diff --git a/Engine/Source/Runtime/Online/BuildPatchServices/Private/Installer/ChunkDbChunkSource.cpp b/Engine/Source/Runtime/Online/BuildPatchServices/Private/Installer/ChunkDbChunkSource.cpp
--- a/Engine/Source/Runtime/Online/BuildPatchServices/Private/Installer/ChunkDbChunkSource.cpp
+++ b/Engine/Source/Runtime/Online/BuildPatchServices/Private/Installer/ChunkDbChunkSource.cpp
@@ -43,7 +43,7 @@ namespace ChunkDbSourceHelpers
 }
 #endif
 
-DECLARE_LOG_CATEGORY_EXTERN(LogChunkDbChunkSource, Warning, All);
+DECLARE_LOG_CATEGORY_EXTERN(LogChunkDbChunkSource, Log, All);
 DEFINE_LOG_CATEGORY(LogChunkDbChunkSource);
 
 namespace ChunkDbSourceHelpers
@@ -222,14 +222,21 @@ namespace BuildPatchServices
 		}
 
 		// Start threaded load worker.
-		TFunction<void()> Task = [this]() { return ThreadRun(); };
-		Future = Async(EAsyncExecution::ThreadIfForkSafe, MoveTemp(Task));
+		if (ChunkDbDataAccesses.Num() > 0)
+		{
+			TFunction<void()> Task = [this]() { return ThreadRun(); };
+			Future = Async(EAsyncExecution::ThreadIfForkSafe, MoveTemp(Task));
+		}
 	}
 
 	FChunkDbChunkSource::~FChunkDbChunkSource()
 	{
 		Abort();
-		Future.Wait();
+
+		if (Future.IsValid())
+		{
+			Future.Wait();
+		}
 	}
 
 	void FChunkDbChunkSource::SetPaused(bool bInIsPaused)
@@ -495,6 +502,8 @@ namespace BuildPatchServices
 
 	IChunkDbChunkSource* FChunkDbChunkSourceFactory::Create(FChunkDbSourceConfig Configuration, IPlatform* Platform, IFileSystem* FileSystem, IChunkStore* ChunkStore, IChunkReferenceTracker* ChunkReferenceTracker, IChunkDataSerialization* ChunkDataSerialization, IMessagePump* MessagePump, IInstallerError* InstallerError, IChunkDbChunkSourceStat* ChunkDbChunkSourceStat)
 	{
+		UE_LOG(LogChunkDbChunkSource, Verbose, TEXT("FChunkDbChunkSourceFactory::Create for %d db files"), Configuration.ChunkDbFiles.Num());
+
 		check(Platform != nullptr);
 		check(FileSystem != nullptr);
 		check(ChunkStore != nullptr);
diff --git a/Engine/Source/Runtime/Online/BuildPatchServices/Private/Installer/CloudChunkSource.cpp b/Engine/Source/Runtime/Online/BuildPatchServices/Private/Installer/CloudChunkSource.cpp
--- a/Engine/Source/Runtime/Online/BuildPatchServices/Private/Installer/CloudChunkSource.cpp
+++ b/Engine/Source/Runtime/Online/BuildPatchServices/Private/Installer/CloudChunkSource.cpp
@@ -750,6 +750,8 @@ namespace BuildPatchServices
 
 	ICloudChunkSource* FCloudChunkSourceFactory::Create(FCloudSourceConfig Configuration, IPlatform* Platform, IChunkStore* ChunkStore, IDownloadService* DownloadService, IChunkReferenceTracker* ChunkReferenceTracker, IChunkDataSerialization* ChunkDataSerialization, IMessagePump* MessagePump, IInstallerError* InstallerError, IDownloadConnectionCount* ConnectionCount, ICloudChunkSourceStat* CloudChunkSourceStat, IBuildManifestSet* ManifestSet, TSet<FGuid> InitialDownloadSet)
 	{
+		UE_LOG(LogCloudChunkSource, Verbose, TEXT("FCloudChunkSourceFactory::Create for %d roots"), Configuration.CloudRoots.Num());
+
 		check(Platform != nullptr);
 		check(ChunkStore != nullptr);
 		check(DownloadService != nullptr);
