Fix TSAN issues by replacing deprecated thread-safe counters by std::atomic

#rnx

diff --git a/Engine/Source/Runtime/Online/HTTP/Private/Curl/CurlHttp.cpp b/Engine/Source/Runtime/Online/HTTP/Private/Curl/CurlHttp.cpp
--- a/Engine/Source/Runtime/Online/HTTP/Private/Curl/CurlHttp.cpp
+++ b/Engine/Source/Runtime/Online/HTTP/Private/Curl/CurlHttp.cpp
@@ -565,8 +565,8 @@ size_t FCurlHttpRequest::ReceiveResponseBodyCallback(void* Ptr, size_t SizeInBlo
 	uint64 SizeToDownload = SizeInBlocks * BlockSizeInBytes;
 
 	UE_LOG(LogHttp, Verbose, TEXT("%p: ReceiveResponseBodyCallback: %llu bytes out of %llu received. (SizeInBlocks=%llu, BlockSizeInBytes=%llu, Response->TotalBytesRead=%llu, Response->GetContentLength()=%llu, SizeToDownload=%llu (<-this will get returned from the callback))"),
-		this, Response->TotalBytesRead.GetValue() + SizeToDownload, Response->GetContentLength(),
-		SizeInBlocks, BlockSizeInBytes, Response->TotalBytesRead.GetValue(), Response->GetContentLength(), SizeToDownload);
+		this, Response->TotalBytesRead.load() + SizeToDownload, Response->GetContentLength(),
+		SizeInBlocks, BlockSizeInBytes, Response->TotalBytesRead.load(), Response->GetContentLength(), SizeToDownload);
 
 	// note that we can be passed 0 bytes if file transmitted has 0 length
 	if (SizeToDownload == 0)
@@ -586,12 +586,12 @@ size_t FCurlHttpRequest::ReceiveResponseBodyCallback(void* Ptr, size_t SizeInBlo
 	else
 	{
 		Response->Payload.AddUninitialized(SizeToDownload);
-		FMemory::Memcpy(static_cast<uint8*>(Response->Payload.GetData()) + Response->TotalBytesRead.GetValue(), Ptr, SizeToDownload);
+		FMemory::Memcpy(static_cast<uint8*>(Response->Payload.GetData()) + Response->TotalBytesRead.load(), Ptr, SizeToDownload);
 
 		NumberOfBytesProcessed = SizeToDownload;
 	}
 
-	Response->TotalBytesRead.Add(NumberOfBytesProcessed);
+	Response->TotalBytesRead += NumberOfBytesProcessed;
 
 	return NumberOfBytesProcessed;
 }
@@ -601,13 +601,13 @@ size_t FCurlHttpRequest::UploadCallback(void* Ptr, size_t SizeInBlocks, size_t B
 	TimeSinceLastResponse = 0.0f;
 
 	size_t MaxBufferSize = SizeInBlocks * BlockSizeInBytes;
-	size_t SizeAlreadySent = BytesSent.GetValue();
+	size_t SizeAlreadySent = BytesSent.load();
 	size_t SizeSentThisTime = RequestPayload->FillOutputBuffer(Ptr, MaxBufferSize, SizeAlreadySent);
-	BytesSent.Add(SizeSentThisTime);
-	TotalBytesSent.Add(SizeSentThisTime);
+	BytesSent += SizeSentThisTime;
+	TotalBytesSent += SizeSentThisTime;
 
 	UE_LOG(LogHttp, Verbose, TEXT("%p: UploadCallback: %llu bytes out of %llu sent (%llu bytes total sent). (SizeInBlocks=%llu, BlockSizeInBytes=%llu, SizeToSendThisTime=%llu (<-this will get returned from the callback))"),
-		this, BytesSent.GetValue(), RequestPayload->GetContentLength(), TotalBytesSent.GetValue(), SizeInBlocks, BlockSizeInBytes, SizeSentThisTime);
+		this, BytesSent.load(), RequestPayload->GetContentLength(), TotalBytesSent.load(), SizeInBlocks, BlockSizeInBytes, SizeSentThisTime);
 
 	return SizeSentThisTime;
 }
@@ -617,8 +617,8 @@ int FCurlHttpRequest::SeekCallback(curl_off_t Offset, int Origin)
 	// Only support seeking to the very beginning
 	if (bIsRequestPayloadSeekable && Origin == SEEK_SET && Offset == 0)
 	{
-		UE_LOG(LogHttp, Log, TEXT("%p: SeekCallback: Resetting to the beginning. We had uploaded %llu bytes"), this, BytesSent.GetValue());
-		BytesSent.Reset();
+		UE_LOG(LogHttp, Log, TEXT("%p: SeekCallback: Resetting to the beginning. We had uploaded %llu bytes"), this, BytesSent.load());
+		BytesSent.store(0);
 		bIsRequestPayloadSeekable = false; // Do not attempt to re-seek
 		return CURL_SEEKFUNC_OK;
 	}
@@ -901,8 +901,8 @@ bool FCurlHttpRequest::SetupRequestHttpThread()
 
 		if (bUseReadFunction)
 		{
-			BytesSent.Reset();
-			TotalBytesSent.Reset();
+			BytesSent.store(0);
+			TotalBytesSent.store(0);
 			curl_easy_setopt(EasyHandle, CURLOPT_READDATA, this);
 			curl_easy_setopt(EasyHandle, CURLOPT_READFUNCTION, StaticUploadCallback);
 		}
@@ -1108,8 +1108,8 @@ void FCurlHttpRequest::Tick(float DeltaSeconds)
 void FCurlHttpRequest::CheckProgressDelegate()
 {
 	QUICK_SCOPE_CYCLE_COUNTER(STAT_FCurlHttpRequest_CheckProgressDelegate);
-	const uint64 CurrentBytesRead = Response.IsValid() ? Response->TotalBytesRead.GetValue() : 0;
-	const uint64 CurrentBytesSent = BytesSent.GetValue();
+	const uint64 CurrentBytesRead = Response.IsValid() ? Response->TotalBytesRead.load() : 0;
+	const uint64 CurrentBytesSent = BytesSent.load();
 
 	const bool bProcessing = CompletionStatus == EHttpRequestStatus::Processing;
 	const bool bBytesSentChanged = (CurrentBytesSent != LastReportedBytesSent);
@@ -1190,14 +1190,14 @@ void FCurlHttpRequest::FinishRequest()
 				else
 				{
 					// If curl did not know how much we downloaded, or we were missing a Content-Length header (Chunked request), set our ContentLength as the amount we downloaded
-					Response->ContentLength = Response->TotalBytesRead.GetValue();
+					Response->ContentLength = Response->TotalBytesRead;
 				}
 			}
 
 			if (Response->HttpCode <= 0 && URL.StartsWith(TEXT("Http"), ESearchCase::IgnoreCase))
 			{
 				UE_LOG(LogHttp, Warning, TEXT("%p: invalid HTTP response code received. URL: %s, HTTP code: %d, content length: %llu, actual payload size: %llu"),
-					this, *GetURL(), Response->HttpCode, Response->ContentLength, Response->TotalBytesRead.GetValue());
+					this, *GetURL(), Response->HttpCode, Response->ContentLength, Response->TotalBytesRead.load());
 				Response->bSucceeded = false;
 			}
 		}
@@ -1222,12 +1222,12 @@ void FCurlHttpRequest::FinishRequest()
 			if (bDebugServerResponse)
 			{
 				UE_LOG(LogHttp, Warning, TEXT("%p: request has been successfully processed. URL: %s, HTTP code: %d, content length: %llu, actual payload size: %llu, elapsed: %.2fs"),
-					this, *GetURL(), Response->HttpCode, Response->ContentLength, Response->TotalBytesRead.GetValue(), ElapsedTime);
+					this, *GetURL(), Response->HttpCode, Response->ContentLength, Response->TotalBytesRead.load(), ElapsedTime);
 			}
 			else
 			{
 				UE_LOG(LogHttp, Log, TEXT("%p: request has been successfully processed. URL: %s, HTTP code: %d, content length: %llu, actual payload size: %llu, elapsed: %.2fs"),
-					this, *GetURL(), Response->HttpCode, Response->ContentLength, Response->TotalBytesRead.GetValue(), ElapsedTime);
+					this, *GetURL(), Response->HttpCode, Response->ContentLength, Response->TotalBytesRead.load(), ElapsedTime);
 			}
 
 			TArray<FString> AllHeaders = Response->GetAllHeaders();
diff --git a/Engine/Source/Runtime/Online/HTTP/Private/Curl/CurlHttp.h b/Engine/Source/Runtime/Online/HTTP/Private/Curl/CurlHttp.h
--- a/Engine/Source/Runtime/Online/HTTP/Private/Curl/CurlHttp.h
+++ b/Engine/Source/Runtime/Online/HTTP/Private/Curl/CurlHttp.h
@@ -377,7 +377,7 @@ private:
 	/** Set to true if request has been canceled */
 	bool			bCanceled;
 	/** Set to true when request has been completed */
-	FThreadSafeBool	bCurlRequestCompleted;
+	std::atomic<bool> bCurlRequestCompleted;
 	/** Set to true when request has "30* Multiple Choices" (e.g. 301 Moved Permanently, 302 temporary redirect, 308 Permanent Redirect, etc.) */
 	bool			bRedirected;
 	/** Set to true if request failed to be added to curl multi */
@@ -401,9 +401,9 @@ private:
 	/** Have we had any HTTP activity with the host? Sending headers, SSL handshake, etc */
 	bool bAnyHttpActivity;
 	/** Number of bytes sent already */
-	FThreadSafeCounter64 BytesSent;
+	std::atomic<int64> BytesSent;
 	/** Total number of bytes sent already (includes data re-sent by seek attempts) */
-	FThreadSafeCounter64 TotalBytesSent;
+	std::atomic<int64> TotalBytesSent;
 	/** Last bytes read reported to progress delegate */
 	uint64 LastReportedBytesRead;
 	/** Last bytes sent reported to progress delegate */
@@ -467,7 +467,7 @@ private:
 	/** The stream to receive response body */
 	TSharedPtr<FArchive> ResponseBodyReceiveStream;
 	/** Caches how many bytes of the response we've read so far */
-	FThreadSafeCounter64 TotalBytesRead;
+	std::atomic<int64> TotalBytesRead;
 	/** Cached key/value header pairs. Parsed once request completes. Only accessible on the game thread. */
 	TMap<FString, FString> Headers;
 	/** Newly received headers we need to inform listeners about */
