Add http request progress delegate 64 bit support.

diff --git a/Engine/Plugins/Experimental/Web/WebAPI/Source/WebAPI/Public/WebAPIHttpRequest.h b/Engine/Plugins/Experimental/Web/WebAPI/Source/WebAPI/Public/WebAPIHttpRequest.h
--- a/Engine/Plugins/Experimental/Web/WebAPI/Source/WebAPI/Public/WebAPIHttpRequest.h
+++ b/Engine/Plugins/Experimental/Web/WebAPI/Source/WebAPI/Public/WebAPIHttpRequest.h
@@ -24,7 +24,9 @@ public:
 
 	UE_NODISCARD TRequest<PayloadType>& AddHeader(TPair<FString, FString> NewHeader) const;
 	UE_NODISCARD TRequest<PayloadType>& BindCompletionCallback(TFunction<void(FHttpRequestPtr, FHttpResponsePtr, bool)> CompletionCallback);
+	UE_DEPRECATED(5.3, "BindProgressCallback has been deprecated, use BindProgressCallback64 instead")
 	UE_NODISCARD TRequest<PayloadType>& BindProgressCallback(TFunction<void(FHttpRequestPtr, int32 /* BytesSent */, int32 /* BytesReceived */)> ProgressCallback);
+	UE_NODISCARD TRequest<PayloadType>& BindProgressCallback64(TFunction<void(FHttpRequestPtr, uint64 /* BytesSent */, uint64 /* BytesReceived */)> ProgressCallback);
 	UE_NODISCARD TRequest<PayloadType>& BindRetryCallback(TFunction<void(FHttpRequestPtr, FHttpResponsePtr, float /* SecondsToRetry */)> RetryCallback);
 	UE_NODISCARD TRequest<PayloadType>& BindHeaderReceivedCallback(TFunction<void(FHttpRequestPtr, const FString& /* HeaderName */, const FString& /* HeaderValue */)> HeaderReceivedCallback);
 	TRequest<PayloadType>& SetPayloadData(const PayloadType& InPayloadData);
@@ -213,12 +215,21 @@ TRequest<PayloadType>& TRequest<PayloadType>::BindCompletionCallback(TFunction<v
 	return *this;
 }
 
+PRAGMA_DISABLE_DEPRECATION_WARNINGS
 template <typename PayloadType>
 TRequest<PayloadType>& TRequest<PayloadType>::BindProgressCallback(TFunction<void(FHttpRequestPtr, int32, int32)> ProgressCallback)
 {
 	InternalRequest->OnRequestProgress().BindLambda(MoveTemp(ProgressCallback));
 	return *this;
 }
+PRAGMA_ENABLE_DEPRECATION_WARNINGS
+
+template <typename PayloadType>
+TRequest<PayloadType>& TRequest<PayloadType>::BindProgressCallback64(TFunction<void(FHttpRequestPtr, uint64, uint64)> ProgressCallback)
+{
+	InternalRequest->OnRequestProgress64().BindLambda(MoveTemp(ProgressCallback));
+	return *this;
+}
 
 template <typename PayloadType>
 TRequest<PayloadType>& TRequest<PayloadType>::BindRetryCallback(TFunction<void(FHttpRequestPtr, FHttpResponsePtr, float)> RetryCallback)
diff --git a/Engine/Plugins/Online/OnlineSubsystem/Source/Public/Interfaces/OnlineUserCloudInterface.h b/Engine/Plugins/Online/OnlineSubsystem/Source/Public/Interfaces/OnlineUserCloudInterface.h
--- a/Engine/Plugins/Online/OnlineSubsystem/Source/Public/Interfaces/OnlineUserCloudInterface.h
+++ b/Engine/Plugins/Online/OnlineSubsystem/Source/Public/Interfaces/OnlineUserCloudInterface.h
@@ -41,8 +41,22 @@ typedef FOnEnumerateUserFilesComplete::FDelegate FOnEnumerateUserFilesCompleteDe
  * @param FileName the name of the file this was for
  *
  */
+UE_DEPRECATED(5.3, "OnWriteUserFileProgress has been deprecated, use OnWriteUserFileProgress64 instead")
 DECLARE_MULTICAST_DELEGATE_ThreeParams(FOnWriteUserFileProgress, int32, const FUniqueNetId&, const FString&);
+PRAGMA_DISABLE_DEPRECATION_WARNINGS
 typedef FOnWriteUserFileProgress::FDelegate FOnWriteUserFileProgressDelegate;
+PRAGMA_ENABLE_DEPRECATION_WARNINGS
+
+/**
+ * Delegate fired at intervals during a user file write to the network platform's storage
+ *
+ * @param BytesWritten the number of bytes written so far
+ * @param UserId User owning the storage
+ * @param FileName the name of the file this was for
+ *
+ */
+DECLARE_MULTICAST_DELEGATE_ThreeParams(FOnWriteUserFileProgress64, uint64, const FUniqueNetId&, const FString&);
+typedef FOnWriteUserFileProgress64::FDelegate FOnWriteUserFileProgress64Delegate;
 
 /**
  * Delegate fired when a user file write to the network platform's storage is complete
@@ -118,7 +132,9 @@ protected:
 	IOnlineUserCloud() {};
 
 public:
+PRAGMA_DISABLE_DEPRECATION_WARNINGS
 	virtual ~IOnlineUserCloud() {};
+PRAGMA_ENABLE_DEPRECATION_WARNINGS
 
 	/**
 	 * Copies the file data into the specified buffer for the specified file
@@ -216,7 +232,45 @@ public:
 	* @param FileName the name of the file this was for
 	*
 	*/
-	DEFINE_ONLINE_DELEGATE_THREE_PARAM(OnWriteUserFileProgress, int32, const FUniqueNetId&, const FString&);
+PRAGMA_DISABLE_DEPRECATION_WARNINGS
+	// 	Expanded macro with deprecate warnings: DEFINE_ONLINE_DELEGATE_THREE_PARAM(OnWriteUserFileProgress, int32, const FUniqueNetId&, const FString&);
+
+	UE_DEPRECATED(5.3, "OnWriteUserFileProgress has been deprecated, use OnWriteUserFileProgress64 instead")
+	FOnWriteUserFileProgress OnWriteUserFileProgressDelegates;
+
+	UE_DEPRECATED(5.3, "AddOnWriteUserFileProgressDelegate_Handle has been deprecated, use AddOnWriteUserFileProgress64Delegate_Handle instead")
+	virtual FDelegateHandle AddOnWriteUserFileProgressDelegate_Handle(const FOnWriteUserFileProgressDelegate& Delegate)
+	{
+		OnWriteUserFileProgressDelegates.Add(Delegate);
+		return Delegate.GetHandle();
+	}
+	UE_DEPRECATED(5.3, "ClearOnWriteUserFileProgressDelegate_Handle has been deprecated, use ClearOnWriteUserFileProgress64Delegate_Handle instead")
+	virtual void ClearOnWriteUserFileProgressDelegate_Handle(FDelegateHandle& Handle)
+	{
+		OnWriteUserFileProgressDelegates.Remove(Handle);
+		Handle.Reset();
+	}
+	UE_DEPRECATED(5.3, "ClearOnWriteUserFileProgressDelegates has been deprecated, use ClearOnWriteUserFileProgress64Delegates instead")
+	virtual void ClearOnWriteUserFileProgressDelegates(void* Object)
+	{
+		OnWriteUserFileProgressDelegates.RemoveAll(Object);
+	}
+	UE_DEPRECATED(5.3, "TriggerOnWriteUserFileProgressDelegates has been deprecated, use TriggerOnWriteUserFileProgress64Delegates instead")
+	virtual void TriggerOnWriteUserFileProgressDelegates(int32 Param1, const FUniqueNetId& Param2, const FString& Param3)
+	{
+		OnWriteUserFileProgressDelegates.Broadcast(Param1, Param2, Param3);
+	}
+PRAGMA_ENABLE_DEPRECATION_WARNINGS
+
+	/**
+	* Delegate fired at intervals during a user file write to the network platform's storage
+	*
+	* @param BytesWritten the number of bytes written so far
+	* @param UserId User owning the storage
+	* @param FileName the name of the file this was for
+	*
+	*/
+	DEFINE_ONLINE_DELEGATE_THREE_PARAM(OnWriteUserFileProgress64, uint64, const FUniqueNetId&, const FString&);
 
 	/**
 	 * Delegate fired when a user file write to the network platform's storage is complete
diff --git a/Engine/Plugins/Online/OnlineSubsystemEOS/Source/OnlineSubsystemEOS/Private/OnlineTitleFileEOS.cpp b/Engine/Plugins/Online/OnlineSubsystemEOS/Source/OnlineSubsystemEOS/Private/OnlineTitleFileEOS.cpp
--- a/Engine/Plugins/Online/OnlineSubsystemEOS/Source/OnlineSubsystemEOS/Private/OnlineTitleFileEOS.cpp
+++ b/Engine/Plugins/Online/OnlineSubsystemEOS/Source/OnlineSubsystemEOS/Private/OnlineTitleFileEOS.cpp
@@ -426,7 +426,7 @@ bool FOnlineTitleFileEOS::HandleTitleFileExec(UWorld* InWorld, const TCHAR* Cmd,
 		// Note: These delegates will trigger for all files being read, if multiple requests are inflight
 		OnReadFileProgressDelegateHandle = AddOnReadFileProgressDelegate_Handle(FOnReadFileProgressDelegate::CreateLambda([this](const FString& FileName, uint64 NumBytes)
 		{
-			UE_LOG_ONLINE(Log, TEXT("ReadFile: Progress (%s) %d bytes read"), *FileName, NumBytes);
+			UE_LOG_ONLINE(Log, TEXT("ReadFile: Progress (%s) %llu bytes read"), *FileName, NumBytes);
 		}));
 
 		OnReadFileCompleteDelegateHandle = AddOnReadFileCompleteDelegate_Handle(FOnReadFileCompleteDelegate::CreateLambda([this](bool bWasSuccessful, const FString& Error)
diff --git a/Engine/Plugins/Online/OnlineSubsystemEOS/Source/OnlineSubsystemEOS/Private/OnlineUserCloudEOS.cpp b/Engine/Plugins/Online/OnlineSubsystemEOS/Source/OnlineSubsystemEOS/Private/OnlineUserCloudEOS.cpp
--- a/Engine/Plugins/Online/OnlineSubsystemEOS/Source/OnlineSubsystemEOS/Private/OnlineUserCloudEOS.cpp
+++ b/Engine/Plugins/Online/OnlineSubsystemEOS/Source/OnlineSubsystemEOS/Private/OnlineUserCloudEOS.cpp
@@ -554,7 +554,10 @@ bool FOnlineUserCloudEOS::WriteUserFile(const FUniqueNetId& UserId, const FStrin
 	CallbackObj->SetNested2CallbackLambda([this, SharedUserId, FileName](const EOS_PlayerDataStorage_FileTransferProgressCallbackInfo* Data)
 		{
 			UE_LOG_ONLINE_CLOUD(VeryVerbose, TEXT("[FOnlineUserCloudEOS::WriteUserFile] File transfer progress for file %s is %d bytes"), *FileName, Data->BytesTransferred);
+PRAGMA_DISABLE_DEPRECATION_WARNINGS
 			TriggerOnWriteUserFileProgressDelegates(Data->BytesTransferred, *SharedUserId, FileName);
+PRAGMA_ENABLE_DEPRECATION_WARNINGS
+			TriggerOnWriteUserFileProgress64Delegates(Data->BytesTransferred, *SharedUserId, FileName);
 		});
 
 	CallbackObj->CallbackLambda = [this, SharedUserId, FileName](const EOS_PlayerDataStorage_WriteFileCallbackInfo* Data)
diff --git a/Engine/Plugins/Online/OnlineSubsystemEOS/Source/OnlineSubsystemEOSPlus/Private/OnlineUserCloudEOSPlus.cpp b/Engine/Plugins/Online/OnlineSubsystemEOS/Source/OnlineSubsystemEOSPlus/Private/OnlineUserCloudEOSPlus.cpp
--- a/Engine/Plugins/Online/OnlineSubsystemEOS/Source/OnlineSubsystemEOSPlus/Private/OnlineUserCloudEOSPlus.cpp
+++ b/Engine/Plugins/Online/OnlineSubsystemEOS/Source/OnlineSubsystemEOSPlus/Private/OnlineUserCloudEOSPlus.cpp
@@ -15,7 +15,10 @@ FOnlineUserCloudEOSPlus::~FOnlineUserCloudEOSPlus()
 	if (BaseUserCloudInterface.IsValid())
 	{
 		BaseUserCloudInterface->ClearOnEnumerateUserFilesCompleteDelegates(this);
+PRAGMA_DISABLE_DEPRECATION_WARNINGS
 		BaseUserCloudInterface->ClearOnWriteUserFileProgressDelegates(this);
+PRAGMA_ENABLE_DEPRECATION_WARNINGS
+		BaseUserCloudInterface->ClearOnWriteUserFileProgress64Delegates(this);
 		BaseUserCloudInterface->ClearOnWriteUserFileCompleteDelegates(this);
 		BaseUserCloudInterface->ClearOnWriteUserFileCanceledDelegates(this);
 		BaseUserCloudInterface->ClearOnReadUserFileCompleteDelegates(this);
@@ -34,7 +37,10 @@ void FOnlineUserCloudEOSPlus::Initialize()
 	if (BaseUserCloudInterface.IsValid())
 	{
 		BaseUserCloudInterface->AddOnEnumerateUserFilesCompleteDelegate_Handle(FOnEnumerateUserFilesCompleteDelegate::CreateThreadSafeSP(this, &FOnlineUserCloudEOSPlus::OnEnumerateUserFilesComplete));
+PRAGMA_DISABLE_DEPRECATION_WARNINGS
 		BaseUserCloudInterface->AddOnWriteUserFileProgressDelegate_Handle(FOnWriteUserFileProgressDelegate::CreateThreadSafeSP(this, &FOnlineUserCloudEOSPlus::OnWriteUserFileProgress));
+PRAGMA_ENABLE_DEPRECATION_WARNINGS
+		BaseUserCloudInterface->AddOnWriteUserFileProgress64Delegate_Handle(FOnWriteUserFileProgress64Delegate::CreateThreadSafeSP(this, &FOnlineUserCloudEOSPlus::OnWriteUserFileProgress64));
 		BaseUserCloudInterface->AddOnWriteUserFileCompleteDelegate_Handle(FOnWriteUserFileCompleteDelegate::CreateThreadSafeSP(this, &FOnlineUserCloudEOSPlus::OnWriteUserFileComplete));
 		BaseUserCloudInterface->AddOnWriteUserFileCanceledDelegate_Handle(FOnWriteUserFileCanceledDelegate::CreateThreadSafeSP(this, &FOnlineUserCloudEOSPlus::OnWriteUserFileCanceled));
 		BaseUserCloudInterface->AddOnReadUserFileCompleteDelegate_Handle(FOnReadUserFileCompleteDelegate::CreateThreadSafeSP(this, &FOnlineUserCloudEOSPlus::OnReadUserFileComplete));
@@ -360,7 +366,20 @@ void FOnlineUserCloudEOSPlus::OnWriteUserFileProgress(int32 BytesWritten, const
 		UE_LOG_ONLINE(Warning, TEXT("[FOnlineUserCloudEOSPlus::OnWriteUserFileProgress] Unknown user (%s)"), *UserId.ToString());
 	}
 
+PRAGMA_DISABLE_DEPRECATION_WARNINGS
 	TriggerOnWriteUserFileProgressDelegates(BytesWritten, *NetIdPlus, FileName);
+PRAGMA_ENABLE_DEPRECATION_WARNINGS
+}
+
+void FOnlineUserCloudEOSPlus::OnWriteUserFileProgress64(uint64 BytesWritten, const FUniqueNetId& UserId, const FString& FileName)
+{
+	const FUniqueNetIdEOSPlusPtr NetIdPlus = GetNetIdPlus(UserId.ToString());
+	if (!NetIdPlus.IsValid())
+	{
+		UE_LOG_ONLINE(Warning, TEXT("[FOnlineUserCloudEOSPlus::OnWriteUserFileProgress64] Unknown user (%s)"), *UserId.ToString());
+	}
+
+	TriggerOnWriteUserFileProgress64Delegates(BytesWritten, *NetIdPlus, FileName);
 }
 
 void FOnlineUserCloudEOSPlus::OnWriteUserFileComplete(bool bWasSuccessful, const FUniqueNetId& UserId, const FString& FileName)
diff --git a/Engine/Plugins/Online/OnlineSubsystemEOS/Source/OnlineSubsystemEOSPlus/Private/OnlineUserCloudEOSPlus.h b/Engine/Plugins/Online/OnlineSubsystemEOS/Source/OnlineSubsystemEOSPlus/Private/OnlineUserCloudEOSPlus.h
--- a/Engine/Plugins/Online/OnlineSubsystemEOS/Source/OnlineSubsystemEOSPlus/Private/OnlineUserCloudEOSPlus.h
+++ b/Engine/Plugins/Online/OnlineSubsystemEOS/Source/OnlineSubsystemEOSPlus/Private/OnlineUserCloudEOSPlus.h
@@ -43,6 +43,7 @@ PACKAGE_SCOPE:
 
 	void OnEnumerateUserFilesComplete(bool bWasSuccessful, const FUniqueNetId& UserId);
 	void OnWriteUserFileProgress(int32 BytesWritten, const FUniqueNetId& UserId, const FString& FileName);
+	void OnWriteUserFileProgress64(uint64 BytesWritten, const FUniqueNetId& UserId, const FString& FileName);
 	void OnWriteUserFileComplete(bool bWasSuccessful, const FUniqueNetId& UserId, const FString& FileName);
 	void OnWriteUserFileCanceled(bool bWasSuccessful, const FUniqueNetId& UserId, const FString& FileName);
 	void OnReadUserFileComplete(bool bWasSuccessful, const FUniqueNetId& UserId, const FString& FileName);
diff --git a/Engine/Plugins/Runtime/ChunkDownloader/Source/Private/PlatformStreamDownload.cpp b/Engine/Plugins/Runtime/ChunkDownloader/Source/Private/PlatformStreamDownload.cpp
--- a/Engine/Plugins/Runtime/ChunkDownloader/Source/Private/PlatformStreamDownload.cpp
+++ b/Engine/Plugins/Runtime/ChunkDownloader/Source/Private/PlatformStreamDownload.cpp
@@ -63,7 +63,7 @@ FDownloadCancel PlatformStreamDownload(const FString& Url, const FString& Target
 	// bind the progress delegate
 	if (Progress)
 	{
-		Request->OnRequestProgress().BindLambda([Progress](FHttpRequestPtr HttpRequest, int32 BytesSent, int32 BytesReceived) {
+		Request->OnRequestProgress64().BindLambda([Progress](FHttpRequestPtr HttpRequest, uint64 BytesSent, uint64 BytesReceived) {
 			Progress(BytesReceived);
 		});
 	}
diff --git a/Engine/Plugins/Runtime/ChunkDownloader/Source/Private/PlatformStreamDownload.h b/Engine/Plugins/Runtime/ChunkDownloader/Source/Private/PlatformStreamDownload.h
--- a/Engine/Plugins/Runtime/ChunkDownloader/Source/Private/PlatformStreamDownload.h
+++ b/Engine/Plugins/Runtime/ChunkDownloader/Source/Private/PlatformStreamDownload.h
@@ -8,7 +8,7 @@ class FString;
 template <typename FuncType> class TFunction;
 
 typedef TFunction<void(int32 HttpStatus)> FDownloadComplete;
-typedef TFunction<void(int32 BytesReceived)> FDownloadProgress;
+typedef TFunction<void(uint64 BytesReceived)> FDownloadProgress;
 typedef TFunction<void(void)> FDownloadCancel;
 
 extern FDownloadCancel PlatformStreamDownload(const FString& Url, const FString& TargetFile, const FDownloadProgress& Progress, const FDownloadComplete& Callback);
diff --git a/Engine/Plugins/Runtime/HTTPChunkInstaller/Source/Private/HTTPChunkInstaller.cpp b/Engine/Plugins/Runtime/HTTPChunkInstaller/Source/Private/HTTPChunkInstaller.cpp
--- a/Engine/Plugins/Runtime/HTTPChunkInstaller/Source/Private/HTTPChunkInstaller.cpp
+++ b/Engine/Plugins/Runtime/HTTPChunkInstaller/Source/Private/HTTPChunkInstaller.cpp
@@ -505,7 +505,7 @@ private:
 			FileProgressRequestsMap.Add(HttpRequest, FPendingFileRequest(FileName));
 
 			HttpRequest->OnProcessRequestComplete().BindRaw(this, &FOnlineTitleFileHttp::ReadFile_HttpRequestComplete);
-			HttpRequest->OnRequestProgress().BindRaw(this, &FOnlineTitleFileHttp::ReadFile_HttpRequestProgress);
+			HttpRequest->OnRequestProgress64().BindRaw(this, &FOnlineTitleFileHttp::ReadFile_HttpRequestProgress);
 			FString RequestUrl;
 			// Grab the file from the specified URL if that was set, otherwise use the old method that hits the game service
 			if (CloudFileHeader != nullptr && !CloudFileHeader->URL.IsEmpty())
@@ -712,7 +712,7 @@ private:
 		FileRequests.Remove(HttpRequest);
 		// remove from progress updates
 		FileProgressRequestsMap.Remove(HttpRequest);
-		HttpRequest->OnRequestProgress().Unbind();
+		HttpRequest->OnRequestProgress64().Unbind();
 
 		// Cloud file being operated on
 		FCloudEntry* CloudFile = GetCloudFile(PendingRequest.FileName, true);
@@ -760,7 +760,7 @@ private:
 	/**
 	* Delegate called as a Http request progresses for reading a cloud file
 	*/
-	void ReadFile_HttpRequestProgress(FHttpRequestPtr HttpRequest, int32 BytesSent, int32 BytesReceived)
+	void ReadFile_HttpRequestProgress(FHttpRequestPtr HttpRequest, uint64 BytesSent, uint64 BytesReceived)
 	{
 		FPendingFileRequest PendingRequest = FileProgressRequestsMap.FindChecked(HttpRequest);
 		// Just forward this to anyone that is listening
diff --git a/Engine/Source/Runtime/Online/BackgroundHTTP/Private/GenericPlatform/GenericPlatformBackgroundHttpRequest.cpp b/Engine/Source/Runtime/Online/BackgroundHTTP/Private/GenericPlatform/GenericPlatformBackgroundHttpRequest.cpp
--- a/Engine/Source/Runtime/Online/BackgroundHTTP/Private/GenericPlatform/GenericPlatformBackgroundHttpRequest.cpp
+++ b/Engine/Source/Runtime/Online/BackgroundHTTP/Private/GenericPlatform/GenericPlatformBackgroundHttpRequest.cpp
@@ -75,7 +75,7 @@ void FGenericPlatformBackgroundHttpRequest::FGenericPlatformBackgroundHttpWrappe
 		//Create new Request
 		HttpRequest = FHttpModule::Get().CreateRequest();
 		HttpRequest->OnProcessRequestComplete().BindRaw(this, &FGenericPlatformBackgroundHttpRequest::FGenericPlatformBackgroundHttpWrapper::HttpRequestComplete);
-		HttpRequest->OnRequestProgress().BindRaw(this, &FGenericPlatformBackgroundHttpRequest::FGenericPlatformBackgroundHttpWrapper::UpdateHttpProgress);
+		HttpRequest->OnRequestProgress64().BindRaw(this, &FGenericPlatformBackgroundHttpRequest::FGenericPlatformBackgroundHttpWrapper::UpdateHttpProgress);
 
 		const FString& RequestURL = GetURLForCurrentRetryNumber();
 		if (!RequestURL.IsEmpty())
@@ -102,7 +102,7 @@ void FGenericPlatformBackgroundHttpRequest::FGenericPlatformBackgroundHttpWrappe
 	if (HttpRequest.IsValid())
 	{
 		HttpRequest->OnProcessRequestComplete().Unbind();
-		HttpRequest->OnRequestProgress().Unbind();
+		HttpRequest->OnRequestProgress64().Unbind();
 		HttpRequest->CancelRequest();
 
 		HttpRequest.Reset();
@@ -168,15 +168,15 @@ bool FGenericPlatformBackgroundHttpRequest::FGenericPlatformBackgroundHttpWrappe
 	return (CurrentRetryNumber <= MaxRetries);
 }
 
-void FGenericPlatformBackgroundHttpRequest::FGenericPlatformBackgroundHttpWrapper::UpdateHttpProgress(FHttpRequestPtr UnderlyingHttpRequest, int32 BytesSent, int32 BytesReceived)
+void FGenericPlatformBackgroundHttpRequest::FGenericPlatformBackgroundHttpWrapper::UpdateHttpProgress(FHttpRequestPtr UnderlyingHttpRequest, uint64 BytesSent, uint64 BytesReceived)
 {
 	if (ensureAlwaysMsgf(OriginalRequest.IsValid(), TEXT("Received UpdateHttpProgress callback with invalid OriginalRequest pointer! Can not update request progress!")))
 	{
-		const int32 ByteDifference = (LastProgressUpdateBytes > 0) ? BytesReceived - LastProgressUpdateBytes : BytesReceived;
-		ensureAlwaysMsgf((ByteDifference >= 0), TEXT("Invalid Byte Difference in UpdateHttpProgress -- ByteDifference:%d | LastProgressUpdateBytes:%d | BytesSent:%d | BytesReceived:%d"), ByteDifference, LastProgressUpdateBytes, BytesSent, BytesReceived);
+		const uint64 ByteDifference = (LastProgressUpdateBytes > 0) ? BytesReceived - LastProgressUpdateBytes : BytesReceived;
+		ensureAlwaysMsgf((ByteDifference >= 0), TEXT("Invalid Byte Difference in UpdateHttpProgress -- ByteDifference:%llu | LastProgressUpdateBytes:%llu | BytesSent:%llu | BytesReceived:%llu"), ByteDifference, LastProgressUpdateBytes, BytesSent, BytesReceived);
 		LastProgressUpdateBytes = BytesReceived;
 
-		UE_LOG(LogBackgroundHttpRequest, VeryVerbose, TEXT("HttpRequest Progress Update- RequestID:%s | BytesSent: %d | BytesReceived:%d | BytesReceivedSinceLastUpdate:%d"), *OriginalRequest.Pin()->GetRequestID(), BytesSent, BytesReceived, ByteDifference);
+		UE_LOG(LogBackgroundHttpRequest, VeryVerbose, TEXT("HttpRequest Progress Update- RequestID:%s | BytesSent: %llu | BytesReceived:%llu | BytesReceivedSinceLastUpdate:%llu"), *OriginalRequest.Pin()->GetRequestID(), BytesSent, BytesReceived, ByteDifference);
 
 		OriginalRequest.Pin()->OnProgressUpdated().ExecuteIfBound(OriginalRequest.Pin(), BytesReceived, ByteDifference);
 	}
diff --git a/Engine/Source/Runtime/Online/BackgroundHTTP/Public/GenericPlatform/GenericPlatformBackgroundHttpRequest.h b/Engine/Source/Runtime/Online/BackgroundHTTP/Public/GenericPlatform/GenericPlatformBackgroundHttpRequest.h
--- a/Engine/Source/Runtime/Online/BackgroundHTTP/Public/GenericPlatform/GenericPlatformBackgroundHttpRequest.h
+++ b/Engine/Source/Runtime/Online/BackgroundHTTP/Public/GenericPlatform/GenericPlatformBackgroundHttpRequest.h
@@ -43,7 +43,7 @@ protected:
 	private:
 		
 		bool HasRetriesRemaining() const;
-		void UpdateHttpProgress(FHttpRequestPtr UnderlyingHttpRequest, int32 BytesSent, int32 BytesReceived);
+		void UpdateHttpProgress(FHttpRequestPtr UnderlyingHttpRequest, uint64 BytesSent, uint64 BytesReceived);
 
 		void CleanUpHttpRequest();
 
diff --git a/Engine/Source/Runtime/Online/BuildPatchServices/Private/Installer/CloudChunkSource.cpp b/Engine/Source/Runtime/Online/BuildPatchServices/Private/Installer/CloudChunkSource.cpp
--- a/Engine/Source/Runtime/Online/BuildPatchServices/Private/Installer/CloudChunkSource.cpp
+++ b/Engine/Source/Runtime/Online/BuildPatchServices/Private/Installer/CloudChunkSource.cpp
@@ -106,7 +106,7 @@ namespace BuildPatchServices
 			FDownloadDelegates(FCloudChunkSource& InCloudChunkSource);
 
 		public:
-			void OnDownloadProgress(int32 RequestId, int32 BytesSoFar);
+			void OnDownloadProgress(int32 RequestId, uint64 BytesSoFar);
 			void OnDownloadComplete(int32 RequestId, const FDownloadRef& Download);
 
 		private:
@@ -155,7 +155,7 @@ namespace BuildPatchServices
 		float GetRetryDelay(int32 RetryNum);
 		EBuildPatchDownloadHealth GetDownloadHealth(bool bIsDisconnected, float ChunkSuccessRate);
 		FGuid GetNextTask(const TMap<FGuid, FTaskInfo>& TaskInfos, const TMap<int32, FGuid>& InFlightDownloads, const TSet<FGuid>& TotalRequiredChunks, const TSet<FGuid>& PriorityRequests, const TSet<FGuid>& FailedDownloads, const TSet<FGuid>& Stored, TArray<FGuid>& DownloadQueue, EBuildPatchDownloadHealth DownloadHealth);
-		void OnDownloadProgress(int32 RequestId, int32 BytesSoFar);
+		void OnDownloadProgress(int32 RequestId, uint64 BytesSoFar);
 		void OnDownloadComplete(int32 RequestId, const FDownloadRef& Download);
 
 	private:
@@ -207,7 +207,7 @@ namespace BuildPatchServices
 	{
 	}
 
-	void FCloudChunkSource::FDownloadDelegates::OnDownloadProgress(int32 RequestId, int32 BytesSoFar)
+	void FCloudChunkSource::FDownloadDelegates::OnDownloadProgress(int32 RequestId, uint64 BytesSoFar)
 	{
 		CloudChunkSource.OnDownloadProgress(RequestId, BytesSoFar);
 	}
@@ -737,7 +737,7 @@ namespace BuildPatchServices
 		Promise.SetValue();
 	}
 
-	void FCloudChunkSource::OnDownloadProgress(int32 RequestId, int32 BytesSoFar)
+	void FCloudChunkSource::OnDownloadProgress(int32 RequestId, uint64 BytesSoFar)
 	{
 		FPlatformAtomics::InterlockedExchange(&CyclesAtLastData, FStatsCollector::GetCycles());
 	}
diff --git a/Engine/Source/Runtime/Online/BuildPatchServices/Private/Installer/DownloadService.cpp b/Engine/Source/Runtime/Online/BuildPatchServices/Private/Installer/DownloadService.cpp
--- a/Engine/Source/Runtime/Online/BuildPatchServices/Private/Installer/DownloadService.cpp
+++ b/Engine/Source/Runtime/Online/BuildPatchServices/Private/Installer/DownloadService.cpp
@@ -210,7 +210,7 @@ namespace BuildPatchServices
 		public:
 			FHttpDelegates(FDownloadService& InDownloadService);
 
-			void HttpRequestProgress(FHttpRequestPtr Request, int32 BytesSent, int32 BytesReceived, int32 RequestId);
+			void HttpRequestProgress(FHttpRequestPtr Request, uint64 BytesSent, uint64 BytesReceived, int32 RequestId);
 			void HttpRequestComplete(FHttpRequestPtr Request, FHttpResponsePtr Response, bool bSucceeded, IDownloadServiceStat::FDownloadRecord DownloadRecord);
 
 		private:
@@ -245,9 +245,9 @@ namespace BuildPatchServices
 		void RegisterRequest(int32 RequestId, TSharedRef<IHttpRequest, ESPMode::ThreadSafe> HttpRequest);
 		void RegisterRequest(int32 RequestId, TUniquePtr<FFileRequest> FileRequest);
 		void UnregisterRequest(int32 RequestId);
-		void HttpRequestProgress(FHttpRequestPtr Request, int32 BytesSent, int32 BytesReceived, int32 RequestId);
+		void HttpRequestProgress(FHttpRequestPtr Request, uint64 BytesSent, uint64 BytesReceived, int32 RequestId);
 		void HttpRequestComplete(FHttpRequestPtr Request, FHttpResponsePtr Response, bool bSuccess, IDownloadServiceStat::FDownloadRecord DownloadRecord);
-		void SetRequestProgress(int32 RequestId, int32 BytesSoFar);
+		void SetRequestProgress(int32 RequestId, uint64 BytesSoFar);
 		void SetFileRequestComplete(int32 RequestId, bool bSuccess, TArray<uint8> FileDataArray, IDownloadServiceStat::FDownloadRecord&& DownloadRecord);
 		void SetHttpRequestComplete(int32 RequestId, bool bSuccess, FHttpResponsePtr Response, IDownloadServiceStat::FDownloadRecord&& DownloadRecord);
 
@@ -278,7 +278,7 @@ namespace BuildPatchServices
 		TMap<int32, TUniquePtr<FFileRequest>> ActiveFileRequests;
 
 		FCriticalSection ProgressUpdatesCS;
-		TMap<int32, int32> ProgressUpdates;
+		TMap<int32, uint64> ProgressUpdates;
 
 		FCriticalSection CompletedRequestsCS;
 		TMap<int32, FDownloadBaseRef> CompletedRequests;
@@ -291,7 +291,7 @@ namespace BuildPatchServices
 	{
 	}
 
-	void FDownloadService::FHttpDelegates::HttpRequestProgress(FHttpRequestPtr Request, int32 BytesSent, int32 BytesReceived, int32 RequestId)
+	void FDownloadService::FHttpDelegates::HttpRequestProgress(FHttpRequestPtr Request, uint64 BytesSent, uint64 BytesReceived, int32 RequestId)
 	{
 		DownloadService.HttpRequestProgress(MoveTemp(Request), BytesSent, BytesReceived, RequestId);
 	}
@@ -341,7 +341,7 @@ namespace BuildPatchServices
 		ActiveRequestsCS.Lock();
 		for (const TPair<int32, TSharedRef<IHttpRequest, ESPMode::ThreadSafe>>& ActiveHttpRequest : ActiveHttpRequests)
 		{
-			ActiveHttpRequest.Value->OnRequestProgress().Unbind();
+			ActiveHttpRequest.Value->OnRequestProgress64().Unbind();
 			ActiveHttpRequest.Value->OnProcessRequestComplete().Unbind();
 			ActiveHttpRequest.Value->CancelRequest();
 		}
@@ -488,7 +488,7 @@ namespace BuildPatchServices
 			{
 				// Kick off http request.
 				TSharedRef<IHttpRequest, ESPMode::ThreadSafe> HttpRequest = HttpManager->CreateRequest();
-				HttpRequest->OnRequestProgress().BindThreadSafeSP(HttpDelegates, &FHttpDelegates::HttpRequestProgress, NewRequest.Key);
+				HttpRequest->OnRequestProgress64().BindThreadSafeSP(HttpDelegates, &FHttpDelegates::HttpRequestProgress, NewRequest.Key);
 				HttpRequest->OnProcessRequestComplete().BindThreadSafeSP(HttpDelegates, &FHttpDelegates::HttpRequestComplete, MakeDownloadRecord(NewRequest.Key, NewRequest.Value));
 				HttpRequest->SetURL(NewRequest.Value);
 				HttpRequest->SetVerb(TEXT("GET"));
@@ -510,14 +510,14 @@ namespace BuildPatchServices
 	void FDownloadService::ProcessProgressUpdates()
 	{
 		// Grab the progress updates for this frame.
-		TMap<int32, int32> FrameProgressUpdates;
+		TMap<int32, uint64> FrameProgressUpdates;
 		ProgressUpdatesCS.Lock();
 		FrameProgressUpdates = MoveTemp(ProgressUpdates);
 		ProgressUpdatesCS.Unlock();
 
 		// Process progress updates.
 		RequestDelegatesCS.Lock();
-		for (const TPair<int32, int32>& FrameProgressUpdate : FrameProgressUpdates)
+		for (const TPair<int32, uint64>& FrameProgressUpdate : FrameProgressUpdates)
 		{
 			if (RequestDelegates.Contains(FrameProgressUpdate.Key))
 			{
@@ -525,7 +525,7 @@ namespace BuildPatchServices
 			}
 		}
 		RequestDelegatesCS.Unlock();
-		for (const TPair<int32, int32>& FrameProgressUpdate : FrameProgressUpdates)
+		for (const TPair<int32, uint64>& FrameProgressUpdate : FrameProgressUpdates)
 		{
 			DownloadServiceStat->OnDownloadProgress(FrameProgressUpdate.Key, FrameProgressUpdate.Value);
 		}
@@ -669,7 +669,7 @@ namespace BuildPatchServices
 		ActiveFileRequests.Remove(RequestId);
 	}
 
-	void FDownloadService::HttpRequestProgress(FHttpRequestPtr Request, int32 BytesSent, int32 BytesReceived, int32 RequestId)
+	void FDownloadService::HttpRequestProgress(FHttpRequestPtr Request, uint64 BytesSent, uint64 BytesReceived, int32 RequestId)
 	{
 		SetRequestProgress(RequestId, BytesReceived);
 	}
@@ -684,7 +684,7 @@ namespace BuildPatchServices
 		SetHttpRequestComplete(DownloadRecord.RequestId, bSuccess, MoveTemp(Response), MoveTemp(DownloadRecord));
 	}
 
-	void FDownloadService::SetRequestProgress(int32 RequestId, int32 BytesSoFar)
+	void FDownloadService::SetRequestProgress(int32 RequestId, uint64 BytesSoFar)
 	{
 		FScopeLock ScopeLock(&ProgressUpdatesCS);
 		ProgressUpdates.Add(RequestId, BytesSoFar);
diff --git a/Engine/Source/Runtime/Online/BuildPatchServices/Private/Installer/DownloadService.h b/Engine/Source/Runtime/Online/BuildPatchServices/Private/Installer/DownloadService.h
--- a/Engine/Source/Runtime/Online/BuildPatchServices/Private/Installer/DownloadService.h
+++ b/Engine/Source/Runtime/Online/BuildPatchServices/Private/Installer/DownloadService.h
@@ -51,7 +51,7 @@ namespace BuildPatchServices
 	 * @param RequestId     The id that was returned by RequestFile(..).
 	 * @param BytesSoFar    The number of bytes received so far.
 	 */
-	DECLARE_DELEGATE_TwoParams(FDownloadProgressDelegate, int32 /* RequestId */, int32 /* BytesSoFar */);
+	DECLARE_DELEGATE_TwoParams(FDownloadProgressDelegate, int32 /* RequestId */, uint64 /* BytesSoFar */);
 
 	/**
 	 * Delegate called for download complete.
@@ -146,7 +146,7 @@ namespace BuildPatchServices
 		 * @param RequestId             The id for the request.
 		 * @param BytesReceived         The bytes received so far.
 		 */
-		virtual void OnDownloadProgress(int32 RequestId, int32 BytesReceived) = 0;
+		virtual void OnDownloadProgress(int32 RequestId, uint64 BytesReceived) = 0;
 
 		/**
 		 * Called for each request completion.
diff --git a/Engine/Source/Runtime/Online/BuildPatchServices/Private/Installer/Statistics/DownloadServiceStatistics.cpp b/Engine/Source/Runtime/Online/BuildPatchServices/Private/Installer/Statistics/DownloadServiceStatistics.cpp
--- a/Engine/Source/Runtime/Online/BuildPatchServices/Private/Installer/Statistics/DownloadServiceStatistics.cpp
+++ b/Engine/Source/Runtime/Online/BuildPatchServices/Private/Installer/Statistics/DownloadServiceStatistics.cpp
@@ -24,7 +24,7 @@ namespace BuildPatchServices
 
 		// IDownloadServiceStat interface begin.
 		virtual void OnDownloadStarted(int32 RequestId, const FString& Uri) override;
-		virtual void OnDownloadProgress(int32 RequestId, int32 BytesReceived) override;
+		virtual void OnDownloadProgress(int32 RequestId, uint64 BytesReceived) override;
 		virtual void OnDownloadComplete(const FDownloadRecord& DownloadRecord) override;
 		// IDownloadServiceStat interface end.
 
@@ -46,7 +46,7 @@ namespace BuildPatchServices
 		FThreadSafeInt32 NumSuccessfulDownloads;
 		FThreadSafeInt32 NumFailedDownloads;
 
-		typedef TTuple<FString, int32> FDownloadTuple;
+		typedef TTuple<FString, uint64> FDownloadTuple;
 		TMap<int32, FDownloadTuple> Downloads;
 		double AccumulatedRequestSpeed;
 		uint32 AverageSpeedSampleCount;
@@ -77,7 +77,7 @@ namespace BuildPatchServices
 		DownloadTuple.Get<1>() = 0;
 	}
 
-	void FDownloadServiceStatistics::OnDownloadProgress(int32 RequestId, int32 BytesReceived)
+	void FDownloadServiceStatistics::OnDownloadProgress(int32 RequestId, uint64 BytesReceived)
 	{
 		checkSlow(IsInGameThread());
 		FDownloadTuple& DownloadTuple = Downloads.FindOrAdd(RequestId);
diff --git a/Engine/Source/Runtime/Online/BuildPatchServices/Private/Tests/Mock/DownloadServiceStat.mock.h b/Engine/Source/Runtime/Online/BuildPatchServices/Private/Tests/Mock/DownloadServiceStat.mock.h
--- a/Engine/Source/Runtime/Online/BuildPatchServices/Private/Tests/Mock/DownloadServiceStat.mock.h
+++ b/Engine/Source/Runtime/Online/BuildPatchServices/Private/Tests/Mock/DownloadServiceStat.mock.h
@@ -22,7 +22,7 @@ namespace BuildPatchServices
 			RxDownloadStarted.Emplace(FStatsCollector::GetSeconds(), RequestId, Uri);
 		}
 
-		virtual void OnDownloadProgress(int32 RequestId, int32 BytesReceived) override
+		virtual void OnDownloadProgress(int32 RequestId, uint64 BytesReceived) override
 		{
 			RxDownloadProgress.Emplace(FStatsCollector::GetSeconds(), RequestId, BytesReceived);
 		}
diff --git a/Engine/Source/Runtime/Online/BuildPatchServices/Private/Tests/Mock/HttpRequest.mock.h b/Engine/Source/Runtime/Online/BuildPatchServices/Private/Tests/Mock/HttpRequest.mock.h
--- a/Engine/Source/Runtime/Online/BuildPatchServices/Private/Tests/Mock/HttpRequest.mock.h
+++ b/Engine/Source/Runtime/Online/BuildPatchServices/Private/Tests/Mock/HttpRequest.mock.h
@@ -156,6 +156,11 @@ namespace BuildPatchServices
 			return HttpRequestProgressDelegate;
 		}
 
+		virtual FHttpRequestProgressDelegate64& OnRequestProgress64() override
+		{
+			return HttpRequestProgressDelegate64;
+		}
+
 		virtual FHttpRequestHeaderReceivedDelegate& OnHeaderReceived() override
 		{
 			return HttpHeaderReceivedDelegate;
@@ -202,6 +207,7 @@ namespace BuildPatchServices
 
 	public:
 		FHttpRequestProgressDelegate HttpRequestProgressDelegate;
+		FHttpRequestProgressDelegate64 HttpRequestProgressDelegate64;
 		FHttpRequestCompleteDelegate HttpRequestCompleteDelegate;
 		FHttpRequestHeaderReceivedDelegate HttpHeaderReceivedDelegate;
 		FHttpRequestWillRetryDelegate HttpRequestWillRetryDelegate;
diff --git a/Engine/Source/Runtime/Online/HTTP/Private/Apple/AppleHTTPNSUrlConnection.cpp b/Engine/Source/Runtime/Online/HTTP/Private/Apple/AppleHTTPNSUrlConnection.cpp
--- a/Engine/Source/Runtime/Online/HTTP/Private/Apple/AppleHTTPNSUrlConnection.cpp
+++ b/Engine/Source/Runtime/Online/HTTP/Private/Apple/AppleHTTPNSUrlConnection.cpp
@@ -508,16 +508,16 @@ void FAppleHttpNSUrlConnectionRequest::Tick(float DeltaSeconds)
 {
 	if (Response.IsValid() && (CompletionStatus == EHttpRequestStatus::Processing || Response->HadError()))
 	{
-		if (OnRequestProgress().IsBound())
+		const uint64 BytesWritten = Response->GetNumBytesWritten();
+		const uint64 BytesRead = Response->GetNumBytesReceived();
+		if (BytesWritten != LastReportedBytesWritten || BytesRead != LastReportedBytesRead)
 		{
-			const uint64 BytesWritten = Response->GetNumBytesWritten();
-			const uint64 BytesRead = Response->GetNumBytesReceived();
-			if (BytesWritten != LastReportedBytesWritten || BytesRead != LastReportedBytesRead)
-			{
-				OnRequestProgress().ExecuteIfBound(SharedThis(this), BytesWritten, BytesRead);
-				LastReportedBytesWritten = BytesWritten;
-				LastReportedBytesRead = BytesRead;
-			}
+PRAGMA_DISABLE_DEPRECATION_WARNINGS
+			OnRequestProgress().ExecuteIfBound(SharedThis(this), BytesWritten, BytesRead);
+PRAGMA_ENABLE_DEPRECATION_WARNINGS
+			OnRequestProgress64().ExecuteIfBound(SharedThis(this), BytesWritten, BytesRead);
+			LastReportedBytesWritten = BytesWritten;
+			LastReportedBytesRead = BytesRead;
 		}
 		if (Response->IsReady())
 		{
diff --git a/Engine/Source/Runtime/Online/HTTP/Private/Apple/AppleHTTPNSUrlSession.cpp b/Engine/Source/Runtime/Online/HTTP/Private/Apple/AppleHTTPNSUrlSession.cpp
--- a/Engine/Source/Runtime/Online/HTTP/Private/Apple/AppleHTTPNSUrlSession.cpp
+++ b/Engine/Source/Runtime/Online/HTTP/Private/Apple/AppleHTTPNSUrlSession.cpp
@@ -814,16 +814,16 @@ void FAppleHttpNSUrlSessionRequest::Tick(float DeltaSeconds)
 {
 	if (Response.IsValid() && (CompletionStatus == EHttpRequestStatus::Processing || Response->HadError()))
 	{
-		if (OnRequestProgress().IsBound())
+		const uint64 BytesWritten = Response->GetNumBytesWritten();
+		const uint64 BytesRead = Response->GetNumBytesReceived();
+		if (BytesWritten != LastReportedBytesWritten || BytesRead != LastReportedBytesRead)
 		{
-			const int32 BytesWritten = Response->GetNumBytesWritten();
-			const int32 BytesRead = Response->GetNumBytesReceived();
-			if (BytesWritten != LastReportedBytesWritten || BytesRead != LastReportedBytesRead)
-			{
-				OnRequestProgress().ExecuteIfBound(SharedThis(this), BytesWritten, BytesRead);
-				LastReportedBytesWritten = BytesWritten;
-				LastReportedBytesRead = BytesRead;
-			}
+PRAGMA_DISABLE_DEPRECATION_WARNINGS
+			OnRequestProgress().ExecuteIfBound(SharedThis(this), BytesWritten, BytesRead);
+PRAGMA_ENABLE_DEPRECATION_WARNINGS
+			OnRequestProgress64().ExecuteIfBound(SharedThis(this), BytesWritten, BytesRead);
+			LastReportedBytesWritten = BytesWritten;
+			LastReportedBytesRead = BytesRead;
 		}
 	}
 }
diff --git a/Engine/Source/Runtime/Online/HTTP/Private/Curl/CurlHttp.cpp b/Engine/Source/Runtime/Online/HTTP/Private/Curl/CurlHttp.cpp
--- a/Engine/Source/Runtime/Online/HTTP/Private/Curl/CurlHttp.cpp
+++ b/Engine/Source/Runtime/Online/HTTP/Private/Curl/CurlHttp.cpp
@@ -520,7 +520,7 @@ size_t FCurlHttpRequest::ReceiveResponseHeaderCallback(void* Ptr, size_t SizeInB
 			HeaderValue.TrimStartInline();
 			if (!HeaderKey.IsEmpty() && !HeaderValue.IsEmpty() && !bRedirected)
 			{
-				//Store the content length so OnRequestProgress() delegates have something to work with
+				//Store the content length so OnRequestProgress64() delegates have something to work with
 				if (HeaderKey == TEXT("Content-Length"))
 				{
 					Response->ContentLength = FCString::Atoi64(*HeaderValue);
@@ -1122,8 +1122,11 @@ void FCurlHttpRequest::CheckProgressDelegate()
 		{
 			LastReportedBytesRead = CurrentBytesRead;
 		}
+PRAGMA_DISABLE_DEPRECATION_WARNINGS
 		// Update response progress
 		OnRequestProgress().ExecuteIfBound(SharedThis(this), LastReportedBytesSent, LastReportedBytesRead);
+PRAGMA_ENABLE_DEPRECATION_WARNINGS
+		OnRequestProgress64().ExecuteIfBound(SharedThis(this), LastReportedBytesSent, LastReportedBytesRead);
 	}
 }
 
diff --git a/Engine/Source/Runtime/Online/HTTP/Private/GenericPlatform/HttpRequestImpl.cpp b/Engine/Source/Runtime/Online/HTTP/Private/GenericPlatform/HttpRequestImpl.cpp
--- a/Engine/Source/Runtime/Online/HTTP/Private/GenericPlatform/HttpRequestImpl.cpp
+++ b/Engine/Source/Runtime/Online/HTTP/Private/GenericPlatform/HttpRequestImpl.cpp
@@ -10,11 +10,19 @@ FHttpRequestCompleteDelegate& FHttpRequestImpl::OnProcessRequestComplete()
 	return RequestCompleteDelegate;
 }
 
+PRAGMA_DISABLE_DEPRECATION_WARNINGS
 FHttpRequestProgressDelegate& FHttpRequestImpl::OnRequestProgress() 
 {
 	UE_LOG(LogHttp, VeryVerbose, TEXT("FHttpRequestImpl::OnRequestProgress()"));
 	return RequestProgressDelegate;
 }
+PRAGMA_ENABLE_DEPRECATION_WARNINGS
+
+FHttpRequestProgressDelegate64& FHttpRequestImpl::OnRequestProgress64() 
+{
+	UE_LOG(LogHttp, VeryVerbose, TEXT("FHttpRequestImpl::OnRequestProgress64()"));
+	return RequestProgressDelegate64;
+}
 
 FHttpRequestHeaderReceivedDelegate& FHttpRequestImpl::OnHeaderReceived()
 {
diff --git a/Engine/Source/Runtime/Online/HTTP/Private/HttpManager.cpp b/Engine/Source/Runtime/Online/HTTP/Private/HttpManager.cpp
--- a/Engine/Source/Runtime/Online/HTTP/Private/HttpManager.cpp
+++ b/Engine/Source/Runtime/Online/HTTP/Private/HttpManager.cpp
@@ -85,7 +85,10 @@ void FHttpManager::Shutdown()
 	{
 		FHttpRequestRef& Request = *It;
 		Request->OnProcessRequestComplete().Unbind();
+PRAGMA_DISABLE_DEPRECATION_WARNINGS
 		Request->OnRequestProgress().Unbind();
+PRAGMA_ENABLE_DEPRECATION_WARNINGS
+		Request->OnRequestProgress64().Unbind();
 		Request->OnHeaderReceived().Unbind();
 
 		// Don't emit these tracking logs in commandlet runs. Build system traps warnings during cook, and these are not truly fatal, but useful for tracking down shutdown issues.
diff --git a/Engine/Source/Runtime/Online/HTTP/Private/HttpRetrySystem.cpp b/Engine/Source/Runtime/Online/HTTP/Private/HttpRetrySystem.cpp
--- a/Engine/Source/Runtime/Online/HTTP/Private/HttpRetrySystem.cpp
+++ b/Engine/Source/Runtime/Online/HTTP/Private/HttpRetrySystem.cpp
@@ -108,7 +108,7 @@ bool FHttpRetrySystem::FRequest::ProcessRequest()
 		SetUrlFromRetryDomains();
 	}
 
-	HttpRequest->OnRequestProgress().BindThreadSafeSP(RetryRequest, &FHttpRetrySystem::FRequest::HttpOnRequestProgress);
+	HttpRequest->OnRequestProgress64().BindThreadSafeSP(RetryRequest, &FHttpRetrySystem::FRequest::HttpOnRequestProgress);
 
 	return RetryManager.ProcessRequest(RetryRequest);
 }
@@ -142,9 +142,12 @@ void FHttpRetrySystem::FRequest::CancelRequest()
 	RetryManager.CancelRequest(RetryRequest);
 }
 
-void FHttpRetrySystem::FRequest::HttpOnRequestProgress(FHttpRequestPtr InHttpRequest, int32 BytesSent, int32 BytesRcv)
+void FHttpRetrySystem::FRequest::HttpOnRequestProgress(FHttpRequestPtr InHttpRequest, uint64 BytesSent, uint64 BytesRcv)
 {
+PRAGMA_DISABLE_DEPRECATION_WARNINGS
 	OnRequestProgress().ExecuteIfBound(AsShared(), BytesSent, BytesRcv);
+PRAGMA_ENABLE_DEPRECATION_WARNINGS
+	OnRequestProgress64().ExecuteIfBound(AsShared(), BytesSent, BytesRcv);
 }
 
 FHttpRetrySystem::FManager::FManager(const FRetryLimitCountSetting& InRetryLimitCountDefault, const FRetryTimeoutRelativeSecondsSetting& InRetryTimeoutRelativeSecondsDefault)
diff --git a/Engine/Source/Runtime/Online/HTTP/Public/GenericPlatform/HttpRequestImpl.h b/Engine/Source/Runtime/Online/HTTP/Public/GenericPlatform/HttpRequestImpl.h
--- a/Engine/Source/Runtime/Online/HTTP/Public/GenericPlatform/HttpRequestImpl.h
+++ b/Engine/Source/Runtime/Online/HTTP/Public/GenericPlatform/HttpRequestImpl.h
@@ -14,6 +14,7 @@ public:
 	// IHttpRequest
 	HTTP_API virtual FHttpRequestCompleteDelegate& OnProcessRequestComplete() override;
 	HTTP_API virtual FHttpRequestProgressDelegate& OnRequestProgress() override;
+	HTTP_API virtual FHttpRequestProgressDelegate64& OnRequestProgress64() override;
 	HTTP_API virtual FHttpRequestHeaderReceivedDelegate& OnHeaderReceived() override;
 	HTTP_API virtual FHttpRequestWillRetryDelegate& OnRequestWillRetry() override;
 
@@ -37,6 +38,9 @@ protected:
 	/** Delegate that will get called once per tick with bytes downloaded so far */
 	FHttpRequestProgressDelegate RequestProgressDelegate;
 
+	/** Delegate that will get called once per tick with bytes downloaded so far */
+	FHttpRequestProgressDelegate64 RequestProgressDelegate64;
+
 	/** Delegate that will get called for each new header received */
 	FHttpRequestHeaderReceivedDelegate HeaderReceivedDelegate;
 	
diff --git a/Engine/Source/Runtime/Online/HTTP/Public/HttpRetrySystem.h b/Engine/Source/Runtime/Online/HTTP/Public/HttpRetrySystem.h
--- a/Engine/Source/Runtime/Online/HTTP/Public/HttpRetrySystem.h
+++ b/Engine/Source/Runtime/Online/HTTP/Public/HttpRetrySystem.h
@@ -113,7 +113,7 @@ namespace FHttpRetrySystem
 			const FRetryDomainsPtr& InRetryDomains = FRetryDomainsPtr()
 			);
 
-		void HttpOnRequestProgress(FHttpRequestPtr InHttpRequest, int32 BytesSent, int32 BytesRcv);
+		void HttpOnRequestProgress(FHttpRequestPtr InHttpRequest, uint64 BytesSent, uint64 BytesRcv);
 
 		/** Update our HTTP request's URL's domain from our RetryDomains */
 		void SetUrlFromRetryDomains();
diff --git a/Engine/Source/Runtime/Online/HTTP/Public/Interfaces/IHttpRequest.h b/Engine/Source/Runtime/Online/HTTP/Public/Interfaces/IHttpRequest.h
--- a/Engine/Source/Runtime/Online/HTTP/Public/Interfaces/IHttpRequest.h
+++ b/Engine/Source/Runtime/Online/HTTP/Public/Interfaces/IHttpRequest.h
@@ -107,6 +107,15 @@ DECLARE_DELEGATE_ThreeParams(FHttpRequestHeaderReceivedDelegate, FHttpRequestPtr
  */
 DECLARE_DELEGATE_ThreeParams(FHttpRequestProgressDelegate, FHttpRequestPtr /*Request*/, int32 /*BytesSent*/, int32 /*BytesReceived*/);
 
+/**
+ * Delegate called per tick to update an Http request upload or download size progress
+ *
+ * @param Request original Http request that started things
+ * @param BytesSent the number of bytes sent / uploaded in the request so far.
+ * @param BytesReceived the number of bytes received / downloaded in the response so far.
+ */
+DECLARE_DELEGATE_ThreeParams(FHttpRequestProgressDelegate64, FHttpRequestPtr /*Request*/, uint64 /*BytesSent*/, uint64 /*BytesReceived*/);
+
 /**
  * Delegate called when an Http request will be retried in the future
  *
@@ -326,7 +335,13 @@ public:
 	/**
 	 * Delegate called to update the request/response progress. See FHttpRequestProgressDelegate
 	 */
+	UE_DEPRECATED(5.3, "OnRequestProgress has been deprecated, use OnRequestProgress64 instead")
 	virtual FHttpRequestProgressDelegate& OnRequestProgress() = 0;
+
+	/**
+	 * Delegate called to update the request/response progress. See FHttpRequestProgressDelegate64
+	 */
+	virtual FHttpRequestProgressDelegate64& OnRequestProgress64() = 0;
 	
 	/**
 	* Delegate called when the request will be retried
