Fix wrong color format used in FWebImage downloaded images depending on the actual image format

diff --git a/Engine/Source/Runtime/Online/ImageDownload/ImageDownload.Build.cs b/Engine/Source/Runtime/Online/ImageDownload/ImageDownload.Build.cs
--- a/Engine/Source/Runtime/Online/ImageDownload/ImageDownload.Build.cs
+++ b/Engine/Source/Runtime/Online/ImageDownload/ImageDownload.Build.cs
@@ -13,6 +13,7 @@ public class ImageDownload : ModuleRules
 				"Core",
 				"HTTP",
 				"SlateCore",
+				"ImageCore",
 				"ImageWrapper",
 			}
 			);
diff --git a/Engine/Source/Runtime/Online/ImageDownload/Private/WebImage.cpp b/Engine/Source/Runtime/Online/ImageDownload/Private/WebImage.cpp
--- a/Engine/Source/Runtime/Online/ImageDownload/Private/WebImage.cpp
+++ b/Engine/Source/Runtime/Online/ImageDownload/Private/WebImage.cpp
@@ -13,7 +13,6 @@
 FWebImage::FWebImage()
 : StandInBrush(FCoreStyle::Get().GetDefaultBrush())
 , bDownloadSuccess(false)
-, RGBFormat(ERGBFormat::RGBA)
 {
 }
 
@@ -112,42 +111,22 @@ bool FWebImage::ProcessHttpResponse(const FString& RequestUrl, FHttpResponsePtr
 	static const FName MODULE_IMAGE_WRAPPER("ImageWrapper");
 	IImageWrapperModule& ImageWrapperModule = FModuleManager::LoadModuleChecked<IImageWrapperModule>(MODULE_IMAGE_WRAPPER);
 
-	// Look at the signature of the downloaded image to detect image type. (and ignore the content type header except for error reporting)
 	const TArray<uint8>& Content = HttpResponse->GetContent();
-	EImageFormat ImageFormat = ImageWrapperModule.DetectImageFormat(Content.GetData(), Content.Num());
-
-	if (ImageFormat == EImageFormat::Invalid)
+	FImage DownloadedImage;
+	if (!ImageWrapperModule.DecompressImage(Content.GetData(), Content.Num(), DownloadedImage))
 	{
 		FString ContentType = HttpResponse->GetContentType();
 		UE_LOG(LogImageDownload, Error, TEXT("Image Download: Could not recognize file type of image downloaded from url %s, server-reported content type: %s"), *RequestUrl, *ContentType);
 		return false;
 	}
 
-	TSharedPtr<IImageWrapper> ImageWrapper = ImageWrapperModule.CreateImageWrapper(ImageFormat);
-	if (!ImageWrapper.IsValid())
-	{
-		UE_LOG(LogImageDownload, Error, TEXT("Image Download: Unable to make image wrapper for image format %d"), (int32)ImageFormat);
-		return false;
-	}
-
-	// parse the content
-	if (!ImageWrapper->SetCompressed(Content.GetData(), Content.Num()))
-	{
-		UE_LOG(LogImageDownload, Error, TEXT("Image Download: Unable to parse image format %d from %s"), (int32)ImageFormat, *RequestUrl);
-		return false;
-	}
-
-	// get the raw image data
-	TArray<uint8> RawImageData;
-	if (!ImageWrapper->GetRaw(RGBFormat, 8, RawImageData))
-	{
-		UE_LOG(LogImageDownload, Error, TEXT("Image Download: Unable to convert image format %d to BGRA 8"), (int32)ImageFormat);
-		return false;
-	}
+	DownloadedImage.ChangeFormat(ERawImageFormat::BGRA8, EGammaSpace::sRGB);
 
 	// make a dynamic image
 	FName ResourceName(*RequestUrl);
-	DownloadedBrush = FSlateDynamicImageBrush::CreateWithImageData(ResourceName, FVector2D((float)ImageWrapper->GetWidth(), (float)ImageWrapper->GetHeight()), RawImageData);
+	// FImage.RawData is a TArray64<uint8> while FSlateDynamicImageBrush::CreateWithImageData expects a TArray<uint8>. We need to copy the image data to adapt
+	TArray<uint8> RawImageCopy(DownloadedImage.RawData);
+	DownloadedBrush = FSlateDynamicImageBrush::CreateWithImageData(ResourceName, FVector2D((float)DownloadedImage.SizeX, (float)DownloadedImage.SizeY), RawImageCopy);
 	return DownloadedBrush.IsValid();
 }
 
diff --git a/Engine/Source/Runtime/Online/ImageDownload/Public/WebImage.h b/Engine/Source/Runtime/Online/ImageDownload/Public/WebImage.h
--- a/Engine/Source/Runtime/Online/ImageDownload/Public/WebImage.h
+++ b/Engine/Source/Runtime/Online/ImageDownload/Public/WebImage.h
@@ -56,7 +56,8 @@ public:
 	FORCEINLINE FWebImage& SetStandInBrush(TAttribute<const FSlateBrush*> StandInBrushIn) { StandInBrush = StandInBrushIn; DownloadedBrush.Reset(); return *this; }
 
 	/** Sets the RGB format used to read the downloaded image data (ERGBFormat::RGBA by default) */
-	void SetRGBFormat(const ERGBFormat InRGBFormat) { RGBFormat = InRGBFormat; }
+	UE_DEPRECATED(5.4, "ERGBFormat format is detected from downloaded image")
+	void SetRGBFormat(const ERGBFormat InRGBFormat) {}
 
 	/** Begin downloading an image. This will automatically set the current brush to the downloaded image when it completes (if successful) */
 	IMAGEDOWNLOAD_API bool BeginDownload(const FString& InUrl, const TOptional<FString>& StandInETag = TOptional<FString>(), const FOnImageDownloaded& DownloadCallback = FOnImageDownloaded());
@@ -125,7 +126,4 @@ private:
 
 	/** The ETag of the downloaded image */
 	TOptional<FString> ETag;
-
-	/** RGB format used to read the downloaded image data */
-	ERGBFormat RGBFormat;
 };
