Fix race condition in FCurlHttpRequest when using EHttpRequestDelegateThreadPolicy::CompleteOnHttpThread

diff --git a/Engine/Source/Runtime/Online/HTTP/Private/Curl/CurlHttp.cpp b/Engine/Source/Runtime/Online/HTTP/Private/Curl/CurlHttp.cpp
--- a/Engine/Source/Runtime/Online/HTTP/Private/Curl/CurlHttp.cpp
+++ b/Engine/Source/Runtime/Online/HTTP/Private/Curl/CurlHttp.cpp
@@ -525,7 +525,15 @@ size_t FCurlHttpRequest::ReceiveResponseHeaderCallback(void* Ptr, size_t SizeInB
 				{
 					Response->ContentLength = FCString::Atoi64(*HeaderValue);
 				}
-				Response->NewlyReceivedHeaders.Enqueue(TPair<FString, FString>(MoveTemp(HeaderKey), MoveTemp(HeaderValue)));
+
+				if (DelegateThreadPolicy == EHttpRequestDelegateThreadPolicy::CompleteOnHttpThread)
+				{
+					BroadcastNewlyReceivedHeader(HeaderKey, HeaderValue);
+				}
+				else
+				{
+					Response->NewlyReceivedHeaders.Enqueue(TPair<FString, FString>(MoveTemp(HeaderKey), MoveTemp(HeaderValue)));
+				}
 			}
 		}
 		else
@@ -1102,7 +1110,11 @@ const FHttpResponsePtr FCurlHttpRequest::GetResponse() const
 void FCurlHttpRequest::Tick(float DeltaSeconds)
 {
 	CheckProgressDelegate();
-	BroadcastNewlyReceivedHeaders();
+
+	if (DelegateThreadPolicy == EHttpRequestDelegateThreadPolicy::CompleteOnGameThread)
+	{
+		BroadcastNewlyReceivedHeaders();
+	}
 }
 
 void FCurlHttpRequest::CheckProgressDelegate()
@@ -1132,31 +1144,38 @@ PRAGMA_ENABLE_DEPRECATION_WARNINGS
 
 void FCurlHttpRequest::BroadcastNewlyReceivedHeaders()
 {
-	check(IsInGameThread() || DelegateThreadPolicy == EHttpRequestDelegateThreadPolicy::CompleteOnHttpThread);
+	check(IsInGameThread());
+
 	if (Response.IsValid())
 	{
 		// Process the headers received on the HTTP thread and merge them into the response's list of headers and then broadcast the new headers
 		TPair<FString, FString> NewHeader;
 		while (Response->NewlyReceivedHeaders.Dequeue(NewHeader))
 		{
-			const FString& HeaderKey = NewHeader.Key;
-			const FString& HeaderValue = NewHeader.Value;
+			BroadcastNewlyReceivedHeader(NewHeader.Key, NewHeader.Value);
+		}
+	}
+}
 
-			FString NewValue;
-			FString* PreviousValue = Response->Headers.Find(HeaderKey);
-			if (PreviousValue != nullptr && !PreviousValue->IsEmpty())
-			{
-				constexpr const int32 SeparatorLength = 2; // Length of ", "
-				NewValue = MoveTemp(*PreviousValue);
-				NewValue.Reserve(NewValue.Len() + SeparatorLength + HeaderValue.Len());
-				NewValue += TEXT(", ");
-			}
-			NewValue += HeaderValue;
-			Response->Headers.Add(HeaderKey, MoveTemp(NewValue));
+void FCurlHttpRequest::BroadcastNewlyReceivedHeader(const FString& HeaderKey, const FString& HeaderValue)
+{
+	check(IsInGameThread() || DelegateThreadPolicy == EHttpRequestDelegateThreadPolicy::CompleteOnHttpThread);
+	check(Response);
 
-			OnHeaderReceived().ExecuteIfBound(SharedThis(this), NewHeader.Key, NewHeader.Value);
-		}
+	const constexpr FStringView Seperator(TEXTVIEW(", "));
+
+	FString NewValue;
+	FString* PreviousValue = Response->Headers.Find(HeaderKey);
+	if (PreviousValue != nullptr && !PreviousValue->IsEmpty())
+	{
+		NewValue = MoveTemp(*PreviousValue);
+		NewValue.Reserve(NewValue.Len() + Seperator.Len() + HeaderValue.Len());
+		NewValue += Seperator;
 	}
+	NewValue += HeaderValue;
+	Response->Headers.Add(HeaderKey, MoveTemp(NewValue));
+
+	OnHeaderReceived().ExecuteIfBound(SharedThis(this), HeaderKey, HeaderValue);
 }
 
 void FCurlHttpRequest::FinishRequest()
@@ -1209,7 +1228,12 @@ void FCurlHttpRequest::FinishRequest()
 	// if just finished, mark as stopped async processing
 	if (Response.IsValid())
 	{
-		BroadcastNewlyReceivedHeaders();
+		// Broadcast any headers we haven't broadcast yet
+		// If using EHttpRequestDelegateThreadPolicy::CompleteOnHttpThread, we don't add to NewlyReceivedHeaders and will have already broadcast
+		if (DelegateThreadPolicy == EHttpRequestDelegateThreadPolicy::CompleteOnGameThread)
+		{
+			BroadcastNewlyReceivedHeaders();
+		}
 		Response->bIsReady = true;
 	}
 
@@ -1253,8 +1277,6 @@ void FCurlHttpRequest::FinishRequest()
 
 		// Mark last request attempt as completed successfully
 		CompletionStatus = EHttpRequestStatus::Succeeded;
-		// Broadcast any headers we haven't broadcast yet
-		BroadcastNewlyReceivedHeaders();
 		// Call delegate with valid request/response objects
 		OnProcessRequestComplete().ExecuteIfBound(SharedThis(this),Response,true);
 	}
diff --git a/Engine/Source/Runtime/Online/HTTP/Private/Curl/CurlHttp.h b/Engine/Source/Runtime/Online/HTTP/Private/Curl/CurlHttp.h
--- a/Engine/Source/Runtime/Online/HTTP/Private/Curl/CurlHttp.h
+++ b/Engine/Source/Runtime/Online/HTTP/Private/Curl/CurlHttp.h
@@ -351,6 +351,7 @@ private:
 
 	/** Broadcast newly received headers */
 	void BroadcastNewlyReceivedHeaders();
+	void BroadcastNewlyReceivedHeader(const FString& HeaderKey, const FString& HeaderValue);
 
 	/** Combine a header's key/value in the format "Key: Value" */
 	static FString CombineHeaderKeyValue(const FString& HeaderKey, const FString& HeaderValue);
