Fix a potential crash in http curl impl when destroy http requests.

diff --git a/Engine/Source/Runtime/Online/HTTP/Private/Curl/CurlHttp.cpp b/Engine/Source/Runtime/Online/HTTP/Private/Curl/CurlHttp.cpp
--- a/Engine/Source/Runtime/Online/HTTP/Private/Curl/CurlHttp.cpp
+++ b/Engine/Source/Runtime/Online/HTTP/Private/Curl/CurlHttp.cpp
@@ -185,9 +185,6 @@ FCurlHttpRequest::~FCurlHttpRequest()
 	checkf(FCurlHttpManager::IsInit(), TEXT("Curl request was held after the library was shutdown."));
 	if (EasyHandle)
 	{
-		// clear to prevent crashing in debug callback when this handle is part of an asynchronous curl_multi_perform()
-		curl_easy_setopt(EasyHandle, CURLOPT_DEBUGDATA, nullptr);
-
 		// cleanup the handle first (that order is used in howtos)
 		curl_easy_cleanup(EasyHandle);
 		EasyHandle = nullptr;
