Fix instances of incorrect usage of UE_LOG.

diff --git a/Engine/Plugins/Developer/Concert/ConcertSync/ConcertSyncClient/Source/ConcertSyncClient/Private/ConcertClientSequencerManager.cpp b/Engine/Plugins/Developer/Concert/ConcertSync/ConcertSyncClient/Source/ConcertSyncClient/Private/ConcertClientSequencerManager.cpp
--- a/Engine/Plugins/Developer/Concert/ConcertSync/ConcertSyncClient/Source/ConcertSyncClient/Private/ConcertClientSequencerManager.cpp
+++ b/Engine/Plugins/Developer/Concert/ConcertSync/ConcertSyncClient/Source/ConcertSyncClient/Private/ConcertClientSequencerManager.cpp
@@ -767,7 +767,7 @@ void FConcertClientSequencerManager::ApplyOpenEvent(const FConcertSequencerOpenE
 		FConcertSyncEncounteredMissingObject MissingObjectDelegate;
 		MissingObjectDelegate.BindLambda([](const FStringView MissingObject)
 			{
-				UE_LOG(LogConcertSequencerSync, Display, TEXT("Missing Object %s when loading PendingTake"), MissingObject);
+				UE_LOG(LogConcertSequencerSync, Display, TEXT("Missing Object %s when loading PendingTake"), *FString(MissingObject));
 			});
 
 		FConcertSyncWorldRemapper Remapper(
diff --git a/Engine/Plugins/Experimental/ChaosVD/Source/ChaosVD/Private/ChaosVDGeometryDataComponent.cpp b/Engine/Plugins/Experimental/ChaosVD/Source/ChaosVD/Private/ChaosVDGeometryDataComponent.cpp
--- a/Engine/Plugins/Experimental/ChaosVD/Source/ChaosVD/Private/ChaosVDGeometryDataComponent.cpp
+++ b/Engine/Plugins/Experimental/ChaosVD/Source/ChaosVD/Private/ChaosVDGeometryDataComponent.cpp
@@ -36,7 +36,7 @@ void FChaosVDGeometryDataComponentBase::UpdateVisibility_Internal(const FChaosVD
 			else
 			{
 				ensure(false);
-				UE_LOG(LogChaosVDEditor, Error, TEXT("[%hs] Failed to get Query Only material for, applying the default mesh to all geometry"), ANSI_TO_TCHAR(__FUNCTION__));
+				UE_LOG(LogChaosVDEditor, Error, TEXT("[%s] Failed to get Query Only material for, applying the default mesh to all geometry"), ANSI_TO_TCHAR(__FUNCTION__));
 				
 				MeshComponent->SetMaterial(0, UMaterial::GetDefaultMaterial(EMaterialDomain::MD_Surface));
 			}
diff --git a/Engine/Plugins/Experimental/ChaosVD/Source/ChaosVD/Private/Trace/ChaosVDTraceManager.cpp b/Engine/Plugins/Experimental/ChaosVD/Source/ChaosVD/Private/Trace/ChaosVDTraceManager.cpp
--- a/Engine/Plugins/Experimental/ChaosVD/Source/ChaosVD/Private/Trace/ChaosVDTraceManager.cpp
+++ b/Engine/Plugins/Experimental/ChaosVD/Source/ChaosVD/Private/Trace/ChaosVDTraceManager.cpp
@@ -93,14 +93,14 @@ FString FChaosVDTraceManager::GetLocalTraceStoreDirPath()
 
 	if (!StoreClient)
 	{
-		UE_LOG(LogChaosVDEditor, Error, TEXT("[%hs] Failed to connect to local Trace Store client"), ANSI_TO_TCHAR(__FUNCTION__));
+		UE_LOG(LogChaosVDEditor, Error, TEXT("[%s] Failed to connect to local Trace Store client"), ANSI_TO_TCHAR(__FUNCTION__));
 		return TEXT("");
 	}
 
 	const UE::Trace::FStoreClient::FStatus* Status = StoreClient->GetStatus();
 	if (!Status)
 	{
-		UE_LOG(LogChaosVDEditor, Error, TEXT("[%hs] Failed to to get Trace Store staus"), ANSI_TO_TCHAR(__FUNCTION__));
+		UE_LOG(LogChaosVDEditor, Error, TEXT("[%s] Failed to to get Trace Store staus"), ANSI_TO_TCHAR(__FUNCTION__));
 		return TEXT("");
 	}
 
diff --git a/Engine/Plugins/Experimental/GameFeatures/Source/GameFeatures/Private/GameFeatureData.cpp b/Engine/Plugins/Experimental/GameFeatures/Source/GameFeatures/Private/GameFeatureData.cpp
--- a/Engine/Plugins/Experimental/GameFeatures/Source/GameFeatures/Private/GameFeatureData.cpp
+++ b/Engine/Plugins/Experimental/GameFeatures/Source/GameFeatures/Private/GameFeatureData.cpp
@@ -295,7 +295,7 @@ void UGameFeatureData::InitializeHierarchicalPluginIniFiles(const FString& Plugi
 					}
 					else
 					{
-						UE_LOG(LogGameFeatures, Warning, TEXT("Game feature '%s' has invalid runtime device profile %s with parent %s, suffix %s, %d CVars, %d fragments"),
+						UE_LOG(LogGameFeatures, Warning, TEXT("Game feature '%s' has invalid runtime device profile with parent %s, suffix %s, %d CVars, %d fragments"),
 							*PluginName, ParentProfileName ? *ParentProfileName->GetValue() : TEXT("null"), ProfileSuffix ? *ProfileSuffix->GetValue() : TEXT("null"),
 							FragmentIncludes.Num(), PluginCVars.Num());
 					}
diff --git a/Engine/Plugins/Experimental/LearningAgents/Source/LearningAgents/Private/LearningAgentsNeuralNetwork.cpp b/Engine/Plugins/Experimental/LearningAgents/Source/LearningAgents/Private/LearningAgentsNeuralNetwork.cpp
--- a/Engine/Plugins/Experimental/LearningAgents/Source/LearningAgents/Private/LearningAgentsNeuralNetwork.cpp
+++ b/Engine/Plugins/Experimental/LearningAgents/Source/LearningAgents/Private/LearningAgentsNeuralNetwork.cpp
@@ -138,7 +138,7 @@ void ULearningAgentsNeuralNetwork::LoadNetworkFromSnapshot(const FFilePath& File
 
 		if (Offset != RecordingData.Num())
 		{
-			UE_LOG(LogLearning, Error, TEXT("%s: Failed to load network. Unexpected end of file."));
+			UE_LOG(LogLearning, Error, TEXT("Failed to load network. Unexpected end of file."));
 			return;
 		}
 
@@ -274,4 +274,4 @@ void ULearningAgentsNeuralNetwork::ForceMarkDirty()
 
 		Package->PackageMarkedDirtyEvent.Broadcast(Package, bIsDirty);
 	}
-}
\ No newline at end of file
+}
diff --git a/Engine/Plugins/Media/AjaMedia/Source/Aja/Private/Device.cpp b/Engine/Plugins/Media/AjaMedia/Source/Aja/Private/Device.cpp
--- a/Engine/Plugins/Media/AjaMedia/Source/Aja/Private/Device.cpp
+++ b/Engine/Plugins/Media/AjaMedia/Source/Aja/Private/Device.cpp
@@ -602,7 +602,7 @@ namespace AJA
 
 				if (Connection.OutputReferenceType == EAJAReferenceType::EAJA_REFERENCETYPE_INPUT && Connection.OutputReferenceChannel != InOutputReferenceChannel)
 				{
-					UE_LOG(LogAjaCore, Error, TEXT("Device: Couldn't set the reference (%S) for output on device %S. The input was already setup with %S.\n"), Helpers::ReferenceTypeToString(InOutputReferenceType), Connection.Card->GetDisplayName().c_str(), Connection.OutputReferenceChannel, Helpers::ReferenceTypeToString(Connection.OutputReferenceType));
+					UE_LOG(LogAjaCore, Error, TEXT("Device: Couldn't set the reference (%S) for output on device %S. The input was already setup with %d.\n"), Helpers::ReferenceTypeToString(InOutputReferenceType), Connection.Card->GetDisplayName().c_str(), Connection.OutputReferenceChannel, Helpers::ReferenceTypeToString(Connection.OutputReferenceType));
 					return false;
 				}
 
diff --git a/Engine/Plugins/Online/IOS/OnlineSubsystemIOS/Source/Private/OnlineTurnBasedInterfaceIOS.cpp b/Engine/Plugins/Online/IOS/OnlineSubsystemIOS/Source/Private/OnlineTurnBasedInterfaceIOS.cpp
--- a/Engine/Plugins/Online/IOS/OnlineSubsystemIOS/Source/Private/OnlineTurnBasedInterfaceIOS.cpp
+++ b/Engine/Plugins/Online/IOS/OnlineSubsystemIOS/Source/Private/OnlineTurnBasedInterfaceIOS.cpp
@@ -167,7 +167,7 @@ int32 FTurnBasedMatchIOS::GetPlayerIndexForPlayer(NSString* PlayerID) const
 			return playerIndex;
 		}
 	}
-	UE_LOG(LogTurnBasedInterfaceIOS, Warning, TEXT("Failed to find participant %s in match"), PlayerID.UTF8String);
+	UE_LOG(LogTurnBasedInterfaceIOS, Warning, TEXT("Failed to find participant %s in match"), StringCast<TCHAR>(reinterpret_cast<const UTF8CHAR*>(PlayerID.UTF8String)).Get());
 	return 0;
 }
 
diff --git a/Engine/Plugins/Runtime/AR/AzureSpatialAnchors/Source/AzureSpatialAnchors/Public/IAzureSpatialAnchors.h b/Engine/Plugins/Runtime/AR/AzureSpatialAnchors/Source/AzureSpatialAnchors/Public/IAzureSpatialAnchors.h
--- a/Engine/Plugins/Runtime/AR/AzureSpatialAnchors/Source/AzureSpatialAnchors/Public/IAzureSpatialAnchors.h
+++ b/Engine/Plugins/Runtime/AR/AzureSpatialAnchors/Source/AzureSpatialAnchors/Public/IAzureSpatialAnchors.h
@@ -70,7 +70,7 @@ public:
 
 	static void OnLog(const wchar_t* LogMsg)
 	{
-		UE_LOG(LogAzureSpatialAnchors, Log, TEXT("%s"), LogMsg);
+		UE_LOG(LogAzureSpatialAnchors, Log, TEXT("%s"), StringCast<TCHAR>(LogMsg).Get());
 	}
 
 public:
diff --git a/Engine/Plugins/Runtime/HairStrands/Source/HairStrandsCore/Private/GroomComponent.cpp b/Engine/Plugins/Runtime/HairStrands/Source/HairStrandsCore/Private/GroomComponent.cpp
--- a/Engine/Plugins/Runtime/HairStrands/Source/HairStrandsCore/Private/GroomComponent.cpp
+++ b/Engine/Plugins/Runtime/HairStrands/Source/HairStrandsCore/Private/GroomComponent.cpp
@@ -4010,7 +4010,7 @@ void DumpLoadedGroomComponent(IConsoleVariable* InCVarPakTesterEnabled)
 						GroupCount,
                         Instance->Cards.LODs.Num(), 
                         InstanceMemory.GetTotalSize()* ToMb,
-                        0,
+                        0.0f,
 						InstanceMemory.Guides * ToMb,
 						InstanceMemory.Strands* ToMb,
 						InstanceMemory.Cards  * ToMb,
@@ -4029,7 +4029,7 @@ void DumpLoadedGroomComponent(IConsoleVariable* InCVarPakTesterEnabled)
 		Total_Component,
 		Total_Group,
         TotalMemory.GetTotalSize() * ToMb,
-		0,
+		0.0f,
 		TotalMemory.Guides * ToMb,
 		TotalMemory.Strands * ToMb,
 		TotalMemory.Cards * ToMb,
diff --git a/Engine/Plugins/Tests/CQTest/Source/CQTest/Private/Commands/TestCommands.cpp b/Engine/Plugins/Tests/CQTest/Source/CQTest/Private/Commands/TestCommands.cpp
--- a/Engine/Plugins/Tests/CQTest/Source/CQTest/Private/Commands/TestCommands.cpp
+++ b/Engine/Plugins/Tests/CQTest/Source/CQTest/Private/Commands/TestCommands.cpp
@@ -26,7 +26,7 @@ bool CQTEST_API FWaitUntil::Update()
 	{
 		if (Description)
 		{
-			UE_LOG(LogCqTest, Log, TEXT("Finished %s after %d milliseconds"), Description, Elapsed.GetTotalMilliseconds());
+			UE_LOG(LogCqTest, Log, TEXT("Finished %s after %d milliseconds"), Description, int(Elapsed.GetTotalMilliseconds()));
 		}
 		return true;
 	}
diff --git a/Engine/Source/Developer/DerivedDataCache/Private/DerivedDataBackends.cpp b/Engine/Source/Developer/DerivedDataCache/Private/DerivedDataBackends.cpp
--- a/Engine/Source/Developer/DerivedDataCache/Private/DerivedDataBackends.cpp
+++ b/Engine/Source/Developer/DerivedDataCache/Private/DerivedDataBackends.cpp
@@ -357,7 +357,7 @@ public:
 		{
 			if (WritePakCache)
 			{
-				UE_LOG(LogDerivedDataCache, Warning, TEXT("Unable to create %s pak cache because only one pak cache write node is supported."), *NodeName);
+				UE_LOG(LogDerivedDataCache, Warning, TEXT("Unable to create %s pak cache because only one pak cache write node is supported."), NodeName);
 				return false;
 			}
 
@@ -399,7 +399,7 @@ public:
 	{
 		if (bVerifyFound)
 		{
-			UE_LOG(LogDerivedDataCache, Warning, TEXT("Unable to create %s Verify because only one Verify node is supported."), *NodeName);
+			UE_LOG(LogDerivedDataCache, Warning, TEXT("Unable to create %s Verify because only one Verify node is supported."), NodeName);
 			return false;
 		}
 
@@ -430,7 +430,7 @@ public:
 	{
 		if (bAsyncFound)
 		{
-			UE_LOG(LogDerivedDataCache, Warning, TEXT("Unable to create %s AsyncPut because only one AsyncPut node is supported."), *NodeName);
+			UE_LOG(LogDerivedDataCache, Warning, TEXT("Unable to create %s AsyncPut because only one AsyncPut node is supported."), NodeName);
 			return false;
 		}
 
@@ -558,7 +558,7 @@ public:
 		UE_LOG(LogDerivedDataCache, Display, TEXT("Boot nodes are deprecated. Please remove the Boot node from the cache graph."));
 		if (bBootFound)
 		{
-			UE_LOG(LogDerivedDataCache, Warning, TEXT("Unable to create %s Boot cache because only one Boot node is supported."), *NodeName);
+			UE_LOG(LogDerivedDataCache, Warning, TEXT("Unable to create %s Boot cache because only one Boot node is supported."), NodeName);
 			return false;
 		}
 
diff --git a/Engine/Source/Developer/DerivedDataCache/Private/DerivedDataCacheReplay.cpp b/Engine/Source/Developer/DerivedDataCache/Private/DerivedDataCacheReplay.cpp
--- a/Engine/Source/Developer/DerivedDataCache/Private/DerivedDataCacheReplay.cpp
+++ b/Engine/Source/Developer/DerivedDataCache/Private/DerivedDataCacheReplay.cpp
@@ -169,7 +169,7 @@ void FCacheStoreReplay::SerializeRequests(
 	{
 		TUtf8StringBuilder<1024> Batch;
 		CompactBinaryToCompactJson(Writer.Save().AsObject(), Batch);
-		UE_LOG(LogDerivedDataCache, Verbose, TEXT("Replay: %hs"), *Batch);
+		UE_LOG(LogDerivedDataCache, Verbose, TEXT("Replay: %s"), StringCast<TCHAR>(*Batch).Get());
 	}
 
 	WriteToArchive(Writer);
@@ -407,7 +407,7 @@ bool FCacheReplayReader::FState::DispatchRequests(
 		{
 			TUtf8StringBuilder<1024> Batch;
 			CompactBinaryToCompactJson(Object, Batch);
-			UE_LOG(LogDerivedDataCache, Verbose, TEXT("Replay: %hs"), *Batch);
+			UE_LOG(LogDerivedDataCache, Verbose, TEXT("Replay: %s"), StringCast<TCHAR>(*Batch).Get());
 		}
 
 		TRACE_CPUPROFILER_EVENT_SCOPE(ReplayDDC_Dispatch);
diff --git a/Engine/Source/Developer/TextureBuild/Private/TextureBuildFunction.cpp b/Engine/Source/Developer/TextureBuild/Private/TextureBuildFunction.cpp
--- a/Engine/Source/Developer/TextureBuild/Private/TextureBuildFunction.cpp
+++ b/Engine/Source/Developer/TextureBuild/Private/TextureBuildFunction.cpp
@@ -514,7 +514,7 @@ void GenericTextureTilingBuildFunction(UE::DerivedData::FBuildContext& Context,
 
 	UE::TextureBuildUtilities::FTextureBuildMetadata BuildMetadata(FCbObject(Context.FindInput(ANSITEXTVIEW("TextureBuildMetadata"))));
 
-	UE_LOG(LogTextureBuildFunction, Display, TEXT("Tiling %s with %hs -> %d source mip(s) with a tail of %d..."), *Context.GetName(), *BuildFunctionName, TextureDescription.NumMips, TextureExtendedData.NumMipsInTail);
+	UE_LOG(LogTextureBuildFunction, Display, TEXT("Tiling %s with %s -> %d source mip(s) with a tail of %d..."), *Context.GetName(), StringCast<TCHAR>(*BuildFunctionName).Get(), TextureDescription.NumMips, TextureExtendedData.NumMipsInTail);
 
 	//
 	// Careful - the linear build might have a different streaming mip count than we output due to mip tail
diff --git a/Engine/Source/Programs/SwitchboardListenerHelper/Private/SBLHelperService.cpp b/Engine/Source/Programs/SwitchboardListenerHelper/Private/SBLHelperService.cpp
--- a/Engine/Source/Programs/SwitchboardListenerHelper/Private/SBLHelperService.cpp
+++ b/Engine/Source/Programs/SwitchboardListenerHelper/Private/SBLHelperService.cpp
@@ -108,7 +108,7 @@ bool FSBLHelperService::ParseIncomingMessage(const FString& Message, const FIPv4
 	}
 	else
 	{
-		UE_LOG(LogSwitchboardListenerHelper, Error, TEXT("Unsupported command '%s'"), &Command);
+		UE_LOG(LogSwitchboardListenerHelper, Error, TEXT("Unsupported command '%s'"), *Command);
 	}
 
 	return false;
diff --git a/Engine/Source/Runtime/Apple/MetalRHI/Private/MetalCommandList.cpp b/Engine/Source/Runtime/Apple/MetalRHI/Private/MetalCommandList.cpp
--- a/Engine/Source/Runtime/Apple/MetalRHI/Private/MetalCommandList.cpp
+++ b/Engine/Source/Runtime/Apple/MetalRHI/Private/MetalCommandList.cpp
@@ -172,7 +172,7 @@ static void ReportMetalCommandBufferFailure(mtlpp::CommandBuffer const& Complete
 							break;
 						}
 						
-                        UE_LOG(LogMetal, Error, TEXT("Command Buffer: %d (%p) Debug Buffer: %p Tile: %u Context: %d Encoder: %d Command: %d PSO: VS: %u_%u, PS: %u_%u"), Command.CmdBuffIndex, Command.CommandBuffer, DebugBuffer.GetPtr(), i, Command.ContextIndex, Command.EncoderIndex, Command.CommandIndex, Command.PSOSignature[0], Command.PSOSignature[1], Command.PSOSignature[2], Command.PSOSignature[3]);
+						UE_LOG(LogMetal, Error, TEXT("Command Buffer: %d (%p) Debug Buffer: %p Tile: %u Context: %d Encoder: %d Command: %d PSO: VS: %u_%u, PS: %u_%u"), Command.CmdBuffIndex, reinterpret_cast<void*>(Command.CommandBuffer), reinterpret_cast<void*>(DebugBuffer.GetPtr()), i, Command.ContextIndex, Command.EncoderIndex, Command.CommandIndex, Command.PSOSignature[0], Command.PSOSignature[1], Command.PSOSignature[2], Command.PSOSignature[3]);
                     }
                 }
 			}
diff --git a/Engine/Source/Runtime/Apple/MetalRHI/Private/MetalRenderPass.cpp b/Engine/Source/Runtime/Apple/MetalRHI/Private/MetalRenderPass.cpp
--- a/Engine/Source/Runtime/Apple/MetalRHI/Private/MetalRenderPass.cpp
+++ b/Engine/Source/Runtime/Apple/MetalRHI/Private/MetalRenderPass.cpp
@@ -352,7 +352,7 @@ void FMetalRenderPass::DrawIndexedPrimitive(FMetalBuffer const& IndexBuffer, uin
 				{
 					// Setting NumInstances to ClampedNumInstances would fix any visual rendering bugs resulting from this bad call but these draw calls are wrong - don't hide the issue
 					UE_LOG(LogMetal, Error, TEXT("Metal DrawIndexedPrimitive requested to draw %d Instances but vertex stream only has %d instance data available. ShaderName: %s, Deficient Attribute Index: %u"), NumInstances, ClampedNumInstances,
-						   *PipelineState->PixelShader->GetShaderName(), VertexElem.AttributeIndex);
+						   PipelineState->PixelShader->GetShaderName(), VertexElem.AttributeIndex);
 				}
 			}
 		}
diff --git a/Engine/Source/Runtime/Apple/MetalRHI/Private/MetalStateCache.cpp b/Engine/Source/Runtime/Apple/MetalRHI/Private/MetalStateCache.cpp
--- a/Engine/Source/Runtime/Apple/MetalRHI/Private/MetalStateCache.cpp
+++ b/Engine/Source/Runtime/Apple/MetalRHI/Private/MetalStateCache.cpp
@@ -1425,7 +1425,7 @@ void FMetalStateCache::SetShaderTexture(EMetalShaderStages const Frequency, FMet
 	check(Index < ML_MaxTextures);
 
 #if (PLATFORM_IOS || PLATFORM_TVOS)
-    UE_CLOG([Texture.GetPtr() storageMode] == MTLStorageModeMemoryless, LogMetal, Fatal, TEXT("FATAL: Attempting to bind a memoryless texture. Stage %u Index %u Texture %@"), Frequency, Index, Texture.GetPtr());
+	UE_CLOG([Texture.GetPtr() storageMode] == MTLStorageModeMemoryless, LogMetal, Fatal, TEXT("FATAL: Attempting to bind a memoryless texture. Stage %u Index %u Texture %s"), Frequency, Index, *FString([Texture.GetPtr() description]));
 #endif
 	
 	if (ShaderTextures[Frequency].Textures[Index] != Texture
diff --git a/Engine/Source/Runtime/ApplicationCore/Private/IOS/IOSInputInterface.cpp b/Engine/Source/Runtime/ApplicationCore/Private/IOS/IOSInputInterface.cpp
--- a/Engine/Source/Runtime/ApplicationCore/Private/IOS/IOSInputInterface.cpp
+++ b/Engine/Source/Runtime/ApplicationCore/Private/IOS/IOSInputInterface.cpp
@@ -79,7 +79,7 @@ uint32 TranslateGCKeyCodeToASCII(GCKeyCode KeyCode)
         else if (KeyCode == GCKeyCodePeriod) { c = '.';}
         else if (KeyCode == GCKeyCodeSlash) { c = '/';}
         
-        UE_LOG(LogIOS, Log, TEXT("char: %c"), (char)c);
+        UE_LOG(LogIOS, Log, TEXT("char: %c"), TCHAR(c));
     }
     return c;
 }
diff --git a/Engine/Source/Runtime/Engine/Private/Materials/MaterialShared.cpp b/Engine/Source/Runtime/Engine/Private/Materials/MaterialShared.cpp
--- a/Engine/Source/Runtime/Engine/Private/Materials/MaterialShared.cpp
+++ b/Engine/Source/Runtime/Engine/Private/Materials/MaterialShared.cpp
@@ -2840,7 +2840,7 @@ bool FMaterial::CacheShaders(const FMaterialShaderMapId& ShaderMapId, EShaderPla
 					IsSpecialEngineMaterial() ? TEXT("Is special engine material.") : TEXT("")
 				);
 #else
-				UE_LOG(LogMaterial, Display, TEXT("%s cached shader map for material %, compiling. %s"),
+				UE_LOG(LogMaterial, Display, TEXT("%s cached shader map for material %s, compiling. %s"),
 					ShaderMapCondition,
 					*GetAssetName(),
 					IsSpecialEngineMaterial() ? TEXT("Is special engine material.") : TEXT("")
diff --git a/Engine/Source/Runtime/Engine/Private/Net/Iris/ReplicationSystem/ActorReplicationBridge.cpp b/Engine/Source/Runtime/Engine/Private/Net/Iris/ReplicationSystem/ActorReplicationBridge.cpp
--- a/Engine/Source/Runtime/Engine/Private/Net/Iris/ReplicationSystem/ActorReplicationBridge.cpp
+++ b/Engine/Source/Runtime/Engine/Private/Net/Iris/ReplicationSystem/ActorReplicationBridge.cpp
@@ -295,14 +295,14 @@ UE::Net::FNetRefHandle UActorReplicationBridge::BeginReplication(AActor* Actor,
 			{
 				LevelGroup = CreateLevelGroup(Level);
 
-				UE_LOG_ACTORREPLICATIONBRIDGE(Log, TEXT("Created new GroupIndex: %u for Level: %s"), LevelGroup, ToCStr(PackageName.ToString()));
+				UE_LOG_ACTORREPLICATIONBRIDGE(Log, TEXT("Created new GroupIndex: %u for Level: %s"), LevelGroup.GetRawValue(), ToCStr(PackageName.ToString()));
 
 				// Update the filtering status of the group based on current level visibility for all connections
 				NetDriver->UpdateGroupFilterStatusForLevel(Level, LevelGroup);
 			}
 
 			// Add object to group
-			UE_LOG_ACTORREPLICATIONBRIDGE(Verbose, TEXT("Added %s to GroupIndex: %u Level: %s"), *ActorRefHandle.ToString(), LevelGroup, ToCStr(PackageName.ToString()));
+			UE_LOG_ACTORREPLICATIONBRIDGE(Verbose, TEXT("Added %s to GroupIndex: %u Level: %s"), *ActorRefHandle.ToString(), LevelGroup.GetRawValue(), ToCStr(PackageName.ToString()));
 			GetReplicationSystem()->AddToGroup(LevelGroup, ActorRefHandle);	
 		}
 	}
diff --git a/Engine/Source/Runtime/Engine/Private/SparseVolumeTexture/SparseVolumeTexture.cpp b/Engine/Source/Runtime/Engine/Private/SparseVolumeTexture/SparseVolumeTexture.cpp
--- a/Engine/Source/Runtime/Engine/Private/SparseVolumeTexture/SparseVolumeTexture.cpp
+++ b/Engine/Source/Runtime/Engine/Private/SparseVolumeTexture/SparseVolumeTexture.cpp
@@ -177,7 +177,7 @@ bool FHeader::Validate(bool bPrintToLog)
 	{
 		if (bPrintToLog)
 		{
-			UE_LOG(LogSparseVolumeTexture, Warning, TEXT("SparseVolumeTexture page table texture memory size (%ll) exceeds the 2048MB GPU resource limit!"), (long long)PageTableSizeBytes);
+			UE_LOG(LogSparseVolumeTexture, Warning, TEXT("SparseVolumeTexture page table texture memory size (%lld) exceeds the 2048MB GPU resource limit!"), (long long)PageTableSizeBytes);
 		}
 		return false;
 	}
diff --git a/Engine/Source/Runtime/Engine/Private/SparseVolumeTexture/SparseVolumeTextureData.cpp b/Engine/Source/Runtime/Engine/Private/SparseVolumeTexture/SparseVolumeTextureData.cpp
--- a/Engine/Source/Runtime/Engine/Private/SparseVolumeTexture/SparseVolumeTextureData.cpp
+++ b/Engine/Source/Runtime/Engine/Private/SparseVolumeTexture/SparseVolumeTextureData.cpp
@@ -700,7 +700,7 @@ bool FTextureData::DeduplicateTiles()
 		const int64 DeduplicatedTileSizeBytes[] = { (int64)TileWritePos * (int64)TileByteSize[0], (int64)TileWritePos * (int64)TileByteSize[1] };
 		if (DeduplicatedTileSizeBytes[0] > SVT::MaxResourceSize || DeduplicatedTileSizeBytes[1] > SVT::MaxResourceSize)
 		{
-			UE_LOG(LogSparseVolumeTextureData, Warning, TEXT("SparseVolumeTexture still requires too much memory (> 2048MB) after tile deduplication! The full mip chain will not fit in memory, so only lower resolution mip levels can be streamed in. Physical tile data A: %ll bytes. Physical tile data B: %ll bytes"),
+			UE_LOG(LogSparseVolumeTextureData, Warning, TEXT("SparseVolumeTexture still requires too much memory (> 2048MB) after tile deduplication! The full mip chain will not fit in memory, so only lower resolution mip levels can be streamed in. Physical tile data A: %lld bytes. Physical tile data B: %lld bytes"),
 				(long long)DeduplicatedTileSizeBytes[0], (long long)DeduplicatedTileSizeBytes[1]);
 		}
 
diff --git a/Engine/Source/Runtime/Experimental/Iris/Core/Private/Iris/ReplicationSystem/Filtering/ReplicationFiltering.cpp b/Engine/Source/Runtime/Experimental/Iris/Core/Private/Iris/ReplicationSystem/Filtering/ReplicationFiltering.cpp
--- a/Engine/Source/Runtime/Experimental/Iris/Core/Private/Iris/ReplicationSystem/Filtering/ReplicationFiltering.cpp
+++ b/Engine/Source/Runtime/Experimental/Iris/Core/Private/Iris/ReplicationSystem/Filtering/ReplicationFiltering.cpp
@@ -1617,7 +1617,7 @@ void FReplicationFiltering::SetSubObjectFilterStatus(FNetObjectGroupHandle Group
 
 	if (ensure(ValidConnections.GetBit(ConnectionId) && SubObjectFilterGroups.GetBit(GroupHandle.GetGroupIndex())))
 	{
-		UE_LOG(LogIrisFiltering, Verbose, TEXT("ReplicationFiltering::SetSubObjectFilterStatus GroupIndex: %u, ConnectionId: %u, FilterStatus: %u"), GroupHandle, ConnectionId, ReplicationStatus == ENetFilterStatus::Allow ? 1U : 0U);
+		UE_LOG(LogIrisFiltering, Verbose, TEXT("ReplicationFiltering::SetSubObjectFilterStatus GroupIndex: %u, ConnectionId: %u, FilterStatus: %u"), GroupHandle.GetRawValue(), ConnectionId, ReplicationStatus == ENetFilterStatus::Allow ? 1U : 0U);
 		FPerObjectInfo* FilterInfo = GetPerObjectInfo(GroupInfos[GroupHandle.GetGroupIndex()].ConnectionStateIndex);
 		SetConnectionFilterStatus(*FilterInfo, ConnectionId, ReplicationStatus);
 		if (!IsConnectionFilterStatusAllowedForAnyConnection(*FilterInfo))
@@ -1938,7 +1938,7 @@ void FReplicationFiltering::NotifyObjectAddedToGroup(FNetObjectGroupHandle Group
 
 void FReplicationFiltering::NotifyObjectRemovedFromGroup(FNetObjectGroupHandle GroupHandle, FInternalNetRefIndex ObjectIndex)
 {
-	UE_LOG(LogIrisFiltering, Verbose, TEXT("ReplicationFiltering::NotifyObjectRemovedFromGroup GroupIndex: %u, %s"), GroupHandle, *(NetRefHandleManager->GetReplicatedObjectDataNoCheck(ObjectIndex).RefHandle.ToString()));
+	UE_LOG(LogIrisFiltering, Verbose, TEXT("ReplicationFiltering::NotifyObjectRemovedFromGroup GroupIndex: %u, %s"), GroupHandle.GetRawValue(), *(NetRefHandleManager->GetReplicatedObjectDataNoCheck(ObjectIndex).RefHandle.ToString()));
 
 	if (FilterGroups.GetBit(GroupHandle.GetGroupIndex()))
 	{
diff --git a/Engine/Source/Runtime/Experimental/Iris/Core/Private/Iris/ReplicationSystem/ObjectReferenceCache.cpp b/Engine/Source/Runtime/Experimental/Iris/Core/Private/Iris/ReplicationSystem/ObjectReferenceCache.cpp
--- a/Engine/Source/Runtime/Experimental/Iris/Core/Private/Iris/ReplicationSystem/ObjectReferenceCache.cpp
+++ b/Engine/Source/Runtime/Experimental/Iris/Core/Private/Iris/ReplicationSystem/ObjectReferenceCache.cpp
@@ -391,7 +391,7 @@ bool FObjectReferenceCache::CreateObjectReferenceInternal(const UObject* Object,
 	}
 #endif
 	
-	UE_LOG_REFERENCECACHE(Verbose, TEXT("ObjectReferenceCache::CreateObjectReferenceHandle Adding %s to ObjectToNetReferenceHandle and Cache for ObjectPath %s, Object (0x%p), OuterNetRefHandle: %s"), 
+	UE_LOG_REFERENCECACHE(Verbose, TEXT("ObjectReferenceCache::CreateObjectReferenceHandle Adding %s to ObjectToNetReferenceHandle and Cache for ObjectPath %s, Name: %s, Object (0x%p), OuterNetRefHandle: %s"), 
 		*CachedObject.NetRefHandle.ToString(), ToCStr(Object->GetPathName()), *GetNameSafe(Object), Object, *CachedObject.OuterNetRefHandle.ToString());
 
 	// Create reference
diff --git a/Engine/Source/Runtime/Experimental/Iris/Core/Private/Iris/ReplicationSystem/ReplicationBridge.cpp b/Engine/Source/Runtime/Experimental/Iris/Core/Private/Iris/ReplicationSystem/ReplicationBridge.cpp
--- a/Engine/Source/Runtime/Experimental/Iris/Core/Private/Iris/ReplicationSystem/ReplicationBridge.cpp
+++ b/Engine/Source/Runtime/Experimental/Iris/Core/Private/Iris/ReplicationSystem/ReplicationBridge.cpp
@@ -844,7 +844,7 @@ UE::Net::FNetRefHandle UReplicationBridge::InternalAddDestructionInfo(FNetRefHan
 		GetReplicationSystem()->SetPrioritizer(DestructionInfoHandle, DefaultSpatialNetObjectPrioritizerHandle);
 	}
 
-	UE_LOG_REPLICATIONBRIDGE(Verbose, TEXT("UReplicationBridge::InternalAddDestructionInfo %s ( InternalIndex: %u ) for %s GroupIndex: %u"), *DestructionInfoHandle.ToString(), InternalReplicationIndex,  *Handle.ToString(), LevelGroupHandle);
+	UE_LOG_REPLICATIONBRIDGE(Verbose, TEXT("UReplicationBridge::InternalAddDestructionInfo %s ( InternalIndex: %u ) for %s GroupIndex: %u"), *DestructionInfoHandle.ToString(), InternalReplicationIndex,  *Handle.ToString(), LevelGroupHandle.GetRawValue());
 	
 	return DestructionInfoHandle;
 }
diff --git a/Engine/Source/Runtime/Experimental/Iris/Core/Public/Iris/ReplicationSystem/NetObjectGroupHandle.h b/Engine/Source/Runtime/Experimental/Iris/Core/Public/Iris/ReplicationSystem/NetObjectGroupHandle.h
--- a/Engine/Source/Runtime/Experimental/Iris/Core/Public/Iris/ReplicationSystem/NetObjectGroupHandle.h
+++ b/Engine/Source/Runtime/Experimental/Iris/Core/Public/Iris/ReplicationSystem/NetObjectGroupHandle.h
@@ -55,6 +55,11 @@ public:
 	/** Special group, SubObjects assigned to this group will replicate if replay netconditions is met  */
 	bool IsNetGroupReplayNetObjectGroup() const { return Index == NetGroupReplayNetObjectGroupIndex; }
 
+	uint32 GetRawValue() const
+	{
+		return Value;
+	}
+	
 private:
 	friend UE::Net::Private::FNetObjectGroups;
 
diff --git a/Engine/Source/Runtime/Online/BuildPatchServices/Private/BuildPatchInstaller.cpp b/Engine/Source/Runtime/Online/BuildPatchServices/Private/BuildPatchInstaller.cpp
--- a/Engine/Source/Runtime/Online/BuildPatchServices/Private/BuildPatchInstaller.cpp
+++ b/Engine/Source/Runtime/Online/BuildPatchServices/Private/BuildPatchInstaller.cpp
@@ -282,9 +282,9 @@ namespace InstallerHelpers
 			UE_LOG(LogBPSInstallerConfig, Log, TEXT("Build Config: CloudDirectories: %s"), *CloudDirectory);
 		}
 
-		UE_LOG(LogBPSInstallerConfig, Log, TEXT("Build Config: InstallMode: %s"), *LexToString(InstallerConfiguration.InstallMode));
-		UE_LOG(LogBPSInstallerConfig, Log, TEXT("Build Config: VerifyMode: %s"), *LexToString(InstallerConfiguration.VerifyMode));
-		UE_LOG(LogBPSInstallerConfig, Log, TEXT("Build Config: DeltaPolicy: %s"), *LexToString(InstallerConfiguration.DeltaPolicy));
+		UE_LOG(LogBPSInstallerConfig, Log, TEXT("Build Config: InstallMode: %s"), LexToString(InstallerConfiguration.InstallMode));
+		UE_LOG(LogBPSInstallerConfig, Log, TEXT("Build Config: VerifyMode: %s"), LexToString(InstallerConfiguration.VerifyMode));
+		UE_LOG(LogBPSInstallerConfig, Log, TEXT("Build Config: DeltaPolicy: %s"), LexToString(InstallerConfiguration.DeltaPolicy));
 		UE_LOG(LogBPSInstallerConfig, Log, TEXT("Build Config: bRunRequiredPrereqs: %s"), (InstallerConfiguration.bRunRequiredPrereqs) ? TEXT("true") : TEXT("false"));
 		UE_LOG(LogBPSInstallerConfig, Log, TEXT("Build Config: bAllowConcurrentExecution: %s"), (InstallerConfiguration.bAllowConcurrentExecution) ? TEXT("true") : TEXT("false"));
 	}
diff --git a/Engine/Source/Runtime/Online/Experimental/EventLoop/Private/EventLoop/BSDSocket/EventLoopIOManagerBSDSocketSelect.cpp b/Engine/Source/Runtime/Online/Experimental/EventLoop/Private/EventLoop/BSDSocket/EventLoopIOManagerBSDSocketSelect.cpp
--- a/Engine/Source/Runtime/Online/Experimental/EventLoop/Private/EventLoop/BSDSocket/EventLoopIOManagerBSDSocketSelect.cpp
+++ b/Engine/Source/Runtime/Online/Experimental/EventLoop/Private/EventLoop/BSDSocket/EventLoopIOManagerBSDSocketSelect.cpp
@@ -130,7 +130,7 @@ void FIOManagerBSDSocketSelectImpl::PollInternal(FTimespan WaitTime)
 		SOCKET RemovedSocket;
 		if (SocketInternalHandleIndex.RemoveAndCopyValue(RemovedHandle, RemovedSocket))
 		{
-			UE_LOG(LogEventLoop, VeryVerbose, TEXT("[FIOManagerBSDSocketSelectImpl::PollInternal] Request removed: Socket: %p, Handle: %s"), RemovedSocket, *RemovedHandle.ToString());
+			UE_LOG(LogEventLoop, VeryVerbose, TEXT("[FIOManagerBSDSocketSelectImpl::PollInternal] Request removed: Socket: %p, Handle: %s"), reinterpret_cast<void*>(RemovedSocket), *RemovedHandle.ToString());
 			SocketIndex.Remove(RemovedSocket);
 		}
 	}
@@ -143,7 +143,7 @@ void FIOManagerBSDSocketSelectImpl::PollInternal(FTimespan WaitTime)
 			// Check whether there are enough resources available to handle the request.
 			if (SocketIndex.Num() >= FD_SETSIZE)
 			{
-				UE_LOG(LogEventLoop, Verbose, TEXT("[FIOManagerBSDSocketSelectImpl::PollInternal] Request failed: Socket: %p, Reason: %s"), NewRequest->Socket, *LexToString(ESocketIoRequestStatus::NoResources));
+				UE_LOG(LogEventLoop, Verbose, TEXT("[FIOManagerBSDSocketSelectImpl::PollInternal] Request failed: Socket: %p, Reason: %s"), reinterpret_cast<void*>(NewRequest->Socket), *LexToString(ESocketIoRequestStatus::NoResources));
 				NewRequest->Callback(NewRequest->Socket, ESocketIoRequestStatus::NoResources, EIOFlags::None);
 				IORequestStorage.Remove(AddedHandle);
 				continue;
@@ -152,13 +152,13 @@ void FIOManagerBSDSocketSelectImpl::PollInternal(FTimespan WaitTime)
 			// Only one request may be active for a socket at one time.
 			if (SocketIndex.Contains(NewRequest->Socket))
 			{
-				UE_LOG(LogEventLoop, Verbose, TEXT("[FIOManagerBSDSocketSelectImpl::PollInternal] Request failed: Socket: %p, Reason: %s"), NewRequest->Socket, *LexToString(ESocketIoRequestStatus::Invalid));
+				UE_LOG(LogEventLoop, Verbose, TEXT("[FIOManagerBSDSocketSelectImpl::PollInternal] Request failed: Socket: %p, Reason: %s"), reinterpret_cast<void*>(NewRequest->Socket), *LexToString(ESocketIoRequestStatus::Invalid));
 				NewRequest->Callback(NewRequest->Socket, ESocketIoRequestStatus::Invalid, EIOFlags::None);
 				IORequestStorage.Remove(AddedHandle);
 				continue;
 			}
 
-			UE_LOG(LogEventLoop, VeryVerbose, TEXT("[FIOManagerBSDSocketSelectImpl::PollInternal] Request added: Socket: %p, Handle: %s, Flags: 0x%08X"), NewRequest->Socket, *AddedHandle.ToString(), NewRequest->Flags);
+			UE_LOG(LogEventLoop, VeryVerbose, TEXT("[FIOManagerBSDSocketSelectImpl::PollInternal] Request added: Socket: %p, Handle: %s, Flags: 0x%08X"), reinterpret_cast<void*>(NewRequest->Socket), *AddedHandle.ToString(), NewRequest->Flags);
 
 			SocketInternalHandleIndex.Add(AddedHandle, NewRequest->Socket);
 			SocketIndex.Add(NewRequest->Socket, AddedHandle);
@@ -208,7 +208,7 @@ void FIOManagerBSDSocketSelectImpl::PollInternal(FTimespan WaitTime)
 			if (SignaledFlags != EIOFlags::None)
 			{
 				// Signal event status.
-				UE_LOG(LogEventLoop, VeryVerbose, TEXT("[FIOManagerBSDSocketSelectImpl::PollInternal] Request signaled: Socket: %p, Handle: %s, Flags: 0x%08X"), IORequest.Socket, *IORequestEntry.Key.ToString(), SignaledFlags);
+				UE_LOG(LogEventLoop, VeryVerbose, TEXT("[FIOManagerBSDSocketSelectImpl::PollInternal] Request signaled: Socket: %p, Handle: %s, Flags: 0x%08X"), reinterpret_cast<void*>(IORequest.Socket), *IORequestEntry.Key.ToString(), SignaledFlags);
 				IORequest.Callback(IORequest.Socket, ESocketIoRequestStatus::Ok, SignaledFlags);
 			}
 		}
diff --git a/Engine/Source/Runtime/Online/HTTP/Private/Curl/CurlSocketEventLoopHttpThread.cpp b/Engine/Source/Runtime/Online/HTTP/Private/Curl/CurlSocketEventLoopHttpThread.cpp
--- a/Engine/Source/Runtime/Online/HTTP/Private/Curl/CurlSocketEventLoopHttpThread.cpp
+++ b/Engine/Source/Runtime/Online/HTTP/Private/Curl/CurlSocketEventLoopHttpThread.cpp
@@ -145,7 +145,7 @@ int FCurlSocketEventLoopHttpThread::HandleCurlSocketCallback(CURL* CurlE, curl_s
 
 	FCurlSocketData* CurlSocketData = reinterpret_cast<FCurlSocketData*>(SocketData);
 
-	UE_LOG(LogHttp, VeryVerbose, TEXT("[FCurlSocketEventLoopHttpThread::HandleCurlSocketCallback] Socket: %p, EventFlags: 0x%08X, SocketData: %p"), Socket, EventFlags, SocketData);
+	UE_LOG(LogHttp, VeryVerbose, TEXT("[FCurlSocketEventLoopHttpThread::HandleCurlSocketCallback] Socket: %p, EventFlags: 0x%08X, SocketData: %p"), reinterpret_cast<void*>(Socket), EventFlags, SocketData);
 
 	switch (EventFlags)
 	{
@@ -224,7 +224,7 @@ void FCurlSocketEventLoopHttpThread::ProcessCurlSocketActions(curl_socket_t Sock
 {
 	QUICK_SCOPE_CYCLE_COUNTER(STAT_FCurlSocketEventLoopHttpThread_ProcessCurlSocketActions);
 
-	UE_LOG(LogHttp, VeryVerbose, TEXT("[FCurlSocketEventLoopHttpThread::ProcessCurlSocketActions] Socket: %p CurlEvents: 0x%08X"), Socket, EventFlags);
+	UE_LOG(LogHttp, VeryVerbose, TEXT("[FCurlSocketEventLoopHttpThread::ProcessCurlSocketActions] Socket: %p CurlEvents: 0x%08X"), reinterpret_cast<void*>(Socket), EventFlags);
 
 	int running_handles;
 	curl_multi_socket_action(FCurlHttpManager::GMultiHandle, Socket, EventFlags, &running_handles);
@@ -235,7 +235,7 @@ void FCurlSocketEventLoopHttpThread::ProcessCurlSocketEvent(curl_socket_t Socket
 {
 	if (Status != UE::EventLoop::ESocketIoRequestStatus::Ok)
 	{
-		UE_LOG(LogHttp, Warning, TEXT("Socket event request failed. Socket %p, Status: %s"), Socket, *LexToString(Status));
+		UE_LOG(LogHttp, Warning, TEXT("Socket event request failed. Socket %p, Status: %s"), reinterpret_cast<void*>(Socket), *LexToString(Status));
 		ProcessCurlSocketActions(CURL_SOCKET_TIMEOUT, 0);
 	}
 	else
diff --git a/Engine/Source/Runtime/Online/WebSockets/Private/WinHttp/Support/WinHttpConnectionWebSocket.cpp b/Engine/Source/Runtime/Online/WebSockets/Private/WinHttp/Support/WinHttpConnectionWebSocket.cpp
--- a/Engine/Source/Runtime/Online/WebSockets/Private/WinHttp/Support/WinHttpConnectionWebSocket.cpp
+++ b/Engine/Source/Runtime/Online/WebSockets/Private/WinHttp/Support/WinHttpConnectionWebSocket.cpp
@@ -25,7 +25,7 @@ void CALLBACK UE_WinHttpWebSocketStatusCallback(HINTERNET hInternet, DWORD_PTR d
 	const EWinHttpCallbackStatus Status = static_cast<EWinHttpCallbackStatus>(dwInternetStatus);
 	if (!IsValidStatus(Status))
 	{
-		UE_LOG(LogWinHttp, Warning, TEXT("WebSocket[%p]: Received unknown WinHttp Status %lu"), dwContext, dwInternetStatus);
+		UE_LOG(LogWinHttp, Warning, TEXT("WebSocket[%p]: Received unknown WinHttp Status %lu"), reinterpret_cast<void*>(dwContext), dwInternetStatus);
 		return;
 	}
 
