FEventLoopHttpThread
- Fix base UpdateConfigs() not being called

diff --git a/Engine/Source/Runtime/Online/HTTP/Private/EventLoopHttpThread.cpp b/Engine/Source/Runtime/Online/HTTP/Private/EventLoopHttpThread.cpp
--- a/Engine/Source/Runtime/Online/HTTP/Private/EventLoopHttpThread.cpp
+++ b/Engine/Source/Runtime/Online/HTTP/Private/EventLoopHttpThread.cpp
@@ -78,6 +78,7 @@ void FEventLoopHttpThread::UpdateConfigs()
 {
 	ResetTickTimer();
 	UpdateEventLoopConfigs();
+	FHttpThreadBase::UpdateConfigs();
 }
 
 void FEventLoopHttpThread::AddRequest(IHttpThreadedRequest* Request)
