Make http retry system thread safe.

diff --git a/Engine/Source/Programs/WebTests/Private/TestHttp.cpp b/Engine/Source/Programs/WebTests/Private/TestHttp.cpp
--- a/Engine/Source/Programs/WebTests/Private/TestHttp.cpp
+++ b/Engine/Source/Programs/WebTests/Private/TestHttp.cpp
@@ -6,6 +6,7 @@
 #include "HAL/Runnable.h"
 #include "HAL/RunnableThread.h"
 #include "HttpManager.h"
+#include "HttpRetrySystem.h"
 #include "Http.h"
 #include "Misc/CommandLine.h"
 #include "TestHarness.h"
@@ -176,7 +177,7 @@ public:
 
 	void WaitUntilAllHttpRequestsComplete()
 	{
-		while (!bRunningThreadRequest && OngoingRequests != 0)
+		while (OngoingRequests != 0)
 		{
 			HttpModule->GetHttpManager().Tick(TickFrequency);
 			FPlatformProcess::Sleep(TickFrequency);
@@ -185,7 +186,6 @@ public:
 
 	uint32 OngoingRequests = 0;
 	float TickFrequency = 1.0f / 60; /*60 FPS*/;
-	bool bRunningThreadRequest = false;
 };
 
 TEST_CASE_METHOD(FWaitUntilCompleteHttpFixture, "Http Methods", HTTP_TAG)
@@ -594,38 +594,61 @@ TEST_CASE_METHOD(FWaitUntilCompleteHttpFixture, "Redirect enabled by default and
 	HttpRequest->ProcessRequest();
 }
 
-class FWaitThreadedHttpFixture : public FWaitUntilCompleteHttpFixture, public FRunnable
+class FThreadedHttpRunnable : public FRunnable
 {
 public:
 	DECLARE_DELEGATE(FRunActualTestCodeDelegate);
 
-	FWaitThreadedHttpFixture()
+	FRunActualTestCodeDelegate& OnRunFromThread()
 	{
-		bRunningThreadRequest = true;
+		return ThreadCallback;
 	}
 
 	// FRunnable interface
 	virtual uint32 Run() override
 	{
 		ThreadCallback.ExecuteIfBound();
-		bRunningThreadRequest = false;
 		return 0;
 	}
 
-	void StartTestHttpThread()
+	void StartTestHttpThread(bool bBlockGameThread)
 	{
+		bBlockingGameThreadTick = bBlockGameThread;
+
 		RunnableThread = TSharedPtr<FRunnableThread>(FRunnableThread::Create(this, TEXT("Test Http Thread")));
+
+		while (bBlockingGameThreadTick)
+		{
+			float TickFrequency = 1.0f / 60; /*60 FPS*/;
+			FPlatformProcess::Sleep(TickFrequency);
+		}
+	}
+
+	void UnblockGameThread()
+	{
+		bBlockingGameThreadTick = false;
 	}
 
+private:
 	FRunActualTestCodeDelegate ThreadCallback;
 	TSharedPtr<FRunnableThread> RunnableThread;
+	std::atomic<bool> bBlockingGameThreadTick = true;
 };
 
-TEST_CASE_METHOD(FWaitThreadedHttpFixture, "Http streaming download request can work in non game thread", HTTP_TAG)
+class FWaitThreadedHttpFixture : public FWaitUntilCompleteHttpFixture
 {
-	std::atomic<bool> bBlockingGameThreadTick = true;
+public:
+	~FWaitThreadedHttpFixture()
+	{
+		WaitUntilAllHttpRequestsComplete();
+	}
 
-	ThreadCallback.BindLambda([this, &bBlockingGameThreadTick]() {
+	FThreadedHttpRunnable ThreadedHttpRunnable;
+};
+
+TEST_CASE_METHOD(FWaitThreadedHttpFixture, "Http streaming download request can work in non game thread", HTTP_TAG)
+{
+	ThreadedHttpRunnable.OnRunFromThread().BindLambda([this]() {
 		TSharedRef<IHttpRequest> HttpRequest = HttpModule->CreateRequest();
 		HttpRequest->SetURL(UrlStreamDownload(3/*Chunks*/, 1024/*ChunkSize*/));
 		HttpRequest->SetVerb(TEXT("GET"));
@@ -642,23 +665,18 @@ TEST_CASE_METHOD(FWaitThreadedHttpFixture, "Http streaming download request can
 		};
 		CHECK(HttpRequest->SetResponseBodyReceiveStream(MakeShared<FTestHttpReceiveStream>()));
 
-		HttpRequest->OnProcessRequestComplete().BindLambda([&bBlockingGameThreadTick](FHttpRequestPtr HttpRequest, FHttpResponsePtr HttpResponse, bool bSucceeded) {
+		HttpRequest->OnProcessRequestComplete().BindLambda([this](FHttpRequestPtr HttpRequest, FHttpResponsePtr HttpResponse, bool bSucceeded) {
 			// EHttpRequestDelegateThreadPolicy::CompleteOnHttpThread was used, so not in game thread here
 			CHECK(!IsInGameThread());
 			CHECK(bSucceeded);
 			CHECK(HttpResponse->GetResponseCode() == 200);
-			bBlockingGameThreadTick = false;
+			ThreadedHttpRunnable.UnblockGameThread();
 		});
 
 		HttpRequest->ProcessRequest();
 	});
 
-	StartTestHttpThread();
-
-	while (bBlockingGameThreadTick)
-	{
-		FPlatformProcess::Sleep(TickFrequency);
-	}
+	ThreadedHttpRunnable.StartTestHttpThread(true/*bBlockGameThread*/);
 }
 
 namespace UE
@@ -708,7 +726,7 @@ TEST_CASE_METHOD(FWaitUntilCompleteHttpFixture, "Http request pre check will fai
 
 TEST_CASE_METHOD(FWaitThreadedHttpFixture, "Threaded http request pre check will fail", HTTP_TAG)
 {
-	ThreadCallback.BindLambda([this]() {
+	ThreadedHttpRunnable.OnRunFromThread().BindLambda([this]() {
 		// Pre check will fail when domain is not allowed
 		UE::TestHttp::SetupURLRequestFilter(HttpModule);
 
@@ -735,10 +753,9 @@ TEST_CASE_METHOD(FWaitThreadedHttpFixture, "Threaded http request pre check will
 		HttpRequest->ProcessRequest();
 	});
 
-	StartTestHttpThread();
+	ThreadedHttpRunnable.StartTestHttpThread(false/*bBlockGameThread*/);
 }
 
-
 TEST_CASE_METHOD(FWaitUntilCompleteHttpFixture, "Cancel http request connect before timeout", HTTP_TAG)
 {
 	TSharedRef<IHttpRequest, ESPMode::ThreadSafe> HttpRequest = HttpModule->CreateRequest();
@@ -757,4 +774,102 @@ TEST_CASE_METHOD(FWaitUntilCompleteHttpFixture, "Cancel http request connect bef
 	FPlatformProcess::Sleep(0.5);
 	HttpRequest->CancelRequest();
 }
+
+class FMockRetryManager : public FHttpRetrySystem::FManager
+{
+public:
+	using FHttpRetrySystem::FManager::FManager;
+	using FHttpRetrySystem::FManager::RequestList;
+	using FHttpRetrySystem::FManager::FHttpRetryRequestEntry;
+	using FHttpRetrySystem::FManager::RetryTimeoutRelativeSecondsDefault;
+};
+
+class FRetryHttpFixture : public FWaitUntilCompleteHttpFixture
+{
+public:
+	FRetryHttpFixture()
+	{
+		HttpRetryManager = MakeShared<FMockRetryManager>(FHttpRetrySystem::FRetryLimitCountSetting(RetryLimitCount), FHttpRetrySystem::FRetryTimeoutRelativeSecondsSetting(/*RetryTimeoutRelativeSeconds*/));
+		HttpRetryManager->RetryTimeoutRelativeSecondsDefault = 2.0f; // Value is set so that retry system will use it at high level, will no longer wait for ever for HTTP code
+	}
+
+	~FRetryHttpFixture()
+	{
+		// Make sure not destroy retry manager before waiting in FWaitUntilCompleteHttpFixture
+		while (HttpRetryManager->RequestList.Num())
+		{
+			HttpRetryManager->Update();
+			HttpModule->GetHttpManager().Tick(TickFrequency);
+			FPlatformProcess::Sleep(TickFrequency);
+		}
+	}
+
+	TSharedPtr<FMockRetryManager> HttpRetryManager;
+	uint32 RetryLimitCount = 2;
+
+	FThreadedHttpRunnable ThreadedHttpRunnable;
+};
+
+// TODO: Currently broken, will fix this in retry system
+//TEST_CASE_METHOD(FRetryHttpFixture, "Http request adaptor in retry system can work in non game thread", HTTP_TAG)
+//{
+//	ThreadedHttpRunnable.OnRunFromThread().BindLambda([this]() {
+//		TSharedRef<IHttpRequest> HttpRequest = HttpRetryManager->CreateRequest();
+//		HttpRequest->SetURL(UrlStreamDownload(3/*Chunks*/, 1024/*ChunkSize*/));
+//		HttpRequest->SetVerb(TEXT("GET"));
+//		HttpRequest->SetDelegateThreadPolicy(EHttpRequestDelegateThreadPolicy::CompleteOnHttpThread);
+//
+//		class FTestHttpReceiveStream final : public FArchive
+//		{
+//		public:
+//			virtual void Serialize(void* V, int64 Length) override
+//			{
+//				// No matter what's the thread policy, Serialize always get called in http thread.
+//				CHECK(!IsInGameThread());
+//			}
+//		};
+//		CHECK(HttpRequest->SetResponseBodyReceiveStream(MakeShared<FTestHttpReceiveStream>()));
+//
+//		HttpRequest->OnProcessRequestComplete().BindLambda([this](FHttpRequestPtr HttpRequest, FHttpResponsePtr HttpResponse, bool bSucceeded) {
+//			// EHttpRequestDelegateThreadPolicy::CompleteOnHttpThread was used, so not in game thread here
+//			CHECK(!IsInGameThread());
+//			CHECK(bSucceeded);
+//			CHECK(HttpResponse->GetResponseCode() == 200);
+//			ThreadedHttpRunnable.UnblockGameThread();
+//		});
+//
+//		HttpRequest->ProcessRequest();
+//	});
+//
+//	ThreadedHttpRunnable.StartTestHttpThread(true/*bBlockGameThread*/);
+//}
+
+class FRetryHttpManagerThreadedRequestsFixture : public FRetryHttpFixture
+{
+public:
+	void LaunchDownloadRequests()
+	{
+		for (int32 i = 0; i < 10; ++i)
+		{
+			TSharedRef<IHttpRequest> HttpRequest = HttpRetryManager->CreateRequest();
+			HttpRequest->SetURL(UrlStreamDownload(3, 1024*1024));
+			HttpRequest->SetVerb(TEXT("GET"));
+			HttpRequest->ProcessRequest();
+		}
+	}
+};
+
+TEST_CASE_METHOD(FRetryHttpManagerThreadedRequestsFixture, "Retry manager is thread safe", HTTP_TAG)
+{
+	ThreadedHttpRunnable.OnRunFromThread().BindLambda([this]() {
+		LaunchDownloadRequests();
+		HttpRetryManager->BlockUntilFlushed(5.0);
+	});
+	ThreadedHttpRunnable.StartTestHttpThread(false/*bBlockGameThread*/);
+
+	LaunchDownloadRequests();
+	HttpRetryManager->BlockUntilFlushed(5.0);
+}
+
+
 // TODO: Add cancel test, with multiple cancel calls
diff --git a/Engine/Source/Runtime/Online/HTTP/Private/Curl/CurlHttp.cpp b/Engine/Source/Runtime/Online/HTTP/Private/Curl/CurlHttp.cpp
--- a/Engine/Source/Runtime/Online/HTTP/Private/Curl/CurlHttp.cpp
+++ b/Engine/Source/Runtime/Online/HTTP/Private/Curl/CurlHttp.cpp
@@ -1141,7 +1141,10 @@ PRAGMA_ENABLE_DEPRECATION_WARNINGS
 
 void FCurlHttpRequest::BroadcastNewlyReceivedHeaders()
 {
-	check(IsInGameThread());
+	if (!OnHeaderReceived().IsBound())
+	{
+		return;
+	}
 
 	if (Response.IsValid())
 	{
@@ -1177,8 +1180,6 @@ void FCurlHttpRequest::BroadcastNewlyReceivedHeader(const FString& HeaderKey, co
 
 void FCurlHttpRequest::FinishRequest()
 {
-	check(IsInGameThread() || DelegateThreadPolicy == EHttpRequestDelegateThreadPolicy::CompleteOnHttpThread);
-
 	QUICK_SCOPE_CYCLE_COUNTER(STAT_FCurlHttpRequest_FinishRequest);
 
 	curl_easy_setopt(EasyHandle, CURLOPT_SHARE, nullptr);
diff --git a/Engine/Source/Runtime/Online/HTTP/Private/HttpRetrySystem.cpp b/Engine/Source/Runtime/Online/HTTP/Private/HttpRetrySystem.cpp
--- a/Engine/Source/Runtime/Online/HTTP/Private/HttpRetrySystem.cpp
+++ b/Engine/Source/Runtime/Online/HTTP/Private/HttpRetrySystem.cpp
@@ -158,6 +158,8 @@ FHttpRetrySystem::FManager::FManager(const FRetryLimitCountSetting& InRetryLimit
 
 FHttpRetrySystem::FManager::~FManager()
 {
+	FScopeLock ScopeLock(&RequestListLock);
+
 	// Decrement retried request for log verbosity tracker
 	for (const FHttpRetryRequestEntry& Request : RequestList)
 	{
@@ -342,6 +344,8 @@ bool FHttpRetrySystem::FManager::Update(uint32* FileCount, uint32* FailingCount,
 	QUICK_SCOPE_CYCLE_COUNTER(STAT_FHttpRetrySystem_FManager_Update);
 	LLM_SCOPE_BYTAG(HTTP);
 
+	FScopeLock ScopeLock(&RequestListLock);
+
 	bool bIsGreen = true;
 
 	if (FileCount != nullptr)
@@ -589,6 +593,7 @@ bool FHttpRetrySystem::FManager::ProcessRequest(TSharedRef<FHttpRetrySystem::FRe
 
 	if (bResult)
 	{
+		FScopeLock ScopeLock(&RequestListLock);
 		RequestList.Add(FHttpRetryRequestEntry(HttpRetryRequest));
 	}
 
@@ -599,6 +604,8 @@ void FHttpRetrySystem::FManager::CancelRequest(TSharedRef<FHttpRetrySystem::FReq
 {
 	QUICK_SCOPE_CYCLE_COUNTER(STAT_FHttpRetrySystem_FManager_CancelRequest);
 
+	FScopeLock ScopeLock(&RequestListLock);
+
 	// Find the existing request entry if is was previously processed.
 	bool bFound = false;
 	for (int32 i = 0; i < RequestList.Num(); ++i)
@@ -629,6 +636,9 @@ void FHttpRetrySystem::FManager::BlockUntilFlushed(float InTimeoutSec)
 	const float SleepInterval = 0.016;
 	float TimeElapsed = 0.0f;
 	uint32 FileCount, FailingCount, FailedCount, CompleteCount;
+
+	FScopeLock ScopeLock(&RequestListLock);
+
 	while (RequestList.Num() > 0 && TimeElapsed < InTimeoutSec)
 	{
 		FHttpModule::Get().GetHttpManager().Tick(SleepInterval);
diff --git a/Engine/Source/Runtime/Online/HTTP/Private/HttpThread.cpp b/Engine/Source/Runtime/Online/HTTP/Private/HttpThread.cpp
--- a/Engine/Source/Runtime/Online/HTTP/Private/HttpThread.cpp
+++ b/Engine/Source/Runtime/Online/HTTP/Private/HttpThread.cpp
@@ -88,7 +88,6 @@ void FHttpThreadBase::CancelRequest(IHttpThreadedRequest* Request)
 
 void FHttpThreadBase::GetCompletedRequests(TArray<IHttpThreadedRequest*>& OutCompletedRequests)
 {
-	check(IsInGameThread());
 	IHttpThreadedRequest* Request = nullptr;
 	while (CompletedThreadedRequests.Dequeue(Request))
 	{
diff --git a/Engine/Source/Runtime/Online/HTTP/Public/HttpRetrySystem.h b/Engine/Source/Runtime/Online/HTTP/Public/HttpRetrySystem.h
--- a/Engine/Source/Runtime/Online/HTTP/Public/HttpRetrySystem.h
+++ b/Engine/Source/Runtime/Online/HTTP/Public/HttpRetrySystem.h
@@ -27,14 +27,14 @@ namespace FHttpRetrySystem
 	typedef uint32 RetryLimitCountType;
 	typedef double RetryTimeoutRelativeSecondsType;
 
-	inline RetryLimitCountType             RetryLimitCount(uint32 Value)             { return Value; }
+	inline RetryLimitCountType             RetryLimitCount(uint32 Value) { return Value; }
 	inline RetryTimeoutRelativeSecondsType RetryTimeoutRelativeSeconds(double Value) { return Value; }
 
 	template <typename  IntrinsicType>
 	IntrinsicType TZero();
 
-	template <> inline float                           TZero<float>()                           { return 0.0f; }
-	template <> inline RetryLimitCountType             TZero<RetryLimitCountType>()             { return RetryLimitCount(0); }
+	template <> inline float                           TZero<float>() { return 0.0f; }
+	template <> inline RetryLimitCountType             TZero<RetryLimitCountType>() { return RetryLimitCount(0); }
 	template <> inline RetryTimeoutRelativeSecondsType TZero<RetryTimeoutRelativeSecondsType>() { return RetryTimeoutRelativeSeconds(0.0); }
 
 	typedef TOptional<float>                           FRandomFailureRateSetting;
@@ -45,7 +45,7 @@ namespace FHttpRetrySystem
 
 	struct FRetryDomains
 	{
-		FRetryDomains(TArray<FString>&& InDomains) 
+		FRetryDomains(TArray<FString>&& InDomains)
 			: Domains(MoveTemp(InDomains))
 			, ActiveIndex(0)
 		{}
@@ -71,36 +71,36 @@ namespace FHttpRetrySystem
 
 namespace FHttpRetrySystem
 {
-    /**
-     * class FRequest is what the retry system accepts as inputs
-     */
-    class FRequest 
+	/**
+	 * class FRequest is what the retry system accepts as inputs
+	 */
+	class FRequest
 		: public FHttpRequestAdapterBase
-    {
-    public:
-        struct EStatus
-        {
-            enum Type
-            {
-                NotStarted = 0,
-                Processing,
-                ProcessingLockout,
-                Cancelled,
-                FailedRetry,
-                FailedTimeout,
-                Succeeded
-            };
-        };
-
-    public:
+	{
+	public:
+		struct EStatus
+		{
+			enum Type
+			{
+				NotStarted = 0,
+				Processing,
+				ProcessingLockout,
+				Cancelled,
+				FailedRetry,
+				FailedTimeout,
+				Succeeded
+			};
+		};
+
+	public:
 		// IHttpRequest interface
 		HTTP_API virtual bool ProcessRequest() override;
 		HTTP_API virtual void CancelRequest() override;
-		
+
 		// FRequest
 		EStatus::Type GetRetryStatus() const { return Status; }
 
-    protected:
+	protected:
 		friend class FManager;
 
 		HTTP_API FRequest(
@@ -108,10 +108,10 @@ namespace FHttpRetrySystem
 			const TSharedRef<IHttpRequest, ESPMode::ThreadSafe>& HttpRequest,
 			const FRetryLimitCountSetting& InRetryLimitCountOverride = FRetryLimitCountSetting(),
 			const FRetryTimeoutRelativeSecondsSetting& InRetryTimeoutRelativeSecondsOverride = FRetryTimeoutRelativeSecondsSetting(),
-            const FRetryResponseCodes& InRetryResponseCodes = FRetryResponseCodes(),
-            const FRetryVerbs& InRetryVerbs = FRetryVerbs(),
+			const FRetryResponseCodes& InRetryResponseCodes = FRetryResponseCodes(),
+			const FRetryVerbs& InRetryVerbs = FRetryVerbs(),
 			const FRetryDomainsPtr& InRetryDomains = FRetryDomainsPtr()
-			);
+		);
 
 		void HttpOnRequestProgress(FHttpRequestPtr InHttpRequest, uint64 BytesSent, uint64 BytesRcv);
 
@@ -122,10 +122,10 @@ namespace FHttpRetrySystem
 
 		EStatus::Type                        Status;
 
-        FRetryLimitCountSetting              RetryLimitCountOverride;
-        FRetryTimeoutRelativeSecondsSetting  RetryTimeoutRelativeSecondsOverride;
+		FRetryLimitCountSetting              RetryLimitCountOverride;
+		FRetryTimeoutRelativeSecondsSetting  RetryTimeoutRelativeSecondsOverride;
 		FRetryResponseCodes					 RetryResponseCodes;
-        FRetryVerbs                          RetryVerbs;
+		FRetryVerbs                          RetryVerbs;
 		FRetryDomainsPtr					 RetryDomains;
 		/** The current index in RetryDomains we are attempting */
 		int32								 RetryDomainsIndex = 0;
@@ -133,113 +133,116 @@ namespace FHttpRetrySystem
 		FString								 OriginalUrl;
 
 		FManager& RetryManager;
-    };
+	};
 }
 
 namespace FHttpRetrySystem
 {
-    class FManager
-    {
-    public:
-        // FManager
-		HTTP_API FManager(const FRetryLimitCountSetting& InRetryLimitCountDefault, const FRetryTimeoutRelativeSecondsSetting& InRetryTimeoutRelativeSecondsDefault);
 
-		/**
-		 * Create a new http request with retries
-		 */
-		HTTP_API TSharedRef<class FHttpRetrySystem::FRequest, ESPMode::ThreadSafe> CreateRequest(
-			const FRetryLimitCountSetting& InRetryLimitCountOverride = FRetryLimitCountSetting(),
-			const FRetryTimeoutRelativeSecondsSetting& InRetryTimeoutRelativeSecondsOverride = FRetryTimeoutRelativeSecondsSetting(),
-			const FRetryResponseCodes& InRetryResponseCodes = FRetryResponseCodes(),
-			const FRetryVerbs& InRetryVerbs = FRetryVerbs(),
-			const FRetryDomainsPtr& InRetryDomains = FRetryDomainsPtr()
-			);
-
-		HTTP_API virtual ~FManager();
-
-        /**
-         * Updates the entries in the list of retry requests. Optional parameters are for future connection health assessment
-         *
-         * @param FileCount       optional parameter that will be filled with the total files updated
-         * @param FailingCount    optional parameter that will be filled with the total files that have are in a retrying state
-         * @param FailedCount     optional parameter that will be filled with the total files that have failed
-         * @param CompletedCount  optional parameter that will be filled with the total files that have completed
-         *
-         * @return                true if there are no failures or retries
-         */
-        HTTP_API bool Update(uint32* FileCount = NULL, uint32* FailingCount = NULL, uint32* FailedCount = NULL, uint32* CompletedCount = NULL);
-		void SetRandomFailureRate(float Value) { RandomFailureRate = FRandomFailureRateSetting(Value); }
-		void SetDefaultRetryLimit(uint32 Value) { RetryLimitCountDefault = FRetryLimitCountSetting(Value); }
-		
-		// @return Block the current process until all requests are flushed, or timeout has elapsed
-		HTTP_API void BlockUntilFlushed(float TimeoutSec);
-
-    protected:
-		friend class FRequest;
-
-        struct FHttpRetryRequestEntry
-        {
-            FHttpRetryRequestEntry(TSharedRef<FRequest, ESPMode::ThreadSafe>& InRequest);
-
-            bool                    bShouldCancel;
-            uint32                  CurrentRetryCount;
-            double                  RequestStartTimeAbsoluteSeconds;
-            double                  LockoutEndTimeAbsoluteSeconds;
-
-            TSharedRef<FRequest, ESPMode::ThreadSafe>	Request;
-        };
-
-		class FHttpLogVerbosityTracker
-		{
-		public:
-			FHttpLogVerbosityTracker();
-			~FHttpLogVerbosityTracker();
-
-			/** Mark that a request is being retried */
-			void IncrementRetriedRequests();
-			/** Mark that a retried request is no longer being retried */
-			void DecrementRetriedRequests();
-			static FHttpLogVerbosityTracker& Get();
-		protected:
-			/** Update settings from config */
-			void UpdateSettingsFromConfig();
-			void OnConfigSectionsChanged(const FString& IniFilename, const TSet<FString>& SectionName);
-
-			/** Number of requests that are in a retried state.  When this is non-zero, verbosity will be adjusted. */
-			int32 NumRetriedRequests = 0;
-			/** Verbosity to restore to when there are no requests being retried */
-			ELogVerbosity::Type OriginalVerbosity = ELogVerbosity::Error;
-			/** Config driven target verbosity to set to when requests are being retried.  NoLogging means the verbosity will not be modified. */
-			ELogVerbosity::Type TargetVerbosity = ELogVerbosity::NoLogging;
-		};
+class FManager
+{
+public:
+	// FManager
+	HTTP_API FManager(const FRetryLimitCountSetting& InRetryLimitCountDefault, const FRetryTimeoutRelativeSecondsSetting& InRetryTimeoutRelativeSecondsDefault);
 
-		bool ProcessRequest(TSharedRef<FRequest, ESPMode::ThreadSafe>& HttpRequest);
-		void CancelRequest(TSharedRef<FRequest, ESPMode::ThreadSafe>& HttpRequest);
+	/**
+	 * Create a new http request with retries
+	 */
+	HTTP_API TSharedRef<class FHttpRetrySystem::FRequest, ESPMode::ThreadSafe> CreateRequest(
+		const FRetryLimitCountSetting& InRetryLimitCountOverride = FRetryLimitCountSetting(),
+		const FRetryTimeoutRelativeSecondsSetting& InRetryTimeoutRelativeSecondsOverride = FRetryTimeoutRelativeSecondsSetting(),
+		const FRetryResponseCodes& InRetryResponseCodes = FRetryResponseCodes(),
+		const FRetryVerbs& InRetryVerbs = FRetryVerbs(),
+		const FRetryDomainsPtr& InRetryDomains = FRetryDomainsPtr()
+	);
 
-        // @return true if there is a no formal response to the request
-        // @TODO return true if a variety of 5xx errors are the result of a formal response
-        bool ShouldRetry(const FHttpRetryRequestEntry& HttpRetryRequestEntry);
+	HTTP_API virtual ~FManager();
 
-        // @return true if retry chances have not been exhausted
-        bool CanRetry(const FHttpRetryRequestEntry& HttpRetryRequestEntry);
+	/**
+	 * Updates the entries in the list of retry requests. Optional parameters are for future connection health assessment
+	 *
+	 * @param FileCount       optional parameter that will be filled with the total files updated
+	 * @param FailingCount    optional parameter that will be filled with the total files that have are in a retrying state
+	 * @param FailedCount     optional parameter that will be filled with the total files that have failed
+	 * @param CompletedCount  optional parameter that will be filled with the total files that have completed
+	 *
+	 * @return                true if there are no failures or retries
+	 */
+	HTTP_API bool Update(uint32* FileCount = NULL, uint32* FailingCount = NULL, uint32* FailedCount = NULL, uint32* CompletedCount = NULL);
+	void SetRandomFailureRate(float Value) { RandomFailureRate = FRandomFailureRateSetting(Value); }
+	void SetDefaultRetryLimit(uint32 Value) { RetryLimitCountDefault = FRetryLimitCountSetting(Value); }
 
-        // @return true if the retry request has timed out
-        bool HasTimedOut(const FHttpRetryRequestEntry& HttpRetryRequestEntry, const double NowAbsoluteSeconds);
+	// @return Block the current process until all requests are flushed, or timeout has elapsed
+	HTTP_API void BlockUntilFlushed(float TimeoutSec);
 
-		/**
-		 * Retry an HTTP request
-		 * @param RequestEntry request to retry
-		 */
-		void RetryHttpRequest(FHttpRetryRequestEntry& RequestEntry);
+protected:
+	friend class FRequest;
+
+	struct FHttpRetryRequestEntry
+	{
+		FHttpRetryRequestEntry(TSharedRef<FRequest, ESPMode::ThreadSafe>& InRequest);
 
-        // @return number of seconds to lockout for
-        float GetLockoutPeriodSeconds(const FHttpRetryRequestEntry& HttpRetryRequestEntry);
+		bool                    bShouldCancel;
+		uint32                  CurrentRetryCount;
+		double                  RequestStartTimeAbsoluteSeconds;
+		double                  LockoutEndTimeAbsoluteSeconds;
+
+		TSharedRef<FRequest, ESPMode::ThreadSafe>	Request;
+	};
+
+	class FHttpLogVerbosityTracker
+	{
+	public:
+		FHttpLogVerbosityTracker();
+		~FHttpLogVerbosityTracker();
+
+		/** Mark that a request is being retried */
+		void IncrementRetriedRequests();
+		/** Mark that a retried request is no longer being retried */
+		void DecrementRetriedRequests();
+		static FHttpLogVerbosityTracker& Get();
+	protected:
+		/** Update settings from config */
+		void UpdateSettingsFromConfig();
+		void OnConfigSectionsChanged(const FString& IniFilename, const TSet<FString>& SectionName);
+
+		/** Number of requests that are in a retried state.  When this is non-zero, verbosity will be adjusted. */
+		int32 NumRetriedRequests = 0;
+		/** Verbosity to restore to when there are no requests being retried */
+		ELogVerbosity::Type OriginalVerbosity = ELogVerbosity::Error;
+		/** Config driven target verbosity to set to when requests are being retried.  NoLogging means the verbosity will not be modified. */
+		ELogVerbosity::Type TargetVerbosity = ELogVerbosity::NoLogging;
+	};
 
-        // Default configuration for the retry system
-        FRandomFailureRateSetting            RandomFailureRate;
-        FRetryLimitCountSetting              RetryLimitCountDefault;
-        FRetryTimeoutRelativeSecondsSetting  RetryTimeoutRelativeSecondsDefault;
+	bool ProcessRequest(TSharedRef<FRequest, ESPMode::ThreadSafe>& HttpRequest);
+	void CancelRequest(TSharedRef<FRequest, ESPMode::ThreadSafe>& HttpRequest);
+
+	// @return true if there is a no formal response to the request
+	// @TODO return true if a variety of 5xx errors are the result of a formal response
+	bool ShouldRetry(const FHttpRetryRequestEntry& HttpRetryRequestEntry);
+
+	// @return true if retry chances have not been exhausted
+	bool CanRetry(const FHttpRetryRequestEntry& HttpRetryRequestEntry);
+
+	// @return true if the retry request has timed out
+	bool HasTimedOut(const FHttpRetryRequestEntry& HttpRetryRequestEntry, const double NowAbsoluteSeconds);
+
+	/**
+	 * Retry an HTTP request
+	 * @param RequestEntry request to retry
+	 */
+	void RetryHttpRequest(FHttpRetryRequestEntry& RequestEntry);
+
+	// @return number of seconds to lockout for
+	float GetLockoutPeriodSeconds(const FHttpRetryRequestEntry& HttpRetryRequestEntry);
+
+	// Default configuration for the retry system
+	FRandomFailureRateSetting RandomFailureRate;
+	FRetryLimitCountSetting RetryLimitCountDefault;
+	FRetryTimeoutRelativeSecondsSetting RetryTimeoutRelativeSecondsDefault;
+
+	TArray<FHttpRetryRequestEntry> RequestList;
+	FCriticalSection RequestListLock;
+};
 
-        TArray<FHttpRetryRequestEntry>       RequestList;
-    };
 }
