[Backout] - CL27428302
Original CL Desc
-----------------------------------------------------------------
Make http retry system thread safe.

diff --git a/Engine/Source/Programs/WebTests/Private/TestHttp.cpp b/Engine/Source/Programs/WebTests/Private/TestHttp.cpp
--- a/Engine/Source/Programs/WebTests/Private/TestHttp.cpp
+++ b/Engine/Source/Programs/WebTests/Private/TestHttp.cpp
@@ -6,7 +6,6 @@
 #include "HAL/Runnable.h"
 #include "HAL/RunnableThread.h"
 #include "HttpManager.h"
-#include "HttpRetrySystem.h"
 #include "Http.h"
 #include "Misc/CommandLine.h"
 #include "TestHarness.h"
@@ -177,7 +176,7 @@ public:
 
 	void WaitUntilAllHttpRequestsComplete()
 	{
-		while (OngoingRequests != 0)
+		while (!bRunningThreadRequest && OngoingRequests != 0)
 		{
 			HttpModule->GetHttpManager().Tick(TickFrequency);
 			FPlatformProcess::Sleep(TickFrequency);
@@ -186,6 +185,7 @@ public:
 
 	uint32 OngoingRequests = 0;
 	float TickFrequency = 1.0f / 60; /*60 FPS*/;
+	bool bRunningThreadRequest = false;
 };
 
 TEST_CASE_METHOD(FWaitUntilCompleteHttpFixture, "Http Methods", HTTP_TAG)
@@ -594,61 +594,38 @@ TEST_CASE_METHOD(FWaitUntilCompleteHttpFixture, "Redirect enabled by default and
 	HttpRequest->ProcessRequest();
 }
 
-class FThreadedHttpRunnable : public FRunnable
+class FWaitThreadedHttpFixture : public FWaitUntilCompleteHttpFixture, public FRunnable
 {
 public:
 	DECLARE_DELEGATE(FRunActualTestCodeDelegate);
 
-	FRunActualTestCodeDelegate& OnRunFromThread()
+	FWaitThreadedHttpFixture()
 	{
-		return ThreadCallback;
+		bRunningThreadRequest = true;
 	}
 
 	// FRunnable interface
 	virtual uint32 Run() override
 	{
 		ThreadCallback.ExecuteIfBound();
+		bRunningThreadRequest = false;
 		return 0;
 	}
 
-	void StartTestHttpThread(bool bBlockGameThread)
+	void StartTestHttpThread()
 	{
-		bBlockingGameThreadTick = bBlockGameThread;
-
 		RunnableThread = TSharedPtr<FRunnableThread>(FRunnableThread::Create(this, TEXT("Test Http Thread")));
-
-		while (bBlockingGameThreadTick)
-		{
-			float TickFrequency = 1.0f / 60; /*60 FPS*/;
-			FPlatformProcess::Sleep(TickFrequency);
-		}
-	}
-
-	void UnblockGameThread()
-	{
-		bBlockingGameThreadTick = false;
 	}
 
-private:
 	FRunActualTestCodeDelegate ThreadCallback;
 	TSharedPtr<FRunnableThread> RunnableThread;
-	std::atomic<bool> bBlockingGameThreadTick = true;
-};
-
-class FWaitThreadedHttpFixture : public FWaitUntilCompleteHttpFixture
-{
-public:
-	~FWaitThreadedHttpFixture()
-	{
-		WaitUntilAllHttpRequestsComplete();
-	}
-
-	FThreadedHttpRunnable ThreadedHttpRunnable;
 };
 
 TEST_CASE_METHOD(FWaitThreadedHttpFixture, "Http streaming download request can work in non game thread", HTTP_TAG)
 {
-	ThreadedHttpRunnable.OnRunFromThread().BindLambda([this]() {
+	// TODO: Block main thread here to verify
+
+	ThreadCallback.BindLambda([this]() {
 		TSharedRef<IHttpRequest> HttpRequest = HttpModule->CreateRequest();
 		HttpRequest->SetURL(UrlStreamDownload(3/*Chunks*/, 1024/*ChunkSize*/));
 		HttpRequest->SetVerb(TEXT("GET"));
@@ -665,18 +642,17 @@ TEST_CASE_METHOD(FWaitThreadedHttpFixture, "Http streaming download request can
 		};
 		CHECK(HttpRequest->SetResponseBodyReceiveStream(MakeShared<FTestHttpReceiveStream>()));
 
-		HttpRequest->OnProcessRequestComplete().BindLambda([this](FHttpRequestPtr HttpRequest, FHttpResponsePtr HttpResponse, bool bSucceeded) {
+		HttpRequest->OnProcessRequestComplete().BindLambda([](FHttpRequestPtr HttpRequest, FHttpResponsePtr HttpResponse, bool bSucceeded) {
 			// EHttpRequestDelegateThreadPolicy::CompleteOnHttpThread was used, so not in game thread here
 			CHECK(!IsInGameThread());
 			CHECK(bSucceeded);
 			CHECK(HttpResponse->GetResponseCode() == 200);
-			ThreadedHttpRunnable.UnblockGameThread();
 		});
 
 		HttpRequest->ProcessRequest();
 	});
 
-	ThreadedHttpRunnable.StartTestHttpThread(true/*bBlockGameThread*/);
+	StartTestHttpThread();
 }
 
 namespace UE
@@ -726,7 +702,7 @@ TEST_CASE_METHOD(FWaitUntilCompleteHttpFixture, "Http request pre check will fai
 
 TEST_CASE_METHOD(FWaitThreadedHttpFixture, "Threaded http request pre check will fail", HTTP_TAG)
 {
-	ThreadedHttpRunnable.OnRunFromThread().BindLambda([this]() {
+	ThreadCallback.BindLambda([this]() {
 		// Pre check will fail when domain is not allowed
 		UE::TestHttp::SetupURLRequestFilter(HttpModule);
 
@@ -753,9 +729,10 @@ TEST_CASE_METHOD(FWaitThreadedHttpFixture, "Threaded http request pre check will
 		HttpRequest->ProcessRequest();
 	});
 
-	ThreadedHttpRunnable.StartTestHttpThread(false/*bBlockGameThread*/);
+	StartTestHttpThread();
 }
 
+
 TEST_CASE_METHOD(FWaitUntilCompleteHttpFixture, "Cancel http request connect before timeout", HTTP_TAG)
 {
 	TSharedRef<IHttpRequest, ESPMode::ThreadSafe> HttpRequest = HttpModule->CreateRequest();
@@ -774,102 +751,4 @@ TEST_CASE_METHOD(FWaitUntilCompleteHttpFixture, "Cancel http request connect bef
 	FPlatformProcess::Sleep(0.5);
 	HttpRequest->CancelRequest();
 }
-
-class FMockRetryManager : public FHttpRetrySystem::FManager
-{
-public:
-	using FHttpRetrySystem::FManager::FManager;
-	using FHttpRetrySystem::FManager::RequestList;
-	using FHttpRetrySystem::FManager::FHttpRetryRequestEntry;
-	using FHttpRetrySystem::FManager::RetryTimeoutRelativeSecondsDefault;
-};
-
-class FRetryHttpFixture : public FWaitUntilCompleteHttpFixture
-{
-public:
-	FRetryHttpFixture()
-	{
-		HttpRetryManager = MakeShared<FMockRetryManager>(FHttpRetrySystem::FRetryLimitCountSetting(RetryLimitCount), FHttpRetrySystem::FRetryTimeoutRelativeSecondsSetting(/*RetryTimeoutRelativeSeconds*/));
-		HttpRetryManager->RetryTimeoutRelativeSecondsDefault = 2.0f; // Value is set so that retry system will use it at high level, will no longer wait for ever for HTTP code
-	}
-
-	~FRetryHttpFixture()
-	{
-		// Make sure not destroy retry manager before waiting in FWaitUntilCompleteHttpFixture
-		while (HttpRetryManager->RequestList.Num())
-		{
-			HttpRetryManager->Update();
-			HttpModule->GetHttpManager().Tick(TickFrequency);
-			FPlatformProcess::Sleep(TickFrequency);
-		}
-	}
-
-	TSharedPtr<FMockRetryManager> HttpRetryManager;
-	uint32 RetryLimitCount = 2;
-
-	FThreadedHttpRunnable ThreadedHttpRunnable;
-};
-
-// TODO: Currently broken, will fix this in retry system
-//TEST_CASE_METHOD(FRetryHttpFixture, "Http request adaptor in retry system can work in non game thread", HTTP_TAG)
-//{
-//	ThreadedHttpRunnable.OnRunFromThread().BindLambda([this]() {
-//		TSharedRef<IHttpRequest> HttpRequest = HttpRetryManager->CreateRequest();
-//		HttpRequest->SetURL(UrlStreamDownload(3/*Chunks*/, 1024/*ChunkSize*/));
-//		HttpRequest->SetVerb(TEXT("GET"));
-//		HttpRequest->SetDelegateThreadPolicy(EHttpRequestDelegateThreadPolicy::CompleteOnHttpThread);
-//
-//		class FTestHttpReceiveStream final : public FArchive
-//		{
-//		public:
-//			virtual void Serialize(void* V, int64 Length) override
-//			{
-//				// No matter what's the thread policy, Serialize always get called in http thread.
-//				CHECK(!IsInGameThread());
-//			}
-//		};
-//		CHECK(HttpRequest->SetResponseBodyReceiveStream(MakeShared<FTestHttpReceiveStream>()));
-//
-//		HttpRequest->OnProcessRequestComplete().BindLambda([this](FHttpRequestPtr HttpRequest, FHttpResponsePtr HttpResponse, bool bSucceeded) {
-//			// EHttpRequestDelegateThreadPolicy::CompleteOnHttpThread was used, so not in game thread here
-//			CHECK(!IsInGameThread());
-//			CHECK(bSucceeded);
-//			CHECK(HttpResponse->GetResponseCode() == 200);
-//			ThreadedHttpRunnable.UnblockGameThread();
-//		});
-//
-//		HttpRequest->ProcessRequest();
-//	});
-//
-//	ThreadedHttpRunnable.StartTestHttpThread(true/*bBlockGameThread*/);
-//}
-
-class FRetryHttpManagerThreadedRequestsFixture : public FRetryHttpFixture
-{
-public:
-	void LaunchDownloadRequests()
-	{
-		for (int32 i = 0; i < 10; ++i)
-		{
-			TSharedRef<IHttpRequest> HttpRequest = HttpRetryManager->CreateRequest();
-			HttpRequest->SetURL(UrlStreamDownload(3, 1024*1024));
-			HttpRequest->SetVerb(TEXT("GET"));
-			HttpRequest->ProcessRequest();
-		}
-	}
-};
-
-TEST_CASE_METHOD(FRetryHttpManagerThreadedRequestsFixture, "Retry manager is thread safe", HTTP_TAG)
-{
-	ThreadedHttpRunnable.OnRunFromThread().BindLambda([this]() {
-		LaunchDownloadRequests();
-		HttpRetryManager->BlockUntilFlushed(5.0);
-	});
-	ThreadedHttpRunnable.StartTestHttpThread(false/*bBlockGameThread*/);
-
-	LaunchDownloadRequests();
-	HttpRetryManager->BlockUntilFlushed(5.0);
-}
-
-
 // TODO: Add cancel test, with multiple cancel calls
diff --git a/Engine/Source/Runtime/Online/HTTP/Private/Curl/CurlHttp.cpp b/Engine/Source/Runtime/Online/HTTP/Private/Curl/CurlHttp.cpp
--- a/Engine/Source/Runtime/Online/HTTP/Private/Curl/CurlHttp.cpp
+++ b/Engine/Source/Runtime/Online/HTTP/Private/Curl/CurlHttp.cpp
@@ -1141,11 +1141,7 @@ PRAGMA_ENABLE_DEPRECATION_WARNINGS
 
 void FCurlHttpRequest::BroadcastNewlyReceivedHeaders()
 {
-	if (!OnHeaderReceived().IsBound())
-	{
-		return;
-	}
-
+	check(IsInGameThread() || DelegateThreadPolicy == EHttpRequestDelegateThreadPolicy::CompleteOnHttpThread);
 	if (Response.IsValid())
 	{
 		// Process the headers received on the HTTP thread and merge them into the response's list of headers and then broadcast the new headers
@@ -1180,6 +1176,8 @@ void FCurlHttpRequest::BroadcastNewlyReceivedHeader(const FString& HeaderKey, co
 
 void FCurlHttpRequest::FinishRequest()
 {
+	check(IsInGameThread() || DelegateThreadPolicy == EHttpRequestDelegateThreadPolicy::CompleteOnHttpThread);
+
 	QUICK_SCOPE_CYCLE_COUNTER(STAT_FCurlHttpRequest_FinishRequest);
 
 	curl_easy_setopt(EasyHandle, CURLOPT_SHARE, nullptr);
diff --git a/Engine/Source/Runtime/Online/HTTP/Private/HttpRetrySystem.cpp b/Engine/Source/Runtime/Online/HTTP/Private/HttpRetrySystem.cpp
--- a/Engine/Source/Runtime/Online/HTTP/Private/HttpRetrySystem.cpp
+++ b/Engine/Source/Runtime/Online/HTTP/Private/HttpRetrySystem.cpp
@@ -158,8 +158,6 @@ FHttpRetrySystem::FManager::FManager(const FRetryLimitCountSetting& InRetryLimit
 
 FHttpRetrySystem::FManager::~FManager()
 {
-	FScopeLock ScopeLock(&RequestListLock);
-
 	// Decrement retried request for log verbosity tracker
 	for (const FHttpRetryRequestEntry& Request : RequestList)
 	{
@@ -344,8 +342,6 @@ bool FHttpRetrySystem::FManager::Update(uint32* FileCount, uint32* FailingCount,
 	QUICK_SCOPE_CYCLE_COUNTER(STAT_FHttpRetrySystem_FManager_Update);
 	LLM_SCOPE_BYTAG(HTTP);
 
-	FScopeLock ScopeLock(&RequestListLock);
-
 	bool bIsGreen = true;
 
 	if (FileCount != nullptr)
@@ -593,7 +589,6 @@ bool FHttpRetrySystem::FManager::ProcessRequest(TSharedRef<FHttpRetrySystem::FRe
 
 	if (bResult)
 	{
-		FScopeLock ScopeLock(&RequestListLock);
 		RequestList.Add(FHttpRetryRequestEntry(HttpRetryRequest));
 	}
 
@@ -604,8 +599,6 @@ void FHttpRetrySystem::FManager::CancelRequest(TSharedRef<FHttpRetrySystem::FReq
 {
 	QUICK_SCOPE_CYCLE_COUNTER(STAT_FHttpRetrySystem_FManager_CancelRequest);
 
-	FScopeLock ScopeLock(&RequestListLock);
-
 	// Find the existing request entry if is was previously processed.
 	bool bFound = false;
 	for (int32 i = 0; i < RequestList.Num(); ++i)
@@ -636,9 +629,6 @@ void FHttpRetrySystem::FManager::BlockUntilFlushed(float InTimeoutSec)
 	const float SleepInterval = 0.016;
 	float TimeElapsed = 0.0f;
 	uint32 FileCount, FailingCount, FailedCount, CompleteCount;
-
-	FScopeLock ScopeLock(&RequestListLock);
-
 	while (RequestList.Num() > 0 && TimeElapsed < InTimeoutSec)
 	{
 		FHttpModule::Get().GetHttpManager().Tick(SleepInterval);
diff --git a/Engine/Source/Runtime/Online/HTTP/Private/HttpThread.cpp b/Engine/Source/Runtime/Online/HTTP/Private/HttpThread.cpp
--- a/Engine/Source/Runtime/Online/HTTP/Private/HttpThread.cpp
+++ b/Engine/Source/Runtime/Online/HTTP/Private/HttpThread.cpp
@@ -88,6 +88,7 @@ void FHttpThreadBase::CancelRequest(IHttpThreadedRequest* Request)
 
 void FHttpThreadBase::GetCompletedRequests(TArray<IHttpThreadedRequest*>& OutCompletedRequests)
 {
+	check(IsInGameThread());
 	IHttpThreadedRequest* Request = nullptr;
 	while (CompletedThreadedRequests.Dequeue(Request))
 	{
diff --git a/Engine/Source/Runtime/Online/HTTP/Public/HttpRetrySystem.h b/Engine/Source/Runtime/Online/HTTP/Public/HttpRetrySystem.h
--- a/Engine/Source/Runtime/Online/HTTP/Public/HttpRetrySystem.h
+++ b/Engine/Source/Runtime/Online/HTTP/Public/HttpRetrySystem.h
@@ -27,14 +27,14 @@ namespace FHttpRetrySystem
 	typedef uint32 RetryLimitCountType;
 	typedef double RetryTimeoutRelativeSecondsType;
 
-	inline RetryLimitCountType             RetryLimitCount(uint32 Value) { return Value; }
+	inline RetryLimitCountType             RetryLimitCount(uint32 Value)             { return Value; }
 	inline RetryTimeoutRelativeSecondsType RetryTimeoutRelativeSeconds(double Value) { return Value; }
 
 	template <typename  IntrinsicType>
 	IntrinsicType TZero();
 
-	template <> inline float                           TZero<float>() { return 0.0f; }
-	template <> inline RetryLimitCountType             TZero<RetryLimitCountType>() { return RetryLimitCount(0); }
+	template <> inline float                           TZero<float>()                           { return 0.0f; }
+	template <> inline RetryLimitCountType             TZero<RetryLimitCountType>()             { return RetryLimitCount(0); }
 	template <> inline RetryTimeoutRelativeSecondsType TZero<RetryTimeoutRelativeSecondsType>() { return RetryTimeoutRelativeSeconds(0.0); }
 
 	typedef TOptional<float>                           FRandomFailureRateSetting;
@@ -45,7 +45,7 @@ namespace FHttpRetrySystem
 
 	struct FRetryDomains
 	{
-		FRetryDomains(TArray<FString>&& InDomains)
+		FRetryDomains(TArray<FString>&& InDomains) 
 			: Domains(MoveTemp(InDomains))
 			, ActiveIndex(0)
 		{}
@@ -71,36 +71,36 @@ namespace FHttpRetrySystem
 
 namespace FHttpRetrySystem
 {
-	/**
-	 * class FRequest is what the retry system accepts as inputs
-	 */
-	class FRequest
+    /**
+     * class FRequest is what the retry system accepts as inputs
+     */
+    class FRequest 
 		: public FHttpRequestAdapterBase
-	{
-	public:
-		struct EStatus
-		{
-			enum Type
-			{
-				NotStarted = 0,
-				Processing,
-				ProcessingLockout,
-				Cancelled,
-				FailedRetry,
-				FailedTimeout,
-				Succeeded
-			};
-		};
-
-	public:
+    {
+    public:
+        struct EStatus
+        {
+            enum Type
+            {
+                NotStarted = 0,
+                Processing,
+                ProcessingLockout,
+                Cancelled,
+                FailedRetry,
+                FailedTimeout,
+                Succeeded
+            };
+        };
+
+    public:
 		// IHttpRequest interface
 		HTTP_API virtual bool ProcessRequest() override;
 		HTTP_API virtual void CancelRequest() override;
-
+		
 		// FRequest
 		EStatus::Type GetRetryStatus() const { return Status; }
 
-	protected:
+    protected:
 		friend class FManager;
 
 		HTTP_API FRequest(
@@ -108,10 +108,10 @@ namespace FHttpRetrySystem
 			const TSharedRef<IHttpRequest, ESPMode::ThreadSafe>& HttpRequest,
 			const FRetryLimitCountSetting& InRetryLimitCountOverride = FRetryLimitCountSetting(),
 			const FRetryTimeoutRelativeSecondsSetting& InRetryTimeoutRelativeSecondsOverride = FRetryTimeoutRelativeSecondsSetting(),
-			const FRetryResponseCodes& InRetryResponseCodes = FRetryResponseCodes(),
-			const FRetryVerbs& InRetryVerbs = FRetryVerbs(),
+            const FRetryResponseCodes& InRetryResponseCodes = FRetryResponseCodes(),
+            const FRetryVerbs& InRetryVerbs = FRetryVerbs(),
 			const FRetryDomainsPtr& InRetryDomains = FRetryDomainsPtr()
-		);
+			);
 
 		void HttpOnRequestProgress(FHttpRequestPtr InHttpRequest, uint64 BytesSent, uint64 BytesRcv);
 
@@ -122,10 +122,10 @@ namespace FHttpRetrySystem
 
 		EStatus::Type                        Status;
 
-		FRetryLimitCountSetting              RetryLimitCountOverride;
-		FRetryTimeoutRelativeSecondsSetting  RetryTimeoutRelativeSecondsOverride;
+        FRetryLimitCountSetting              RetryLimitCountOverride;
+        FRetryTimeoutRelativeSecondsSetting  RetryTimeoutRelativeSecondsOverride;
 		FRetryResponseCodes					 RetryResponseCodes;
-		FRetryVerbs                          RetryVerbs;
+        FRetryVerbs                          RetryVerbs;
 		FRetryDomainsPtr					 RetryDomains;
 		/** The current index in RetryDomains we are attempting */
 		int32								 RetryDomainsIndex = 0;
@@ -133,116 +133,113 @@ namespace FHttpRetrySystem
 		FString								 OriginalUrl;
 
 		FManager& RetryManager;
-	};
+    };
 }
 
 namespace FHttpRetrySystem
 {
+    class FManager
+    {
+    public:
+        // FManager
+		HTTP_API FManager(const FRetryLimitCountSetting& InRetryLimitCountDefault, const FRetryTimeoutRelativeSecondsSetting& InRetryTimeoutRelativeSecondsDefault);
 
-class FManager
-{
-public:
-	// FManager
-	HTTP_API FManager(const FRetryLimitCountSetting& InRetryLimitCountDefault, const FRetryTimeoutRelativeSecondsSetting& InRetryTimeoutRelativeSecondsDefault);
-
-	/**
-	 * Create a new http request with retries
-	 */
-	HTTP_API TSharedRef<class FHttpRetrySystem::FRequest, ESPMode::ThreadSafe> CreateRequest(
-		const FRetryLimitCountSetting& InRetryLimitCountOverride = FRetryLimitCountSetting(),
-		const FRetryTimeoutRelativeSecondsSetting& InRetryTimeoutRelativeSecondsOverride = FRetryTimeoutRelativeSecondsSetting(),
-		const FRetryResponseCodes& InRetryResponseCodes = FRetryResponseCodes(),
-		const FRetryVerbs& InRetryVerbs = FRetryVerbs(),
-		const FRetryDomainsPtr& InRetryDomains = FRetryDomainsPtr()
-	);
-
-	HTTP_API virtual ~FManager();
-
-	/**
-	 * Updates the entries in the list of retry requests. Optional parameters are for future connection health assessment
-	 *
-	 * @param FileCount       optional parameter that will be filled with the total files updated
-	 * @param FailingCount    optional parameter that will be filled with the total files that have are in a retrying state
-	 * @param FailedCount     optional parameter that will be filled with the total files that have failed
-	 * @param CompletedCount  optional parameter that will be filled with the total files that have completed
-	 *
-	 * @return                true if there are no failures or retries
-	 */
-	HTTP_API bool Update(uint32* FileCount = NULL, uint32* FailingCount = NULL, uint32* FailedCount = NULL, uint32* CompletedCount = NULL);
-	void SetRandomFailureRate(float Value) { RandomFailureRate = FRandomFailureRateSetting(Value); }
-	void SetDefaultRetryLimit(uint32 Value) { RetryLimitCountDefault = FRetryLimitCountSetting(Value); }
-
-	// @return Block the current process until all requests are flushed, or timeout has elapsed
-	HTTP_API void BlockUntilFlushed(float TimeoutSec);
-
-protected:
-	friend class FRequest;
-
-	struct FHttpRetryRequestEntry
-	{
-		FHttpRetryRequestEntry(TSharedRef<FRequest, ESPMode::ThreadSafe>& InRequest);
-
-		bool                    bShouldCancel;
-		uint32                  CurrentRetryCount;
-		double                  RequestStartTimeAbsoluteSeconds;
-		double                  LockoutEndTimeAbsoluteSeconds;
-
-		TSharedRef<FRequest, ESPMode::ThreadSafe>	Request;
-	};
-
-	class FHttpLogVerbosityTracker
-	{
-	public:
-		FHttpLogVerbosityTracker();
-		~FHttpLogVerbosityTracker();
-
-		/** Mark that a request is being retried */
-		void IncrementRetriedRequests();
-		/** Mark that a retried request is no longer being retried */
-		void DecrementRetriedRequests();
-		static FHttpLogVerbosityTracker& Get();
-	protected:
-		/** Update settings from config */
-		void UpdateSettingsFromConfig();
-		void OnConfigSectionsChanged(const FString& IniFilename, const TSet<FString>& SectionName);
-
-		/** Number of requests that are in a retried state.  When this is non-zero, verbosity will be adjusted. */
-		int32 NumRetriedRequests = 0;
-		/** Verbosity to restore to when there are no requests being retried */
-		ELogVerbosity::Type OriginalVerbosity = ELogVerbosity::Error;
-		/** Config driven target verbosity to set to when requests are being retried.  NoLogging means the verbosity will not be modified. */
-		ELogVerbosity::Type TargetVerbosity = ELogVerbosity::NoLogging;
-	};
-
-	bool ProcessRequest(TSharedRef<FRequest, ESPMode::ThreadSafe>& HttpRequest);
-	void CancelRequest(TSharedRef<FRequest, ESPMode::ThreadSafe>& HttpRequest);
+		/**
+		 * Create a new http request with retries
+		 */
+		HTTP_API TSharedRef<class FHttpRetrySystem::FRequest, ESPMode::ThreadSafe> CreateRequest(
+			const FRetryLimitCountSetting& InRetryLimitCountOverride = FRetryLimitCountSetting(),
+			const FRetryTimeoutRelativeSecondsSetting& InRetryTimeoutRelativeSecondsOverride = FRetryTimeoutRelativeSecondsSetting(),
+			const FRetryResponseCodes& InRetryResponseCodes = FRetryResponseCodes(),
+			const FRetryVerbs& InRetryVerbs = FRetryVerbs(),
+			const FRetryDomainsPtr& InRetryDomains = FRetryDomainsPtr()
+			);
+
+		HTTP_API virtual ~FManager();
+
+        /**
+         * Updates the entries in the list of retry requests. Optional parameters are for future connection health assessment
+         *
+         * @param FileCount       optional parameter that will be filled with the total files updated
+         * @param FailingCount    optional parameter that will be filled with the total files that have are in a retrying state
+         * @param FailedCount     optional parameter that will be filled with the total files that have failed
+         * @param CompletedCount  optional parameter that will be filled with the total files that have completed
+         *
+         * @return                true if there are no failures or retries
+         */
+        HTTP_API bool Update(uint32* FileCount = NULL, uint32* FailingCount = NULL, uint32* FailedCount = NULL, uint32* CompletedCount = NULL);
+		void SetRandomFailureRate(float Value) { RandomFailureRate = FRandomFailureRateSetting(Value); }
+		void SetDefaultRetryLimit(uint32 Value) { RetryLimitCountDefault = FRetryLimitCountSetting(Value); }
+		
+		// @return Block the current process until all requests are flushed, or timeout has elapsed
+		HTTP_API void BlockUntilFlushed(float TimeoutSec);
+
+    protected:
+		friend class FRequest;
+
+        struct FHttpRetryRequestEntry
+        {
+            FHttpRetryRequestEntry(TSharedRef<FRequest, ESPMode::ThreadSafe>& InRequest);
+
+            bool                    bShouldCancel;
+            uint32                  CurrentRetryCount;
+            double                  RequestStartTimeAbsoluteSeconds;
+            double                  LockoutEndTimeAbsoluteSeconds;
+
+            TSharedRef<FRequest, ESPMode::ThreadSafe>	Request;
+        };
+
+		class FHttpLogVerbosityTracker
+		{
+		public:
+			FHttpLogVerbosityTracker();
+			~FHttpLogVerbosityTracker();
+
+			/** Mark that a request is being retried */
+			void IncrementRetriedRequests();
+			/** Mark that a retried request is no longer being retried */
+			void DecrementRetriedRequests();
+			static FHttpLogVerbosityTracker& Get();
+		protected:
+			/** Update settings from config */
+			void UpdateSettingsFromConfig();
+			void OnConfigSectionsChanged(const FString& IniFilename, const TSet<FString>& SectionName);
+
+			/** Number of requests that are in a retried state.  When this is non-zero, verbosity will be adjusted. */
+			int32 NumRetriedRequests = 0;
+			/** Verbosity to restore to when there are no requests being retried */
+			ELogVerbosity::Type OriginalVerbosity = ELogVerbosity::Error;
+			/** Config driven target verbosity to set to when requests are being retried.  NoLogging means the verbosity will not be modified. */
+			ELogVerbosity::Type TargetVerbosity = ELogVerbosity::NoLogging;
+		};
 
-	// @return true if there is a no formal response to the request
-	// @TODO return true if a variety of 5xx errors are the result of a formal response
-	bool ShouldRetry(const FHttpRetryRequestEntry& HttpRetryRequestEntry);
+		bool ProcessRequest(TSharedRef<FRequest, ESPMode::ThreadSafe>& HttpRequest);
+		void CancelRequest(TSharedRef<FRequest, ESPMode::ThreadSafe>& HttpRequest);
 
-	// @return true if retry chances have not been exhausted
-	bool CanRetry(const FHttpRetryRequestEntry& HttpRetryRequestEntry);
+        // @return true if there is a no formal response to the request
+        // @TODO return true if a variety of 5xx errors are the result of a formal response
+        bool ShouldRetry(const FHttpRetryRequestEntry& HttpRetryRequestEntry);
 
-	// @return true if the retry request has timed out
-	bool HasTimedOut(const FHttpRetryRequestEntry& HttpRetryRequestEntry, const double NowAbsoluteSeconds);
+        // @return true if retry chances have not been exhausted
+        bool CanRetry(const FHttpRetryRequestEntry& HttpRetryRequestEntry);
 
-	/**
-	 * Retry an HTTP request
-	 * @param RequestEntry request to retry
-	 */
-	void RetryHttpRequest(FHttpRetryRequestEntry& RequestEntry);
+        // @return true if the retry request has timed out
+        bool HasTimedOut(const FHttpRetryRequestEntry& HttpRetryRequestEntry, const double NowAbsoluteSeconds);
 
-	// @return number of seconds to lockout for
-	float GetLockoutPeriodSeconds(const FHttpRetryRequestEntry& HttpRetryRequestEntry);
+		/**
+		 * Retry an HTTP request
+		 * @param RequestEntry request to retry
+		 */
+		void RetryHttpRequest(FHttpRetryRequestEntry& RequestEntry);
 
-	// Default configuration for the retry system
-	FRandomFailureRateSetting RandomFailureRate;
-	FRetryLimitCountSetting RetryLimitCountDefault;
-	FRetryTimeoutRelativeSecondsSetting RetryTimeoutRelativeSecondsDefault;
+        // @return number of seconds to lockout for
+        float GetLockoutPeriodSeconds(const FHttpRetryRequestEntry& HttpRetryRequestEntry);
 
-	TArray<FHttpRetryRequestEntry> RequestList;
-	FCriticalSection RequestListLock;
-};
+        // Default configuration for the retry system
+        FRandomFailureRateSetting            RandomFailureRate;
+        FRetryLimitCountSetting              RetryLimitCountDefault;
+        FRetryTimeoutRelativeSecondsSetting  RetryTimeoutRelativeSecondsDefault;
 
+        TArray<FHttpRetryRequestEntry>       RequestList;
+    };
 }
