Fix ghost threads in BPS that don't immediately destroy their context after exiting, leading to runaway thread counts

#rnx

diff --git a/Engine/Source/Runtime/Online/BuildPatchServices/Private/BuildPatchInstaller.cpp b/Engine/Source/Runtime/Online/BuildPatchServices/Private/BuildPatchInstaller.cpp
--- a/Engine/Source/Runtime/Online/BuildPatchServices/Private/BuildPatchInstaller.cpp
+++ b/Engine/Source/Runtime/Online/BuildPatchServices/Private/BuildPatchInstaller.cpp
@@ -504,12 +504,9 @@ namespace BuildPatchServices
 		{
 			InstallerError->SetError(EBuildPatchInstallError::ApplicationClosing, ApplicationClosedErrorCodes::ApplicationClosed);
 		}
-		if (Thread != nullptr)
-		{
-			Thread->WaitForCompletion();
-			delete Thread;
-			Thread = nullptr;
-		}
+
+		CleanupThread();
+
 		if (InstallerAnalytics.IsValid())
 		{
 			InstallerAnalytics->Flush();
@@ -523,6 +520,7 @@ namespace BuildPatchServices
 		if (IsComplete())
 		{
 			ExecuteCompleteDelegate();
+			CleanupThread();
 			bStillTicking = false;
 		}
 		return bStillTicking;
@@ -1835,6 +1833,16 @@ namespace BuildPatchServices
 		return bBackupSuccess;
 	}
 
+	void FBuildPatchInstaller::CleanupThread()
+	{
+		if (Thread != nullptr)
+		{
+			Thread->WaitForCompletion();
+			delete Thread;
+			Thread = nullptr;
+		}
+	}
+
 	double FBuildPatchInstaller::GetDownloadSpeed() const
 	{
 		return DownloadSpeedRecorder->GetAverageSpeed(ConfigHelpers::DownloadSpeedAverageTime());
diff --git a/Engine/Source/Runtime/Online/BuildPatchServices/Private/BuildPatchInstaller.h b/Engine/Source/Runtime/Online/BuildPatchServices/Private/BuildPatchInstaller.h
--- a/Engine/Source/Runtime/Online/BuildPatchServices/Private/BuildPatchInstaller.h
+++ b/Engine/Source/Runtime/Online/BuildPatchServices/Private/BuildPatchInstaller.h
@@ -457,6 +457,11 @@ namespace BuildPatchServices
 		 * Builds the cloud source configuration struct.
 		 */
 		FDownloadConnectionCountConfig BuildConnectionCountConfig();
+
+		/**
+		 * Destroys the internal worker thread
+		 */
+		void CleanupThread();
 	};
 }
 
