Remove ticking from BPS Installer, especially in cases where it leads to a round trip to the game thread.

FChunkDataSizeProvider - made theadsafe
FSpeedRecorder - made queue MPSC since its accessed from multiple threads - this still ticks but it shouldn't block anything
Download Service - made threadsafe and no longer ticks
FInstallerAnalytics - use null implementation if there is no analytics provider
FOptimisedDelta - Move post download processing into a task
FDownloadServiceStatistics - made threadsafe
FFileOperationTracker - add null implementation and flag to turn off file op tracking

diff --git a/Engine/Source/Runtime/Online/BuildPatchServices/Private/BuildPatchInstaller.cpp b/Engine/Source/Runtime/Online/BuildPatchServices/Private/BuildPatchInstaller.cpp
--- a/Engine/Source/Runtime/Online/BuildPatchServices/Private/BuildPatchInstaller.cpp
+++ b/Engine/Source/Runtime/Online/BuildPatchServices/Private/BuildPatchInstaller.cpp
@@ -453,7 +453,7 @@ namespace BuildPatchServices
 		, InstallerError(FInstallerErrorFactory::Create())
 		, Analytics(MoveTemp(InAnalytics))
 		, InstallerAnalytics(FInstallerAnalyticsFactory::Create(Analytics.Get()))
-		, FileOperationTracker(FFileOperationTrackerFactory::Create(FTSTicker::GetCoreTicker()))
+		, FileOperationTracker(Configuration.bTrackFileOperations ? FFileOperationTrackerFactory::Create(FTSTicker::GetCoreTicker()) : FFileOperationTrackerFactory::CreateNull())
 		, MemoryChunkStoreStatistics(FMemoryChunkStoreStatisticsFactory::Create(FileOperationTracker.Get()))
 		, DiskChunkStoreStatistics(FDiskChunkStoreStatisticsFactory::Create(InstallerAnalytics.Get(), FileOperationTracker.Get()))
 		, DownloadSpeedRecorder(FSpeedRecorderFactory::Create())
@@ -467,7 +467,7 @@ namespace BuildPatchServices
 		, CloudChunkSourceStatistics(FCloudChunkSourceStatisticsFactory::Create(InstallerAnalytics.Get(), &BuildProgress, FileOperationTracker.Get()))
 		, FileConstructorStatistics(FFileConstructorStatisticsFactory::Create(DiskReadSpeedRecorder.Get(), DiskWriteSpeedRecorder.Get(), &BuildProgress, FileOperationTracker.Get()))
 		, VerifierStatistics(FVerifierStatisticsFactory::Create(DiskReadSpeedRecorder.Get(), &BuildProgress, FileOperationTracker.Get()))
-		, DownloadService(FDownloadServiceFactory::Create(FTSTicker::GetCoreTicker(), HttpManager.Get(), FileSystem.Get(), DownloadServiceStatistics.Get(), InstallerAnalytics.Get()))
+		, DownloadService(FDownloadServiceFactory::Create(HttpManager.Get(), FileSystem.Get(), DownloadServiceStatistics.Get(), InstallerAnalytics.Get()))
 		, MessagePump(FMessagePumpFactory::Create())
 		, Controllables()
 	{
@@ -478,6 +478,10 @@ namespace BuildPatchServices
 		FString InstallDirectory = Configuration.InstallDirectory;
 		FPaths::NormalizeDirectoryName(InstallDirectory);
 		FPaths::CollapseRelativeDirectories(InstallDirectory);
+
+		TArray<FBuildPatchAppManifestPtr> Manifests;
+		Manifests.Reserve(Configuration.InstallerActions.Num() * 2);
+
 		for (const FInstallerAction& InstallerAction : Configuration.InstallerActions)
 		{
 			FBuildPatchInstallerAction& BuildPatchInstallerAction = InstallerActions.Emplace_GetRef(InstallerAction);
@@ -487,9 +491,17 @@ namespace BuildPatchServices
 				InstallationInfo.Add(InstallDirectory / BuildPatchInstallerAction.GetInstallSubdirectory(), BuildPatchInstallerAction.GetSharedCurrentManifest());
 			}
 			// Cache chunk sizes too
-			ChunkDataSizeProvider->AddManifestData(BuildPatchInstallerAction.TryGetSharedCurrentManifest());
-			ChunkDataSizeProvider->AddManifestData(BuildPatchInstallerAction.TryGetSharedInstallManifest());
+			if (FBuildPatchAppManifestPtr SharedCurrentManifest = BuildPatchInstallerAction.TryGetSharedCurrentManifest())
+			{
+				Manifests.Emplace(MoveTemp(SharedCurrentManifest));
+			}
+			if (FBuildPatchAppManifestPtr SharedInstallManifest = BuildPatchInstallerAction.TryGetSharedInstallManifest())
+			{
+				Manifests.Emplace(MoveTemp(SharedInstallManifest));
+			}
 		}
+
+		ChunkDataSizeProvider->AddManifestData(Manifests);
 	}
 
 	FBuildPatchInstaller::~FBuildPatchInstaller()
@@ -687,22 +699,22 @@ namespace BuildPatchServices
 		}
 
 		// do the delta optimization
-		FOptimisedDeltaDependencies OptimisedDeltaDependencies = InstallerHelpers::BuildOptimisedDeltaDependencies(DownloadService);
-		typedef TTuple<FBuildPatchInstallerAction&, IOptimisedDelta*> FManifestInfoDeltaPair;
+		typedef TTuple<FBuildPatchInstallerAction&, TUniquePtr<IOptimisedDelta>> FManifestInfoDeltaPair;
 		TArray<FManifestInfoDeltaPair> RunningOptimisedDeltas;
 		for (FBuildPatchInstallerAction& InstallerAction : InstallerActions)
 		{
 			if (InstallerAction.IsUpdate())
 			{
-				IOptimisedDelta* OptimisedDelta = FOptimisedDeltaFactory::Create(ConfigHelpers::BuildOptimisedDeltaConfig(Configuration, InstallerAction), OptimisedDeltaDependencies);
-				OptimisedDeltas.Add(TUniquePtr<IOptimisedDelta>(OptimisedDelta));
-				RunningOptimisedDeltas.Add(FManifestInfoDeltaPair{ InstallerAction, OptimisedDelta });
+				TUniquePtr<IOptimisedDelta> OptimisedDelta(FOptimisedDeltaFactory::Create(
+					ConfigHelpers::BuildOptimisedDeltaConfig(Configuration, InstallerAction), 
+					InstallerHelpers::BuildOptimisedDeltaDependencies(DownloadService)));
+				RunningOptimisedDeltas.Add(FManifestInfoDeltaPair{ InstallerAction, MoveTemp(OptimisedDelta) });
 			}
 		}
 		for (FManifestInfoDeltaPair& RunningOptimisedDelta : RunningOptimisedDeltas)
 		{
 			FBuildPatchInstallerAction& InstallerAction = RunningOptimisedDelta.Get<0>();
-			IOptimisedDelta* const OptimisedDelta = RunningOptimisedDelta.Get<1>();
+			const TUniquePtr<IOptimisedDelta>& OptimisedDelta = RunningOptimisedDelta.Get<1>();
 			const IOptimisedDelta::FResultValueOrError& OptimisedDeltaResult = OptimisedDelta->GetResult();
 			PreviousTotalDownloadRequired.Add(OptimisedDelta->GetMetaDownloadSize());
 			// The OptimiseDelta class handles policy, so if we get a nullptr back, that is a hard error.
@@ -718,6 +730,7 @@ namespace BuildPatchServices
 				InstallerAction.SetDeltaManifest(OptimisedDeltaResult.GetValue().ToSharedRef());
 			}
 		}
+		RunningOptimisedDeltas.Empty();
 
 		// We can now build out any systems that need late construction but can survive between retries.
 		ManifestSet.Reset(FBuildManifestSetFactory::Create(InstallerActions));
@@ -729,14 +742,16 @@ namespace BuildPatchServices
 			Controllables.Add(Verifier.Get());
 		}
 
-		// Queue update to chunk data size cache on main thread
-		AsyncHelpers::ExecuteOnGameThread<void>([this]()
+		// Update to chunk data size cache since we may have delta manifests now
 		{
+			TArray<FBuildPatchAppManifestPtr> ActionManifests;
+			ActionManifests.Reserve(InstallerActions.Num());
 			for (const FBuildPatchInstallerAction& InstallerAction : InstallerActions)
 			{
-				ChunkDataSizeProvider->AddManifestData(InstallerAction.TryGetSharedInstallManifest());
+				ActionManifests.Emplace(InstallerAction.TryGetSharedInstallManifest());
 			}
-		}).Wait();
+			ChunkDataSizeProvider->AddManifestData(ActionManifests);
+		}
 
 		// Init build statistics that are known.
 		{
@@ -2059,13 +2074,13 @@ namespace BuildPatchServices
 	{
 		check(IsInGameThread());
 		check(MessageHandler != nullptr);
-		MessageHandlers.AddUnique(MessageHandler);
+		MessagePump->RegisterMessageHandler(MessageHandler);
 	}
 
 	void FBuildPatchInstaller::UnregisterMessageHandler(FMessageHandler* MessageHandler)
 	{
 		check(IsInGameThread());
-		MessageHandlers.Remove(MessageHandler);
+		MessagePump->UnregisterMessageHandler(MessageHandler);
 	}
 
 	void FBuildPatchInstaller::ExecuteCompleteDelegate()
@@ -2103,7 +2118,7 @@ namespace BuildPatchServices
 	void FBuildPatchInstaller::PumpMessages()
 	{
 		check(IsInGameThread());
-		MessagePump->PumpMessages(MessageHandlers);
+		MessagePump->PumpMessages();
 	}
 
 	void FBuildPatchInstaller::WaitForThread() const
diff --git a/Engine/Source/Runtime/Online/BuildPatchServices/Private/BuildPatchInstaller.h b/Engine/Source/Runtime/Online/BuildPatchServices/Private/BuildPatchInstaller.h
--- a/Engine/Source/Runtime/Online/BuildPatchServices/Private/BuildPatchInstaller.h
+++ b/Engine/Source/Runtime/Online/BuildPatchServices/Private/BuildPatchInstaller.h
@@ -182,9 +182,6 @@ namespace BuildPatchServices
 		// The message pump controller.
 		TUniquePtr<IMessagePump> MessagePump;
 
-		// Holds the optimised delta interfaces created during Initialize.
-		TArray<TUniquePtr<IOptimisedDelta>> OptimisedDeltas;
-
 		// The interface for manifest data aggregation.
 		TUniquePtr<IBuildManifestSet> ManifestSet;
 
@@ -194,9 +191,6 @@ namespace BuildPatchServices
 		// List of controllable classes that have been constructed.
 		TArray<IControllable*> Controllables;
 
-		// List of message handlers that have been registered.
-		TArray<FMessageHandler*> MessageHandlers;
-
 		// Stage timers for build stats.
 		FProcessTimer InitializeTimer;
 		FProcessTimer ConstructTimer;
diff --git a/Engine/Source/Runtime/Online/BuildPatchServices/Private/BuildPatchSettings.cpp b/Engine/Source/Runtime/Online/BuildPatchServices/Private/BuildPatchSettings.cpp
--- a/Engine/Source/Runtime/Online/BuildPatchServices/Private/BuildPatchSettings.cpp
+++ b/Engine/Source/Runtime/Online/BuildPatchServices/Private/BuildPatchSettings.cpp
@@ -151,6 +151,7 @@ namespace BuildPatchServices
 		, DeltaPolicy(EDeltaPolicy::Skip)
 		, bRunRequiredPrereqs(true)
 		, bAllowConcurrentExecution(false)
+		, bTrackFileOperations(true)
 	{
 	}
 
@@ -166,6 +167,7 @@ namespace BuildPatchServices
 		, DeltaPolicy(MoveFrom.DeltaPolicy)
 		, bRunRequiredPrereqs(MoveFrom.bRunRequiredPrereqs)
 		, bAllowConcurrentExecution(MoveFrom.bAllowConcurrentExecution)
+		, bTrackFileOperations(MoveFrom.bTrackFileOperations)
 	{
 	}
 
@@ -181,6 +183,7 @@ namespace BuildPatchServices
 		, DeltaPolicy(CopyFrom.DeltaPolicy)
 		, bRunRequiredPrereqs(CopyFrom.bRunRequiredPrereqs)
 		, bAllowConcurrentExecution(CopyFrom.bAllowConcurrentExecution)
+		, bTrackFileOperations(CopyFrom.bTrackFileOperations)
 	{
 	}
 
diff --git a/Engine/Source/Runtime/Online/BuildPatchServices/Private/Common/ChunkDataSizeProvider.cpp b/Engine/Source/Runtime/Online/BuildPatchServices/Private/Common/ChunkDataSizeProvider.cpp
--- a/Engine/Source/Runtime/Online/BuildPatchServices/Private/Common/ChunkDataSizeProvider.cpp
+++ b/Engine/Source/Runtime/Online/BuildPatchServices/Private/Common/ChunkDataSizeProvider.cpp
@@ -13,49 +13,104 @@ namespace BuildPatchServices
 		// Begin IDataSizeProvider
 		virtual uint64 GetDownloadSize(const FString& Identifier) const override
 		{
-			checkSlow(IsInGameThread());
-
 			FGuid ChunkGUID;
-			bool bFoundGUID = FBuildPatchUtils::GetGUIDFromFilename(Identifier, ChunkGUID);
+			if (!FBuildPatchUtils::GetGUIDFromFilename(Identifier, ChunkGUID))
+			{
+				return INDEX_NONE;
+			}
 
-			uint64 DownloadSize = INDEX_NONE;
-			if (bFoundGUID)
+			FScopeLock Lock(&ManifestsCS);
+			return GetDownloadSizeUnsafe(ChunkGUID);
+		}
+
+		virtual void GetDownloadSize(TConstArrayView<FString> Identifiers, TArray<uint64>& OutSizes) const override
+		{
+			TArray<FGuid> ChunkGUIDs;
+			ChunkGUIDs.Reserve(Identifiers.Num());
+			for (const FString& Identifier : Identifiers)
 			{
-				for (int32 i = Manifests.Num() - 1; i >= 0; --i)
+				FGuid ChunkGUID;
+				if (!FBuildPatchUtils::GetGUIDFromFilename(Identifier, ChunkGUID))
 				{
-					const FBuildPatchAppManifestPtr& Manifest = Manifests[i];
-					const BuildPatchServices::FChunkInfo* ChunkInfo = Manifest->GetChunkInfo(ChunkGUID);
-					if (ChunkInfo)
-					{
-						DownloadSize = ChunkInfo->FileSize;
-						break;
-					}
+					ChunkGUID.Invalidate();
 				}
+				ChunkGUIDs.Emplace(ChunkGUID);
+			}
+
+			OutSizes.Reserve(OutSizes.Num() + Identifiers.Num());
+
+			FScopeLock Lock(&ManifestsCS);
+			for (const FGuid& ChunkGUID : ChunkGUIDs)
+			{
+				if (!ChunkGUID.IsValid())
+				{
+					OutSizes.Emplace(INDEX_NONE);
+					continue;
+				}
+
+				OutSizes.Emplace(GetDownloadSizeUnsafe(ChunkGUID));
 			}
-			return DownloadSize;
 		}
 		// End IDataSizeProvider
 
 		// Begin IChunkDataSizeProvider
 		virtual void AddManifestData(FBuildPatchAppManifestPtr Manifest) override
 		{
-			check(IsInGameThread());
+			if (!Manifest)
+			{
+				return;
+			}
+
+			FScopeLock Lock(&ManifestsCS);
+			AddManifestDataUnsafe(MoveTemp(Manifest));
+		}
 
-			if (Manifest)
+		virtual void AddManifestData(TConstArrayView<FBuildPatchAppManifestPtr> InManifests) override
+		{
+			FScopeLock Lock(&ManifestsCS);
+			for (FBuildPatchAppManifestPtr Manifest : InManifests)
 			{
-				int32 Index = Manifests.Find(Manifest);
-				if (Index != INDEX_NONE)
+				if (Manifest)
 				{
-					Manifests.RemoveAt(Index, 1, false);
+					AddManifestDataUnsafe(MoveTemp(Manifest));
 				}
-
-				Manifests.Add(MoveTemp(Manifest));  // Last manifest added always goes on the end
 			}
 		}
 		// End IChunkDataSizeProvider
 
+		uint64 GetDownloadSizeUnsafe(const FGuid& ChunkGUID) const
+		{
+			checkSlow(ChunkGUID.IsValid());
+
+			uint64 DownloadSize = INDEX_NONE;
+			for (int32 i = Manifests.Num() - 1; i >= 0; --i)
+			{
+				const FBuildPatchAppManifestPtr& Manifest = Manifests[i];
+				const BuildPatchServices::FChunkInfo* ChunkInfo = Manifest->GetChunkInfo(ChunkGUID);
+				if (ChunkInfo)
+				{
+					DownloadSize = ChunkInfo->FileSize;
+					break;
+				}
+			}
+			return DownloadSize;
+		}
+
+		void AddManifestDataUnsafe(FBuildPatchAppManifestPtr Manifest)
+		{
+			checkSlow(Manifest);
+			int32 Index = Manifests.Find(Manifest);
+			if (Index != INDEX_NONE)
+			{
+				Manifests.RemoveAt(Index, 1, false);
+			}
+
+			Manifests.Add(MoveTemp(Manifest));  // Last manifest added always goes on the end
+		}
+
 	private:
 		TArray<FBuildPatchAppManifestPtr> Manifests;
+		mutable FCriticalSection ManifestsCS;
 	};
 	
 	IChunkDataSizeProvider* FChunkDataSizeProviderFactory::Create()
diff --git a/Engine/Source/Runtime/Online/BuildPatchServices/Private/Common/ChunkDataSizeProvider.h b/Engine/Source/Runtime/Online/BuildPatchServices/Private/Common/ChunkDataSizeProvider.h
--- a/Engine/Source/Runtime/Online/BuildPatchServices/Private/Common/ChunkDataSizeProvider.h
+++ b/Engine/Source/Runtime/Online/BuildPatchServices/Private/Common/ChunkDataSizeProvider.h
@@ -3,6 +3,7 @@
 
 #include "Common/DataSizeProvider.h"
 #include "BuildPatchManifest.h"
+#include "Containers/ContainersFwd.h"
 
 namespace BuildPatchServices
 {
@@ -11,6 +12,7 @@ namespace BuildPatchServices
 	{
 	public:
 		virtual void AddManifestData(FBuildPatchAppManifestPtr Manifest) = 0;
+		virtual void AddManifestData(TConstArrayView<FBuildPatchAppManifestPtr> Manifests) = 0;
 	};
 
 	class FChunkDataSizeProviderFactory
diff --git a/Engine/Source/Runtime/Online/BuildPatchServices/Private/Common/DataSizeProvider.h b/Engine/Source/Runtime/Online/BuildPatchServices/Private/Common/DataSizeProvider.h
--- a/Engine/Source/Runtime/Online/BuildPatchServices/Private/Common/DataSizeProvider.h
+++ b/Engine/Source/Runtime/Online/BuildPatchServices/Private/Common/DataSizeProvider.h
@@ -2,6 +2,7 @@
 #pragma once
 
 #include "CoreMinimal.h"
+#include "Containers/ContainersFwd.h"
 
 namespace BuildPatchServices
 {
@@ -11,5 +12,6 @@ namespace BuildPatchServices
 		virtual ~IDataSizeProvider() {}
 
 		virtual uint64 GetDownloadSize(const FString& Uri) const = 0;
+		virtual void GetDownloadSize(TConstArrayView<FString> InUris, TArray<uint64>& OutSizes) const = 0;
 	};
 }
\ No newline at end of file
diff --git a/Engine/Source/Runtime/Online/BuildPatchServices/Private/Common/SpeedRecorder.cpp b/Engine/Source/Runtime/Online/BuildPatchServices/Private/Common/SpeedRecorder.cpp
--- a/Engine/Source/Runtime/Online/BuildPatchServices/Private/Common/SpeedRecorder.cpp
+++ b/Engine/Source/Runtime/Online/BuildPatchServices/Private/Common/SpeedRecorder.cpp
@@ -108,7 +108,7 @@ namespace BuildPatchServices
 	private:
 		const uint64 CyclesConfidenceFactor;
 		const uint64 SizeConfidenceFactor;
-		TQueue<ISpeedRecorder::FRecord> RecordsQueue;
+		TQueue<ISpeedRecorder::FRecord, EQueueMode::Mpsc> RecordsQueue;
 		TArray<ISpeedRecorder::FRecord> Records;
 		mutable TArray<ISpeedRecorder::FRecord> Temp;
 
diff --git a/Engine/Source/Runtime/Online/BuildPatchServices/Private/Diffing/DiffManifests.cpp b/Engine/Source/Runtime/Online/BuildPatchServices/Private/Diffing/DiffManifests.cpp
--- a/Engine/Source/Runtime/Online/BuildPatchServices/Private/Diffing/DiffManifests.cpp
+++ b/Engine/Source/Runtime/Online/BuildPatchServices/Private/Diffing/DiffManifests.cpp
@@ -332,7 +332,7 @@ namespace BuildPatchServices
 		, DownloadSpeedRecorder(FSpeedRecorderFactory::Create())
 		, InstallerAnalytics(FInstallerAnalyticsFactory::Create(nullptr))
 		, DownloadServiceStatistics(FDownloadServiceStatisticsFactory::Create(DownloadSpeedRecorder.Get(), ChunkDataSizeProvider.Get(), InstallerAnalytics.Get()))
-		, DownloadService(FDownloadServiceFactory::Create(CoreTicker, HttpManager.Get(), FileSystem.Get(), DownloadServiceStatistics.Get(), InstallerAnalytics.Get()))
+		, DownloadService(FDownloadServiceFactory::Create(HttpManager.Get(), FileSystem.Get(), DownloadServiceStatistics.Get(), InstallerAnalytics.Get()))
 		, StatsCollector(FStatsCollectorFactory::Create())
 		, bShouldRun(true)
 		, RequestIdManifestA(INDEX_NONE)
@@ -417,7 +417,7 @@ namespace BuildPatchServices
 			OptimisedDeltaConfiguration.CloudDirectories = { FPaths::GetPath(Configuration.ManifestBUri) };
 			FOptimisedDeltaDependencies OptimisedDeltaDependencies;
 			OptimisedDeltaDependencies.DownloadService = DownloadService.Get();
-			TUniquePtr<IOptimisedDelta> OptimisedDelta(FOptimisedDeltaFactory::Create(OptimisedDeltaConfiguration, OptimisedDeltaDependencies));
+			TUniquePtr<IOptimisedDelta> OptimisedDelta(FOptimisedDeltaFactory::Create(OptimisedDeltaConfiguration, MoveTemp(OptimisedDeltaDependencies)));
 			ManifestB = OptimisedDelta->GetResult().GetValue();
 			const int32 MetaDownloadBytes = OptimisedDelta->GetMetaDownloadSize();
 
diff --git a/Engine/Source/Runtime/Online/BuildPatchServices/Private/Generation/ChunkDeltaOptimiser.cpp b/Engine/Source/Runtime/Online/BuildPatchServices/Private/Generation/ChunkDeltaOptimiser.cpp
--- a/Engine/Source/Runtime/Online/BuildPatchServices/Private/Generation/ChunkDeltaOptimiser.cpp
+++ b/Engine/Source/Runtime/Online/BuildPatchServices/Private/Generation/ChunkDeltaOptimiser.cpp
@@ -586,7 +586,6 @@ namespace BuildPatchServices
 		TUniquePtr<IDownloadService> DownloadService;
 		TUniquePtr<IMessagePump> MessagePump;
 		TUniquePtr<FStatsCollector> StatsCollector;
-		TArray<FMessageHandler*> MessageHandlers;
 		FThreadSafeBool bShouldRun;
 		FThreadSafeBool bSuccess;
 
@@ -610,7 +609,7 @@ namespace BuildPatchServices
 		, DownloadSpeedRecorder(FSpeedRecorderFactory::Create())
 		, InstallerAnalytics(FInstallerAnalyticsFactory::Create(nullptr))
 		, DownloadServiceStatistics(FDownloadServiceStatisticsFactory::Create(DownloadSpeedRecorder.Get(), ChunkDataSizeProvider.Get(), InstallerAnalytics.Get()))
-		, DownloadService(FDownloadServiceFactory::Create(CoreTicker, HttpManager.Get(), FileSystem.Get(), DownloadServiceStatistics.Get(), InstallerAnalytics.Get()))
+		, DownloadService(FDownloadServiceFactory::Create(HttpManager.Get(), FileSystem.Get(), DownloadServiceStatistics.Get(), InstallerAnalytics.Get()))
 		, MessagePump(FMessagePumpFactory::Create())
 		, StatsCollector(FStatsCollectorFactory::Create())
 		, bShouldRun(true)
@@ -669,7 +668,7 @@ namespace BuildPatchServices
 			FTSTicker::GetCoreTicker().Tick(DeltaTime);
 
 			// Message pump.
-			MessagePump->PumpMessages(MessageHandlers);
+			MessagePump->PumpMessages();
 
 			// Log collected stats.
 			GLog->FlushThreadedLogs();
diff --git a/Engine/Source/Runtime/Online/BuildPatchServices/Private/Generation/PackageChunkData.cpp b/Engine/Source/Runtime/Online/BuildPatchServices/Private/Generation/PackageChunkData.cpp
--- a/Engine/Source/Runtime/Online/BuildPatchServices/Private/Generation/PackageChunkData.cpp
+++ b/Engine/Source/Runtime/Online/BuildPatchServices/Private/Generation/PackageChunkData.cpp
@@ -135,7 +135,6 @@ namespace BuildPatchServices
 		FDownloadProgressDelegate DownloadProgressDelegate;
 
 		// Process control.
-		TArray<FMessageHandler*> MessageHandlers;
 		bool bManifestsProcessed;
 		FThreadSafeBool bShouldRun;
 		FThreadSafeBool bSuccess;
@@ -205,7 +204,7 @@ namespace BuildPatchServices
 		, ChunkDataSizeProvider(FChunkDataSizeProviderFactory::Create())
 		, InstallerAnalytics(FInstallerAnalyticsFactory::Create(nullptr))
 		, DownloadServiceStatistics(FDownloadServiceStatisticsFactory::Create(DownloadSpeedRecorder.Get(), ChunkDataSizeProvider.Get(), InstallerAnalytics.Get()))
-		, DownloadService(FDownloadServiceFactory::Create(CoreTicker, HttpManager.Get(), FileSystem.Get(), DownloadServiceStatistics.Get(), InstallerAnalytics.Get()))
+		, DownloadService(FDownloadServiceFactory::Create(HttpManager.Get(), FileSystem.Get(), DownloadServiceStatistics.Get(), InstallerAnalytics.Get()))
 		, FileOperationTracker(FFileOperationTrackerFactory::Create(CoreTicker))
 	{
 		// Make sure the cloud chunk source gets the abort signal if an error occurred.
@@ -253,7 +252,7 @@ namespace BuildPatchServices
 			FTSTicker::GetCoreTicker().Tick(DeltaTime);
 
 			// Message pump.
-			MessagePump->PumpMessages(MessageHandlers);
+			MessagePump->PumpMessages();
 
 			// Flush any threaded logging.
 			GLog->FlushThreadedLogs();
@@ -328,7 +327,7 @@ namespace BuildPatchServices
 				FOptimisedDeltaDependencies OptimisedDeltaDependencies;
 				OptimisedDeltaDependencies.DownloadService = DownloadService.Get();
 				OptimisedDeltaDependencies.OnComplete = MakeOptimiseCompleteDelegate(&FPackageChunks::HandleManifestSelection);
-				OptimisedDelta.Reset(FOptimisedDeltaFactory::Create(OptimisedDeltaConfiguration, OptimisedDeltaDependencies));
+				OptimisedDelta.Reset(FOptimisedDeltaFactory::Create(OptimisedDeltaConfiguration, MoveTemp(OptimisedDeltaDependencies)));
 			}
 		}
 	}
diff --git a/Engine/Source/Runtime/Online/BuildPatchServices/Private/Installer/BuildInstallStreamer.cpp b/Engine/Source/Runtime/Online/BuildPatchServices/Private/Installer/BuildInstallStreamer.cpp
--- a/Engine/Source/Runtime/Online/BuildPatchServices/Private/Installer/BuildInstallStreamer.cpp
+++ b/Engine/Source/Runtime/Online/BuildPatchServices/Private/Installer/BuildInstallStreamer.cpp
@@ -65,7 +65,7 @@ namespace BuildPatchServices
 		, CloudChunkSourceStatistics(FCloudChunkSourceStatisticsFactory::Create(InstallerAnalytics.Get(), &BuildProgress, FileOperationTracker.Get()))
 		, FileConstructorStatistics(FFileConstructorStatisticsFactory::Create(FileReadSpeedRecorder.Get(), FileWriteSpeedRecorder.Get(), &BuildProgress, FileOperationTracker.Get()))
 		, DownloadConnectionCount(FDownloadConnectionCountFactory::Create(BuildConnectionCountConfig(), DownloadServiceStatistics.Get()))
-		, DownloadService(FDownloadServiceFactory::Create(FTSTicker::GetCoreTicker(), HttpManager.Get(), FileSystem.Get(), DownloadServiceStatistics.Get(), InstallerAnalytics.Get()))
+		, DownloadService(FDownloadServiceFactory::Create(HttpManager.Get(), FileSystem.Get(), DownloadServiceStatistics.Get(), InstallerAnalytics.Get()))
 		, MessagePump(FMessagePumpFactory::Create())
 		, bIsShuttingDown(false)
 		, RequestTrigger(FPlatformProcess::GetSynchEventFromPool(false))
@@ -150,13 +150,8 @@ namespace BuildPatchServices
 
 	void FBuildInstallStreamer::Initialise()
 	{
-		// Queue update to chunk data size cache on main thread
-		AsyncHelpers::ExecuteOnCustomThread<void>(
-			[this]()
-			{
-				ChunkDataSizeProvider->AddManifestData(BuildPatchManifest);
-			},
-			TickQueue).Wait();
+		// update to chunk data size cache
+		ChunkDataSizeProvider->AddManifestData(BuildPatchManifest);
 	}
 
 	void FBuildInstallStreamer::RequestWorkerThread()
@@ -281,7 +276,7 @@ namespace BuildPatchServices
 
 				TProcessTimer<class FStatsCollector, false> RequestTimer;
 				RequestTimer.Start();
-				AsyncHelpers::ExecuteOnCustomThread<void>([this]() { DownloadServiceStatistics->Reset(); }, TickQueue);
+				DownloadServiceStatistics->Reset();
 				TSharedPtr<IVirtualFileCache, ESPMode::ThreadSafe> VirtualFileCache = IVirtualFileCache::CreateVirtualFileCache();
 				InstallerError->Reset();
 
@@ -427,7 +422,7 @@ namespace BuildPatchServices
 	{
 		const bool bKeepTicking = true;
 
-		MessagePump->PumpMessages(MessageHandlers);
+		MessagePump->PumpMessages();
 
 		TFunction<void()> TickFunc;
 		while (TickQueue.Dequeue(TickFunc))
diff --git a/Engine/Source/Runtime/Online/BuildPatchServices/Private/Installer/BuildInstallStreamer.h b/Engine/Source/Runtime/Online/BuildPatchServices/Private/Installer/BuildInstallStreamer.h
--- a/Engine/Source/Runtime/Online/BuildPatchServices/Private/Installer/BuildInstallStreamer.h
+++ b/Engine/Source/Runtime/Online/BuildPatchServices/Private/Installer/BuildInstallStreamer.h
@@ -88,7 +88,6 @@ namespace BuildPatchServices
 		TUniquePtr<IDownloadConnectionCount> DownloadConnectionCount;
 		TUniquePtr<IDownloadService> DownloadService;
 		TUniquePtr<IMessagePump> MessagePump;
-		TArray<FMessageHandler*> MessageHandlers;
 		FThreadSafeBool bIsShuttingDown;
 		TFuture<void> RequestWorker;
 		TFuture<void> CloudWorker;
diff --git a/Engine/Source/Runtime/Online/BuildPatchServices/Private/Installer/CloudChunkSource.cpp b/Engine/Source/Runtime/Online/BuildPatchServices/Private/Installer/CloudChunkSource.cpp
--- a/Engine/Source/Runtime/Online/BuildPatchServices/Private/Installer/CloudChunkSource.cpp
+++ b/Engine/Source/Runtime/Online/BuildPatchServices/Private/Installer/CloudChunkSource.cpp
@@ -443,10 +443,8 @@ namespace BuildPatchServices
 
 		// Chunk Uri Processing
 		typedef TTuple<FGuid, FChunkUriResponse> FGuidUriResponse;
-		typedef TQueue<FGuidUriResponse> FGuidUriResponseQueue;
-		TSharedRef<FGuidUriResponseQueue, ESPMode::ThreadSafe> ChunkUriResponsesRef = MakeShareable(new FGuidUriResponseQueue());
-		TWeakPtr<FGuidUriResponseQueue, ESPMode::ThreadSafe> WeakChunkUriResponses = ChunkUriResponsesRef;
-		FGuidUriResponseQueue& ChunkUriResponses = ChunkUriResponsesRef.Get();
+		typedef TQueue<FGuidUriResponse, EQueueMode::Mpsc> FGuidUriResponseQueue; // use Mpsc, message pump callback may be on this thread or message pump thread
+		TSharedRef<FGuidUriResponseQueue> ChunkUriResponsesRef = MakeShared<FGuidUriResponseQueue>();
 		TSet<FGuid> RequestedChunkUris;
 		TMap<FGuid, FChunkUriResponse> ChunkUris;
 
@@ -484,7 +482,10 @@ namespace BuildPatchServices
 				}
 			}
 			// Select the next X chunks that are for downloading, so we can request URIs.
-			TFunction<bool(const FGuid&)> SelectPredicate = [&TotalRequiredChunks, &RequestedChunkUris](const FGuid& ChunkId) { return TotalRequiredChunks.Contains(ChunkId) && !RequestedChunkUris.Contains(ChunkId); };
+			TFunction<bool(const FGuid&)> SelectPredicate = [&TotalRequiredChunks, &RequestedChunkUris](const FGuid& ChunkId) 
+			{ 
+				return TotalRequiredChunks.Contains(ChunkId) && !RequestedChunkUris.Contains(ChunkId); 
+			};
 			TArray<FGuid> ChunkUrisToRequest = ChunkReferenceTracker->SelectFromNextReferences(Configuration.PreFetchMaximum, SelectPredicate);
 			for (const FGuid& ChunkUriToRequest : ChunkUrisToRequest)
 			{
@@ -496,18 +497,14 @@ namespace BuildPatchServices
 				ChunkUriRequest.RelativePath = ManifestSet->GetDataFilename(ChunkUriToRequest);
 				ChunkUriRequest.RelativePath.RemoveFromStart(TEXT("/"));
 
-				MessagePump->SendRequest(ChunkUriRequest, [WeakChunkUriResponses, ChunkUriToRequest](FChunkUriResponse Response)
+				MessagePump->SendRequest(ChunkUriRequest, [ChunkUriResponsesRef, ChunkUriToRequest](FChunkUriResponse Response)
 				{
-					TSharedPtr<FGuidUriResponseQueue, ESPMode::ThreadSafe> ChunkUriResponsesPtr = WeakChunkUriResponses.Pin();
-					if (ChunkUriResponsesPtr.IsValid())
-					{
-						ChunkUriResponsesPtr->Enqueue(FGuidUriResponse(ChunkUriToRequest, MoveTemp(Response)));
-					}
+					ChunkUriResponsesRef->Enqueue(FGuidUriResponse(ChunkUriToRequest, MoveTemp(Response)));
 				});
 			}
 			// Process new chunk uri responses.
 			FGuidUriResponse ChunkUriResponse;
-			while (ChunkUriResponses.Dequeue(ChunkUriResponse))
+			for (FGuidUriResponseQueue& ChunkUriResponses = ChunkUriResponsesRef.Get(); ChunkUriResponses.Dequeue(ChunkUriResponse);)
 			{
 				ChunkUris.Add(MoveTemp(ChunkUriResponse.Get<0>()), MoveTemp(ChunkUriResponse.Get<1>()));
 			}
diff --git a/Engine/Source/Runtime/Online/BuildPatchServices/Private/Installer/DownloadConnectionCount.cpp b/Engine/Source/Runtime/Online/BuildPatchServices/Private/Installer/DownloadConnectionCount.cpp
--- a/Engine/Source/Runtime/Online/BuildPatchServices/Private/Installer/DownloadConnectionCount.cpp
+++ b/Engine/Source/Runtime/Online/BuildPatchServices/Private/Installer/DownloadConnectionCount.cpp
@@ -33,8 +33,10 @@ namespace BuildPatchServices
 
 	public:
 		FDownloadConnectionCount(FDownloadConnectionCountConfig InConfiguration, IDownloadServiceStatistics* InDownloadStatistics);
-		~FDownloadConnectionCount();
-		uint32 GetAdjustedCount(uint32 NumProcessing, EBuildPatchDownloadHealth CurrentHealth);
+		virtual ~FDownloadConnectionCount() override;
+
+		// IDownloadConnectionCount Interface
+		virtual uint32 GetAdjustedCount(uint32 NumProcessing, EBuildPatchDownloadHealth CurrentHealth) override;
 
 	private:
 		EDownloadSpeedStatus GetSpeedStatus(double Speed) const;
diff --git a/Engine/Source/Runtime/Online/BuildPatchServices/Private/Installer/DownloadService.cpp b/Engine/Source/Runtime/Online/BuildPatchServices/Private/Installer/DownloadService.cpp
--- a/Engine/Source/Runtime/Online/BuildPatchServices/Private/Installer/DownloadService.cpp
+++ b/Engine/Source/Runtime/Online/BuildPatchServices/Private/Installer/DownloadService.cpp
@@ -4,6 +4,7 @@
 #include "HAL/ThreadSafeBool.h"
 #include "Misc/ScopeLock.h"
 #include "Core/AsyncHelpers.h"
+#include "Tasks/Task.h"
 #include "Interfaces/IHttpRequest.h"
 #include "Interfaces/IHttpResponse.h"
 #include "Common/StatsCollector.h"
@@ -25,140 +26,59 @@ namespace BuildPatchServices
 
 	struct FFileRequest
 	{
-	public:
-		FFileRequest();
-
-	public:
-		TFuture<void> Future;
-		FThreadSafeBool ShouldCancel;
-	};
-
-	FFileRequest::FFileRequest()
-		: Future()
-		, ShouldCancel(false)
-	{
-	}
-
-	class FDownloadBase
-		: public IDownload
-	{
-	public:
-		FDownloadBase(IDownloadServiceStat::FDownloadRecord&& InDownloadRecord);
-
-	public:
-		IDownloadServiceStat::FDownloadRecord DownloadRecord;
+		UE::Tasks::FTask Task;
+		UE::Tasks::FCancellationToken ShouldCancel;
 	};
-	typedef TSharedRef<FDownloadBase, ESPMode::ThreadSafe> FDownloadBaseRef;
-
-	FDownloadBase::FDownloadBase(IDownloadServiceStat::FDownloadRecord&& InDownloadRecord)
-		: DownloadRecord(MoveTemp(InDownloadRecord))
-	{
-	}
 
 	class FHttpDownload
-		: public FDownloadBase
+		: public IDownload
 	{
 	public:
-		FHttpDownload(const FHttpResponsePtr& HttpResponse, bool bSuccess, IDownloadServiceStat::FDownloadRecord&& DownloadRecord);
-		~FHttpDownload();
+		FHttpDownload(FHttpRequestPtr InHttpRequest, bool bInSuccess)
+			: HttpRequest(MoveTemp(InHttpRequest))
+			, HttpResponse(HttpRequest ? HttpRequest->GetResponse() : nullptr)
+			, bSuccess(bInSuccess)
+		{}
 
 		// IDownload interface begin.
-		virtual bool RequestSuccessful() const override;
-		virtual bool ResponseSuccessful() const override;
-		virtual int32 GetResponseCode() const override;
-		virtual const TArray<uint8>& GetData() const override;
+		virtual bool RequestSuccessful() const override { return bSuccess; }
+		virtual bool ResponseSuccessful() const override { return EHttpResponseCodes::IsOk(GetResponseCode()); }
+		virtual int32 GetResponseCode() const override { return HttpResponse.IsValid() ? HttpResponse->GetResponseCode() : INDEX_NONE; }
+		virtual const TArray<uint8>& GetData() const override { return HttpResponse.IsValid() ? HttpResponse->GetContent() : NoData; }
 		// IDownload interface end.
 
 	private:
+		// Make sure to hold on the the request pointer because the response could reference it
+		FHttpRequestPtr HttpRequest;
 		FHttpResponsePtr HttpResponse;
-		bool bSuccess;
+		bool bSuccess = false;
 	};
 
-	FHttpDownload::FHttpDownload(const FHttpResponsePtr& InHttpResponse, bool bInSuccess, IDownloadServiceStat::FDownloadRecord&& InDownloadRecord)
-		: FDownloadBase(MoveTemp(InDownloadRecord))
-		, HttpResponse(InHttpResponse)
-		, bSuccess(bInSuccess)
-	{
-	}
-
-	FHttpDownload::~FHttpDownload()
-	{
-	}
-
-	bool FHttpDownload::RequestSuccessful() const
-	{
-		return bSuccess;
-	}
-
-	bool FHttpDownload::ResponseSuccessful() const
-	{
-		return EHttpResponseCodes::IsOk(GetResponseCode());
-	}
-
-	int32 FHttpDownload::GetResponseCode() const
-	{
-		return HttpResponse.IsValid() ? HttpResponse->GetResponseCode() : INDEX_NONE;
-	}
-
-	const TArray<uint8>& FHttpDownload::GetData() const
-	{
-		return HttpResponse.IsValid() ? HttpResponse->GetContent() : NoData;
-	}
-
 	class FFileDownload
-		: public FDownloadBase
+		: public IDownload
 	{
 	public:
-		FFileDownload(TArray<uint8> DataArray, bool bSuccess, IDownloadServiceStat::FDownloadRecord&& DownloadRecord);
-		~FFileDownload();
+		FFileDownload(TArray<uint8> InDataArray, bool bInSuccess)
+			: DataArray(MoveTemp(InDataArray))
+			, bSuccess(bInSuccess)
+		{}
 
 		// IDownload interface begin.
-		virtual bool RequestSuccessful() const override;
-		virtual bool ResponseSuccessful() const override;
-		virtual int32 GetResponseCode() const override;
-		virtual const TArray<uint8>& GetData() const override;
+		virtual bool RequestSuccessful() const override { return bSuccess; }
+		virtual bool ResponseSuccessful() const override { return EHttpResponseCodes::IsOk(GetResponseCode()); }
+		virtual int32 GetResponseCode() const override { return RequestSuccessful() ? EHttpResponseCodes::Ok : EHttpResponseCodes::NotFound; }
+		virtual const TArray<uint8>& GetData() const override { return DataArray; }
 		// IDownload interface end.
 
 	private:
 		TArray<uint8> DataArray;
-		bool bSuccess;
+		bool bSuccess = false;
 	};
 
-	FFileDownload::FFileDownload(TArray<uint8> InDataArray, bool bInSuccess, IDownloadServiceStat::FDownloadRecord&& InDownloadRecord)
-		: FDownloadBase(MoveTemp(InDownloadRecord))
-		, DataArray(MoveTemp(InDataArray))
-		, bSuccess(bInSuccess)
-	{
-	}
-
-	FFileDownload::~FFileDownload()
-	{
-	}
-
-	bool FFileDownload::RequestSuccessful() const
-	{
-		return bSuccess;
-	}
-
-	bool FFileDownload::ResponseSuccessful() const
-	{
-		return EHttpResponseCodes::IsOk(GetResponseCode());
-	}
-
-	int32 FFileDownload::GetResponseCode() const
-	{
-		return RequestSuccessful() ? EHttpResponseCodes::Ok : EHttpResponseCodes::NotFound;
-	}
-
-	const TArray<uint8>& FFileDownload::GetData() const
-	{
-		return DataArray;
-	}
-
 	struct FDownloadDelegates
 	{
 	public:
-		FDownloadDelegates();
+		FDownloadDelegates() = default;
 		FDownloadDelegates(const FDownloadCompleteDelegate& InOnCompleteDelegate, const FDownloadProgressDelegate& InOnProgressDelegate);
 
 	public:
@@ -166,549 +86,323 @@ namespace BuildPatchServices
 		FDownloadProgressDelegate OnProgressDelegate;
 	};
 
-	FDownloadDelegates::FDownloadDelegates()
-		: OnCompleteDelegate()
-		, OnProgressDelegate()
-	{
-	}
-
 	FDownloadDelegates::FDownloadDelegates(const FDownloadCompleteDelegate& InOnCompleteDelegate, const FDownloadProgressDelegate& InOnProgressDelegate)
 		: OnCompleteDelegate(InOnCompleteDelegate)
 		, OnProgressDelegate(InOnProgressDelegate)
 	{
 	}
 
-	namespace DownloadServiceHelpers
-	{
-		void ExecuteCancelled(int32 RequestId, const FDownloadDelegates& DownloadDelegates)
-		{
-			DownloadDelegates.OnCompleteDelegate.ExecuteIfBound(RequestId, MakeShareable(new FFileDownload(TArray<uint8>(), false, IDownloadServiceStat::FDownloadRecord())));
-		}
-
-		void ExecuteCancelledAndReset(int32 RequestId, FDownloadDelegates& DownloadDelegates)
-		{
-			ExecuteCancelled(RequestId, DownloadDelegates);
-			DownloadDelegates = FDownloadDelegates();
-		}
-
-		void Reset(FDownloadDelegates& DownloadDelegates)
-		{
-			DownloadDelegates = FDownloadDelegates();
-		}
-	}
-
-	class FDownloadService
-		: public IDownloadService
+	/**
+	 * Use Pimpl pattern for binding thread safe shared ptr delegates for the http module, without having to enforce that
+	 * this service should be made using TShared* reference controllers.
+	 */
+	class FDownloadServiceImpl
+		: public TSharedFromThis<FDownloadServiceImpl>
 	{
-	private:
-		/**
-		 * This is a wrapper class for binding thread safe shared ptr delegates for the http module, without having to enforce that
-		 * this service should be made using TShared* reference controllers.
-		 */
-		class FHttpDelegates
-		{
-		public:
-			FHttpDelegates(FDownloadService& InDownloadService);
-
-			void HttpRequestProgress(FHttpRequestPtr Request, uint64 BytesSent, uint64 BytesReceived, int32 RequestId);
-			void HttpRequestComplete(FHttpRequestPtr Request, FHttpResponsePtr Response, bool bSucceeded, IDownloadServiceStat::FDownloadRecord DownloadRecord);
-
-		private:
-			FDownloadService& DownloadService;
-		};
-
 	public:
-		FDownloadService(FTSTicker& Ticker, IHttpManager* HttpManager, IFileSystem* FileSystem, IDownloadServiceStat* DownloadServiceStat, IInstallerAnalytics* InstallerAnalytics);
-		~FDownloadService();
+		FDownloadServiceImpl(IHttpManager* HttpManager, IFileSystem* FileSystem, IDownloadServiceStat* DownloadServiceStat, IInstallerAnalytics* InstallerAnalytics);
+		~FDownloadServiceImpl();
+
+		void CancelAllRequests();
 
 		// IDownloadService interface begin.
-		virtual int32 RequestFile(const FString& FileUri, const FDownloadCompleteDelegate& OnCompleteDelegate, const FDownloadProgressDelegate& OnProgressDelegate) override;
-		virtual void RequestCancel(int32 RequestId) override;
-		virtual void RequestAbandon(int32 RequestId) override;
+		int32 RequestFile(const FString& FileUri, const FDownloadCompleteDelegate& OnCompleteDelegate, const FDownloadProgressDelegate& OnProgressDelegate);
+		void RequestCancel(int32 RequestId);
+		void RequestAbandon(int32 RequestId);
 		// IDownloadService interface end.
 
 	private:
-		bool Tick(float DeltaTime);
-		void ProcessCancelRequests();
-		void ProcessAbandonRequests();
-		void ProcessNewRequests();
-		void ProcessProgressUpdates();
-		void ProcessCompletedRequests();
+		UE::Tasks::FTask MakeFileLoadTask(int32 RequestId, const FString& FileUri, FFileRequest* FileRequest);
 
-		TSet<int32> CancelNewRequests(const TArray<int32>& RequestsIDs);
-		TSet<int32> CancelOngoingRequests(const TArray<int32>& RequestsIDs);
-
-		int32 MakeRequestId();
 		IDownloadServiceStat::FDownloadRecord MakeDownloadRecord(int32 RequestId, FString Uri);
-		TFunction<void()> MakeFileLoadTask(const int32& RequestId, const FString& FileUri, FFileRequest* FileRequest);
-
-		void RegisterRequest(int32 RequestId, TSharedRef<IHttpRequest, ESPMode::ThreadSafe> HttpRequest);
-		void RegisterRequest(int32 RequestId, TUniquePtr<FFileRequest> FileRequest);
-		void UnregisterRequest(int32 RequestId);
-		void HttpRequestProgress(FHttpRequestPtr Request, uint64 BytesSent, uint64 BytesReceived, int32 RequestId);
-		void HttpRequestComplete(FHttpRequestPtr Request, FHttpResponsePtr Response, bool bSuccess, IDownloadServiceStat::FDownloadRecord DownloadRecord);
-		void SetRequestProgress(int32 RequestId, uint64 BytesSoFar);
-		void SetFileRequestComplete(int32 RequestId, bool bSuccess, TArray<uint8> FileDataArray, IDownloadServiceStat::FDownloadRecord&& DownloadRecord);
-		void SetHttpRequestComplete(int32 RequestId, bool bSuccess, FHttpResponsePtr Response, IDownloadServiceStat::FDownloadRecord&& DownloadRecord);
 
 	private:
-		TSharedRef<FHttpDelegates, ESPMode::ThreadSafe> HttpDelegates;
-		FTSTicker& Ticker;
-		IHttpManager* HttpManager;
-		IFileSystem* FileSystem;
-		IDownloadServiceStat* DownloadServiceStat;
-		IInstallerAnalytics* InstallerAnalytics;
-		TSharedRef<FThreadSafeBool, ESPMode::ThreadSafe> SharedShouldRunState;
-		FThreadSafeCounter RequestIdCounter;
-
-		FCriticalSection RequestDelegatesCS;
-		TMap<int32, FDownloadDelegates> RequestDelegates;
-
-		FCriticalSection NewRequestsCS;
-		TMap<int32, FString> NewRequests;
+		IHttpManager* HttpManager = nullptr;
+		IFileSystem* FileSystem = nullptr;
+		IDownloadServiceStat* DownloadServiceStat = nullptr;
+		IInstallerAnalytics* InstallerAnalytics = nullptr;
 
-		FCriticalSection CancelRequestsCS;
-		TArray<int32> CancelRequests;
-
-		FCriticalSection AbandonRequestsCS;
-		TArray<int32> AbandonRequests;
+		std::atomic<int32> RequestIdCounter;
 
 		FCriticalSection ActiveRequestsCS;
 		TMap<int32, TSharedRef<IHttpRequest, ESPMode::ThreadSafe>> ActiveHttpRequests;
 		TMap<int32, TUniquePtr<FFileRequest>> ActiveFileRequests;
 
-		FCriticalSection ProgressUpdatesCS;
-		TMap<int32, uint64> ProgressUpdates;
-
-		FCriticalSection CompletedRequestsCS;
-		TMap<int32, FDownloadBaseRef> CompletedRequests;
-
-		FTSTicker::FDelegateHandle TickerHandle;
+		FCriticalSection RequestDelegatesCS;
+		TMap<int32, FDownloadDelegates> RequestDelegates;
 	};
 
-	FDownloadService::FHttpDelegates::FHttpDelegates(FDownloadService& InDownloadService)
-		: DownloadService(InDownloadService)
+	class FDownloadService
+		: public IDownloadService
 	{
-	}
+	public:
+		FDownloadService(IHttpManager* HttpManager, IFileSystem* FileSystem, IDownloadServiceStat* DownloadServiceStat, IInstallerAnalytics* InstallerAnalytics)
+			: Impl(MakeShared<FDownloadServiceImpl>(HttpManager, FileSystem, DownloadServiceStat, InstallerAnalytics))
+		{}
+		~FDownloadService()
+		{
+			Impl->CancelAllRequests();
+		}
 
-	void FDownloadService::FHttpDelegates::HttpRequestProgress(FHttpRequestPtr Request, uint64 BytesSent, uint64 BytesReceived, int32 RequestId)
-	{
-		DownloadService.HttpRequestProgress(MoveTemp(Request), BytesSent, BytesReceived, RequestId);
-	}
+		// IDownloadService interface begin.
+		virtual int32 RequestFile(const FString& FileUri, const FDownloadCompleteDelegate& OnCompleteDelegate, const FDownloadProgressDelegate& OnProgressDelegate) override
+		{
+			return Impl->RequestFile(FileUri, OnCompleteDelegate, OnProgressDelegate);
+		}
+		virtual void RequestCancel(int32 RequestId) override { Impl->RequestCancel(RequestId); }
+		virtual void RequestAbandon(int32 RequestId) override { Impl->RequestAbandon(RequestId); }
+		// IDownloadService interface end.
 
-	void FDownloadService::FHttpDelegates::HttpRequestComplete(FHttpRequestPtr Request, FHttpResponsePtr Response, bool bSucceeded, IDownloadServiceStat::FDownloadRecord DownloadRecord)
-	{
-		DownloadService.HttpRequestComplete(MoveTemp(Request), MoveTemp(Response), bSucceeded, MoveTemp(DownloadRecord));
-	}
+	private:
+		TSharedRef<FDownloadServiceImpl> Impl;
+	};
 
-	FDownloadService::FDownloadService(FTSTicker& InTicker, IHttpManager* InHttpManager, IFileSystem* InFileSystem, IDownloadServiceStat* InDownloadServiceStat, IInstallerAnalytics* InInstallerAnalytics)
-		: HttpDelegates(MakeShareable(new FHttpDelegates(*this)))
-		, Ticker(InTicker)
-		, HttpManager(InHttpManager)
+	FDownloadServiceImpl::FDownloadServiceImpl(IHttpManager* InHttpManager, IFileSystem* InFileSystem, IDownloadServiceStat* InDownloadServiceStat, IInstallerAnalytics* InInstallerAnalytics)
+		: HttpManager(InHttpManager)
 		, FileSystem(InFileSystem)
 		, DownloadServiceStat(InDownloadServiceStat)
 		, InstallerAnalytics(InInstallerAnalytics)
-		, SharedShouldRunState(MakeShareable(new FThreadSafeBool(true)))
-		, RequestIdCounter(0)
-		, RequestDelegatesCS()
-		, RequestDelegates()
-		, NewRequestsCS()
-		, NewRequests()
-		, CancelRequestsCS()
-		, CancelRequests()
-		, AbandonRequestsCS()
-		, AbandonRequests()
-		, ActiveRequestsCS()
-		, ActiveHttpRequests()
-		, ActiveFileRequests()
-		, ProgressUpdatesCS()
-		, ProgressUpdates()
-		, CompletedRequestsCS()
-		, CompletedRequests()
+	{}
+
+	FDownloadServiceImpl::~FDownloadServiceImpl()
 	{
-		check(IsInGameThread());
-		TickerHandle = Ticker.AddTicker(FTickerDelegate::CreateRaw(this, &FDownloadService::Tick));
 	}
 
-	FDownloadService::~FDownloadService()
+	void FDownloadServiceImpl::CancelAllRequests()
 	{
-		check(IsInGameThread());
-		// Remove ticker.
-		Ticker.RemoveTicker(TickerHandle);
-		// Make sure our file threads will exit if they continue.
-		SharedShouldRunState->AtomicSet(false);
 		// Cancel all HTTP requests and wait for all file downloads threads to exit.
-		ActiveRequestsCS.Lock();
-		for (const TPair<int32, TSharedRef<IHttpRequest, ESPMode::ThreadSafe>>& ActiveHttpRequest : ActiveHttpRequests)
-		{
-			ActiveHttpRequest.Value->OnRequestProgress64().Unbind();
-			ActiveHttpRequest.Value->OnProcessRequestComplete().Unbind();
-			ActiveHttpRequest.Value->CancelRequest();
-		}
-		for (const TPair<int32, TUniquePtr<FFileRequest>>& ActiveFileRequest : ActiveFileRequests)
 		{
-			ActiveFileRequest.Value->Future.Wait();
-		}
-		ActiveHttpRequests.Empty();
-		ActiveFileRequests.Empty();
-		ActiveRequestsCS.Unlock();
-		// Kick off any remaining delegates next frame.
-		RequestDelegatesCS.Lock();
-		for (const TPair<int32, FDownloadDelegates>& RequestDelegate : RequestDelegates)
-		{
-			DownloadServiceHelpers::ExecuteCancelled(RequestDelegate.Key, RequestDelegate.Value);
+			FScopeLock ScopeLock(&ActiveRequestsCS);
+			for (const TPair<int32, TSharedRef<IHttpRequest>>& ActiveHttpRequest : ActiveHttpRequests)
+			{
+				ActiveHttpRequest.Value->CancelRequest();
+			}
+			for (const TPair<int32, TUniquePtr<FFileRequest>>& ActiveFileRequest : ActiveFileRequests)
+			{
+				ActiveFileRequest.Value->ShouldCancel.Cancel();
+			}
 		}
-		RequestDelegates.Empty();
-		RequestDelegatesCS.Unlock();
-		// By this point all other references to SharedShouldRunState should have destructed.
-		check(SharedShouldRunState.IsUnique());
-	}
-
-	int32 FDownloadService::RequestFile(const FString& FileUri, const FDownloadCompleteDelegate& OnCompleteDelegate, const FDownloadProgressDelegate& OnProgressDelegate)
-	{
-		const int32 NewRequestId = MakeRequestId();
-
-		// Save the delegates.
-		RequestDelegatesCS.Lock();
-		RequestDelegates.Emplace(NewRequestId, FDownloadDelegates(OnCompleteDelegate, OnProgressDelegate));
-		RequestDelegatesCS.Unlock();
-
-		// Add the request.
-		NewRequestsCS.Lock();
-		NewRequests.Emplace(NewRequestId, FileUri);
-		NewRequestsCS.Unlock();
-
-		return NewRequestId;
-	}
-
-	void FDownloadService::RequestCancel(int32 RequestId)
-	{
-		// Add the request.
-		FScopeLock ScopeLock(&CancelRequestsCS);
-		CancelRequests.Add(RequestId);
-	}
-
-	void FDownloadService::RequestAbandon(int32 RequestId)
-	{
-		// Add the request.
-		FScopeLock ScopeLock(&AbandonRequestsCS);
-		AbandonRequests.Add(RequestId);
-	}
-
-	bool FDownloadService::Tick(float DeltaTime)
-	{
-		QUICK_SCOPE_CYCLE_COUNTER(STAT_FDownloadService_Tick);
-
-		ProcessCancelRequests();
 
-		ProcessAbandonRequests();
-
-		ProcessNewRequests();
-
-		ProcessProgressUpdates();
-
-		ProcessCompletedRequests();
-
-		return true;
+		// callbacks and removal handled by completion delegates
 	}
 
-	void FDownloadService::ProcessCancelRequests()
+	int32 FDownloadServiceImpl::RequestFile(const FString& FileUri, const FDownloadCompleteDelegate& OnCompleteDelegate, const FDownloadProgressDelegate& OnProgressDelegate)
 	{
-		// Grab new cancel requests for this frame.
-		TArray<int32> FrameCancelRequests;
-		CancelRequestsCS.Lock();
-		FrameCancelRequests = MoveTemp(CancelRequests);
-		CancelRequestsCS.Unlock();
+		const int32 RequestId = ++RequestIdCounter;
 
-		// Cancel new requests that were not processed yet.
-		TSet<int32> UnstartedRequests = CancelNewRequests(FrameCancelRequests);
-
-		// Cancel ongoing requests.
-		TSet<int32> CancelledRequests = CancelOngoingRequests(FrameCancelRequests);
+		// Save the delegates.
+		{
+			FScopeLock ScopeLock(&RequestDelegatesCS);
+			RequestDelegates.Emplace(RequestId, FDownloadDelegates(OnCompleteDelegate, OnProgressDelegate));
+		}
 
-		// Call and remove delegates for new requests that were canceled.
-		RequestDelegatesCS.Lock();
-		for (const int32& RequestId : UnstartedRequests.Union(CancelledRequests))
+		const bool bIsHttpRequest = FileUri.StartsWith(TEXT("http"), ESearchCase::IgnoreCase);
+		if (bIsHttpRequest)
 		{
-			if (ensureMsgf(RequestDelegates.Contains(RequestId), TEXT("Missing request delegates for %d"), RequestId))
+			// Kick off http request.
+			FHttpRequestPtr HttpRequest;
+			{
+				FScopeLock ScopeLock(&ActiveRequestsCS);
+				HttpRequest = ActiveHttpRequests.Emplace(RequestId, HttpManager->CreateRequest());
+			}
+
+			HttpRequest->SetDelegateThreadPolicy(EHttpRequestDelegateThreadPolicy::CompleteOnHttpThread);
+			HttpRequest->SetURL(FileUri);
+			HttpRequest->SetVerb(TEXT("GET"));
+			
+			HttpRequest->OnRequestProgress64().BindSPLambda(this, [this, RequestId](FHttpRequestPtr Request, uint64 BytesSent, uint64 BytesReceived)
 			{
-				DownloadServiceHelpers::ExecuteCancelledAndReset(RequestId, RequestDelegates[RequestId]);
-				if (UnstartedRequests.Contains(RequestId))
 				{
-					RequestDelegates.Remove(RequestId);
+					FScopeLock ScopeLock(&RequestDelegatesCS);
+					if (FDownloadDelegates* MaybeDelegates = RequestDelegates.Find(RequestId))
+					{
+						MaybeDelegates->OnProgressDelegate.ExecuteIfBound(RequestId, BytesReceived);
+					}
 				}
-			}
-		}
-		RequestDelegatesCS.Unlock();
-	}
 
-	void FDownloadService::ProcessAbandonRequests()
-	{
-		// Grab new abandon requests for this frame.
-		TArray<int32> FrameAbandonRequests;
-		AbandonRequestsCS.Lock();
-		FrameAbandonRequests = MoveTemp(AbandonRequests);
-		AbandonRequestsCS.Unlock();
+				DownloadServiceStat->OnDownloadProgress(RequestId, BytesReceived);
+			});
 
-		// Cancel new requests that were not processed yet.
-		TSet<int32> UnstartedRequests = CancelNewRequests(FrameAbandonRequests);
+			HttpRequest->OnProcessRequestComplete().BindSPLambda(this,
+			[this, DownloadRecord = MakeDownloadRecord(RequestId, FileUri)](FHttpRequestPtr Request, FHttpResponsePtr Response, bool bSucceeded) mutable
+			{
+				InstallerAnalytics->TrackRequest(Request);
 
-		// Cancel ongoing requests.
-		TSet<int32> CancelledRequests = CancelOngoingRequests(FrameAbandonRequests);
+				DownloadRecord.bSuccess = bSucceeded;
+				DownloadRecord.ResponseCode = Response.IsValid() ? Response->GetResponseCode() : INDEX_NONE;
+				DownloadRecord.SpeedRecord.CyclesEnd = FStatsCollector::GetCycles();
+				DownloadRecord.SpeedRecord.Size = Response.IsValid() ? Response->GetContent().Num() : 0;
+				
+				FDownloadDelegates Delegates;
+				bool bFoundDelegates = false;
+				{
+					FScopeLock ScopeLock(&RequestDelegatesCS);
+					bFoundDelegates = RequestDelegates.RemoveAndCopyValue(DownloadRecord.RequestId, Delegates);
+				}
 
-		// Remove delegates for new requests that were canceled.
-		RequestDelegatesCS.Lock();
-		for (const int32& RequestId : UnstartedRequests.Union(CancelledRequests))
-		{
-			if (ensureMsgf(RequestDelegates.Contains(RequestId), TEXT("Missing request delegates for %d"), RequestId))
-			{
-				DownloadServiceHelpers::Reset(RequestDelegates[RequestId]);
-				if (UnstartedRequests.Contains(RequestId))
+				if (bFoundDelegates)
 				{
-					RequestDelegates.Remove(RequestId);
+					TSharedRef<FHttpDownload> CompletedDownload = MakeShared<FHttpDownload>(Request, bSucceeded);
+					Delegates.OnCompleteDelegate.ExecuteIfBound(DownloadRecord.RequestId, CompletedDownload);
 				}
-			}
-		}
-		RequestDelegatesCS.Unlock();
-	}
 
-	void FDownloadService::ProcessNewRequests()
-	{
-		// Grab new http requests for this frame.
-		TMap<int32, FString> FrameNewRequests;
-		NewRequestsCS.Lock();
-		FrameNewRequests = MoveTemp(NewRequests);
-		NewRequestsCS.Unlock();
-
-		// Start new requests.
-		for (const TPair<int32, FString>& NewRequest : FrameNewRequests)
-		{
-			const bool bIsHttpRequest = NewRequest.Value.StartsWith(TEXT("http"), ESearchCase::IgnoreCase);
-			if (bIsHttpRequest)
-			{
-				// Kick off http request.
-				TSharedRef<IHttpRequest, ESPMode::ThreadSafe> HttpRequest = HttpManager->CreateRequest();
-				HttpRequest->OnRequestProgress64().BindThreadSafeSP(HttpDelegates, &FHttpDelegates::HttpRequestProgress, NewRequest.Key);
-				HttpRequest->OnProcessRequestComplete().BindThreadSafeSP(HttpDelegates, &FHttpDelegates::HttpRequestComplete, MakeDownloadRecord(NewRequest.Key, NewRequest.Value));
-				HttpRequest->SetDelegateThreadPolicy(EHttpRequestDelegateThreadPolicy::CompleteOnHttpThread);
-				HttpRequest->SetURL(NewRequest.Value);
-				HttpRequest->SetVerb(TEXT("GET"));
-				HttpRequest->ProcessRequest();
-				RegisterRequest(NewRequest.Key, MoveTemp(HttpRequest));
-			}
-			else
-			{
-				// Load file from drive/network.
-				TUniquePtr<FFileRequest> FileRequest(new FFileRequest());
-				FFileRequest* FileRequestPtr = FileRequest.Get();
-				FileRequest->Future = Async(EAsyncExecution::ThreadPool, MakeFileLoadTask(NewRequest.Key, NewRequest.Value, FileRequestPtr));
-				RegisterRequest(NewRequest.Key, MoveTemp(FileRequest));
-			}
-			DownloadServiceStat->OnDownloadStarted(NewRequest.Key, NewRequest.Value);
-		}
-	}
+				{
+					FScopeLock ScopeLock(&ActiveRequestsCS);
+					ActiveHttpRequests.Remove(DownloadRecord.RequestId);
+				}
 
-	void FDownloadService::ProcessProgressUpdates()
-	{
-		// Grab the progress updates for this frame.
-		TMap<int32, uint64> FrameProgressUpdates;
-		ProgressUpdatesCS.Lock();
-		FrameProgressUpdates = MoveTemp(ProgressUpdates);
-		ProgressUpdatesCS.Unlock();
-
-		// Process progress updates.
-		RequestDelegatesCS.Lock();
-		for (const TPair<int32, uint64>& FrameProgressUpdate : FrameProgressUpdates)
+				DownloadServiceStat->OnDownloadComplete(DownloadRecord);
+			});
+
+			HttpRequest->ProcessRequest();
+		}
+		else
 		{
-			if (RequestDelegates.Contains(FrameProgressUpdate.Key))
+			// Load file from drive/network.
+			TUniquePtr<FFileRequest> FileRequest = MakeUnique<FFileRequest>();
 			{
-				RequestDelegates[FrameProgressUpdate.Key].OnProgressDelegate.ExecuteIfBound(FrameProgressUpdate.Key, FrameProgressUpdate.Value);
+				FScopeLock ScopeLock(&ActiveRequestsCS);
+				ActiveFileRequests.Add(RequestId, MoveTemp(FileRequest));
 			}
+
+			FFileRequest* FileRequestPtr = FileRequest.Get();
+			FileRequest->Task = MakeFileLoadTask(RequestId, FileUri, FileRequestPtr);
 		}
-		RequestDelegatesCS.Unlock();
-		for (const TPair<int32, uint64>& FrameProgressUpdate : FrameProgressUpdates)
-		{
-			DownloadServiceStat->OnDownloadProgress(FrameProgressUpdate.Key, FrameProgressUpdate.Value);
-		}
+		DownloadServiceStat->OnDownloadStarted(RequestId, FileUri);
+
+		return RequestId;
 	}
 
-	void FDownloadService::ProcessCompletedRequests()
+	void FDownloadServiceImpl::RequestCancel(int32 RequestId)
 	{
-		// Grab the completed requests for this frame.
-		TMap<int32, FDownloadBaseRef> FrameCompletedRequests;
-		CompletedRequestsCS.Lock();
-		FrameCompletedRequests = MoveTemp(CompletedRequests);
-		CompletedRequestsCS.Unlock();
-
-		// Process completed requests.
-		RequestDelegatesCS.Lock();
-		for (const TPair<int32, FDownloadBaseRef>& FrameCompletedRequest : FrameCompletedRequests)
-		{
-			if (ensureMsgf(RequestDelegates.Contains(FrameCompletedRequest.Key), TEXT("Missing request delegates for %d"), FrameCompletedRequest.Key))
-			{
-				RequestDelegates[FrameCompletedRequest.Key].OnCompleteDelegate.ExecuteIfBound(FrameCompletedRequest.Key, FrameCompletedRequest.Value);
-				UnregisterRequest(FrameCompletedRequest.Key);
-				RequestDelegates.Remove(FrameCompletedRequest.Key);
-			}
-		}
-		RequestDelegatesCS.Unlock();
-		for (const TPair<int32, FDownloadBaseRef>& FrameCompletedRequest : FrameCompletedRequests)
+		FScopeLock ScopeLock(&ActiveRequestsCS);
+		
+		if (FHttpRequestRef* MaybeRequest = ActiveHttpRequests.Find(RequestId))
 		{
-			DownloadServiceStat->OnDownloadComplete(MoveTemp(FrameCompletedRequest.Value->DownloadRecord));
+			(*MaybeRequest)->CancelRequest();
 		}
-	}
-
-	TSet<int32> FDownloadService::CancelNewRequests(const TArray<int32>& RequestsIDs)
-	{
-		TSet<int32> UnstartedRequests;
-		NewRequestsCS.Lock();
-		for (const int32& RequestId : RequestsIDs)
+		if (TUniquePtr<FFileRequest>* MaybeRequest = ActiveFileRequests.Find(RequestId))
 		{
-			if (NewRequests.Remove(RequestId) > 0)
-			{
-				UnstartedRequests.Add(RequestId);
-			}
+			(*MaybeRequest)->ShouldCancel.Cancel();
 		}
-		NewRequestsCS.Unlock();
 
-		return UnstartedRequests;
+		// callbacks and removal handled by completion delegates
 	}
 
-	TSet<int32> FDownloadService::CancelOngoingRequests(const TArray<int32>& RequestsIDs)
+	void FDownloadServiceImpl::RequestAbandon(int32 RequestId)
 	{
-		TSet<int32> CancelledRequests;
-		ActiveRequestsCS.Lock();
-		for (const int32& RequestId : RequestsIDs)
 		{
-			if (ActiveHttpRequests.Contains(RequestId))
-			{
-				ActiveHttpRequests[RequestId]->CancelRequest();
-				CancelledRequests.Add(RequestId);
-			}
-			if (ActiveFileRequests.Contains(RequestId))
-			{
-				ActiveFileRequests[RequestId]->ShouldCancel = true;
-				CancelledRequests.Add(RequestId);
-			}
+			FScopeLock ScopeLock(&RequestDelegatesCS);
+			RequestDelegates.Remove(RequestId);
 		}
-		ActiveRequestsCS.Unlock();
 
-		return CancelledRequests;
+		RequestCancel(RequestId);
 	}
 
-	int32 FDownloadService::MakeRequestId()
+	UE::Tasks::FTask FDownloadServiceImpl::MakeFileLoadTask(const int32 RequestId, const FString& FileUri, FFileRequest* FileRequest)
 	{
-		return RequestIdCounter.Increment();
-	}
-
-	IDownloadServiceStat::FDownloadRecord FDownloadService::MakeDownloadRecord(int32 RequestId, FString Uri)
-	{
-		IDownloadServiceStat::FDownloadRecord DownloadRecord;
-		DownloadRecord.RequestId = RequestId;
-		DownloadRecord.Uri = MoveTemp(Uri);
-		DownloadRecord.bSuccess = false;
-		DownloadRecord.ResponseCode = INDEX_NONE;
-		DownloadRecord.SpeedRecord.CyclesStart = FStatsCollector::GetCycles();
-		DownloadRecord.SpeedRecord.CyclesEnd = DownloadRecord.SpeedRecord.CyclesStart;
-		DownloadRecord.SpeedRecord.Size = 0;
-		return DownloadRecord;
-	}
-
-	TFunction<void()> FDownloadService::MakeFileLoadTask(const int32& RequestId, const FString& FileUri, FFileRequest* FileRequest)
-	{
-		TWeakPtr<FThreadSafeBool, ESPMode::ThreadSafe> WeakShouldRun = SharedShouldRunState;
-		return [this, RequestId, FileUri, FileRequest, WeakShouldRun]()
-		{
-			TSharedPtr<FThreadSafeBool, ESPMode::ThreadSafe> LambdaShouldRun = WeakShouldRun.Pin();
-			TArray<uint8> FileDataArray;
-			bool bSuccess = LambdaShouldRun.IsValid() && *LambdaShouldRun.Get() && !FileRequest->ShouldCancel;
-			IDownloadServiceStat::FDownloadRecord DownloadRecord = MakeDownloadRecord(RequestId, FileUri);
-			if (bSuccess)
+		return UE::Tasks::Launch(UE_SOURCE_LOCATION,
+			[this, WeakThis = AsWeak(), RequestId, FileUri, FileRequest]
 			{
-				TUniquePtr<FArchive> Reader = FileSystem->CreateFileReader(*FileUri);
-				bSuccess = Reader.IsValid();
+				TSharedPtr<FDownloadServiceImpl> PinThis = WeakThis.Pin();
+				if (!PinThis)
+				{
+					return;
+				}
+
+				TArray<uint8> FileDataArray;
+				bool bSuccess = !FileRequest->ShouldCancel.IsCanceled();
+				IDownloadServiceStat::FDownloadRecord DownloadRecord = MakeDownloadRecord(RequestId, FileUri);
 				if (bSuccess)
 				{
-					const int64 FileSize = Reader->TotalSize();
-					FileDataArray.Reset();
-					FileDataArray.AddUninitialized(FileSize);
-					int64 BytesRead = 0;
-					while (BytesRead < FileSize && *LambdaShouldRun.Get() && !FileRequest->ShouldCancel)
+					TUniquePtr<FArchive> Reader = FileSystem->CreateFileReader(*FileUri);
+					bSuccess = Reader.IsValid();
+					if (bSuccess)
 					{
-						const int64 ReadLen = FMath::Min<int64>(FileReaderBufferSize, FileSize - BytesRead);
-						Reader->Serialize(FileDataArray.GetData() + BytesRead, ReadLen);
-						BytesRead += ReadLen;
-						SetRequestProgress(RequestId, BytesRead);
+						const int64 FileSize = Reader->TotalSize();
+						FileDataArray.Reset();
+						FileDataArray.AddUninitialized(FileSize);
+						int64 BytesRead = 0;
+						bool bIsCanceled = false;
+						while (BytesRead < FileSize && !FileRequest->ShouldCancel.IsCanceled())
+						{
+							bIsCanceled = !FileRequest->ShouldCancel.IsCanceled();
+							if (bIsCanceled)
+							{
+								bSuccess = false;
+								break;
+							}
+
+							const int64 ReadLen = FMath::Min<int64>(FileReaderBufferSize, FileSize - BytesRead);
+							Reader->Serialize(FileDataArray.GetData() + BytesRead, ReadLen);
+							BytesRead += ReadLen;
+
+							{
+								FScopeLock ScopeLock(&RequestDelegatesCS);
+								if (FDownloadDelegates* MaybeDelegates = RequestDelegates.Find(RequestId))
+								{
+									MaybeDelegates->OnProgressDelegate.ExecuteIfBound(RequestId, BytesRead);
+								}
+							}
+
+							DownloadServiceStat->OnDownloadProgress(RequestId, BytesRead);
+						}
+						DownloadRecord.SpeedRecord.Size = BytesRead;
+						bSuccess = bSuccess && Reader->Close() && BytesRead == FileSize;
 					}
-					DownloadRecord.SpeedRecord.Size = BytesRead;
-					bSuccess = Reader->Close() && BytesRead == FileSize;
+					DownloadRecord.SpeedRecord.CyclesEnd = FStatsCollector::GetCycles();
+					DownloadRecord.bSuccess = bSuccess;
+					DownloadRecord.ResponseCode = bSuccess ? EHttpResponseCodes::Ok : EHttpResponseCodes::NotFound;
 				}
-				DownloadRecord.SpeedRecord.CyclesEnd = FStatsCollector::GetCycles();
-				DownloadRecord.bSuccess = bSuccess;
-				DownloadRecord.ResponseCode = bSuccess ? EHttpResponseCodes::Ok : EHttpResponseCodes::NotFound;
-			}
-			SetFileRequestComplete(RequestId, bSuccess, MoveTemp(FileDataArray), MoveTemp(DownloadRecord));
-		};
-	}
 
-	void FDownloadService::RegisterRequest(int32 RequestId, TUniquePtr<FFileRequest> FileRequest)
-	{
-		FScopeLock ScopeLock(&ActiveRequestsCS);
-		ActiveFileRequests.Add(RequestId, MoveTemp(FileRequest));
-	}
-
-	void FDownloadService::RegisterRequest(int32 RequestId, TSharedRef<IHttpRequest, ESPMode::ThreadSafe> HttpRequest)
-	{
-		FScopeLock ScopeLock(&ActiveRequestsCS);
-		ActiveHttpRequests.Add(RequestId, MoveTemp(HttpRequest));
-	}
-
-	void FDownloadService::UnregisterRequest(int32 RequestId)
-	{
-		FScopeLock ScopeLock(&ActiveRequestsCS);
-		ActiveHttpRequests.Remove(RequestId);
-		ActiveFileRequests.Remove(RequestId);
-	}
+				if (!bSuccess)
+				{
+					FileDataArray.Empty();
+				}
 
-	void FDownloadService::HttpRequestProgress(FHttpRequestPtr Request, uint64 BytesSent, uint64 BytesReceived, int32 RequestId)
-	{
-		SetRequestProgress(RequestId, BytesReceived);
-	}
+				FDownloadDelegates Delegates;
+				bool bFoundDelegates = false;
+				{
+					FScopeLock ScopeLock(&RequestDelegatesCS);
+					bFoundDelegates = RequestDelegates.RemoveAndCopyValue(RequestId, Delegates);
+				}
 
-	void FDownloadService::HttpRequestComplete(FHttpRequestPtr Request, FHttpResponsePtr Response, bool bSuccess, IDownloadServiceStat::FDownloadRecord DownloadRecord)
-	{
-		InstallerAnalytics->TrackRequest(Request);
-		DownloadRecord.bSuccess = bSuccess;
-		DownloadRecord.ResponseCode = Response.IsValid() ? Response->GetResponseCode() : INDEX_NONE;
-		DownloadRecord.SpeedRecord.CyclesEnd = FStatsCollector::GetCycles();
-		DownloadRecord.SpeedRecord.Size = Response.IsValid() ? Response->GetContent().Num() : 0;
-		SetHttpRequestComplete(DownloadRecord.RequestId, bSuccess, MoveTemp(Response), MoveTemp(DownloadRecord));
-	}
+				if (bFoundDelegates)
+				{
+					TSharedRef<FFileDownload> CompletedDownload = MakeShared<FFileDownload>(MoveTemp(FileDataArray), bSuccess);
+					Delegates.OnCompleteDelegate.ExecuteIfBound(RequestId, CompletedDownload);
+				}
+				
+				{
+					FScopeLock ScopeLock(&ActiveRequestsCS);
+					ActiveFileRequests.Remove(RequestId);
 
-	void FDownloadService::SetRequestProgress(int32 RequestId, uint64 BytesSoFar)
-	{
-		FScopeLock ScopeLock(&ProgressUpdatesCS);
-		ProgressUpdates.Add(RequestId, BytesSoFar);
-	}
+					// FileRequest is dangling at this point
+				}
 
-	void FDownloadService::SetFileRequestComplete(int32 RequestId, bool bSuccess, TArray<uint8> FileDataArray, IDownloadServiceStat::FDownloadRecord&& DownloadRecord)
-	{
-		FScopeLock ScopeLock(&CompletedRequestsCS);
-		CompletedRequests.Emplace(RequestId, MakeShareable(new FFileDownload(MoveTemp(FileDataArray), bSuccess, MoveTemp(DownloadRecord))));
+				DownloadServiceStat->OnDownloadComplete(DownloadRecord);
+			}, UE::Tasks::ETaskPriority::BackgroundHigh);
 	}
 
-	void FDownloadService::SetHttpRequestComplete(int32 RequestId, bool bSuccess, FHttpResponsePtr Response, IDownloadServiceStat::FDownloadRecord&& DownloadRecord)
+	IDownloadServiceStat::FDownloadRecord FDownloadServiceImpl::MakeDownloadRecord(int32 RequestId, FString Uri)
 	{
-		FScopeLock ScopeLock(&CompletedRequestsCS);
-		CompletedRequests.Emplace(RequestId, MakeShareable(new FHttpDownload(MoveTemp(Response), bSuccess, MoveTemp(DownloadRecord))));
+		IDownloadServiceStat::FDownloadRecord DownloadRecord;
+		DownloadRecord.RequestId = RequestId;
+		DownloadRecord.Uri = MoveTemp(Uri);
+		DownloadRecord.bSuccess = false;
+		DownloadRecord.ResponseCode = INDEX_NONE;
+		DownloadRecord.SpeedRecord.CyclesStart = FStatsCollector::GetCycles();
+		DownloadRecord.SpeedRecord.CyclesEnd = DownloadRecord.SpeedRecord.CyclesStart;
+		DownloadRecord.SpeedRecord.Size = 0;
+		return DownloadRecord;
 	}
 
-	IDownloadService* FDownloadServiceFactory::Create(FTSTicker& Ticker, IHttpManager* HttpManager, IFileSystem* FileSystem, IDownloadServiceStat* DownloadServiceStat, IInstallerAnalytics* InstallerAnalytics)
+	IDownloadService* FDownloadServiceFactory::Create(IHttpManager* HttpManager, IFileSystem* FileSystem, IDownloadServiceStat* DownloadServiceStat, IInstallerAnalytics* InstallerAnalytics)
 	{
 		check(HttpManager != nullptr);
 		check(FileSystem != nullptr);
 		check(DownloadServiceStat != nullptr);
 		check(InstallerAnalytics != nullptr);
-		return new FDownloadService(Ticker, HttpManager, FileSystem, DownloadServiceStat, InstallerAnalytics);
+		return new FDownloadService(HttpManager, FileSystem, DownloadServiceStat, InstallerAnalytics);
 	}
 }
diff --git a/Engine/Source/Runtime/Online/BuildPatchServices/Private/Installer/DownloadService.h b/Engine/Source/Runtime/Online/BuildPatchServices/Private/Installer/DownloadService.h
--- a/Engine/Source/Runtime/Online/BuildPatchServices/Private/Installer/DownloadService.h
+++ b/Engine/Source/Runtime/Online/BuildPatchServices/Private/Installer/DownloadService.h
@@ -98,14 +98,13 @@ namespace BuildPatchServices
 	public:
 		/**
 		 * Instantiates an instance of an IDownloadService, using the HTTP module, and platform file API.
-		 * @param Ticker                The ticker which can be used to register a Tick function on.
 		 * @param HttpManager           The HTTP manager interface for making HTTP(s) requests.
 		 * @param FileSystem            The file system interface for network and disk file loading.
 		 * @param DownloadServiceStat   The class to receive statistics and event information.
 		 * @param InstallerAnalytics    The analytics implementation for tracking HTTP requests.
 		 * @return the new IDownloadService instance created.
 		 */
-		static IDownloadService* Create(FTSTicker& Ticker, IHttpManager* HttpManager, IFileSystem* FileSystem, IDownloadServiceStat* DownloadServiceStat, IInstallerAnalytics* InstallerAnalytics);
+		static IDownloadService* Create(IHttpManager* HttpManager, IFileSystem* FileSystem, IDownloadServiceStat* DownloadServiceStat, IInstallerAnalytics* InstallerAnalytics);
 	};
 
 	/**
diff --git a/Engine/Source/Runtime/Online/BuildPatchServices/Private/Installer/InstallerAnalytics.cpp b/Engine/Source/Runtime/Online/BuildPatchServices/Private/Installer/InstallerAnalytics.cpp
--- a/Engine/Source/Runtime/Online/BuildPatchServices/Private/Installer/InstallerAnalytics.cpp
+++ b/Engine/Source/Runtime/Online/BuildPatchServices/Private/Installer/InstallerAnalytics.cpp
@@ -34,6 +34,21 @@ namespace BuildPatchServices
 	{
 	}
 
+	class FNullInstallerAnalytics
+		: public IInstallerAnalytics
+	{
+	public:
+		// IInstallerAnalytics interface begin.
+		virtual void RecordChunkDownloadError(const FString& ChunkUrl, int32 ResponseCode, const FString& ErrorString) override {}
+		virtual void RecordChunkDownloadAborted(const FString& ChunkUrl, double ChunkTime, double ChunkMean, double ChunkStd, double BreakingPoint) override {}
+		virtual void RecordChunkCacheError(const FGuid& ChunkGuid, const FString& Filename, int32 LastError, const FString& SystemName, const FString& ErrorString) override {}
+		virtual void RecordConstructionError(const FString& Filename, int32 LastError, const FString& ErrorString) override {}
+		virtual void RecordPrereqInstallationError(const FString& AppName, const FString& AppVersion, const FString& Filename, const FString& CommandLine, int32 ErrorCode, const FString& ErrorString) override {}
+		virtual void TrackRequest(const FHttpRequestPtr& Request) override {}
+		virtual void Flush() override {}
+		// IInstallerAnalytics interface end.
+	};
+
 	class FInstallerAnalytics
 		: public IInstallerAnalytics
 	{
@@ -73,6 +88,7 @@ namespace BuildPatchServices
 		, AnalyticsEventQueueCS()
 		, AnalyticsEventQueue()
 	{
+		check(Analytics);
 		TickerHandle = FTSTicker::GetCoreTicker().AddTicker(FTickerDelegate::CreateRaw(this, &FInstallerAnalytics::Tick));
 	}
 
@@ -145,25 +161,23 @@ namespace BuildPatchServices
 
 	void FInstallerAnalytics::TrackRequest(const FHttpRequestPtr& Request)
 	{
-		static const FName EndpointName = TEXT("CDN.Chunk");
+		//static const FName EndpointName = TEXT("CDN.Chunk");
 	}
 
 	void FInstallerAnalytics::Flush()
 	{
 		check(IsInGameThread());
+		check(Analytics);
 
         QUICK_SCOPE_CYCLE_COUNTER(STAT_FInstallerAnalytics_Tick);
 
-		if (Analytics != nullptr)
+		// Process the Analytics Event queue
+		FScopeLock ScopeLock(&AnalyticsEventQueueCS);
+		for (FAnalyticsEventInfo& AnalyticsEvent : AnalyticsEventQueue)
 		{
-			// Process the Analytics Event queue
-			FScopeLock ScopeLock(&AnalyticsEventQueueCS);
-			for (FAnalyticsEventInfo& AnalyticsEvent : AnalyticsEventQueue)
-			{
-				Analytics->RecordEvent(AnalyticsEvent.EventName, AnalyticsEvent.Attributes);
-			}
-			AnalyticsEventQueue.Reset();
+			Analytics->RecordEvent(AnalyticsEvent.EventName, AnalyticsEvent.Attributes);
 		}
+		AnalyticsEventQueue.Reset();
 	}
 
 	void FInstallerAnalytics::QueueAnalyticsEvent(FString EventName, TArray<FAnalyticsEventAttribute> Attributes)
@@ -180,6 +194,11 @@ namespace BuildPatchServices
 
 	IInstallerAnalytics* FInstallerAnalyticsFactory::Create(IAnalyticsProvider* AnalyticsProvider)
 	{
-		return new FInstallerAnalytics(AnalyticsProvider);
+		if (AnalyticsProvider)
+		{
+			return new FInstallerAnalytics(AnalyticsProvider);
+		}
+
+		return new FNullInstallerAnalytics();
 	}
 }
diff --git a/Engine/Source/Runtime/Online/BuildPatchServices/Private/Installer/MessagePump.cpp b/Engine/Source/Runtime/Online/BuildPatchServices/Private/Installer/MessagePump.cpp
--- a/Engine/Source/Runtime/Online/BuildPatchServices/Private/Installer/MessagePump.cpp
+++ b/Engine/Source/Runtime/Online/BuildPatchServices/Private/Installer/MessagePump.cpp
@@ -3,6 +3,7 @@
 #include "Installer/MessagePump.h"
 #include "Containers/Union.h"
 #include "Containers/Queue.h"
+#include <atomic>
 
 namespace BuildPatchServices
 {
@@ -42,29 +43,24 @@ namespace BuildPatchServices
 		: public IMessagePump
 	{
 	public:
-		FMessagePump();
-		~FMessagePump();
-
 		// IMessagePump interface begin.
 		virtual void SendMessage(FChunkSourceEvent Message) override;
 		virtual void SendMessage(FInstallationFileAction Message) override;
 		virtual void SendRequest(FChunkUriRequest Request, TFunction<void(FChunkUriResponse)> OnResponse) override;
-		virtual void PumpMessages(const TArray<FMessageHandler*>& Handlers) override;
+		virtual void PumpMessages() override;
+		virtual void RegisterMessageHandler(FMessageHandler* MessageHandler) override;
+		virtual void UnregisterMessageHandler(FMessageHandler* MessageHandler) override;
 		// IMessagePump interface end.
 
 	private:
 		FMessageQueue MessageQueue;
 		FChunkUriQueue ChunkUriRequestQueue;
+		FDefaultMessageHandler DefaultMessageHandler;
+		TArray<FMessageHandler*> Handlers;
+		// std::atomic will not expose integral operations with enum types.
+		std::atomic<uint32> MessageRequests{0};
 	};
 
-	FMessagePump::FMessagePump()
-	{
-	}
-
-	FMessagePump::~FMessagePump()
-	{
-	}
-
 	void FMessagePump::SendMessage(FChunkSourceEvent Message)
 	{
 		MessageQueue.Enqueue(FMessageUnion(MoveTemp(Message)));
@@ -77,10 +73,17 @@ namespace BuildPatchServices
 
 	void FMessagePump::SendRequest(FChunkUriRequest Request, TFunction<void(FChunkUriResponse)> OnResponse)
 	{
-		ChunkUriRequestQueue.Enqueue(FChunkUriRequestResponse(MoveTemp(Request), MoveTemp(OnResponse)));
+		if (!EnumHasAllFlags(static_cast<EMessageRequests>(MessageRequests.load()), EMessageRequests::ChunkUriRequest))
+		{
+			DefaultMessageHandler.HandleRequest(Request, OnResponse);
+		}
+		else
+		{
+			ChunkUriRequestQueue.Enqueue(FChunkUriRequestResponse(MoveTemp(Request), MoveTemp(OnResponse)));
+		}
 	}
 
-	void FMessagePump::PumpMessages(const TArray<FMessageHandler*>& Handlers)
+	void FMessagePump::PumpMessages()
 	{
 		FMessageUnion MessageUnion;
 		while (MessageQueue.Dequeue(MessageUnion))
@@ -95,7 +98,6 @@ namespace BuildPatchServices
 			}
 		}
 
-		FDefaultMessageHandler DefaultMessageHandler;
 		FChunkUriRequestResponse ChunkUriRequestResponse;
 		while (ChunkUriRequestQueue.Dequeue(ChunkUriRequestResponse))
 		{
@@ -115,6 +117,22 @@ namespace BuildPatchServices
 		}
 	}
 
+	void FMessagePump::RegisterMessageHandler(FMessageHandler* InHandler)
+	{
+		MessageRequests |= static_cast<uint32>(InHandler->GetMessageRequests());
+		Handlers.Add(InHandler);
+	}
+
+	void FMessagePump::UnregisterMessageHandler(FMessageHandler* InHandler)
+	{
+		MessageRequests = 0;
+		Handlers.Remove(InHandler);
+		for (FMessageHandler* Handler : Handlers)
+		{
+			MessageRequests |= static_cast<uint32>(Handler->GetMessageRequests());
+		}
+	}
+
 	IMessagePump* FMessagePumpFactory::Create()
 	{
 		return new FMessagePump();
diff --git a/Engine/Source/Runtime/Online/BuildPatchServices/Private/Installer/MessagePump.h b/Engine/Source/Runtime/Online/BuildPatchServices/Private/Installer/MessagePump.h
--- a/Engine/Source/Runtime/Online/BuildPatchServices/Private/Installer/MessagePump.h
+++ b/Engine/Source/Runtime/Online/BuildPatchServices/Private/Installer/MessagePump.h
@@ -38,9 +38,23 @@ namespace BuildPatchServices
 
 		/**
 		 * Dequeues received messages, pushing them to the provided handlers.
-		 * @param Handlers      The array of handlers.
+		 * NOTE: PumpMessages, RegisterMessageHandler, and UnregisterMessageHandler MUST all be called from the same thread.
 		 */
-		virtual void PumpMessages(const TArray<FMessageHandler*>& Handlers) = 0;
+		virtual void PumpMessages() = 0;
+
+		/**
+		 * Registers a message handler.
+		 * @param MessageHandler    Ptr to the message handler to add. Must not be null.
+		 * NOTE: PumpMessages, RegisterMessageHandler, and UnregisterMessageHandler MUST all be called from the same thread.
+		 */
+		virtual void RegisterMessageHandler(FMessageHandler* MessageHandler) = 0;
+
+		/**
+		 * Unregisters a message handler.
+		 * @param MessageHandler    Ptr to the message handler to remove.
+		 * NOTE: PumpMessages, RegisterMessageHandler, and UnregisterMessageHandler MUST all be called from the same thread.
+		 */
+		virtual void UnregisterMessageHandler(FMessageHandler* MessageHandler) = 0;
 	};
 
 	/**
diff --git a/Engine/Source/Runtime/Online/BuildPatchServices/Private/Installer/OptimisedDelta.cpp b/Engine/Source/Runtime/Online/BuildPatchServices/Private/Installer/OptimisedDelta.cpp
--- a/Engine/Source/Runtime/Online/BuildPatchServices/Private/Installer/OptimisedDelta.cpp
+++ b/Engine/Source/Runtime/Online/BuildPatchServices/Private/Installer/OptimisedDelta.cpp
@@ -3,6 +3,7 @@
 #include "Installer/OptimisedDelta.h"
 
 #include "Async/Future.h"
+#include "Tasks/Task.h"
 #include "Misc/ConfigCacheIni.h"
 
 #include "Installer/DownloadService.h"
@@ -29,7 +30,7 @@ namespace BuildPatchServices
 		: public IOptimisedDelta
 	{
 	public:
-		FOptimisedDelta(const FOptimisedDeltaConfiguration& Configuration, const FOptimisedDeltaDependencies& Dependencies);
+		FOptimisedDelta(const FOptimisedDeltaConfiguration& Configuration, FOptimisedDeltaDependencies&& Dependencies);
 
 		// IOptimisedDelta interface begin.
 		virtual const IOptimisedDelta::FResultValueOrError& GetResult() const override;
@@ -45,9 +46,10 @@ namespace BuildPatchServices
 
 	private:
 		const FOptimisedDeltaConfiguration Configuration;
-		const FOptimisedDeltaDependencies Dependencies;
+		IDownloadService* const DownloadService;
 		const FString RelativeDeltaFilePath;
 		const int32 DeltaRetries;
+		TUniqueFunction<void(const IOptimisedDelta::FResultValueOrError&)> OnComplete;
 		EDeltaPolicy DeltaPolicy;
 		int32 CloudDirIdx;
 		int32 RetryCount;
@@ -60,11 +62,12 @@ namespace BuildPatchServices
 		const TCHAR* ErrorCode;
 	};
 
-	FOptimisedDelta::FOptimisedDelta(const FOptimisedDeltaConfiguration& InConfiguration, const FOptimisedDeltaDependencies& InDependencies)
+	FOptimisedDelta::FOptimisedDelta(const FOptimisedDeltaConfiguration& InConfiguration, FOptimisedDeltaDependencies&& InDependencies)
 		: Configuration(InConfiguration)
-		, Dependencies(InDependencies)
+		, DownloadService(InDependencies.DownloadService)
 		, RelativeDeltaFilePath(Configuration.SourceManifest.IsValid() ? FBuildPatchUtils::GetChunkDeltaFilename(*Configuration.SourceManifest.Get(), Configuration.DestinationManifest.Get()) : TEXT(""))
 		, DeltaRetries(ConfigHelpers::LoadDeltaRetries(Configuration.CloudDirectories.Num()))
+		, OnComplete(MoveTemp(InDependencies.OnComplete))
 		, DeltaPolicy(Configuration.DeltaPolicy)
 		, CloudDirIdx(0)
 		, RetryCount(0)
@@ -87,7 +90,7 @@ namespace BuildPatchServices
 		if (DeltaPolicy != EDeltaPolicy::Skip)
 		{
 			UE_LOG(LogOptimisedDelta, Log, TEXT("Requesting optimised delta file %s"), *RelativeDeltaFilePath);
-			Dependencies.DownloadService->RequestFile(Configuration.CloudDirectories[CloudDirIdx] / RelativeDeltaFilePath, ChunkDeltaComplete, ChunkDeltaProgress);
+			DownloadService->RequestFile(Configuration.CloudDirectories[CloudDirIdx] / RelativeDeltaFilePath, ChunkDeltaComplete, ChunkDeltaProgress);
 		}
 		// Otherwise we provide the standard destination manifest.
 		else
@@ -134,87 +137,92 @@ namespace BuildPatchServices
 
 	void FOptimisedDelta::OnDownloadComplete(int32 RequestId, const FDownloadRef& Download)
 	{
-		if (Download->ResponseSuccessful())
+		// Kick a task to finish this so as not to block the callback thread
+		UE::Tasks::Launch(UE_SOURCE_LOCATION, [this, RequestId, Download] 
 		{
-			// Perform a merge with current manifest so that the delta can support missing out unnecessary information.
-			FBuildPatchAppManifestPtr NewManifest;
-			FBuildPatchAppManifestRef DeltaManifest = MakeShareable(new FBuildPatchAppManifest());
-			if (DeltaManifest->DeserializeFromData(Download->GetData()))
+			if (Download->ResponseSuccessful())
 			{
-				// Verify each file received matches what we expect.
-				TSet<FString> DeltaFilenameList;
-				DeltaManifest->GetFileList(DeltaFilenameList);
-				bool bDeltaAccepted = true;
-				for (const FString& DeltaFilename : DeltaFilenameList)
+				// Perform a merge with current manifest so that the delta can support missing out unnecessary information.
+				FBuildPatchAppManifestPtr NewManifest;
+				FBuildPatchAppManifestRef DeltaManifest = MakeShareable(new FBuildPatchAppManifest());
+				if (DeltaManifest->DeserializeFromData(Download->GetData()))
 				{
-					const FFileManifest* OriginalFileManifest = Configuration.DestinationManifest->GetFileManifest(DeltaFilename);
-					const FFileManifest* DeltaFileManifest = DeltaManifest->GetFileManifest(DeltaFilename);
-					if (OriginalFileManifest != nullptr && DeltaFileManifest != nullptr)
+					// Verify each file received matches what we expect.
+					TSet<FString> DeltaFilenameList;
+					DeltaManifest->GetFileList(DeltaFilenameList);
+					bool bDeltaAccepted = true;
+					for (const FString& DeltaFilename : DeltaFilenameList)
 					{
-						FSHAHash OriginalFileHash;
-						FSHAHash DeltaFileHash;
-						Configuration.DestinationManifest->GetFileHash(DeltaFilename, OriginalFileHash);
-						DeltaManifest->GetFileHash(DeltaFilename, DeltaFileHash);
-						if (OriginalFileHash != DeltaFileHash)
+						const FFileManifest* OriginalFileManifest = Configuration.DestinationManifest->GetFileManifest(DeltaFilename);
+						const FFileManifest* DeltaFileManifest = DeltaManifest->GetFileManifest(DeltaFilename);
+						if (OriginalFileManifest != nullptr && DeltaFileManifest != nullptr)
 						{
-							UE_LOG(LogOptimisedDelta, Log, TEXT("Rejected optimised delta due to file SHA1 mismatch."));
+							FSHAHash OriginalFileHash;
+							FSHAHash DeltaFileHash;
+							Configuration.DestinationManifest->GetFileHash(DeltaFilename, OriginalFileHash);
+							DeltaManifest->GetFileHash(DeltaFilename, DeltaFileHash);
+							if (OriginalFileHash != DeltaFileHash)
+							{
+								UE_LOG(LogOptimisedDelta, Log, TEXT("Rejected optimised delta due to file SHA1 mismatch."));
+								ErrorCode = DownloadErrorCodes::RejectedDeltaFile;
+								bDeltaAccepted = false;
+								break;
+							}
+						}
+						else
+						{
+							UE_LOG(LogOptimisedDelta, Log, TEXT("Rejected optimised delta due to file list mismatch."));
 							ErrorCode = DownloadErrorCodes::RejectedDeltaFile;
 							bDeltaAccepted = false;
 							break;
 						}
 					}
-					else
+					// Only use if accepted so far.
+					if (bDeltaAccepted)
 					{
-						UE_LOG(LogOptimisedDelta, Log, TEXT("Rejected optimised delta due to file list mismatch."));
-						ErrorCode = DownloadErrorCodes::RejectedDeltaFile;
-						bDeltaAccepted = false;
-						break;
+						NewManifest = FBuildMergeManifests::MergeDeltaManifest(Configuration.DestinationManifest, DeltaManifest);
+						if (!NewManifest.IsValid())
+						{
+							UE_LOG(LogOptimisedDelta, Log, TEXT("Rejected optimised delta due to merge failure."));
+							ErrorCode = DownloadErrorCodes::RejectedDeltaFile;
+							bDeltaAccepted = false;
+						}
 					}
 				}
-				// Only use if accepted so far.
-				if (bDeltaAccepted)
+				else
 				{
-					NewManifest = FBuildMergeManifests::MergeDeltaManifest(Configuration.DestinationManifest, DeltaManifest);
-					if (!NewManifest.IsValid())
-					{
-						UE_LOG(LogOptimisedDelta, Log, TEXT("Rejected optimised delta due to merge failure."));
-						ErrorCode = DownloadErrorCodes::RejectedDeltaFile;
-						bDeltaAccepted = false;
-					}
+					ErrorCode = DownloadErrorCodes::UnserialisableDeltaFile;
+				}
+
+				if (NewManifest.IsValid())
+				{
+					UE_LOG(LogOptimisedDelta, Log, TEXT("Received optimised delta file successfully %s"), *RelativeDeltaFilePath);
+					DownloadedBytes.Set(Download->GetData().Num());
+					SetResult(NewManifest);
+				}
+				else if (ShouldRetry(Download))
+				{
+					++RetryCount;
+					CloudDirIdx = (CloudDirIdx + RetryCount) % Configuration.CloudDirectories.Num();
+					DownloadService->RequestFile(Configuration.CloudDirectories[CloudDirIdx] / RelativeDeltaFilePath, ChunkDeltaComplete, ChunkDeltaProgress);
+				}
+				else
+				{
+					SetFailedDownload();
 				}
-			}
-			else
-			{
-				ErrorCode = DownloadErrorCodes::UnserialisableDeltaFile;
-			}
-			if (NewManifest.IsValid())
-			{
-				UE_LOG(LogOptimisedDelta, Log, TEXT("Received optimised delta file successfully %s"), *RelativeDeltaFilePath);
-				DownloadedBytes.Set(Download->GetData().Num());
-				SetResult(NewManifest);
 			}
 			else if (ShouldRetry(Download))
 			{
 				++RetryCount;
 				CloudDirIdx = (CloudDirIdx + RetryCount) % Configuration.CloudDirectories.Num();
-				Dependencies.DownloadService->RequestFile(Configuration.CloudDirectories[CloudDirIdx] / RelativeDeltaFilePath, ChunkDeltaComplete, ChunkDeltaProgress);
+				DownloadService->RequestFile(Configuration.CloudDirectories[CloudDirIdx] / RelativeDeltaFilePath, ChunkDeltaComplete, ChunkDeltaProgress);
 			}
 			else
 			{
+				ErrorCode = DownloadErrorCodes::MissingDeltaFile;
 				SetFailedDownload();
 			}
-		}
-		else if (ShouldRetry(Download))
-		{
-			++RetryCount;
-			CloudDirIdx = (CloudDirIdx + RetryCount) % Configuration.CloudDirectories.Num();
-			Dependencies.DownloadService->RequestFile(Configuration.CloudDirectories[CloudDirIdx] / RelativeDeltaFilePath, ChunkDeltaComplete, ChunkDeltaProgress);
-		}
-		else
-		{
-			ErrorCode = DownloadErrorCodes::MissingDeltaFile;
-			SetFailedDownload();
-		}
+		}, UE::Tasks::ETaskPriority::BackgroundHigh);
 	}
 
 	bool FOptimisedDelta::ShouldRetry(const FDownloadRef& Download)
@@ -229,8 +237,14 @@ namespace BuildPatchServices
 
 	void FOptimisedDelta::SetResult(const FBuildPatchAppManifestPtr& ResultManifest)
 	{
-		ChunkDeltaPromise.SetValue(MakeShared<FResultValueOrError>(MakeValue(ResultManifest)));
-		Dependencies.OnComplete(GetResult());
+		TSharedPtr<FResultValueOrError> Result = MakeShared<FResultValueOrError>(MakeValue(ResultManifest));
+		TUniqueFunction<void(const IOptimisedDelta::FResultValueOrError&)> LocalOnComplete = MoveTemp(OnComplete);
+
+		ChunkDeltaPromise.SetValue(Result);
+		if (LocalOnComplete)
+		{
+			LocalOnComplete(*Result);
+		}
 	}
 
 	void FOptimisedDelta::SetFailedDownload()
@@ -243,8 +257,15 @@ namespace BuildPatchServices
 		else
 		{
 			UE_LOG(LogOptimisedDelta, Log, TEXT("Failed optimised delta file fetch %s"), *RelativeDeltaFilePath);
-			ChunkDeltaPromise.SetValue(MakeShared<FResultValueOrError>(MakeError(ErrorCode)));
-			Dependencies.OnComplete(GetResult());
+
+			TSharedPtr<FResultValueOrError> Result = MakeShared<FResultValueOrError>(MakeError(ErrorCode));
+			TUniqueFunction<void(const IOptimisedDelta::FResultValueOrError&)> LocalOnComplete = MoveTemp(OnComplete);
+
+			ChunkDeltaPromise.SetValue(Result);
+			if (LocalOnComplete)
+			{
+				LocalOnComplete(*Result);
+			}
 		}
 	}
 
@@ -255,15 +276,9 @@ namespace BuildPatchServices
 	{
 	}
 
-	FOptimisedDeltaDependencies::FOptimisedDeltaDependencies()
-		: DownloadService(nullptr)
-		, OnComplete([](const IOptimisedDelta::FResultValueOrError&) {})
-	{
-	}
-
-	IOptimisedDelta* FOptimisedDeltaFactory::Create(const FOptimisedDeltaConfiguration& Configuration, const FOptimisedDeltaDependencies& Dependencies)
+	IOptimisedDelta* FOptimisedDeltaFactory::Create(const FOptimisedDeltaConfiguration& Configuration, FOptimisedDeltaDependencies&& Dependencies)
 	{
 		check(Dependencies.DownloadService != nullptr);
-		return new FOptimisedDelta(Configuration, Dependencies);
+		return new FOptimisedDelta(Configuration, MoveTemp(Dependencies));
 	}
 }
diff --git a/Engine/Source/Runtime/Online/BuildPatchServices/Private/Installer/OptimisedDelta.h b/Engine/Source/Runtime/Online/BuildPatchServices/Private/Installer/OptimisedDelta.h
--- a/Engine/Source/Runtime/Online/BuildPatchServices/Private/Installer/OptimisedDelta.h
+++ b/Engine/Source/Runtime/Online/BuildPatchServices/Private/Installer/OptimisedDelta.h
@@ -62,22 +62,16 @@ namespace BuildPatchServices
 	 */
 	struct FOptimisedDeltaDependencies
 	{
-	public:
-		/**
-		 * Constructor setting up default values.
-		 */
-		FOptimisedDeltaDependencies();
-
 	public:
 		// A download service instance.
-		IDownloadService* DownloadService;
+		IDownloadService* DownloadService = nullptr;
 		// Function to call once the destination manifest has been selected. Receives manifest or error.
-		TFunction<void(const IOptimisedDelta::FResultValueOrError&)> OnComplete;
+		TUniqueFunction<void(const IOptimisedDelta::FResultValueOrError&)> OnComplete;
 	};
 
 	class FOptimisedDeltaFactory
 	{
 	public:
-		static IOptimisedDelta* Create(const FOptimisedDeltaConfiguration& Configuration, const FOptimisedDeltaDependencies& Dependencies);
+		static IOptimisedDelta* Create(const FOptimisedDeltaConfiguration& Configuration, FOptimisedDeltaDependencies&& Dependencies);
 	};
 }
diff --git a/Engine/Source/Runtime/Online/BuildPatchServices/Private/Installer/Statistics/DownloadServiceStatistics.cpp b/Engine/Source/Runtime/Online/BuildPatchServices/Private/Installer/Statistics/DownloadServiceStatistics.cpp
--- a/Engine/Source/Runtime/Online/BuildPatchServices/Private/Installer/Statistics/DownloadServiceStatistics.cpp
+++ b/Engine/Source/Runtime/Online/BuildPatchServices/Private/Installer/Statistics/DownloadServiceStatistics.cpp
@@ -20,7 +20,6 @@ namespace BuildPatchServices
 	{
 	public:
 		FDownloadServiceStatistics(ISpeedRecorder* SpeedRecorder, IDataSizeProvider* DataSizeProvider, IInstallerAnalytics* InstallerAnalytics);
-		~FDownloadServiceStatistics();
 
 		// IDownloadServiceStat interface begin.
 		virtual void OnDownloadStarted(int32 RequestId, const FString& Uri) override;
@@ -39,19 +38,21 @@ namespace BuildPatchServices
 		// IDownloadServiceStatistics interface end.
 
 	private:
-		ISpeedRecorder* SpeedRecorder;
+		ISpeedRecorder* SpeedRecorder; // AddRecord is safe, but SpeedRecorder is ticked
 		IDataSizeProvider* DataSizeProvider;
 		IInstallerAnalytics* InstallerAnalytics;
-		FThreadSafeInt64 TotalBytesReceived;
-		FThreadSafeInt32 NumSuccessfulDownloads;
-		FThreadSafeInt32 NumFailedDownloads;
+		
+		std::atomic<uint64> TotalBytesReceived;
+		std::atomic<int32> NumSuccessfulDownloads;
+		std::atomic<int32> NumFailedDownloads;
 
 		typedef TTuple<FString, uint64> FDownloadTuple;
 		TMap<int32, FDownloadTuple> Downloads;
+		mutable FCriticalSection DownloadsCriticalSection;
+
 		double AccumulatedRequestSpeed;
 		uint32 AverageSpeedSampleCount;
-		FCriticalSection AverageSpeedCriticalSection;
-		
+		mutable FCriticalSection AverageSpeedCriticalSection;
 	};
 
 	FDownloadServiceStatistics::FDownloadServiceStatistics(ISpeedRecorder* InSpeedRecorder, IDataSizeProvider* InDataSizeProvider, IInstallerAnalytics* InInstallerAnalytics)
@@ -61,17 +62,14 @@ namespace BuildPatchServices
 		, TotalBytesReceived(0)
 		, NumSuccessfulDownloads(0)
 		, NumFailedDownloads(0)
+		, AccumulatedRequestSpeed(0.0)
 		, AverageSpeedSampleCount(0U)
 	{
 	}
 
-	FDownloadServiceStatistics::~FDownloadServiceStatistics()
-	{
-	}
-
 	void FDownloadServiceStatistics::OnDownloadStarted(int32 RequestId, const FString& Uri)
 	{
-		checkSlow(IsInGameThread());
+		FScopeLock Lock(&DownloadsCriticalSection);
 		FDownloadTuple& DownloadTuple = Downloads.FindOrAdd(RequestId);
 		DownloadTuple.Get<0>() = Uri;
 		DownloadTuple.Get<1>() = 0;
@@ -79,73 +77,94 @@ namespace BuildPatchServices
 
 	void FDownloadServiceStatistics::OnDownloadProgress(int32 RequestId, uint64 BytesReceived)
 	{
-		checkSlow(IsInGameThread());
+		FScopeLock Lock(&DownloadsCriticalSection);
 		FDownloadTuple& DownloadTuple = Downloads.FindOrAdd(RequestId);
 		DownloadTuple.Get<1>() = BytesReceived;
 	}
 
 	void FDownloadServiceStatistics::OnDownloadComplete(const FDownloadRecord& DownloadRecord)
 	{
-		checkSlow(IsInGameThread());
-		Downloads.Remove(DownloadRecord.RequestId);
+		{
+			FScopeLock Lock(&DownloadsCriticalSection);
+			Downloads.Remove(DownloadRecord.RequestId);
+		}
+
 		if (DownloadRecord.bSuccess && EHttpResponseCodes::IsOk(DownloadRecord.ResponseCode))
 		{
-			TotalBytesReceived.Add(DownloadRecord.SpeedRecord.Size);
-			NumSuccessfulDownloads.Increment();
+			TotalBytesReceived += DownloadRecord.SpeedRecord.Size;
+			NumSuccessfulDownloads++;
 			SpeedRecorder->AddRecord(DownloadRecord.SpeedRecord);
+
+			// update average speed
 			const uint64 Cycles = DownloadRecord.SpeedRecord.CyclesEnd - DownloadRecord.SpeedRecord.CyclesStart;
 			double Speed = Cycles > 0U ? DownloadRecord.SpeedRecord.Size / FStatsCollector::CyclesToSeconds(Cycles) : 0.0;
 			FScopeLock Lock(&AverageSpeedCriticalSection);
-			AccumulatedRequestSpeed +=  Speed;
+			AccumulatedRequestSpeed += Speed;
 			AverageSpeedSampleCount += 1;
 		}
 		else
 		{
-			NumFailedDownloads.Increment();
+			NumFailedDownloads++;
 			InstallerAnalytics->RecordChunkDownloadError(DownloadRecord.Uri, DownloadRecord.ResponseCode, TEXT("DownloadFail"));
 		}
 	}
 
 	uint64 FDownloadServiceStatistics::GetBytesDownloaded() const
 	{
-		return TotalBytesReceived.GetValue();
+		return TotalBytesReceived;
 	}
 
 	int32 FDownloadServiceStatistics::GetNumSuccessfulChunkDownloads() const
 	{
-		return NumSuccessfulDownloads.GetValue();
+		return NumSuccessfulDownloads;
 	}
 
 	int32 FDownloadServiceStatistics::GetNumFailedChunkDownloads() const
 	{
-		return NumFailedDownloads.GetValue();
+		return NumFailedDownloads;
 	}
 
 	int32 FDownloadServiceStatistics::GetNumCurrentDownloads() const
 	{
+		FScopeLock Lock(&DownloadsCriticalSection);
 		return Downloads.Num();
 	}
 
 	TArray<FDownload> FDownloadServiceStatistics::GetCurrentDownloads() const
 	{
-		checkSlow(IsInGameThread());
+		TArray<FString> DownloadData;
+		TArray<uint64> DownloadSizes;
 		TArray<FDownload> Result;
-		Result.Empty(Downloads.Num());
-		for (const TPair<int32, FDownloadTuple>& Download : Downloads)
+
 		{
-			Result.AddDefaulted();
-			FDownload& Element = Result.Last();
-			Element.Data = FPaths::GetCleanFilename(Download.Value.Get<0>());
-			Element.Size = DataSizeProvider->GetDownloadSize(Element.Data);
-			Element.Received = Download.Value.Get<1>();
+			FScopeLock Lock(&DownloadsCriticalSection);
+			DownloadData.Empty(Downloads.Num());
+			Result.Empty(Downloads.Num());
+			for (const TPair<int32, FDownloadTuple>& Download : Downloads)
+			{
+				DownloadData.Emplace(FPaths::GetCleanFilename(Download.Value.Get<0>()));
+
+				Result.AddDefaulted();
+				FDownload& Element = Result.Last();
+				Element.Received = Download.Value.Get<1>();
+			}
+		}
+
+		DataSizeProvider->GetDownloadSize(DownloadData, DownloadSizes);
+		check(DownloadSizes.Num() == Result.Num());
+
+		for (int32 Index = 0; FDownload& Element : Result)
+		{
+			Element.Data = MoveTemp(DownloadData[Index]);
+			Element.Size = DownloadSizes[Index];
+
+			++Index;
 		}
 		return Result;
 	}
 
 	TPair<double, uint32> FDownloadServiceStatistics::GetImmediateAverageSpeedPerRequest(uint32 MinCount)
 	{
-		double Average = 0.0L;
-		double OverallAverage = 0.0L;
 		uint32 Count = 0U;
 		double Result = 0.0L;
 		FScopeLock Lock(&AverageSpeedCriticalSection);
@@ -164,13 +183,19 @@ namespace BuildPatchServices
 		checkSlow(IsInGameThread());
 
 		TotalBytesReceived = 0;
-		NumSuccessfulDownloads.Set(0);
-		NumFailedDownloads.Set(0);
-		Downloads.Empty();
+		NumSuccessfulDownloads = 0;
+		NumFailedDownloads = 0;
 
-		FScopeLock Lock(&AverageSpeedCriticalSection);
-		AccumulatedRequestSpeed = 0;
-		AverageSpeedSampleCount = 0;
+		{
+			FScopeLock Lock(&DownloadsCriticalSection);
+			Downloads.Empty();
+		}
+
+		{
+			FScopeLock Lock(&AverageSpeedCriticalSection);
+			AccumulatedRequestSpeed = 0;
+			AverageSpeedSampleCount = 0;
+		}
 	}
 
 	IDownloadServiceStatistics* FDownloadServiceStatisticsFactory::Create(ISpeedRecorder* SpeedRecorder, IDataSizeProvider* DataSizeProvider, IInstallerAnalytics* InstallerAnalytics)
diff --git a/Engine/Source/Runtime/Online/BuildPatchServices/Private/Installer/Statistics/FileOperationTracker.cpp b/Engine/Source/Runtime/Online/BuildPatchServices/Private/Installer/Statistics/FileOperationTracker.cpp
--- a/Engine/Source/Runtime/Online/BuildPatchServices/Private/Installer/Statistics/FileOperationTracker.cpp
+++ b/Engine/Source/Runtime/Online/BuildPatchServices/Private/Installer/Statistics/FileOperationTracker.cpp
@@ -23,6 +23,26 @@ namespace BuildPatchServices
 	typedef TTuple<FString, FByteRange, EFileOperationState> FFileByteRangeState;
 	typedef TUnion<FOperationInitialiser, FDataState, FFileState, FFileByteRangeState> FUpdateMessage;
 
+	class FNullFileOperationTracker 
+		: public IFileOperationTracker
+	{
+	public:
+		// IFileOperationTracker interface begin.
+		virtual const TArray<FFileOperation>& GetStates() const override { return FileOperationStates; }
+		virtual void OnManifestSelection(IBuildManifestSet* ManifestSet) override {}
+		virtual void OnDataStateUpdate(const FGuid& DataId, EFileOperationState State) override {}
+		virtual void OnDataStateUpdate(const TSet<FGuid>& DataIds, EFileOperationState State) override {}
+		virtual void OnDataStateUpdate(const TArray<FGuid>& DataIds, EFileOperationState State) override {}
+		virtual void OnFileStateUpdate(const FString& Filename, EFileOperationState State) override {}
+		virtual void OnFileStateUpdate(const TSet<FString>& Filenames, EFileOperationState State) override {}
+		virtual void OnFileStateUpdate(const TArray<FString>& Filenames, EFileOperationState State) override {}
+		virtual void OnFileByteRangeStateUpdate(const FString& Filename, FByteRange ByteRange, EFileOperationState State) override {}
+		// IFileOperationTracker interface end.
+
+	private:
+		TArray<FFileOperation> FileOperationStates;
+	};
+
 	class FFileOperationTracker
 		: public IFileOperationTracker
 	{
@@ -269,4 +289,9 @@ namespace BuildPatchServices
 	{
 		return new FFileOperationTracker(Ticker);
 	}
+
+	IFileOperationTracker* FFileOperationTrackerFactory::CreateNull()
+	{
+		return new FNullFileOperationTracker();
+	}
 }
\ No newline at end of file
diff --git a/Engine/Source/Runtime/Online/BuildPatchServices/Private/Installer/Statistics/FileOperationTracker.h b/Engine/Source/Runtime/Online/BuildPatchServices/Private/Installer/Statistics/FileOperationTracker.h
--- a/Engine/Source/Runtime/Online/BuildPatchServices/Private/Installer/Statistics/FileOperationTracker.h
+++ b/Engine/Source/Runtime/Online/BuildPatchServices/Private/Installer/Statistics/FileOperationTracker.h
@@ -96,5 +96,13 @@ namespace BuildPatchServices
 		 * @return the new IFileOperationTracker instance created.
 		 */
 		static IFileOperationTracker* Create(FTSTicker& Ticker);
+
+		/**
+		 * Creates a null implementation of a file operation tracker which serves as the dependency for for systems providing these updates.
+		 * @param   Ticker      The ticker to register for main thread ticks with.
+		 * @param   Manifest    The manifest for the build being installed.
+		 * @return the new IFileOperationTracker instance created.
+		 */
+		static IFileOperationTracker* CreateNull();
 	};
 }
\ No newline at end of file
diff --git a/Engine/Source/Runtime/Online/BuildPatchServices/Private/Tests/Mock/MessagePump.mock.h b/Engine/Source/Runtime/Online/BuildPatchServices/Private/Tests/Mock/MessagePump.mock.h
--- a/Engine/Source/Runtime/Online/BuildPatchServices/Private/Tests/Mock/MessagePump.mock.h
+++ b/Engine/Source/Runtime/Online/BuildPatchServices/Private/Tests/Mock/MessagePump.mock.h
@@ -19,13 +19,21 @@ namespace BuildPatchServices
 		{
 		}
 
-		virtual void PumpMessages(const TArray<FMessageHandler*>& Handlers) override
+		virtual void PumpMessages() override
 		{
 		}
 
 		virtual void SendRequest(FChunkUriRequest Request, TFunction<void(FChunkUriResponse)> OnResponse) override
 		{
 		}
+
+		virtual void RegisterMessageHandler(FMessageHandler* MessageHandler) override
+		{
+		}
+
+		virtual void UnregisterMessageHandler(FMessageHandler* MessageHandler) override
+		{
+		}
 	};
 }
 
diff --git a/Engine/Source/Runtime/Online/BuildPatchServices/Private/Tests/Unit/DownloadService.spec.cpp b/Engine/Source/Runtime/Online/BuildPatchServices/Private/Tests/Unit/DownloadService.spec.cpp
--- a/Engine/Source/Runtime/Online/BuildPatchServices/Private/Tests/Unit/DownloadService.spec.cpp
+++ b/Engine/Source/Runtime/Online/BuildPatchServices/Private/Tests/Unit/DownloadService.spec.cpp
@@ -63,7 +63,6 @@ void FDownloadServiceSpec::Define()
 		MockDownloadServiceStat.Reset(new FMockDownloadServiceStat());
 		MockInstallerAnalytics.Reset(new FMockInstallerAnalytics());
 		DownloadService.Reset(FDownloadServiceFactory::Create(
-			Ticker,
 			FakeHttpModule.Get(),
 			MockFileSystem.Get(),
 			MockDownloadServiceStat.Get(),
diff --git a/Engine/Source/Runtime/Online/BuildPatchServices/Public/BuildPatchMessage.h b/Engine/Source/Runtime/Online/BuildPatchServices/Public/BuildPatchMessage.h
--- a/Engine/Source/Runtime/Online/BuildPatchServices/Public/BuildPatchMessage.h
+++ b/Engine/Source/Runtime/Online/BuildPatchServices/Public/BuildPatchMessage.h
@@ -65,14 +65,29 @@ namespace BuildPatchServices
 		FString Uri;
 	};
 
+	// Flag for which Requests a message handler expects to receive, allows for internal implementation optimisation.
+	enum class EMessageRequests : uint32
+	{
+		// Does not respond to any requests - message listener only.
+		None = 0,
+		// Will respond to chunk URI requests.
+		ChunkUriRequest = 0x1,
+		// Further request types to follow in future.
+
+	};
+	ENUM_CLASS_FLAGS(EMessageRequests);
+
 	/**
 	 * Base class of a message handler, this should be inherited from and passed to an installer to receive messages that you want to handle.
 	 */
 	class FMessageHandler
 	{
 	public:
-		FMessageHandler() {}
-		virtual ~FMessageHandler() {}
+		FMessageHandler(EMessageRequests InMessageRequests)
+			: MessageRequests(InMessageRequests)
+		{}
+
+		virtual ~FMessageHandler() = default;
 
 		/**
 		 * Handles a chunk source event message.
@@ -92,12 +107,25 @@ namespace BuildPatchServices
 		 * @param OnResponse   The function to callback once the chunk has been found
 		 */
 		virtual bool HandleRequest(const FChunkUriRequest& Request, TFunction<void(FChunkUriResponse)> OnResponse) { return false; }
+
+		/**
+		 * @return the message request flags.
+		 */
+		EMessageRequests GetMessageRequests() const { return MessageRequests; }
+
+	private:
+		const EMessageRequests MessageRequests;
 	};
 
 	class FDefaultMessageHandler
 		: public FMessageHandler
 	{
 	public:
+		// Default constructor passes flags to always support all requests
+		FDefaultMessageHandler()
+			: FMessageHandler(static_cast<EMessageRequests>(0xFFFFFFFF))
+		{}
+
 		virtual bool HandleRequest(const FChunkUriRequest& Request, TFunction<void(FChunkUriResponse)> OnResponse) override
 		{
 			OnResponse({ Request.CloudDirectory / Request.RelativePath });
diff --git a/Engine/Source/Runtime/Online/BuildPatchServices/Public/BuildPatchSettings.h b/Engine/Source/Runtime/Online/BuildPatchServices/Public/BuildPatchSettings.h
--- a/Engine/Source/Runtime/Online/BuildPatchServices/Public/BuildPatchSettings.h
+++ b/Engine/Source/Runtime/Online/BuildPatchServices/Public/BuildPatchSettings.h
@@ -288,6 +288,8 @@ namespace BuildPatchServices
 		bool bRunRequiredPrereqs;
 		// Whether to allow this installation to run concurrently with any existing installations.
 		bool bAllowConcurrentExecution;
+		// Whether to gather individual file operation statistics during install
+		bool bTrackFileOperations;
 	};
 
 	/**
