Use LogInit log channel when init http.

diff --git a/Engine/Source/Runtime/Online/HTTP/Private/Curl/CurlHttpManager.cpp b/Engine/Source/Runtime/Online/HTTP/Private/Curl/CurlHttpManager.cpp
--- a/Engine/Source/Runtime/Online/HTTP/Private/Curl/CurlHttpManager.cpp
+++ b/Engine/Source/Runtime/Online/HTTP/Private/Curl/CurlHttpManager.cpp
@@ -456,22 +456,22 @@ FHttpThreadBase* FCurlHttpManager::CreateHttpThread()
 	if (bUseEventLoop)
 	{
 #if WITH_CURL_MULTIPOLL
-		UE_LOG(LogHttp, Log, TEXT("CreateHttpThread using FCurlMultiPollEventLoopHttpThread"));
+		UE_LOG(LogInit, Log, TEXT("CreateHttpThread using FCurlMultiPollEventLoopHttpThread"));
 		return new FCurlMultiPollEventLoopHttpThread();
 #endif // WITH_CURL_MULTIPOLL
 
 #if WITH_CURL_MULTISOCKET
-		UE_LOG(LogHttp, Log, TEXT("CreateHttpThread using FCurlSocketEventLoopHttpThread"));
+		UE_LOG(LogInit, Log, TEXT("CreateHttpThread using FCurlSocketEventLoopHttpThread"));
 		return new FCurlSocketEventLoopHttpThread();
 #endif // WITH_CURL_MULTISOCKET
 
 #if WITH_CURL_MULTIWAIT
-		UE_LOG(LogHttp, Log, TEXT("CreateHttpThread using FCurlMultiWaitEventLoopHttpThread"));
+		UE_LOG(LogInit, Log, TEXT("CreateHttpThread using FCurlMultiWaitEventLoopHttpThread"));
 		return new FCurlMultiWaitEventLoopHttpThread();
 #endif // WITH_CURL_MULTIWAIT
 	}
 
-	UE_LOG(LogHttp, Log, TEXT("CreateHttpThread using FCurlHttpThread"));
+	UE_LOG(LogInit, Log, TEXT("CreateHttpThread using FCurlHttpThread"));
 	return new FCurlHttpThread();
 }
 
