Enable more platforms for WebTests, fix corresponding warning or errors.

diff --git a/Engine/Source/Programs/WebTests/Private/TestHttp.cpp b/Engine/Source/Programs/WebTests/Private/TestHttp.cpp
--- a/Engine/Source/Programs/WebTests/Private/TestHttp.cpp
+++ b/Engine/Source/Programs/WebTests/Private/TestHttp.cpp
@@ -28,6 +28,8 @@
 #define HTTP_TAG "[HTTP]"
 #define HTTP_TIME_DIFF_TOLERANCE 0.5f
 
+extern TAutoConsoleVariable<bool> CVarHttpInsecureProtocolEnabled;
+
 class FHttpModuleTestFixture
 {
 public:
@@ -42,6 +44,8 @@ public:
 		HttpModule = new FHttpModule();
 		IModuleInterface* Module = HttpModule;
 		Module->StartupModule();
+
+		CVarHttpInsecureProtocolEnabled->Set(true);
 	}
 
 	virtual ~FHttpModuleTestFixture()
diff --git a/Engine/Source/Programs/WebTests/Private/TestWebSockets.cpp b/Engine/Source/Programs/WebTests/Private/TestWebSockets.cpp
--- a/Engine/Source/Programs/WebTests/Private/TestWebSockets.cpp
+++ b/Engine/Source/Programs/WebTests/Private/TestWebSockets.cpp
@@ -98,6 +98,7 @@ public:
 		while (!bQuitRequested)
 		{
 			FTSBackgroundableTicker::GetCoreTicker().Tick(TickFrequency);
+			FTSTicker::GetCoreTicker().Tick(TickFrequency);
 
 			FPlatformProcess::Sleep(TickFrequency);
 		}
diff --git a/Engine/Source/Runtime/Online/HTTP/Internal/WinHttp/WinHttpHttpManager.h b/Engine/Source/Runtime/Online/HTTP/Internal/WinHttp/WinHttpHttpManager.h
--- a/Engine/Source/Runtime/Online/HTTP/Internal/WinHttp/WinHttpHttpManager.h
+++ b/Engine/Source/Runtime/Online/HTTP/Internal/WinHttp/WinHttpHttpManager.h
@@ -67,8 +67,6 @@ protected:
 	HTTP_API FWinHttpSession* FindOrCreateSession(const uint32 SecurityProtocols);
 
 protected:
-	bool bPlatformForcesSecureConnections = false;
-
 	/** Map of Security Flags to WinHttp Session objects */
 	TMap<uint32, TUniquePtr<FWinHttpSession>> ActiveSessions;
 };
diff --git a/Engine/Source/Runtime/Online/HTTP/Private/HttpModule.cpp b/Engine/Source/Runtime/Online/HTTP/Private/HttpModule.cpp
--- a/Engine/Source/Runtime/Online/HTTP/Private/HttpModule.cpp
+++ b/Engine/Source/Runtime/Online/HTTP/Private/HttpModule.cpp
@@ -10,6 +10,18 @@
 #include "Misc/CommandLine.h"
 #include "Misc/CoreDelegates.h"
 
+#if !UE_BUILD_SHIPPING
+#include "HAL/IConsoleManager.h"
+#endif
+
+#if !UE_BUILD_SHIPPING
+TAutoConsoleVariable<bool> CVarHttpInsecureProtocolEnabled(
+	TEXT("Http.InsecureProtocolEnabled"),
+	false,
+	TEXT("Enable insecure http protocol")
+);
+#endif
+
 DEFINE_LOG_CATEGORY(LogHttp);
 
 // FHttpModule
diff --git a/Engine/Source/Runtime/Online/HTTP/Private/WinHttp/Support/WinHttpSession.cpp b/Engine/Source/Runtime/Online/HTTP/Private/WinHttp/Support/WinHttpSession.cpp
--- a/Engine/Source/Runtime/Online/HTTP/Private/WinHttp/Support/WinHttpSession.cpp
+++ b/Engine/Source/Runtime/Online/HTTP/Private/WinHttp/Support/WinHttpSession.cpp
@@ -83,10 +83,10 @@ FWinHttpSession::FWinHttpSession(uint32 SecurityProtocolFlags, const bool bInFor
 	}
 
 	FHttpModule& HttpModule = FHttpModule::Get();
-	const FTimespan ConnectionTimeout = FTimespan::FromSeconds(HttpModule.GetHttpConnectionTimeout());
+	const FTimespan ConnectionTimeout = FTimespan::FromSeconds(HttpModule.GetHttpConnectionTimeout() > 0 ? HttpModule.GetHttpConnectionTimeout() : 0);
 	const FTimespan ResolveTimeout = ConnectionTimeout;
-	const FTimespan SendTimeout = FTimespan::FromSeconds(HttpModule.GetHttpSendTimeout());
-	const FTimespan ReceiveTimeout = FTimespan::FromSeconds(HttpModule.GetHttpReceiveTimeout());
+	const FTimespan SendTimeout = FTimespan::FromSeconds(HttpModule.GetHttpSendTimeout() > 0 ? HttpModule.GetHttpSendTimeout() : 0);
+	const FTimespan ReceiveTimeout = FTimespan::FromSeconds(HttpModule.GetHttpReceiveTimeout() > 0 ? HttpModule.GetHttpReceiveTimeout() : 0);
 
 	if (!WinHttpSetTimeouts(SessionHandle.Get(), ResolveTimeout.GetTotalMilliseconds(), ConnectionTimeout.GetTotalMilliseconds(), SendTimeout.GetTotalMilliseconds(), ReceiveTimeout.GetTotalMilliseconds()))
 	{
diff --git a/Engine/Source/Runtime/Online/HTTP/Private/WinHttp/WinHttpHttpManager.cpp b/Engine/Source/Runtime/Online/HTTP/Private/WinHttp/WinHttpHttpManager.cpp
--- a/Engine/Source/Runtime/Online/HTTP/Private/WinHttp/WinHttpHttpManager.cpp
+++ b/Engine/Source/Runtime/Online/HTTP/Private/WinHttp/WinHttpHttpManager.cpp
@@ -9,6 +9,14 @@
 #include "Misc/CoreDelegates.h"
 #include "Stats/Stats.h"
 
+#if !UE_BUILD_SHIPPING
+#include "HAL/IConsoleManager.h"
+#endif
+
+#if !UE_BUILD_SHIPPING
+extern TAutoConsoleVariable<bool> CVarHttpInsecureProtocolEnabled;
+#endif
+
 namespace
 {
 	FWinHttpHttpManager* GWinHttpManager = nullptr;
@@ -133,7 +141,12 @@ FWinHttpSession* FWinHttpHttpManager::FindOrCreateSession(const uint32 SecurityP
 	FWinHttpSession* SessionPtr = SessionPtrPtr ? SessionPtrPtr->Get() : nullptr;
 	if (!SessionPtr)
 	{
-		SessionPtr = ActiveSessions.Emplace(SecurityProtocols, MakeUnique<FWinHttpSession>(SecurityProtocols, bPlatformForcesSecureConnections)).Get();
+#if !UE_BUILD_SHIPPING
+		bool bForceSecureConnections = !CVarHttpInsecureProtocolEnabled.GetValueOnAnyThread();
+#else
+		bool bForceSecureConnections = true;
+#endif
+		SessionPtr = ActiveSessions.Emplace(SecurityProtocols, MakeUnique<FWinHttpSession>(SecurityProtocols, bForceSecureConnections)).Get();
 	}
 
 	return SessionPtr;
