Remove WinHttp request, since we switched to CurlHttp.

diff --git a/Engine/Source/Runtime/Online/HTTP/Internal/WinHttp/WinHttpHttpManager.h b/Engine/Source/Runtime/Online/HTTP/Internal/WinHttp/WinHttpHttpManager.h
--- a/Engine/Source/Runtime/Online/HTTP/Internal/WinHttp/WinHttpHttpManager.h
+++ b/Engine/Source/Runtime/Online/HTTP/Internal/WinHttp/WinHttpHttpManager.h
@@ -13,8 +13,7 @@ DECLARE_DELEGATE_OneParam(FWinHttpQuerySessionComplete, FWinHttpSession* /*HttpS
 
 class IWinHttpConnection;
 
-class FWinHttpHttpManager
-	: public FHttpManager
+class FWinHttpHttpManager : public FTSTickerObjectBase
 {
 public:
 	static HTTP_API FWinHttpHttpManager* GetManager();
@@ -28,7 +27,7 @@ public:
 	 * @param Url The URL to find or create a WinHttp session for
 	 * @param Delegate The delegate that is called with the WinHttp session pointer if successful, or null otherwise
 	 */
-	HTTP_API virtual void QuerySessionForUrl(const FString& Url, FWinHttpQuerySessionComplete&& Delegate);
+	HTTP_API virtual void QuerySessionForUrl(const FString& Url, FWinHttpQuerySessionComplete&& Delegate) = 0;
 
 	/**
 	 * Validate the provided connection before we start sending our request.
@@ -38,7 +37,7 @@ public:
 	 * @param Connection the connection to validate
 	 * @return True if the request validated successfully, false if not
 	 */
-	HTTP_API virtual bool ValidateRequestCertificates(IWinHttpConnection& Connection);
+	HTTP_API virtual bool ValidateRequestCertificates(IWinHttpConnection& Connection) = 0;
 
 	/**
 	 * Release any resources that were pinned for this request.
@@ -47,28 +46,7 @@ public:
 	 *
 	 * @param Connection The connection to release resources for
 	 */
-	HTTP_API virtual void ReleaseRequestResources(IWinHttpConnection& Connection);
-
-	//~ Begin FHttpManager Interface
-	HTTP_API virtual void OnBeforeFork() override;
-	//~ End FHttpManager Interface
-
-	/**
-	 * Handle the application preparing to suspend
-	 */
-	HTTP_API virtual void HandleApplicationSuspending();
-
-	/**
-	 * Handle the application preparing to suspend
-	 */
-	HTTP_API virtual void HandleApplicationResuming();
-
-protected:
-	HTTP_API FWinHttpSession* FindOrCreateSession(const uint32 SecurityProtocols);
-
-protected:
-	/** Map of Security Flags to WinHttp Session objects */
-	TMap<uint32, TUniquePtr<FWinHttpSession>> ActiveSessions;
+	HTTP_API virtual void ReleaseRequestResources(IWinHttpConnection& Connection) = 0;
 };
 
 #endif // WITH_WINHTTP
diff --git a/Engine/Source/Runtime/Online/HTTP/Private/WinHttp/WinHttpHttpManager.cpp b/Engine/Source/Runtime/Online/HTTP/Private/WinHttp/WinHttpHttpManager.cpp
--- a/Engine/Source/Runtime/Online/HTTP/Private/WinHttp/WinHttpHttpManager.cpp
+++ b/Engine/Source/Runtime/Online/HTTP/Private/WinHttp/WinHttpHttpManager.cpp
@@ -3,10 +3,11 @@
 #if WITH_WINHTTP
 
 #include "WinHttp/WinHttpHttpManager.h"
+#include "Containers/BackgroundableTicker.h"
+#include "HttpModule.h"
 #include "WinHttp/Support/WinHttpSession.h"
 #include "WinHttp/Support/WinHttpTypes.h"
 #include "Http.h"
-#include "Misc/CoreDelegates.h"
 #include "Stats/Stats.h"
 
 #if !UE_BUILD_SHIPPING
@@ -20,26 +21,6 @@ extern TAutoConsoleVariable<bool> CVarHttpInsecureProtocolEnabled;
 namespace
 {
 	FWinHttpHttpManager* GWinHttpManager = nullptr;
-
-	DWORD GetPlatformProtocolFlags()
-	{
-		// Enable "all" protocols (but not SSL2, it is insecure).
-		// For legacy reasons, "all" isn't actually all protocols, so we explicitly enable some below based on windows versions
-		DWORD ProtocolFlags = WINHTTP_FLAG_SECURE_PROTOCOL_ALL & ~WINHTTP_FLAG_SECURE_PROTOCOL_SSL2;
-#if PLATFORM_WINDOWS
-		const bool bIsWindows7OrGreater = FPlatformMisc::VerifyWindowsVersion(6, 1);
-		if (bIsWindows7OrGreater)
-		{
-			ProtocolFlags |= WINHTTP_FLAG_SECURE_PROTOCOL_TLS1_1;
-		}
-		const bool bIsWindows8Point1OrGreater = FPlatformMisc::VerifyWindowsVersion(6, 3);
-		if (bIsWindows8Point1OrGreater)
-		{
-			ProtocolFlags |= WINHTTP_FLAG_SECURE_PROTOCOL_TLS1_2;
-		}
-#endif // PLATFORM_WINDOWS
-		return ProtocolFlags;
-	}
 }
 
 FWinHttpHttpManager* FWinHttpHttpManager::GetManager()
@@ -48,26 +29,12 @@ FWinHttpHttpManager* FWinHttpHttpManager::GetManager()
 }
 
 FWinHttpHttpManager::FWinHttpHttpManager()
+	: FTSTickerObjectBase(0.0f, FTSBackgroundableTicker::GetCoreTicker())
 {
 	if (GWinHttpManager == nullptr)
 	{
 		GWinHttpManager = this;
 	}
-
-	FCoreDelegates::ApplicationWillEnterBackgroundDelegate.AddLambda([]()
-	{
-		if (FWinHttpHttpManager* const Manager = FWinHttpHttpManager::GetManager())
-		{
-			Manager->HandleApplicationSuspending();
-		}
-	});
-	FCoreDelegates::ApplicationHasEnteredForegroundDelegate.AddLambda([]()
-	{
-		if (FWinHttpHttpManager* const Manager = FWinHttpHttpManager::GetManager())
-		{
-			Manager->HandleApplicationResuming();
-		}
-	});
 }
 
 FWinHttpHttpManager::~FWinHttpHttpManager()
@@ -78,78 +45,4 @@ FWinHttpHttpManager::~FWinHttpHttpManager()
 	}
 }
 
-void FWinHttpHttpManager::OnBeforeFork()
-{
-	// FHttpManager's OnBeforeFork will flush all active requests, so it will be safe to reset our active sessions
-	FHttpManager::OnBeforeFork();
-	ActiveSessions.Reset();
-}
-
-void FWinHttpHttpManager::HandleApplicationSuspending()
-{
-	SCOPED_ENTER_BACKGROUND_EVENT(FWinHttpHttpManager_HandleApplicationSuspending);
-
-	Flush(EHttpFlushReason::Background);
-	ActiveSessions.Reset();
-}
-
-void FWinHttpHttpManager::HandleApplicationResuming()
-{
-	// No-op
-}
-
-void FWinHttpHttpManager::QuerySessionForUrl(const FString& /*UnusedUrl*/, FWinHttpQuerySessionComplete&& Delegate)
-{
-	// Pretend to be async here so applications properly react on platforms where this is actually async
-	AddGameThreadTask([LambdaDelegate = MoveTemp(Delegate)]()
-	{
-		QUICK_SCOPE_CYCLE_COUNTER(STAT_FWinHttpHttpManager_QuerySessionForUrlLambda);
-
-		FWinHttpSession* SessionPtr = nullptr;
-		if (FWinHttpHttpManager* HttpManager = GetManager())
-		{
-			const uint32 DefaultProtocolFlags = GetPlatformProtocolFlags();
-			SessionPtr = HttpManager->FindOrCreateSession(DefaultProtocolFlags);
-		}
-
-		LambdaDelegate.ExecuteIfBound(SessionPtr);
-	});
-}
-
-bool FWinHttpHttpManager::ValidateRequestCertificates(IWinHttpConnection& Connection)
-{
-	// WinHttp already does regular validation of certificates, this is for additional validation
-	// NOTE: this is usually not called on the game thread! Everything that happens here must be thread safe!
-
-	// TODO: Add support for client cert pinning here when we want to support that on Windows
-
-	// True means this connection did not fail validation
-	return true;
-}
-
-void FWinHttpHttpManager::ReleaseRequestResources(IWinHttpConnection& Connection)
-{
-	// No-op
-}
-
-FWinHttpSession* FWinHttpHttpManager::FindOrCreateSession(const uint32 SecurityProtocols)
-{
-	check(IsInGameThread());
-
-	TUniquePtr<FWinHttpSession>* SessionPtrPtr = ActiveSessions.Find(SecurityProtocols);
-
-	FWinHttpSession* SessionPtr = SessionPtrPtr ? SessionPtrPtr->Get() : nullptr;
-	if (!SessionPtr)
-	{
-#if !UE_BUILD_SHIPPING
-		bool bForceSecureConnections = !CVarHttpInsecureProtocolEnabled.GetValueOnAnyThread();
-#else
-		bool bForceSecureConnections = true;
-#endif
-		SessionPtr = ActiveSessions.Emplace(SecurityProtocols, MakeUnique<FWinHttpSession>(SecurityProtocols, bForceSecureConnections)).Get();
-	}
-
-	return SessionPtr;
-}
-
 #endif //WITH_WINHTTP
diff --git a/Engine/Source/Runtime/Online/HTTP/Private/WinHttp/WinHttpHttpRequest.cpp b/Engine/Source/Runtime/Online/HTTP/Private/WinHttp/WinHttpHttpRequest.cpp
deleted file mode 100644
diff --git a/Engine/Source/Runtime/Online/HTTP/Private/WinHttp/WinHttpHttpRequest.h b/Engine/Source/Runtime/Online/HTTP/Private/WinHttp/WinHttpHttpRequest.h
deleted file mode 100644
diff --git a/Engine/Source/Runtime/Online/HTTP/Private/WinHttp/WinHttpHttpResponse.cpp b/Engine/Source/Runtime/Online/HTTP/Private/WinHttp/WinHttpHttpResponse.cpp
deleted file mode 100644
diff --git a/Engine/Source/Runtime/Online/HTTP/Private/WinHttp/WinHttpHttpResponse.h b/Engine/Source/Runtime/Online/HTTP/Private/WinHttp/WinHttpHttpResponse.h
deleted file mode 100644
diff --git a/Engine/Source/Runtime/Online/HTTP/Private/Windows/WindowsPlatformHttp.cpp b/Engine/Source/Runtime/Online/HTTP/Private/Windows/WindowsPlatformHttp.cpp
--- a/Engine/Source/Runtime/Online/HTTP/Private/Windows/WindowsPlatformHttp.cpp
+++ b/Engine/Source/Runtime/Online/HTTP/Private/Windows/WindowsPlatformHttp.cpp
@@ -108,9 +108,9 @@ FHttpManager * FWindowsPlatformHttp::CreatePlatformHttpManager()
 IHttpRequest* FWindowsPlatformHttp::ConstructRequest()
 {
 #if WITH_CURL
-		return new FCurlHttpRequest();
+	return new FCurlHttpRequest();
 #else // ^^^ WITH_CURL  ^^^ // vvv WITH_CURL  vvv
-		return new FWinHttpHttpRequest();
+	return FGenericPlatformHttp::ConstructRequest();
 #endif // !WITH_CURL 
 }
 
diff --git a/Engine/Source/Runtime/Online/WebSockets/Private/WinHttp/WinHttpWebSocketsManager.cpp b/Engine/Source/Runtime/Online/WebSockets/Private/WinHttp/WinHttpWebSocketsManager.cpp
--- a/Engine/Source/Runtime/Online/WebSockets/Private/WinHttp/WinHttpWebSocketsManager.cpp
+++ b/Engine/Source/Runtime/Online/WebSockets/Private/WinHttp/WinHttpWebSocketsManager.cpp
@@ -28,17 +28,16 @@ void FWinHttpWebSocketsManager::InitWebSockets(TArrayView<const FString> Protoco
 
 void FWinHttpWebSocketsManager::InitHttpManager()
 {
-	HttpManager = FPlatformHttp::CreateWinHttpHttpManager();
-	check(HttpManager);
-	HttpManager->Initialize();
+	WinHttpHttpManager = FPlatformHttp::CreateWinHttpHttpManager();
+	check(WinHttpHttpManager);
 }
 
 void FWinHttpWebSocketsManager::ShutdownWebSockets()
 {
-	if (HttpManager)
+	if (WinHttpHttpManager)
 	{
-		delete HttpManager;
-		HttpManager = nullptr;
+		delete WinHttpHttpManager;
+		WinHttpHttpManager = nullptr;
 	}
 
 	for (TWeakPtr<FWinHttpWebSocket>& WeakWebSocket : ActiveWebSockets)
diff --git a/Engine/Source/Runtime/Online/WebSockets/Private/WinHttp/WinHttpWebSocketsManager.h b/Engine/Source/Runtime/Online/WebSockets/Private/WinHttp/WinHttpWebSocketsManager.h
--- a/Engine/Source/Runtime/Online/WebSockets/Private/WinHttp/WinHttpWebSocketsManager.h
+++ b/Engine/Source/Runtime/Online/WebSockets/Private/WinHttp/WinHttpWebSocketsManager.h
@@ -8,7 +8,7 @@
 #include "IWebSocketsManager.h"
 #include "Containers/Ticker.h"
 
-class FHttpManager;
+class FWinHttpHttpManager;
 class FWinHttpWebSocket;
 
 class FWinHttpWebSocketsManager
@@ -31,7 +31,7 @@ protected:
 protected:
 	FTSTicker::FDelegateHandle TickHandle;
 
-	FHttpManager* HttpManager = nullptr;
+	FWinHttpHttpManager* WinHttpHttpManager = nullptr;
 
 	TArray<TWeakPtr<FWinHttpWebSocket>> ActiveWebSockets;
 };
