Fix the WinGDK compile error.

diff --git a/Engine/Source/Runtime/Online/HTTP/Private/Windows/WindowsPlatformHttp.cpp b/Engine/Source/Runtime/Online/HTTP/Private/Windows/WindowsPlatformHttp.cpp
--- a/Engine/Source/Runtime/Online/HTTP/Private/Windows/WindowsPlatformHttp.cpp
+++ b/Engine/Source/Runtime/Online/HTTP/Private/Windows/WindowsPlatformHttp.cpp
@@ -9,7 +9,6 @@
 	#include "Curl/CurlHttpManager.h"
 #else // ^^^ WITH_CURL  ^^^ // vvv !WITH_CURL  vvv
 	#include "WinHttp/WinHttpHttpManager.h"
-	#include "WinHttp/WinHttpHttpRequest.h"
 #endif // !WITH_CURL 
 #include "Http.h"
 #include "WinHttp/Support/WinHttpTypes.h" // Always include for OS proxy settings
@@ -101,7 +100,7 @@ FHttpManager * FWindowsPlatformHttp::CreatePlatformHttpManager()
 #if WITH_CURL
 	return new FCurlHttpManager();
 #else // ^^^ WITH_CURL  ^^^ // vvv WITH_CURL  vvv
-	return new FWinHttpHttpManager();
+	return FGenericPlatformHttp::CreatePlatformHttpManager();
 #endif // !WITH_CURL 
 }
 
