Remove the load of http module from websockets manager, websockets doesn't depend on it anymore.

diff --git a/Engine/Source/Runtime/Online/WebSockets/Private/WinHttp/WinHttpWebSocketsManager.cpp b/Engine/Source/Runtime/Online/WebSockets/Private/WinHttp/WinHttpWebSocketsManager.cpp
--- a/Engine/Source/Runtime/Online/WebSockets/Private/WinHttp/WinHttpWebSocketsManager.cpp
+++ b/Engine/Source/Runtime/Online/WebSockets/Private/WinHttp/WinHttpWebSocketsManager.cpp
@@ -13,8 +13,6 @@
 
 void FWinHttpWebSocketsManager::InitWebSockets(TArrayView<const FString> Protocols)
 {
-	(void)FModuleManager::LoadModuleChecked<FHttpModule>(TEXT("Http"));
-
 	if (FWinHttpHttpManager::GetManager() == nullptr)
 	{
 		InitHttpManager();
