Enable web tests for more platforms.

diff --git a/Engine/Source/Developer/S3Client/Private/S3Client.cpp b/Engine/Source/Developer/S3Client/Private/S3Client.cpp
--- a/Engine/Source/Developer/S3Client/Private/S3Client.cpp
+++ b/Engine/Source/Developer/S3Client/Private/S3Client.cpp
@@ -27,12 +27,10 @@
 #include "Microsoft/HideMicrosoftPlatformTypes.h"
 #endif
 
-#if WITH_SSL
 #include "Ssl.h"
 #include <openssl/hmac.h>
 #include <openssl/sha.h>
 #include <openssl/ssl.h>
-#endif
 
 IMPLEMENT_MODULE(FDefaultModuleImpl, S3Client);
 
diff --git a/Engine/Source/Programs/WebTests/Private/TestWebSockets.cpp b/Engine/Source/Programs/WebTests/Private/TestWebSockets.cpp
--- a/Engine/Source/Programs/WebTests/Private/TestWebSockets.cpp
+++ b/Engine/Source/Programs/WebTests/Private/TestWebSockets.cpp
@@ -1,6 +1,7 @@
 // Copyright Epic Games, Inc. All Rights Reserved.
 #include "Containers/BackgroundableTicker.h"
 #include "CoreMinimal.h"
+#include "HAL/IConsoleManager.h"
 #include "HAL/PlatformProcess.h"
 #include "Http.h"
 #include "Misc/CommandLine.h"
@@ -23,6 +24,8 @@
 
 #define WEBSOCKETS_TAG "[WebSockets]"
 
+extern TAutoConsoleVariable<bool> CVarHttpInsecureProtocolEnabled;
+
 class FWebSocketsModuleTestFixture
 {
 public:
@@ -31,6 +34,8 @@ public:
 		, WebServerWebSocketsPort(8000)
 		, OldVerbosity(LogWebSockets.GetVerbosity())
 	{
+		CVarHttpInsecureProtocolEnabled->Set(true);
+
 		ParseSettingsFromCommandLine();
 
 		// Init HTTP module because websockets module has dependency on it when get proxy
diff --git a/Engine/Source/Runtime/Online/HTTP/Private/Windows/WindowsPlatformHttp.cpp b/Engine/Source/Runtime/Online/HTTP/Private/Windows/WindowsPlatformHttp.cpp
--- a/Engine/Source/Runtime/Online/HTTP/Private/Windows/WindowsPlatformHttp.cpp
+++ b/Engine/Source/Runtime/Online/HTTP/Private/Windows/WindowsPlatformHttp.cpp
@@ -4,12 +4,8 @@
 #include "Misc/CommandLine.h"
 #include "Misc/Parse.h"
 #include "Misc/ConfigCacheIni.h"
-#if WITH_CURL
-	#include "Curl/CurlHttp.h"
-	#include "Curl/CurlHttpManager.h"
-#else // ^^^ WITH_CURL  ^^^ // vvv !WITH_CURL  vvv
-	#include "WinHttp/WinHttpHttpManager.h"
-#endif // !WITH_CURL 
+#include "Curl/CurlHttp.h"
+#include "Curl/CurlHttpManager.h"
 #include "Http.h"
 #include "WinHttp/Support/WinHttpTypes.h" // Always include for OS proxy settings
 
@@ -81,36 +77,24 @@ void FWindowsPlatformHttp::Init()
 
 	FGenericPlatformHttp::Init();
 
-#if WITH_CURL
 	FCurlHttpManager::InitCurl();
-#endif // WITH_WINHTTP
 }
 
 void FWindowsPlatformHttp::Shutdown()
 {
-#if WITH_CURL
 	FCurlHttpManager::ShutdownCurl();
-#endif // WITH_CURL 
 
 	FGenericPlatformHttp::Shutdown();
 }
 
 FHttpManager * FWindowsPlatformHttp::CreatePlatformHttpManager()
 {
-#if WITH_CURL
 	return new FCurlHttpManager();
-#else // ^^^ WITH_CURL  ^^^ // vvv WITH_CURL  vvv
-	return FGenericPlatformHttp::CreatePlatformHttpManager();
-#endif // !WITH_CURL 
 }
 
 IHttpRequest* FWindowsPlatformHttp::ConstructRequest()
 {
-#if WITH_CURL
 	return new FCurlHttpRequest();
-#else // ^^^ WITH_CURL  ^^^ // vvv WITH_CURL  vvv
-	return FGenericPlatformHttp::ConstructRequest();
-#endif // !WITH_CURL 
 }
 
 FString FWindowsPlatformHttp::GetMimeType(const FString& FilePath)
@@ -202,10 +186,7 @@ bool FWindowsPlatformHttp::IsOperatingSystemProxyInformationSupported()
 
 bool FWindowsPlatformHttp::VerifyPeerSslCertificate(bool verify)
 {
-	bool prev = false;
-#if WITH_CURL
-	prev = FCurlHttpManager::CurlRequestOptions.bVerifyPeer;
+	bool bPreviousValue = FCurlHttpManager::CurlRequestOptions.bVerifyPeer;
 	FCurlHttpManager::CurlRequestOptions.bVerifyPeer = verify;
-#endif // #if WITH_CURL 
-	return prev;
+	return bPreviousValue;
 }
