Trigger delegate to process events for ios/mac also when new data is received so the progress delegate is called when possible

diff --git a/Engine/Source/Runtime/Online/HTTP/Private/Apple/AppleEventLoopHttpThread.cpp b/Engine/Source/Runtime/Online/HTTP/Private/Apple/AppleEventLoopHttpThread.cpp
--- a/Engine/Source/Runtime/Online/HTTP/Private/Apple/AppleEventLoopHttpThread.cpp
+++ b/Engine/Source/Runtime/Online/HTTP/Private/Apple/AppleEventLoopHttpThread.cpp
@@ -8,7 +8,7 @@ bool FAppleEventLoopHttpThread::StartThreadedRequest(IHttpThreadedRequest* Reque
 {
 	FHttpResponsePtr Response = Request->GetResponse();
 	auto AppleResponse = StaticCastSharedPtr<FAppleHttpNSUrlSessionResponse>(Response);
-	AppleResponse->SetInternalTaskCompleteDelegate(FTaskCompleteDelegate::CreateLambda([IOAccess = EventLoop->GetIOAccess()]() mutable
+	AppleResponse->SetNewAppleHttpEventDelegate(FNewAppleHttpEventDelegate::CreateLambda([IOAccess = EventLoop->GetIOAccess()]() mutable
 	{
 		IOAccess.Notify();
 	}));
diff --git a/Engine/Source/Runtime/Online/HTTP/Private/Apple/AppleHTTPNSUrlSession.cpp b/Engine/Source/Runtime/Online/HTTP/Private/Apple/AppleHTTPNSUrlSession.cpp
--- a/Engine/Source/Runtime/Online/HTTP/Private/Apple/AppleHTTPNSUrlSession.cpp
+++ b/Engine/Source/Runtime/Online/HTTP/Private/Apple/AppleHTTPNSUrlSession.cpp
@@ -41,8 +41,8 @@ enum class EAppleHttpRequestResponseState: uint8
 	/** flag meant to reduce locking on ResponseStreamLock*/
 	@public BOOL bInitializedWithValidStream;
 
-	/** Delegate invoked after processing URLSession:task:didCompleteWithError:*/
-	@public FTaskCompleteDelegate TaskCompleteDelegate;
+	/** Delegate invoked after processing URLSession:dataTask:didReceiveData or URLSession:task:didCompleteWithError:*/
+	@public FNewAppleHttpEventDelegate NewAppleHttpEventDelegate;
 }
 
 /** A handle for the response */
@@ -212,7 +212,7 @@ enum class EAppleHttpRequestResponseState: uint8
 			}
 		}
 	}
-	TaskCompleteDelegate.ExecuteIfBound();
+	NewAppleHttpEventDelegate.ExecuteIfBound();
 }
 
 - (void)URLSession:(NSURLSession *)session dataTask:(NSURLSessionDataTask *)dataTask willCacheResponse:(NSCachedURLResponse *)proposedResponse completionHandler:(void (^)(NSCachedURLResponse *cachedResponse))completionHandler
@@ -743,8 +743,6 @@ bool FAppleHttpNSUrlSessionRequest::StartRequest()
 
 void FAppleHttpNSUrlSessionRequest::FinishRequest()
 {
-	check(IsInGameThread() || DelegateThreadPolicy == EHttpRequestDelegateThreadPolicy::CompleteOnHttpThread);
-
 	UE_LOG(LogHttp, Verbose, TEXT("FAppleHttpNSUrlSessionRequest::FinishRequest()"));
 
 	// Clean up session/request handles that may have been created
@@ -881,9 +879,9 @@ FAppleHttpNSUrlSessionResponse::~FAppleHttpNSUrlSessionResponse()
 	ResponseDelegate = nil;
 }
 
-void FAppleHttpNSUrlSessionResponse::SetInternalTaskCompleteDelegate(FTaskCompleteDelegate&& Delegate)
+void FAppleHttpNSUrlSessionResponse::SetNewAppleHttpEventDelegate(FNewAppleHttpEventDelegate&& Delegate)
 {	
-	ResponseDelegate->TaskCompleteDelegate = MoveTemp(Delegate);
+	ResponseDelegate->NewAppleHttpEventDelegate = MoveTemp(Delegate);
 }
 
 void FAppleHttpNSUrlSessionResponse::CleanSharedObjects()
diff --git a/Engine/Source/Runtime/Online/HTTP/Private/Apple/AppleHTTPNSUrlSession.h b/Engine/Source/Runtime/Online/HTTP/Private/Apple/AppleHTTPNSUrlSession.h
--- a/Engine/Source/Runtime/Online/HTTP/Private/Apple/AppleHTTPNSUrlSession.h
+++ b/Engine/Source/Runtime/Online/HTTP/Private/Apple/AppleHTTPNSUrlSession.h
@@ -12,7 +12,7 @@
  * Delegate invoked when in progress NSUrlSessionTask completes. It is invoked in an out of our control thread
  *
  */
-DECLARE_DELEGATE(FTaskCompleteDelegate);
+DECLARE_DELEGATE(FNewAppleHttpEventDelegate);
 
 /**
  * Apple implementation of an Http request
@@ -200,10 +200,10 @@ public:
 	void CleanSharedObjects();
 
 	/**
-	 * Sets delegate invoked after processing URLSession:task:didCompleteWithError:
+	 * Sets delegate invoked when  URLSession:dataTask:didReceiveData or URLSession:task:didCompleteWithError: are triggered
 	 * Should be set right before task is started 
 	*/
-	void SetInternalTaskCompleteDelegate(FTaskCompleteDelegate&& Delegate);
+	void SetNewAppleHttpEventDelegate(FNewAppleHttpEventDelegate&& Delegate);
 
 	/**
 	 * Constructor
