Use thread-safe delegates for http requests.

diff --git a/Engine/Source/Runtime/Online/HTTP/Public/Interfaces/IHttpRequest.h b/Engine/Source/Runtime/Online/HTTP/Public/Interfaces/IHttpRequest.h
--- a/Engine/Source/Runtime/Online/HTTP/Public/Interfaces/IHttpRequest.h
+++ b/Engine/Source/Runtime/Online/HTTP/Public/Interfaces/IHttpRequest.h
@@ -87,7 +87,7 @@ typedef TSharedRef<IHttpResponse, ESPMode::ThreadSafe> FHttpResponseRef;
  * @param Response response received from the server if a successful connection was established
  * @param bConnectedSuccessfully - indicates whether or not the request was able to connect successfully
  */
-DECLARE_DELEGATE_ThreeParams(FHttpRequestCompleteDelegate, FHttpRequestPtr /*Request*/, FHttpResponsePtr /*Response*/, bool /*bConnectedSuccessfully*/);
+using FHttpRequestCompleteDelegate = TTSDelegate<void(FHttpRequestPtr /*Request*/, FHttpResponsePtr /*Response*/, bool /*bConnectedSuccessfully*/)>;
 
 /**
  * Delegate called when an Http request receives a header
@@ -96,7 +96,7 @@ DECLARE_DELEGATE_ThreeParams(FHttpRequestCompleteDelegate, FHttpRequestPtr /*Req
  * @param HeaderName the name of the header
  * @param NewHeaderValue the value of the header
  */
-DECLARE_DELEGATE_ThreeParams(FHttpRequestHeaderReceivedDelegate, FHttpRequestPtr /*Request*/, const FString& /*HeaderName*/, const FString& /*NewHeaderValue*/);
+using FHttpRequestHeaderReceivedDelegate = TTSDelegate<void(FHttpRequestPtr /*Request*/, const FString& /*HeaderName*/, const FString& /*NewHeaderValue*/)>;
 
 /**
  * Delegate called per tick to update an Http request upload or download size progress
@@ -105,7 +105,7 @@ DECLARE_DELEGATE_ThreeParams(FHttpRequestHeaderReceivedDelegate, FHttpRequestPtr
  * @param BytesSent the number of bytes sent / uploaded in the request so far.
  * @param BytesReceived the number of bytes received / downloaded in the response so far.
  */
-DECLARE_DELEGATE_ThreeParams(FHttpRequestProgressDelegate, FHttpRequestPtr /*Request*/, int32 /*BytesSent*/, int32 /*BytesReceived*/);
+using FHttpRequestProgressDelegate = TTSDelegate<void(FHttpRequestPtr /*Request*/, int32 /*BytesSent*/, int32 /*BytesReceived*/)>;
 
 /**
  * Delegate called per tick to update an Http request upload or download size progress
@@ -114,7 +114,7 @@ DECLARE_DELEGATE_ThreeParams(FHttpRequestProgressDelegate, FHttpRequestPtr /*Req
  * @param BytesSent the number of bytes sent / uploaded in the request so far.
  * @param BytesReceived the number of bytes received / downloaded in the response so far.
  */
-DECLARE_DELEGATE_ThreeParams(FHttpRequestProgressDelegate64, FHttpRequestPtr /*Request*/, uint64 /*BytesSent*/, uint64 /*BytesReceived*/);
+using FHttpRequestProgressDelegate64 = TTSDelegate<void(FHttpRequestPtr /*Request*/, uint64 /*BytesSent*/, uint64 /*BytesReceived*/)>;
 
 /**
  * Delegate called when an Http request will be retried in the future
@@ -123,7 +123,7 @@ DECLARE_DELEGATE_ThreeParams(FHttpRequestProgressDelegate64, FHttpRequestPtr /*R
  * @param Response - response received from the server if a successful connection was established
  * @param SecondsToRetry - seconds in the future when the response will be retried
  */
-DECLARE_DELEGATE_ThreeParams(FHttpRequestWillRetryDelegate, FHttpRequestPtr /*Request*/, FHttpResponsePtr /*Response*/, float /*SecondsToRetry*/);
+using FHttpRequestWillRetryDelegate = TTSDelegate<void(FHttpRequestPtr /*Request*/, FHttpResponsePtr /*Response*/, float /*SecondsToRetry*/)>;
 
 /**
  * Delegate called when an Http request will send/recv data through stream
@@ -132,8 +132,7 @@ DECLARE_DELEGATE_ThreeParams(FHttpRequestWillRetryDelegate, FHttpRequestPtr /*Re
  * @param Length - The length of buffer to read/write
  * @return true if succeed, false if failed to read/write data
  */
-DECLARE_DELEGATE_RetVal_TwoParams(bool, FHttpRequestStreamDelegate, void*/*Ptr*/, int64/*Length*/);
-
+using FHttpRequestStreamDelegate = TTSDelegate<bool(void*/*Ptr*/, int64/*Length*/)>;
 
 /**
  * Delegate version of FArchive, for streaming interface
