Optimizations in BuildPatchServices when calling GetRequiredDiskSpace. These are particularly impactful for Android low-spec devices.
- Only run GetOutdatedFiles on tagged files when supplied instead of on all files and then filtering that list by the tagged files.
- GetOutdatedFiles now only gets disk space information in cases it actually will use it.
#tests Android Dev Client

diff --git a/Engine/Source/Runtime/Online/BuildPatchServices/Private/BuildPatchFileConstructor.cpp b/Engine/Source/Runtime/Online/BuildPatchServices/Private/BuildPatchFileConstructor.cpp
--- a/Engine/Source/Runtime/Online/BuildPatchServices/Private/BuildPatchFileConstructor.cpp
+++ b/Engine/Source/Runtime/Online/BuildPatchServices/Private/BuildPatchFileConstructor.cpp
@@ -78,16 +78,14 @@ namespace FileConstructorHelpers
 		}
 		InstallTags.Add(TEXT(""));
 		// Calculate the files that need constructing.
-		FString DummyString;
-		TSet<FString> FilesToConstruct;
-		BuildManifest->GetOutdatedFiles(CurrentManifest, DummyString, FilesToConstruct);
 		TSet<FString> TaggedFiles;
 		BuildManifest->GetTaggedFileList(InstallTags, TaggedFiles);
-		FilesToConstruct = FilesToConstruct.Intersect(TaggedFiles);
+		FString DummyString;
+		TSet<FString> FilesToConstruct;
+		BuildManifest->GetOutdatedFiles(CurrentManifest.Get(), DummyString, TaggedFiles, FilesToConstruct);
 		// Count disk space needed by each operation.
 		int64 DiskSpaceDeltaPeak = 0;
-		const bool bCurrentManifestIsValid = CurrentManifest.IsValid();
-		if (InstallMode == EInstallMode::DestructiveInstall && bCurrentManifestIsValid)
+		if (InstallMode == EInstallMode::DestructiveInstall && CurrentManifest.IsValid())
 		{
 			// The simplest method will be to run through each high level file operation, tracking peak disk usage delta.
 			int64 DiskSpaceDelta = 0;
diff --git a/Engine/Source/Runtime/Online/BuildPatchServices/Private/BuildPatchManifest.cpp b/Engine/Source/Runtime/Online/BuildPatchServices/Private/BuildPatchManifest.cpp
--- a/Engine/Source/Runtime/Online/BuildPatchServices/Private/BuildPatchManifest.cpp
+++ b/Engine/Source/Runtime/Online/BuildPatchServices/Private/BuildPatchManifest.cpp
@@ -1371,16 +1371,19 @@ void FBuildPatchAppManifest::GetOutdatedFiles(const FBuildPatchAppManifest* OldM
 			const FFileManifest* NewFile = GetFileManifest(FileToCheck);
 			if (NewFile != nullptr)
 			{
-				const int64 ExistingFileSize = IFileManager::Get().FileSize(*(InstallDirectory / NewFile->Filename));
 				// Check changed
 				if (IsFileOutdated(*OldManifest, NewFile->Filename))
 				{
 					OutDatedFiles.Add(NewFile->Filename);
 				}
 				// Double check an unchanged file is not missing (size will be -1) or is incorrect size
-				else if (bCheckExistingFile && (ExistingFileSize < 0 || ExistingFileSize != NewFile->FileSize))
+				else if (bCheckExistingFile)
 				{
-					OutDatedFiles.Add(NewFile->Filename);
+					const int64 ExistingFileSize = IFileManager::Get().FileSize(*(InstallDirectory / NewFile->Filename));
+					if ((ExistingFileSize < 0) || (ExistingFileSize != NewFile->FileSize))
+					{
+						OutDatedFiles.Add(NewFile->Filename);
+					}
 				}
 			}
 		}
