Fix the compile error when compile EngineTest under Win64 using GDK option.

diff --git a/Engine/Source/Runtime/Online/HTTP/HTTP.Build.cs b/Engine/Source/Runtime/Online/HTTP/HTTP.Build.cs
--- a/Engine/Source/Runtime/Online/HTTP/HTTP.Build.cs
+++ b/Engine/Source/Runtime/Online/HTTP/HTTP.Build.cs
@@ -21,13 +21,12 @@ public class HTTP : ModuleRules
 				Target.IsInPlatformGroup(UnrealPlatformGroup.Android);
 		}
 	}
-	protected virtual bool bPlatformSupportsXCurl { get { return false; } }
+	protected virtual bool bPlatformSupportsXCurl { get { return Target.Platform.IsInGroup(UnrealPlatformGroup.Windows) && Target.WindowsPlatform.bUseXCurl; } }
+	protected virtual bool bPlatformSupportsCurlMultiSocket { get { return !bPlatformSupportsXCurl; } }
 
 	protected virtual bool bPlatformSupportsCurlMultiPoll { get { return true; } }
 
 	protected virtual bool bPlatformSupportsCurlMultiWait { get { return false; } }
-	protected virtual bool bPlatformSupportsCurlMultiSocket { get { return true; } }
-
 	private bool bPlatformSupportsCurl { get { return bPlatformSupportsLibCurl || bPlatformSupportsXCurl; } }
 
 	protected virtual bool bPlatformRequiresOpenSSL
