Trigger delegate to process events also when new data is received so the progress delegate is called when possible (this call was missing from a previous review)

diff --git a/Engine/Source/Runtime/Online/HTTP/Private/Apple/AppleHTTPNSUrlSession.cpp b/Engine/Source/Runtime/Online/HTTP/Private/Apple/AppleHTTPNSUrlSession.cpp
--- a/Engine/Source/Runtime/Online/HTTP/Private/Apple/AppleHTTPNSUrlSession.cpp
+++ b/Engine/Source/Runtime/Online/HTTP/Private/Apple/AppleHTTPNSUrlSession.cpp
@@ -153,8 +153,9 @@ enum class EAppleHttpRequestResponseState: uint8
 	}
 	// Keep BytesReceived as a separated value to avoid concurrent accesses to Payload
 	self.BytesReceived += NewBytesReceived;
-	
 	UE_LOG(LogHttp, Verbose, TEXT("URLSession:dataTask:didReceiveData with %llu bytes. After Append, Payload Length = %llu: %p"), NewBytesReceived, self.BytesReceived, self);
+	
+	NewAppleHttpEventDelegate.ExecuteIfBound();
 }
 
 - (void)URLSession:(NSURLSession *)session task:(NSURLSessionTask *)task didCompleteWithError:(nullable NSError *)error
