Set correct host info for libWebsockets.

#10674

diff --git a/Engine/Source/Runtime/Online/WebSockets/Private/Lws/LwsWebSocket.cpp b/Engine/Source/Runtime/Online/WebSockets/Private/Lws/LwsWebSocket.cpp
--- a/Engine/Source/Runtime/Online/WebSockets/Private/Lws/LwsWebSocket.cpp
+++ b/Engine/Source/Runtime/Online/WebSockets/Private/Lws/LwsWebSocket.cpp
@@ -759,6 +759,22 @@ void FLwsWebSocket::ConnectInternal(struct lws_context &LwsContext)
 		return;
 	}
 
+	FString ParsedPortString = FString::FromInt(ParsedPort);
+	FTCHARToUTF8 ParsedPortStringUTF8(*ParsedPortString);
+	// Length is based on format [addr]:[port]\0
+	int32 ParsedAddressLength = FCStringAnsi::Strlen(ParsedAddress);
+	bool bUsingDefaultPort = (SslConnection == 0 && ParsedPort == 80) || (SslConnection != 0 && ParsedPort == 443);
+	int32 HostLength = bUsingDefaultPort ? (ParsedAddressLength + 1) : (ParsedAddressLength + 1 + ParsedPortStringUTF8.Length() + 1); 
+	TArray<ANSICHAR> Host;
+	Host.Reserve(HostLength);
+	Host.Append(ParsedAddress, ParsedAddressLength);
+	if (!bUsingDefaultPort)
+	{
+		Host.Add(':');
+		Host.Append(ParsedPortStringUTF8.Get(), ParsedPortStringUTF8.Length());
+	}
+	Host.Add('\0');
+
 	// If we specify protocols, convert the comma separated list into UTF8
 	TOptional<FTCHARToUTF8> OptionalCombinedProtocolsUTF8;
 	if (Protocols.Num() > 0)
@@ -775,7 +791,7 @@ void FLwsWebSocket::ConnectInternal(struct lws_context &LwsContext)
 	ConnectInfo.port = ParsedPort;
 	ConnectInfo.ssl_connection = SslConnection;
 	ConnectInfo.path = UrlPath.GetData();
-	ConnectInfo.host = ConnectInfo.address;
+	ConnectInfo.host = Host.GetData();
 	ConnectInfo.origin = ConnectInfo.address;
 	ConnectInfo.protocol = CombinedProtocolsUTF8;
 	ConnectInfo.ietf_version_or_minus_one = -1;
