BPS Installer shared context and thread free list.
- Currently just allows a group of installers to re-use threads, mitigating the cost of thread creation.
- This is a jumping off point for component sharing, or running multiple components on the same thread  (ex. multiple cloud chunk sources sharing a thread).

diff --git a/Engine/Source/Runtime/Online/BuildPatchServices/Private/BuildPatchInstaller.cpp b/Engine/Source/Runtime/Online/BuildPatchServices/Private/BuildPatchInstaller.cpp
--- a/Engine/Source/Runtime/Online/BuildPatchServices/Private/BuildPatchInstaller.cpp
+++ b/Engine/Source/Runtime/Online/BuildPatchServices/Private/BuildPatchInstaller.cpp
@@ -39,6 +39,7 @@
 #include "Installer/Verifier.h"
 #include "Installer/FileAttribution.h"
 #include "Installer/InstallerAnalytics.h"
+#include "Installer/InstallerSharedContext.h"
 #include "Installer/Prerequisites.h"
 #include "Installer/MachineConfig.h"
 #include "Installer/MessagePump.h"
@@ -502,6 +503,14 @@ namespace BuildPatchServices
 		}
 
 		ChunkDataSizeProvider->AddManifestData(Manifests);
+
+		if (!Configuration.SharedContext)
+		{
+			Configuration.SharedContext = FBuildInstallerSharedContextFactory::Create(TEXT("BuildPatchInstaller"));
+			const bool bUseChunkDBs = !Configuration.ChunkDatabaseFiles.IsEmpty();
+			const uint32 NumExpectedThreads = Configuration.SharedContext->NumThreadsPerInstaller(bUseChunkDBs);
+			Configuration.SharedContext->PreallocateThreads(NumExpectedThreads);
+		}
 	}
 
 	FBuildPatchInstaller::~FBuildPatchInstaller()
@@ -608,19 +617,8 @@ namespace BuildPatchServices
 		if (Thread == nullptr)
 		{
 			// Start thread!
-			const TCHAR* ThreadName = TEXT("BuildPatchInstallerThread");
-
-			// Ideally this would check if we were forkable or a forked child process but there is 
-			// currently no in-engine way to check that.  Since BPS does not currently support
-			// FRunnableThread::ThreadType::Fake or FRunnableThread::ThreadType::Forkable 
-			// this check ends up being equivalent for now.  We most likely *never* want support 
-			// forking while an installer is running.
-			if (FPlatformProcess::SupportsMultithreading() || FForkProcessHelper::IsForkedMultithreadInstance())
-			{
-				Thread = FForkProcessHelper::CreateForkableThread(this, ThreadName);
-				check(Thread != nullptr);
-				check(Thread->GetThreadType() == FRunnableThread::ThreadType::Real);
-			}
+			Thread = Configuration.SharedContext->CreateThread();
+			Thread->RunTask([this]{ Run(); });
 
 			StartDelegate.ExecuteIfBound(AsShared());
 		}
@@ -992,6 +990,8 @@ namespace BuildPatchServices
 	{
 		FChunkDbSourceConfig ChunkDbSourceConfig(Configuration.ChunkDatabaseFiles);
 
+		ChunkDbSourceConfig.SharedContext = Configuration.SharedContext.Get();
+
 		// Load batch fetch config.
 		GConfig->GetInt(TEXT("Portal.BuildPatch"), TEXT("ChunkDbSourcePreFetchMinimum"), ChunkDbSourceConfig.PreFetchMinimum, GEngineIni);
 		GConfig->GetInt(TEXT("Portal.BuildPatch"), TEXT("ChunkDbSourcePreFetchMaximum"), ChunkDbSourceConfig.PreFetchMaximum, GEngineIni);
@@ -1023,6 +1023,8 @@ namespace BuildPatchServices
 	{
 		FCloudSourceConfig CloudSourceConfig(Configuration.CloudDirectories);
 
+		CloudSourceConfig.SharedContext = Configuration.SharedContext.Get();
+
 		// Load max download retry count from engine config.
 		GConfig->GetInt(TEXT("Portal.BuildPatch"), TEXT("ChunkRetries"), CloudSourceConfig.MaxRetryCount, GEngineIni);
 		CloudSourceConfig.MaxRetryCount = FMath::Clamp<int32>(CloudSourceConfig.MaxRetryCount, -1, 1000);
@@ -1286,21 +1288,20 @@ namespace BuildPatchServices
 			TSet<FGuid> ReferencedChunks = ChunkReferenceTracker->GetReferencedChunks();
 			TUniquePtr<IChunkEvictionPolicy> MemoryEvictionPolicy(FChunkEvictionPolicyFactory::Create(
 				ChunkReferenceTracker.Get()));
+			TUniquePtr<IDiskChunkStore> DiskOverflowStore;
 #if ENABLE_PATCH_DISK_OVERFLOW_STORE
-			TUniquePtr<IDiskChunkStore> DiskOverflowStore(FDiskChunkStoreFactory::Create(
+			FDiskChunkStoreConfig DiskChunkStoreConfig(DataStagingDir);
+			DiskChunkStoreConfig.SharedContext = Configuration.SharedContext.Get();
+			DiskOverflowStore.Reset(FDiskChunkStoreFactory::Create(
 				FileSystem.Get(),
 				ChunkDataSerialization.Get(),
 				DiskChunkStoreStatistics.Get(),
-				FDiskChunkStoreConfig(DataStagingDir)));
+				MoveTemp(DiskChunkStoreConfig)));
 #endif // ENABLE_PATCH_DISK_OVERFLOW_STORE
 			TUniquePtr<IMemoryChunkStore> CloudChunkStore(FMemoryChunkStoreFactory::Create(
 				ChunkStoreMemorySize,
 				MemoryEvictionPolicy.Get(),
-#if ENABLE_PATCH_DISK_OVERFLOW_STORE
 				DiskOverflowStore.Get(),
-#else
-				nullptr,
-#endif // ENABLE_PATCH_DISK_OVERFLOW_STORE
 				MemoryChunkStoreStatistics.Get()));
 			TUniquePtr<IChunkDbChunkSource> ChunkDbChunkSource(FChunkDbChunkSourceFactory::Create(
 				BuildChunkDbSourceConfig(),
@@ -1321,7 +1322,7 @@ namespace BuildPatchServices
 				InstallChunkSourceStatistics.Get(),
 				InstallationInfo,
 				ManifestSet.Get()));
-			const TSet<FGuid> InitialDownloadChunks = ReferencedChunks.Difference(InstallChunkSource->GetAvailableChunks()).Difference(ChunkDbChunkSource->GetAvailableChunks());
+			TSet<FGuid> InitialDownloadChunks = ReferencedChunks.Difference(InstallChunkSource->GetAvailableChunks()).Difference(ChunkDbChunkSource->GetAvailableChunks());
 			FileOperationTracker->OnDataStateUpdate(ReferencedChunks.Intersect(ChunkDbChunkSource->GetAvailableChunks()), EFileOperationState::PendingLocalChunkDbData);
 			FileOperationTracker->OnDataStateUpdate(ReferencedChunks.Intersect(InstallChunkSource->GetAvailableChunks()).Difference(ChunkDbChunkSource->GetAvailableChunks()), EFileOperationState::PendingLocalInstallData);
 			FileOperationTracker->OnDataStateUpdate(InitialDownloadChunks, EFileOperationState::PendingRemoteCloudData);
@@ -1338,7 +1339,7 @@ namespace BuildPatchServices
 				DownloadConnectionCount.Get(),
 				CloudChunkSourceStatistics.Get(),
 				ManifestSet.Get(),
-				InitialDownloadChunks));
+				MoveTemp(InitialDownloadChunks)));
 			TArray<IChunkSource*> ChunkSources;
 			ChunkSources.Add(ChunkDbChunkSource.Get());
 			ChunkSources.Add(InstallChunkSource.Get());
@@ -1850,12 +1851,8 @@ namespace BuildPatchServices
 
 	void FBuildPatchInstaller::CleanupThread()
 	{
-		if (Thread != nullptr)
-		{
-			Thread->WaitForCompletion();
-			delete Thread;
-			Thread = nullptr;
-		}
+		Configuration.SharedContext->ReleaseThread(Thread);
+		Thread = nullptr;
 	}
 
 	double FBuildPatchInstaller::GetDownloadSpeed() const
@@ -2120,12 +2117,4 @@ namespace BuildPatchServices
 		check(IsInGameThread());
 		MessagePump->PumpMessages();
 	}
-
-	void FBuildPatchInstaller::WaitForThread() const
-	{
-		if (Thread != nullptr)
-		{
-			Thread->WaitForCompletion();
-		}
-	}
 }
diff --git a/Engine/Source/Runtime/Online/BuildPatchServices/Private/BuildPatchInstaller.h b/Engine/Source/Runtime/Online/BuildPatchServices/Private/BuildPatchInstaller.h
--- a/Engine/Source/Runtime/Online/BuildPatchServices/Private/BuildPatchInstaller.h
+++ b/Engine/Source/Runtime/Online/BuildPatchServices/Private/BuildPatchInstaller.h
@@ -65,7 +65,7 @@ namespace BuildPatchServices
 		FGuid SessionId;
 
 		// Hold a pointer to my thread for easier deleting.
-		FRunnableThread* Thread;
+		IBuildInstallerThread* Thread;
 
 		// The delegates that we will be calling when started.
 		const FBuildPatchInstallerDelegate StartDelegate;
@@ -74,7 +74,7 @@ namespace BuildPatchServices
 		const FBuildPatchInstallerDelegate CompleteDelegate;
 
 		// The installer configuration.
-		const FBuildInstallerConfiguration Configuration;
+		FBuildInstallerConfiguration Configuration;
 
 		// The Configuration.InstallerActions array converted into private class type.
 		TArray<FBuildPatchInstallerAction> InstallerActions;
@@ -263,11 +263,6 @@ namespace BuildPatchServices
 		 */
 		bool Tick();
 
-		/**
-		 * Only returns once the thread has finished running.
-		 */
-		void WaitForThread() const;
-
 		/**
 		 * Called by the module during shutdown.
 		 */
diff --git a/Engine/Source/Runtime/Online/BuildPatchServices/Private/BuildPatchServicesModule.cpp b/Engine/Source/Runtime/Online/BuildPatchServices/Private/BuildPatchServicesModule.cpp
--- a/Engine/Source/Runtime/Online/BuildPatchServices/Private/BuildPatchServicesModule.cpp
+++ b/Engine/Source/Runtime/Online/BuildPatchServices/Private/BuildPatchServicesModule.cpp
@@ -22,6 +22,7 @@
 #include "Generation/ChunkDeltaOptimiser.h"
 #include "Generation/PackageChunkData.h"
 #include "Installer/BuildStatistics.h"
+#include "Installer/InstallerSharedContext.h"
 #include "Installer/MachineConfig.h"
 #include "BuildPatchMergeManifests.h"
 #include "BuildPatchHash.h"
@@ -128,10 +129,15 @@ IBuildInstallerRef FBuildPatchServicesModule::CreateBuildInstaller(BuildPatchSer
 	{
 		Configuration.bRunRequiredPrereqs = false;
 	}
-	FBuildPatchInstallerRef Installer = MakeShareable(new FBuildPatchInstaller(MoveTemp(Configuration), AvailableInstallations, LocalMachineConfigFile, Analytics, InstallerStartDelegate, MoveTemp(CompleteDelegate)));
+	FBuildPatchInstallerRef Installer = MakeShared<FBuildPatchInstaller>(MoveTemp(Configuration), AvailableInstallations, LocalMachineConfigFile, Analytics, InstallerStartDelegate, MoveTemp(CompleteDelegate));
 	return Installer;
 }
 
+IBuildInstallerSharedContextRef FBuildPatchServicesModule::CreateBuildInstallerSharedContext(const TCHAR* DebugName) const
+{
+	return BuildPatchServices::FBuildInstallerSharedContextFactory::Create(DebugName);
+}
+
 IBuildStatisticsRef FBuildPatchServicesModule::CreateBuildStatistics(const IBuildInstallerRef& Installer) const
 {
 	checkSlow(IsInGameThread());
diff --git a/Engine/Source/Runtime/Online/BuildPatchServices/Private/BuildPatchServicesModule.h b/Engine/Source/Runtime/Online/BuildPatchServices/Private/BuildPatchServicesModule.h
--- a/Engine/Source/Runtime/Online/BuildPatchServices/Private/BuildPatchServicesModule.h
+++ b/Engine/Source/Runtime/Online/BuildPatchServices/Private/BuildPatchServicesModule.h
@@ -38,6 +38,7 @@ public:
 	// IBuildPatchServicesModule interface begin.
 	virtual IBuildInstallStreamerRef CreateBuildInstallStreamer(BuildPatchServices::FBuildInstallStreamerConfiguration Configuration) override;
 	virtual IBuildInstallerRef CreateBuildInstaller(BuildPatchServices::FBuildInstallerConfiguration Configuration, FBuildPatchInstallerDelegate OnComplete) const override;
+	virtual IBuildInstallerSharedContextRef CreateBuildInstallerSharedContext(const TCHAR* DebugName) const override;
 	virtual BuildPatchServices::IBuildStatisticsRef CreateBuildStatistics(const IBuildInstallerRef& Installer) const override;
 	virtual BuildPatchServices::IPatchDataEnumerationRef CreatePatchDataEnumeration(BuildPatchServices::FPatchDataEnumerationConfiguration Configuration) const override;
 	virtual IBuildManifestPtr LoadManifestFromFile(const FString& Filename) override;
diff --git a/Engine/Source/Runtime/Online/BuildPatchServices/Private/BuildPatchSettings.cpp b/Engine/Source/Runtime/Online/BuildPatchServices/Private/BuildPatchSettings.cpp
--- a/Engine/Source/Runtime/Online/BuildPatchServices/Private/BuildPatchSettings.cpp
+++ b/Engine/Source/Runtime/Online/BuildPatchServices/Private/BuildPatchSettings.cpp
@@ -155,38 +155,6 @@ namespace BuildPatchServices
 	{
 	}
 
-	FBuildInstallerConfiguration::FBuildInstallerConfiguration(FBuildInstallerConfiguration&& MoveFrom)
-		: InstallerActions(MoveTemp(MoveFrom.InstallerActions))
-		, InstallDirectory(MoveTemp(MoveFrom.InstallDirectory))
-		, StagingDirectory(MoveTemp(MoveFrom.StagingDirectory))
-		, BackupDirectory(MoveTemp(MoveFrom.BackupDirectory))
-		, ChunkDatabaseFiles(MoveTemp(MoveFrom.ChunkDatabaseFiles))
-		, CloudDirectories(MoveTemp(MoveFrom.CloudDirectories))
-		, InstallMode(MoveFrom.InstallMode)
-		, VerifyMode(MoveFrom.VerifyMode)
-		, DeltaPolicy(MoveFrom.DeltaPolicy)
-		, bRunRequiredPrereqs(MoveFrom.bRunRequiredPrereqs)
-		, bAllowConcurrentExecution(MoveFrom.bAllowConcurrentExecution)
-		, bTrackFileOperations(MoveFrom.bTrackFileOperations)
-	{
-	}
-
-	FBuildInstallerConfiguration::FBuildInstallerConfiguration(const FBuildInstallerConfiguration& CopyFrom)
-		: InstallerActions(CopyFrom.InstallerActions)
-		, InstallDirectory(CopyFrom.InstallDirectory)
-		, StagingDirectory(CopyFrom.StagingDirectory)
-		, BackupDirectory(CopyFrom.BackupDirectory)
-		, ChunkDatabaseFiles(CopyFrom.ChunkDatabaseFiles)
-		, CloudDirectories(CopyFrom.CloudDirectories)
-		, InstallMode(CopyFrom.InstallMode)
-		, VerifyMode(CopyFrom.VerifyMode)
-		, DeltaPolicy(CopyFrom.DeltaPolicy)
-		, bRunRequiredPrereqs(CopyFrom.bRunRequiredPrereqs)
-		, bAllowConcurrentExecution(CopyFrom.bAllowConcurrentExecution)
-		, bTrackFileOperations(CopyFrom.bTrackFileOperations)
-	{
-	}
-
 	FChunkBuildConfiguration::FChunkBuildConfiguration()
 		: FeatureLevel(EFeatureLevel::Latest)
 		, AppId(0)
diff --git a/Engine/Source/Runtime/Online/BuildPatchServices/Private/Installer/ChunkDbChunkSource.cpp b/Engine/Source/Runtime/Online/BuildPatchServices/Private/Installer/ChunkDbChunkSource.cpp
--- a/Engine/Source/Runtime/Online/BuildPatchServices/Private/Installer/ChunkDbChunkSource.cpp
+++ b/Engine/Source/Runtime/Online/BuildPatchServices/Private/Installer/ChunkDbChunkSource.cpp
@@ -14,6 +14,7 @@
 #include "Installer/ChunkReferenceTracker.h"
 #include "Installer/MessagePump.h"
 #include "Installer/InstallerError.h"
+#include "Installer/InstallerSharedContext.h"
 
 #if PLATFORM_WINDOWS
 #include "Windows/WindowsHWrapper.h"
@@ -133,20 +134,22 @@ namespace BuildPatchServices
 		// Configuration.
 		const FChunkDbSourceConfig Configuration;
 		// Dependencies.
-		IPlatform* Platform;
-		IFileSystem* FileSystem;
-		IChunkStore* ChunkStore;
-		IChunkReferenceTracker* ChunkReferenceTracker;
-		IChunkDataSerialization* ChunkDataSerialization;
-		IMessagePump* MessagePump;
-		IInstallerError* InstallerError;
+		IPlatform* Platform = nullptr;
+		IFileSystem* FileSystem = nullptr;
+		IChunkStore* ChunkStore = nullptr;
+		IChunkReferenceTracker* ChunkReferenceTracker = nullptr;
+		IChunkDataSerialization* ChunkDataSerialization = nullptr;
+		IMessagePump* MessagePump = nullptr;
+		IInstallerError* InstallerError = nullptr;
 		IChunkDbChunkSourceStat* ChunkDbChunkSourceStat;
-		// The future for our worker thread.
+		// Worker thread lifetime management
+		TPromise<void> Promise;
 		TFuture<void> Future;
+		IBuildInstallerThread* Thread = nullptr;
 		// Control signals.
-		FThreadSafeBool bIsPaused;
-		FThreadSafeBool bShouldAbort;
-		FThreadSafeBool bStartedLoading;
+		std::atomic<bool> bIsPaused = false;
+		std::atomic<bool> bShouldAbort = false;
+		std::atomic<bool> bStartedLoading = false;
 		// Handling of chunks we lose access to.
 		TFunction<void(TSet<FGuid>)> UnavailableChunksCallback;
 		TSet<FGuid> UnavailableChunks;
@@ -173,10 +176,7 @@ namespace BuildPatchServices
 		, MessagePump(InMessagePump)
 		, InstallerError(InInstallerError)
 		, ChunkDbChunkSourceStat(InChunkDbChunkSourceStat)
-		, bIsPaused(false)
-		, bShouldAbort(false)
-		, bStartedLoading(InConfiguration.bBeginLoadsOnFirstGet == false)
-		, UnavailableChunksCallback(nullptr)
+		, bStartedLoading(!InConfiguration.bBeginLoadsOnFirstGet)
 	{
 		// Allow OS intervention only once.
 		bool bResetOsIntervention = false;
@@ -224,8 +224,9 @@ namespace BuildPatchServices
 		// Start threaded load worker.
 		if (ChunkDbDataAccesses.Num() > 0)
 		{
-			TFunction<void()> Task = [this]() { return ThreadRun(); };
-			Future = Async(EAsyncExecution::ThreadIfForkSafe, MoveTemp(Task));
+			Future = Promise.GetFuture();
+			Thread = Configuration.SharedContext->CreateThread();
+			Thread->RunTask([this]() { ThreadRun(); });
 		}
 	}
 
@@ -233,10 +234,15 @@ namespace BuildPatchServices
 	{
 		Abort();
 
-		if (Future.IsValid())
+		if (!Thread)
 		{
-			Future.Wait();
+			// The thread was never started, so complete the promise here
+			Promise.SetValue();
 		}
+
+		Future.Wait();
+		Configuration.SharedContext->ReleaseThread(Thread);
+		Thread = nullptr;
 	}
 
 	void FChunkDbChunkSource::SetPaused(bool bInIsPaused)
@@ -344,6 +350,8 @@ namespace BuildPatchServices
 				Platform->Sleep(0.1f);
 			}
 		}
+
+		Promise.SetValue();
 	}
 
 	void FChunkDbChunkSource::LoadChunk(const FGuid& DataId)
diff --git a/Engine/Source/Runtime/Online/BuildPatchServices/Private/Installer/ChunkDbChunkSource.h b/Engine/Source/Runtime/Online/BuildPatchServices/Private/Installer/ChunkDbChunkSource.h
--- a/Engine/Source/Runtime/Online/BuildPatchServices/Private/Installer/ChunkDbChunkSource.h
+++ b/Engine/Source/Runtime/Online/BuildPatchServices/Private/Installer/ChunkDbChunkSource.h
@@ -6,6 +6,8 @@
 #include "Containers/Set.h"
 #include "Common/SpeedRecorder.h"
 
+class IBuildInstallerSharedContext;
+
 namespace BuildPatchServices
 {
 	class IPlatform;
@@ -42,14 +44,16 @@ namespace BuildPatchServices
 		// An array of chunkdb full file paths.
 		TArray<FString> ChunkDbFiles;
 		// The minimum number of chunks to load ahead of what is required, depending on store slack.
-		int32 PreFetchMinimum;
+		int32 PreFetchMinimum = 10;
 		// The maximum number of chunks to load ahead of what is required, depending on store slack.
-		int32 PreFetchMaximum;
+		int32 PreFetchMaximum = 40;
 		// The time in seconds to wait until trying to open a chunkdb file again after we lost the file handle (e.g. due to device eject or network error).
-		float ChunkDbOpenRetryTime;
+		float ChunkDbOpenRetryTime = 5.0f;
 		// If true, the loading will not begin until the first Get request is made. It is fairly fundamental to stop loading of chunks until resume
 		// data is processed, but can be special case disabled.
-		bool bBeginLoadsOnFirstGet;
+		bool bBeginLoadsOnFirstGet = true;
+		// The context for allocating shared resources.
+		IBuildInstallerSharedContext* SharedContext = nullptr;
 
 		/**
 		 * Constructor which sets usual defaults, and takes params for values that cannot use a default.
@@ -57,12 +61,7 @@ namespace BuildPatchServices
 		 */
 		FChunkDbSourceConfig(TArray<FString> InChunkDbFiles)
 			: ChunkDbFiles(MoveTemp(InChunkDbFiles))
-			, PreFetchMinimum(10)
-			, PreFetchMaximum(40)
-			, ChunkDbOpenRetryTime(5.0f)
-			, bBeginLoadsOnFirstGet(true)
-		{
-		}
+		{}
 	};
 
 	/**
diff --git a/Engine/Source/Runtime/Online/BuildPatchServices/Private/Installer/CloudChunkSource.cpp b/Engine/Source/Runtime/Online/BuildPatchServices/Private/Installer/CloudChunkSource.cpp
--- a/Engine/Source/Runtime/Online/BuildPatchServices/Private/Installer/CloudChunkSource.cpp
+++ b/Engine/Source/Runtime/Online/BuildPatchServices/Private/Installer/CloudChunkSource.cpp
@@ -17,6 +17,7 @@
 #include "Installer/MessagePump.h"
 #include "Common/StatsCollector.h"
 #include "Interfaces/IBuildInstaller.h"
+#include "Interfaces/IBuildInstallerSharedContext.h"
 #include "BuildPatchUtil.h"
 #include "Installer/Statistics/DownloadServiceStatistics.h"
 
@@ -173,6 +174,7 @@ namespace BuildPatchServices
 		const TSet<FGuid> InitialDownloadSet;
 		TPromise<void> Promise;
 		TFuture<void> Future;
+		IBuildInstallerThread* Thread = nullptr;
 		FDownloadProgressDelegate OnDownloadProgressDelegate;
 		FDownloadCompleteDelegate OnDownloadCompleteDelegate;
 
@@ -240,6 +242,7 @@ namespace BuildPatchServices
 		, InitialDownloadSet(MoveTemp(InInitialDownloadSet))
 		, Promise()
 		, Future()
+		, Thread(nullptr)
 		, OnDownloadProgressDelegate(FDownloadProgressDelegate::CreateThreadSafeSP(DownloadDelegates, &FDownloadDelegates::OnDownloadProgress))
 		, OnDownloadCompleteDelegate(FDownloadCompleteDelegate::CreateThreadSafeSP(DownloadDelegates, &FDownloadDelegates::OnDownloadComplete))
 		, CyclesAtLastData(0)
@@ -251,14 +254,12 @@ namespace BuildPatchServices
 		, RequestedDownloads()
 		, DownloadCount(InDownloadConnectionCount)
 	{
+		Future = Promise.GetFuture();
 		if (Configuration.bRunOwnThread)
 		{
-			TFunction<void()> Task = [this]() { return ThreadRun(); };
-			Future = Async(EAsyncExecution::ThreadIfForkSafe, MoveTemp(Task));
-		}
-		else
-		{
-			Future = Promise.GetFuture();
+			check(Configuration.SharedContext);
+			Thread = Configuration.SharedContext->CreateThread();
+			Thread->RunTask([this]() { ThreadRun(); });
 		}
 	}
 
@@ -266,6 +267,12 @@ namespace BuildPatchServices
 	{
 		bShouldAbort = true;
 		Future.Wait();
+
+		if (Thread)
+		{
+			Configuration.SharedContext->ReleaseThread(Thread);
+			Thread = nullptr;
+		}
 	}
 
 	void FCloudChunkSource::SetPaused(bool bInIsPaused)
diff --git a/Engine/Source/Runtime/Online/BuildPatchServices/Private/Installer/CloudChunkSource.h b/Engine/Source/Runtime/Online/BuildPatchServices/Private/Installer/CloudChunkSource.h
--- a/Engine/Source/Runtime/Online/BuildPatchServices/Private/Installer/CloudChunkSource.h
+++ b/Engine/Source/Runtime/Online/BuildPatchServices/Private/Installer/CloudChunkSource.h
@@ -8,6 +8,7 @@
 #include "Interfaces/IBuildInstaller.h"
 
 enum class EBuildPatchDownloadHealth;
+class IBuildInstallerSharedContext;
 
 namespace BuildPatchServices
 {
@@ -48,11 +49,11 @@ namespace BuildPatchServices
 		TArray<FString> CloudRoots;
 		// The maximum number of times that a single chunk should retry, before registering a fatal error.
 		// Infinite can be specified as < 0.
-		int32 MaxRetryCount;
+		int32 MaxRetryCount = 6;
 		// The minimum number of chunks to request ahead of what is required, depending on store slack.
-		int32 PreFetchMinimum;
+		int32 PreFetchMinimum  = 16;
 		// The maximum number of chunks to request ahead of what is required, depending on store slack.
-		int32 PreFetchMaximum;
+		int32 PreFetchMaximum = 256;
 		// Array of times in seconds, representing the time between each retry upon failure. The last entry will be used
 		// indefinitely once it is reached.
 		TArray<float> RetryDelayTimes;
@@ -60,14 +61,16 @@ namespace BuildPatchServices
 		// EBuildPatchDownloadHealth::NUM_Values entries in this array.
 		TArray<float> HealthPercentages;
 		// When all requests are failing, how many seconds before a success until we determine the state as disconnected.
-		float DisconnectedDelay;
+		float DisconnectedDelay = 5.0f;
 		// If true, the downloads will not begin until the first Get request is made. It is fairly fundamental to stop downloads of chunks until resume
 		// data is processed, but can be special case disabled.
-		bool bBeginDownloadsOnFirstGet;
+		bool bBeginDownloadsOnFirstGet = true;
 		// The minimum time to allow a http download before assessing it as affected by TCP zero window issue.
-		float TcpZeroWindowMinimumSeconds;
-		// Whether cloud source should run its own thread.
-		bool bRunOwnThread;
+		float TcpZeroWindowMinimumSeconds = 20.0f;
+		// The context for allocating shared resources.
+		IBuildInstallerSharedContext* SharedContext = nullptr;
+		// Whether cloud source should run its own thread. If true, SharedContext is required.
+		bool bRunOwnThread = true;
 
 		/**
 		 * Constructor which sets usual defaults, and takes params for values that cannot use a default.
@@ -75,13 +78,6 @@ namespace BuildPatchServices
 		 */
 		FCloudSourceConfig(TArray<FString> InCloudRoots)
 			: CloudRoots(MoveTemp(InCloudRoots))
-			, MaxRetryCount(6)
-			, PreFetchMinimum(16)
-			, PreFetchMaximum(256)
-			, DisconnectedDelay(5.0f)
-			, bBeginDownloadsOnFirstGet(true)
-			, TcpZeroWindowMinimumSeconds(20.0f)
-			, bRunOwnThread(true)
 		{
 			const float RetryFloats[] = {0.5f, 1.0f, 1.0f, 3.0f, 3.0f, 10.0f, 10.0f, 20.0f, 20.0f, 30.0f};
 			RetryDelayTimes.Empty(UE_ARRAY_COUNT(RetryFloats));
@@ -118,7 +114,18 @@ namespace BuildPatchServices
 		 * @param InitialDownloadSet        The initial set of chunks to be sourced from cloud.
 		 * @return the new ICloudChunkSource instance created.
 		 */
-		static ICloudChunkSource* Create(FCloudSourceConfig Configuration, IPlatform* Platform, IChunkStore* ChunkStore, IDownloadService* DownloadService, IChunkReferenceTracker* ChunkReferenceTracker, IChunkDataSerialization* ChunkDataSerialization, IMessagePump* MessagePump, IInstallerError* InstallerError, IDownloadConnectionCount* InConnectionCount, ICloudChunkSourceStat* CloudChunkSourceStat, IBuildManifestSet* ManifestSet, TSet<FGuid> InitialDownloadSet);
+		static ICloudChunkSource* Create(FCloudSourceConfig Configuration, 
+			IPlatform* Platform, 
+			IChunkStore* ChunkStore, 
+			IDownloadService* DownloadService, 
+			IChunkReferenceTracker* ChunkReferenceTracker, 
+			IChunkDataSerialization* ChunkDataSerialization, 
+			IMessagePump* MessagePump, 
+			IInstallerError* InstallerError, 
+			IDownloadConnectionCount* InConnectionCount, 
+			ICloudChunkSourceStat* CloudChunkSourceStat, 
+			IBuildManifestSet* ManifestSet, 
+			TSet<FGuid> InitialDownloadSet);
 	};
 
 	/**
diff --git a/Engine/Source/Runtime/Online/BuildPatchServices/Private/Installer/DiskChunkStore.cpp b/Engine/Source/Runtime/Online/BuildPatchServices/Private/Installer/DiskChunkStore.cpp
--- a/Engine/Source/Runtime/Online/BuildPatchServices/Private/Installer/DiskChunkStore.cpp
+++ b/Engine/Source/Runtime/Online/BuildPatchServices/Private/Installer/DiskChunkStore.cpp
@@ -1,6 +1,7 @@
 // Copyright Epic Games, Inc. All Rights Reserved.
 
 #include "Installer/DiskChunkStore.h"
+#include "Installer/InstallerSharedContext.h"
 #include "HAL/ThreadSafeBool.h"
 #include "Containers/Union.h"
 #include "Containers/Queue.h"
@@ -44,22 +45,24 @@ namespace BuildPatchServices
 		bool IsExpectedChunk(IChunkDataAccess* ChunkDataAccess, const FGuid& ChunkId);
 
 	private:
-		IFileSystem* const FileSystem;
-		IChunkDataSerialization* const Serializer;
-		IDiskChunkStoreStat* const DiskChunkStoreStat;
+		IFileSystem* const FileSystem = nullptr;
+		IChunkDataSerialization* const Serializer = nullptr;
+		IDiskChunkStoreStat* const DiskChunkStoreStat = nullptr;
 		const FDiskChunkStoreConfig Configuration;
-		mutable FCriticalSection LostChunkCallbackCs;
+		FCriticalSection LostChunkCallbackCs;
 		TFunction<void(const FGuid&)> LostChunkCallback;
-		mutable FCriticalSection ThreadLockCs;
+		FCriticalSection ThreadLockCs;
 		FGuid LastGetId;
 		TUniquePtr<IChunkDataAccess> LastGetData;
 		TSet<FGuid> PlacedInStore;
-		FThreadSafeBool bShouldRun;
+		std::atomic<bool> bShouldRun = true;
 		FChunkQueue Queue;
-		FThreadSafeInt32 QueueNum;
-		FEvent* QueueTrigger;
+		std::atomic<int32> QueueNum = 0;
+		FEvent* QueueTrigger = nullptr;
+		TPromise<void> IoThreadPromise;
 		TFuture<void> IoThreadFuture;
-		FEvent* IoThreadTrigger;
+		IBuildInstallerThread* Thread = nullptr;
+		FEvent* IoThreadTrigger = nullptr;
 	};
 
 	FDiskChunkStore::FDiskChunkStore(IFileSystem* InFileSystem, IChunkDataSerialization* InSerializer, IDiskChunkStoreStat* InDiskChunkStoreStat, FDiskChunkStoreConfig InConfiguration)
@@ -67,21 +70,12 @@ namespace BuildPatchServices
 		, Serializer(InSerializer)
 		, DiskChunkStoreStat(InDiskChunkStoreStat)
 		, Configuration(MoveTemp(InConfiguration))
-		, LostChunkCallbackCs()
-		, LostChunkCallback(nullptr)
-		, ThreadLockCs()
-		, LastGetId()
-		, LastGetData(nullptr)
-		, PlacedInStore()
-		, bShouldRun(true)
-		, QueueNum(0)
+		, QueueTrigger(FPlatformProcess::GetSynchEventFromPool(true))
+		, IoThreadTrigger(FPlatformProcess::GetSynchEventFromPool(true))
 	{
-		QueueTrigger = FPlatformProcess::GetSynchEventFromPool(true);
-		IoThreadTrigger = FPlatformProcess::GetSynchEventFromPool(true);
-		IoThreadFuture = Async(EAsyncExecution::ThreadIfForkSafe, [this]()
-		{
-			return IoThread();
-		});
+		IoThreadFuture = IoThreadPromise.GetFuture();
+		Thread = Configuration.SharedContext->CreateThread();
+		Thread->RunTask([this]() { IoThread(); });
 	}
 
 	FDiskChunkStore::~FDiskChunkStore()
@@ -91,9 +85,15 @@ namespace BuildPatchServices
 		QueueTrigger->Trigger();
 		IoThreadTrigger->Trigger();
 		IoThreadFuture.Wait();
+
 		// Return events.
 		FPlatformProcess::ReturnSynchEventToPool(IoThreadTrigger);
 		FPlatformProcess::ReturnSynchEventToPool(QueueTrigger);
+
+		// Return thread.
+		Configuration.SharedContext->ReleaseThread(Thread);
+		Thread = nullptr;
+
 		// Clean up allocations left in the queue.
 		FQueuedChunk QueuedChunk;
 		while (Queue.Dequeue(QueuedChunk))
@@ -209,7 +209,7 @@ namespace BuildPatchServices
 		{
 			if (Queue.Dequeue(QueuedChunk))
 			{
-				QueueNum.Decrement();
+				QueueNum--;
 				QueueTrigger->Trigger();
 				if (QueuedChunk.HasSubtype<FQueuedChunkLoad>())
 				{
@@ -325,18 +325,20 @@ namespace BuildPatchServices
 		DumpWriter.Reset();
 		DumpReader.Reset();
 		FileSystem->DeleteFile(*DumpFilename);
+
+		IoThreadPromise.SetValue();
 	}
 
 	BuildPatchServices::FLoadFuture FDiskChunkStore::QueueLoadRequest(const FGuid& DataId)
 	{
-		while (QueueNum.GetValue() > Configuration.QueueSize && bShouldRun)
+		while (QueueNum > Configuration.QueueSize && bShouldRun)
 		{
 			// Wait 1 second max in case of abort.
 			static const uint32 WaitTime = 1000;
 			QueueTrigger->Wait(WaitTime);
 			QueueTrigger->Reset();
 		}
-		QueueNum.Increment();
+		QueueNum++;
 		TUniquePtr<FLoadPromise> PromisePtr(new FLoadPromise());
 		FLoadFuture Future = PromisePtr->GetFuture();
 		if (bShouldRun)
@@ -353,14 +355,14 @@ namespace BuildPatchServices
 
 	void FDiskChunkStore::QueueSaveRequest(const FGuid& DataId, TUniquePtr<IChunkDataAccess>&& ChunkData)
 	{
-		while (QueueNum.GetValue() > Configuration.QueueSize && bShouldRun)
+		while (QueueNum > Configuration.QueueSize && bShouldRun)
 		{
 			// Wait 1 second max in case of abort.
 			static const uint32 WaitTime = 1000;
 			QueueTrigger->Wait(WaitTime);
 			QueueTrigger->Reset();
 		}
-		QueueNum.Increment();
+		QueueNum++;
 		Queue.Enqueue(FQueuedChunk(FQueuedChunkSave(DataId, ChunkData.Release())));
 		IoThreadTrigger->Trigger();
 	}
diff --git a/Engine/Source/Runtime/Online/BuildPatchServices/Private/Installer/DiskChunkStore.h b/Engine/Source/Runtime/Online/BuildPatchServices/Private/Installer/DiskChunkStore.h
--- a/Engine/Source/Runtime/Online/BuildPatchServices/Private/Installer/DiskChunkStore.h
+++ b/Engine/Source/Runtime/Online/BuildPatchServices/Private/Installer/DiskChunkStore.h
@@ -3,6 +3,8 @@
 
 #include "Installer/ChunkStore.h"
 
+class IBuildInstallerSharedContext;
+
 namespace BuildPatchServices
 {
 	class IFileSystem;
@@ -26,19 +28,18 @@ namespace BuildPatchServices
 		// The directory path for where to store chunk files.
 		FString StoreRootPath;
 		// The number of requests to allow to queue up before blocking Put/Get calls.
-		int32 QueueSize;
+		int32 QueueSize = 50;
 		// The max time in seconds to wait before retrying file handle for dump file.
-		double MaxRetryTime;
+		double MaxRetryTime = 2.0;
+		// The context for allocating shared resources.
+		IBuildInstallerSharedContext* SharedContext = nullptr;
 
 		/**
 		 * Constructor which sets usual defaults.
 		 */
 		FDiskChunkStoreConfig(FString InStoreRootPath)
 			: StoreRootPath(MoveTemp(InStoreRootPath))
-			, QueueSize(50)
-			, MaxRetryTime(2.0)
-		{
-		}
+		{}
 	};
 
 	/**
diff --git a/Engine/Source/Runtime/Online/BuildPatchServices/Private/Installer/InstallerSharedContext.cpp b/Engine/Source/Runtime/Online/BuildPatchServices/Private/Installer/InstallerSharedContext.cpp
new file mode 100644
--- /dev/null
+++ b/Engine/Source/Runtime/Online/BuildPatchServices/Private/Installer/InstallerSharedContext.cpp
@@ -0,0 +1,201 @@
+// Copyright Epic Games, Inc. All Rights Reserved.
+
+#include "Installer/InstallerSharedContext.h"
+
+#include "BuildPatchServicesPrivate.h"
+#include "HAL/Event.h"
+#include "HAL/PlatformProcess.h"
+#include "HAL/RunnableThread.h"
+#include "Misc/Fork.h"
+
+namespace BuildPatchServices
+{
+	ENUM_CLASS_FLAGS(FBuildInstallerThread::EMsg);
+
+	FBuildInstallerThread::~FBuildInstallerThread()
+	{
+		if (Thread != nullptr)
+		{
+			delete Thread; // Thread destructor should wait for completion
+			Thread = nullptr;
+		}
+
+		FPlatformProcess::ReturnSynchEventToPool(DoWorkEvent);
+		DoWorkEvent = nullptr;
+	}
+
+	bool FBuildInstallerThread::StartThread(const TCHAR* DebugName)
+	{
+		// Ideally this would check if we were forkable or a forked child process but there is 
+		// currently no in-engine way to check that.  Since BPS does not currently support
+		// FRunnableThread::ThreadType::Fake or FRunnableThread::ThreadType::Forkable 
+		// this check ends up being equivalent for now.  We most likely *never* want support 
+		// forking while an installer is running.
+		if (FPlatformProcess::SupportsMultithreading() || FForkProcessHelper::IsForkedMultithreadInstance())
+		{
+			DoWorkEvent = FPlatformProcess::GetSynchEventFromPool();
+
+			Thread = FForkProcessHelper::CreateForkableThread(this, DebugName);
+			check(Thread != nullptr);
+			check(Thread->GetThreadType() == FRunnableThread::ThreadType::Real);
+		}
+
+		return DoWorkEvent && Thread;
+	}
+
+	void FBuildInstallerThread::RunTask(TUniqueFunction<void()> InTask)
+	{
+		check(InTask);
+
+		MsgQueue.Enqueue(MoveTemp(InTask), EMsg::RunTask);
+
+		DoWorkEvent->Trigger();
+	}
+
+	uint32 FBuildInstallerThread::Run()
+	{
+		bool bExit = false;
+		while (!bExit)
+		{
+			// Wait for some work to do
+			DoWorkEvent->Wait();
+
+			check(!MsgQueue.IsEmpty()); // Where is the job or termination request if the thread is awake?
+
+			FMsg Msg;
+			while (MsgQueue.Dequeue(Msg))
+			{
+				if (EnumHasAnyFlags(Msg.Msg, EMsg::RunTask))
+				{
+					check(!bExit);
+					Msg.Task();
+				}
+
+				if (EnumHasAnyFlags(Msg.Msg, EMsg::Exit))
+				{
+					bExit = true;
+				}
+			}
+		}
+
+		return 0;
+	}
+
+	void FBuildInstallerThread::Stop()
+	{
+		MsgQueue.Enqueue(nullptr, EMsg::Exit);
+
+		DoWorkEvent->Trigger();
+	}
+
+	IBuildInstallerThread* FBuildInstallerSharedContext::CreateThreadInternal()
+	{
+		TStringBuilder<512> ThreadNameBuilder;
+		ThreadNameBuilder.Appendf(TEXT("%s #%d"), DebugName, ThreadCount);
+
+		FBuildInstallerThread* Thread = new FBuildInstallerThread();
+		if (Thread->StartThread(ThreadNameBuilder.ToString()))
+		{
+			++ThreadCount;
+		}
+		else
+		{
+			delete Thread;
+			Thread = nullptr;
+		}
+
+		return Thread;
+	}
+
+	void FBuildInstallerSharedContext::PreallocateThreads(uint32 NumThreads)
+	{
+		if (NumThreads == 0)
+		{
+			return;
+		}
+
+		FScopeLock Lock(&ThreadFreeListCS);
+		ThreadFreeList.Reserve(NumThreads);
+		for (uint32 i = 0; i < NumThreads; ++i)
+		{
+			if (IBuildInstallerThread* Thread = CreateThreadInternal())
+			{
+				ThreadFreeList.Push(Thread);
+			}
+		}
+
+		bWarnOnCreateThread = true;
+	}
+
+	uint32 FBuildInstallerSharedContext::NumThreadsPerInstaller(bool bUseChunkDBs) const
+	{
+		const uint32 NumInstallerMainThreads = 1;
+		const uint32 NumCloudChunkSourceThreads = 1;
+#if ENABLE_PATCH_DISK_OVERFLOW_STORE
+		const uint32 NumDiskChunkStoreThreads = 1;
+#else
+		const uint32 NumDiskChunkStoreThreads = 0;
+#endif
+		const uint32 NumChunkDBThreads = bUseChunkDBs ? 1 : 0;
+		const uint32 NumExpectedThreads =
+			NumInstallerMainThreads + 
+			NumCloudChunkSourceThreads + 
+			NumDiskChunkStoreThreads + 
+			NumChunkDBThreads;
+
+		return NumExpectedThreads;
+	}
+
+	FBuildInstallerSharedContext::~FBuildInstallerSharedContext()
+	{
+		FScopeLock Lock(&ThreadFreeListCS);
+
+		// All threads should have been freed before deleting the context
+		check(ThreadFreeList.Num() == ThreadCount);
+		for (IBuildInstallerThread* Thread : ThreadFreeList)
+		{
+			FBuildInstallerThread* ThreadActual = static_cast<FBuildInstallerThread*>(Thread);
+			ThreadActual->Stop();
+		}
+
+		for (IBuildInstallerThread* Thread : ThreadFreeList)
+		{
+			FBuildInstallerThread* ThreadActual = static_cast<FBuildInstallerThread*>(Thread);
+			delete ThreadActual;
+		}
+
+		ThreadFreeList.Empty();
+	}
+
+	IBuildInstallerThread* FBuildInstallerSharedContext::CreateThread()
+	{
+		IBuildInstallerThread* Thread = nullptr;
+		{
+			FScopeLock Lock(&ThreadFreeListCS);
+			Thread = ThreadFreeList.IsEmpty() ? nullptr : ThreadFreeList.Pop(false);
+		}
+
+		if (!Thread)
+		{
+			UE_CLOG(bWarnOnCreateThread, LogBuildPatchServices, Warning, TEXT("Allocating installer thread, free list exhausted, check PreallocateResources()"));
+			Thread = CreateThreadInternal();
+		}
+
+		return Thread;
+	}
+
+	void FBuildInstallerSharedContext::ReleaseThread(IBuildInstallerThread* Thread)
+	{
+		if (Thread)
+		{
+			FScopeLock Lock(&ThreadFreeListCS);
+			check(!ThreadFreeList.Contains(Thread)); // Double free is bad y'all
+			ThreadFreeList.Push(Thread);
+		}
+	}
+
+	IBuildInstallerSharedContextRef FBuildInstallerSharedContextFactory::Create(const TCHAR* DebugName)
+	{
+		return MakeShared<FBuildInstallerSharedContext>(DebugName);
+	}
+}
diff --git a/Engine/Source/Runtime/Online/BuildPatchServices/Private/Installer/InstallerSharedContext.h b/Engine/Source/Runtime/Online/BuildPatchServices/Private/Installer/InstallerSharedContext.h
new file mode 100644
--- /dev/null
+++ b/Engine/Source/Runtime/Online/BuildPatchServices/Private/Installer/InstallerSharedContext.h
@@ -0,0 +1,91 @@
+// Copyright Epic Games, Inc. All Rights Reserved.
+
+#pragma once
+
+#include "HAL/Runnable.h"
+#include "Containers/MpscQueue.h"
+#include "Interfaces/IBuildInstallerSharedContext.h"
+
+namespace BuildPatchServices
+{
+	/**
+	 * The concrete implementation of IBuildInstallerThread
+	 */
+	class FBuildInstallerThread
+		: public IBuildInstallerThread
+		, public FRunnable
+	{
+		enum class EMsg : uint8
+		{
+			None = 0,
+			RunTask = (1 << 0),
+			Exit = (1 << 1)
+		};
+		FRIEND_ENUM_CLASS_FLAGS(EMsg);
+
+		struct FMsg
+		{
+			TUniqueFunction<void()> Task;
+			EMsg Msg = EMsg::None;
+		};
+
+	public:
+		~FBuildInstallerThread();
+
+		bool StartThread(const TCHAR* DebugName);
+
+		// IBuildInstallerThread interface begin
+		virtual void RunTask(TUniqueFunction<void()> Task) override;
+		// IBuildInstallerThread interface end
+
+		// FRunnable interface begin
+		virtual uint32 Run() override;
+		virtual void Stop() override;
+		// FRunnable interface end
+
+	private:
+		FEvent* DoWorkEvent = nullptr;
+		FRunnableThread* Thread = nullptr;
+		TMpscQueue<FMsg> MsgQueue;
+	};
+
+	/**
+	 * The concrete implementation of IBuildInstallerSharedContext
+	 */
+	class FBuildInstallerSharedContext : public IBuildInstallerSharedContext
+	{
+	public:
+		FBuildInstallerSharedContext(const TCHAR* InDebugName) : DebugName(InDebugName) {}
+		~FBuildInstallerSharedContext();
+
+		// IBuildInstallerSharedContext interface begin
+		virtual IBuildInstallerThread* CreateThread() override;
+		virtual void ReleaseThread(IBuildInstallerThread* Thread) override;
+		virtual void PreallocateThreads(uint32 NumThreads) override;
+		virtual uint32 NumThreadsPerInstaller(bool bUseChunkDBs) const override;
+		// IBuildInstallerSharedContext interface end
+
+	private:
+		IBuildInstallerThread* CreateThreadInternal();
+
+	private:
+		TArray<IBuildInstallerThread*> ThreadFreeList;
+		FCriticalSection ThreadFreeListCS;
+		const TCHAR* DebugName;
+		uint32 ThreadCount = 0;
+		bool bWarnOnCreateThread = false;
+	};
+
+	/**
+	 * A factory for creating an IBuildInstallerSharedContext instance.
+	 */
+	class FBuildInstallerSharedContextFactory
+	{
+	public:
+		/**
+		 * @param DebugName          Used to tag resources allocated with the shared context.
+		 * @return an instance of an IBuildInstallerSharedContex implementation.
+		 */
+		static IBuildInstallerSharedContextRef Create(const TCHAR* DebugName);
+	};
+}
diff --git a/Engine/Source/Runtime/Online/BuildPatchServices/Private/Tests/Unit/CloudChunkSource.spec.cpp b/Engine/Source/Runtime/Online/BuildPatchServices/Private/Tests/Unit/CloudChunkSource.spec.cpp
--- a/Engine/Source/Runtime/Online/BuildPatchServices/Private/Tests/Unit/CloudChunkSource.spec.cpp
+++ b/Engine/Source/Runtime/Online/BuildPatchServices/Private/Tests/Unit/CloudChunkSource.spec.cpp
@@ -20,6 +20,7 @@
 #include "Tests/Mock/MessagePump.mock.h"
 #include "Tests/Mock/DownloadConnectionCount.mock.h"
 #include "Installer/CloudChunkSource.h"
+#include "Installer/InstallerSharedContext.h"
 #include "Core/Platform.h"
 #include "BuildPatchHash.h"
 
@@ -31,6 +32,7 @@ namespace BuildPatchServices
 BEGIN_DEFINE_SPEC(FCloudChunkSourceSpec, "BuildPatchServices.Unit", EAutomationTestFlags::ProductFilter | EAutomationTestFlags::ApplicationContextMask)
 // Unit.
 TUniquePtr<BuildPatchServices::ICloudChunkSource> CloudChunkSource;
+TUniquePtr<BuildPatchServices::FBuildInstallerSharedContext> InstallerSharedContext;
 // Mock/Fake.
 TUniquePtr<BuildPatchServices::FMockPlatform> MockPlatform;
 TUniquePtr<BuildPatchServices::FFakeChunkStore> FakeChunkStore;
@@ -87,7 +89,9 @@ void FCloudChunkSourceSpec::Define()
 	BeforeEach([this]()
 	{
 		MockPlatform.Reset(new FMockPlatform());
+		InstallerSharedContext.Reset(new FBuildInstallerSharedContext(TEXT("CloudChunkSourceSpec")));
 		CloudSourceConfig.Reset(new FCloudSourceConfig({ "http://download.mydomain.com/clouddata" }));
+		CloudSourceConfig->SharedContext = InstallerSharedContext.Get();
 		FakeChunkStore.Reset(new FFakeChunkStore());
 		FakeDownloadService.Reset(new FFakeDownloadService());
 		MockChunkReferenceTracker.Reset(new FFakeChunkReferenceTracker());
diff --git a/Engine/Source/Runtime/Online/BuildPatchServices/Private/Tests/Unit/DiskChunkStore.spec.cpp b/Engine/Source/Runtime/Online/BuildPatchServices/Private/Tests/Unit/DiskChunkStore.spec.cpp
--- a/Engine/Source/Runtime/Online/BuildPatchServices/Private/Tests/Unit/DiskChunkStore.spec.cpp
+++ b/Engine/Source/Runtime/Online/BuildPatchServices/Private/Tests/Unit/DiskChunkStore.spec.cpp
@@ -7,6 +7,7 @@
 #include "Tests/Mock/ChunkDataSerialization.mock.h"
 #include "Tests/Mock/DiskChunkStoreStat.mock.h"
 #include "Installer/DiskChunkStore.h"
+#include "Installer/InstallerSharedContext.h"
 #include "BuildPatchHash.h"
 
 #if WITH_DEV_AUTOMATION_TESTS
@@ -14,6 +15,7 @@
 BEGIN_DEFINE_SPEC(FDiskChunkStoreSpec, "BuildPatchServices.Unit", EAutomationTestFlags::ProductFilter | EAutomationTestFlags::ApplicationContextMask)
 // Unit
 TUniquePtr<BuildPatchServices::IDiskChunkStore> DiskChunkStore;
+TUniquePtr<BuildPatchServices::FBuildInstallerSharedContext> InstallerSharedContext;
 // Mock
 TUniquePtr<BuildPatchServices::FFakeFileSystem> FakeFileSystem;
 TUniquePtr<BuildPatchServices::FMockChunkDataSerialization> MockChunkDataSerialization;
@@ -47,6 +49,7 @@ void FDiskChunkStoreSpec::Define()
 	// Specs.
 	BeforeEach([this]()
 	{
+		InstallerSharedContext.Reset(new FBuildInstallerSharedContext(TEXT("DiskChunkStoreSpec")));
 		FakeFileSystem.Reset(new FFakeFileSystem());
 		MockChunkDataSerialization.Reset(new FMockChunkDataSerialization());
 		MockDiskChunkStoreStat.Reset(new FMockDiskChunkStoreStat());
@@ -364,6 +367,7 @@ void FDiskChunkStoreSpec::MakeUnit()
 {
 	using namespace BuildPatchServices;
 	FDiskChunkStoreConfig DiskChunkStoreConfig(StoreRootPath);
+	DiskChunkStoreConfig.SharedContext = InstallerSharedContext.Get();
 	DiskChunkStoreConfig.MaxRetryTime = 0.01;
 	DiskChunkStore.Reset(FDiskChunkStoreFactory::Create(
 		FakeFileSystem.Get(),
diff --git a/Engine/Source/Runtime/Online/BuildPatchServices/Public/BuildPatchSettings.h b/Engine/Source/Runtime/Online/BuildPatchServices/Public/BuildPatchSettings.h
--- a/Engine/Source/Runtime/Online/BuildPatchServices/Public/BuildPatchSettings.h
+++ b/Engine/Source/Runtime/Online/BuildPatchServices/Public/BuildPatchSettings.h
@@ -13,6 +13,7 @@
 #include "CoreMinimal.h"
 #include "HAL/Platform.h"
 #include "Interfaces/IBuildManifest.h"
+#include "Interfaces/IBuildInstallerSharedContext.h"
 #include "Misc/Variant.h"
 #include "Templates/UnrealTemplate.h"
 
@@ -255,19 +256,11 @@ namespace BuildPatchServices
 		 */
 		BUILDPATCHSERVICES_API FBuildInstallerConfiguration(TArray<FInstallerAction> InstallerActions);
 
-		/**
-		 * Copy constructor.
-		 */
-		BUILDPATCHSERVICES_API FBuildInstallerConfiguration(const FBuildInstallerConfiguration& CopyFrom);
-
-		/**
-		 * RValue constructor to allow move semantics.
-		 */
-		BUILDPATCHSERVICES_API FBuildInstallerConfiguration(FBuildInstallerConfiguration&& MoveFrom);
-
 	public:
 		// The array of intended actions to perform.
 		TArray<FInstallerAction> InstallerActions;
+		// The context for allocating shared resources.
+		IBuildInstallerSharedContextPtr SharedContext;
 		// The directory to install to.
 		FString InstallDirectory;
 		// The directory for storing the intermediate files. This would usually be inside the InstallDirectory. Empty string will use module's global setting.
diff --git a/Engine/Source/Runtime/Online/BuildPatchServices/Public/Interfaces/IBuildInstallerSharedContext.h b/Engine/Source/Runtime/Online/BuildPatchServices/Public/Interfaces/IBuildInstallerSharedContext.h
new file mode 100644
--- /dev/null
+++ b/Engine/Source/Runtime/Online/BuildPatchServices/Public/Interfaces/IBuildInstallerSharedContext.h
@@ -0,0 +1,65 @@
+// Copyright Epic Games, Inc. All Rights Reserved.
+
+#pragma once
+
+#include "Templates/Function.h"
+#include "Templates/SharedPointer.h"
+
+namespace BuildPatchServices
+{
+	/**
+	 * Interface for a re-usable thread
+	 */
+	class IBuildInstallerThread
+	{
+	protected:
+		/**
+		 * Virtual destructor.
+		 */
+		virtual ~IBuildInstallerThread() {}
+
+	public:
+		/**
+		 * Adds a task to the thread's queue
+		 */
+		virtual void RunTask(TUniqueFunction<void()> Task) = 0;
+	};
+}
+
+/**
+ * An interface for sharing threads and components between multiple BPS installers
+ */
+class IBuildInstallerSharedContext
+{
+public:
+	/**
+	 * Virtual destructor.
+	 */
+	virtual ~IBuildInstallerSharedContext() {}
+
+	/**
+	 * Return an existing free thread or allocate one.
+	 * @return an IBuildInstallerThread
+	 */
+	virtual BuildPatchServices::IBuildInstallerThread* CreateThread() = 0;
+
+	/**
+	 * Relinquish the thread and add to the free list. It should not be used again after this.
+	 */
+	virtual void ReleaseThread(BuildPatchServices::IBuildInstallerThread* Thread) = 0;
+
+	/**
+	 * Preallocate thread free list with NumThreads.
+	 * If more than NumThreads are required, a warning will be logged.
+	 */
+	virtual void PreallocateThreads(uint32 NumThreads) = 0;
+
+	/**
+	 * @return The number of threads required per installer
+	 */
+	virtual uint32 NumThreadsPerInstaller(bool bUseChunkDBs) const = 0;
+};
+
+using IBuildInstallerSharedContextRef = TSharedRef<IBuildInstallerSharedContext>;
+using IBuildInstallerSharedContextPtr = TSharedPtr<IBuildInstallerSharedContext>;
+using IBuildInstallerSharedContextWeakPtr = TWeakPtr<IBuildInstallerSharedContext>;
diff --git a/Engine/Source/Runtime/Online/BuildPatchServices/Public/Interfaces/IBuildPatchServicesModule.h b/Engine/Source/Runtime/Online/BuildPatchServices/Public/Interfaces/IBuildPatchServicesModule.h
--- a/Engine/Source/Runtime/Online/BuildPatchServices/Public/Interfaces/IBuildPatchServicesModule.h
+++ b/Engine/Source/Runtime/Online/BuildPatchServices/Public/Interfaces/IBuildPatchServicesModule.h
@@ -53,6 +53,13 @@ public:
 	 */
 	virtual IBuildInstallerRef CreateBuildInstaller(BuildPatchServices::FBuildInstallerConfiguration Configuration, FBuildPatchInstallerDelegate OnComplete) const = 0;
 
+	/**
+	 * Factory providing construction of a build installer shared context class.
+	 * @param DebugName          Used to tag resources allocated with the shared context.
+	 * @return an instance of an IBuildInstallerSharedContex implementation.
+	 */
+	virtual IBuildInstallerSharedContextRef CreateBuildInstallerSharedContext(const TCHAR* DebugName) const = 0;
+
 	/**
 	 * Factory providing construction of a build statistics class.
 	 * @param Installer     The installer to create a build statistics for.
