CIS fix - Missing constructors.  This should be OK with C++20 aggregate initialization but I guess we don't fully support it on all toolchains yet.

diff --git a/Engine/Source/Runtime/Online/BuildPatchServices/Private/Installer/InstallerSharedContext.h b/Engine/Source/Runtime/Online/BuildPatchServices/Private/Installer/InstallerSharedContext.h
--- a/Engine/Source/Runtime/Online/BuildPatchServices/Private/Installer/InstallerSharedContext.h
+++ b/Engine/Source/Runtime/Online/BuildPatchServices/Private/Installer/InstallerSharedContext.h
@@ -25,6 +25,12 @@ namespace BuildPatchServices
 
 		struct FMsg
 		{
+			FMsg() = default;
+			FMsg(TUniqueFunction<void()> InTask, EMsg InMsg)
+				: Task(MoveTemp(InTask))
+				, Msg(InMsg)
+			{}
+
 			TUniqueFunction<void()> Task;
 			EMsg Msg = EMsg::None;
 		};
