- Remove ultra spammy LogHTTP logging

diff --git a/Engine/Source/Runtime/Online/HTTP/Private/GenericPlatform/HttpRequestImpl.cpp b/Engine/Source/Runtime/Online/HTTP/Private/GenericPlatform/HttpRequestImpl.cpp
--- a/Engine/Source/Runtime/Online/HTTP/Private/GenericPlatform/HttpRequestImpl.cpp
+++ b/Engine/Source/Runtime/Online/HTTP/Private/GenericPlatform/HttpRequestImpl.cpp
@@ -6,33 +6,28 @@
 
 FHttpRequestCompleteDelegate& FHttpRequestImpl::OnProcessRequestComplete()
 {
-	UE_LOG(LogHttp, VeryVerbose, TEXT("FHttpRequestImpl::OnProcessRequestComplete()"));
 	return RequestCompleteDelegate;
 }
 
 PRAGMA_DISABLE_DEPRECATION_WARNINGS
 FHttpRequestProgressDelegate& FHttpRequestImpl::OnRequestProgress() 
 {
-	UE_LOG(LogHttp, VeryVerbose, TEXT("FHttpRequestImpl::OnRequestProgress()"));
 	return RequestProgressDelegate;
 }
 PRAGMA_ENABLE_DEPRECATION_WARNINGS
 
 FHttpRequestProgressDelegate64& FHttpRequestImpl::OnRequestProgress64() 
 {
-	UE_LOG(LogHttp, VeryVerbose, TEXT("FHttpRequestImpl::OnRequestProgress64()"));
 	return RequestProgressDelegate64;
 }
 
 FHttpRequestHeaderReceivedDelegate& FHttpRequestImpl::OnHeaderReceived()
 {
-	UE_LOG(LogHttp, VeryVerbose, TEXT("FHttpRequestImpl::OnHeaderReceived()"));
 	return HeaderReceivedDelegate;
 }
 
 FHttpRequestWillRetryDelegate& FHttpRequestImpl::OnRequestWillRetry()
 {
-	UE_LOG(LogHttp, VeryVerbose, TEXT("FHttpRequestImpl::OnRequestWillRetry()"));
 	return OnRequestWillRetryDelegate;
 }
 
