Add logging to FOptimisedDelta to help try to pinpoint hang.

diff --git a/Engine/Source/Runtime/Online/BuildPatchServices/Private/Installer/OptimisedDelta.cpp b/Engine/Source/Runtime/Online/BuildPatchServices/Private/Installer/OptimisedDelta.cpp
--- a/Engine/Source/Runtime/Online/BuildPatchServices/Private/Installer/OptimisedDelta.cpp
+++ b/Engine/Source/Runtime/Online/BuildPatchServices/Private/Installer/OptimisedDelta.cpp
@@ -137,6 +137,8 @@ namespace BuildPatchServices
 
 	void FOptimisedDelta::OnDownloadComplete(int32 RequestId, const FDownloadRef& Download)
 	{
+		UE_LOG(LogOptimisedDelta, Log, TEXT("Download complete for %s, launching comlpletion task"), *RelativeDeltaFilePath);
+
 		// Kick a task to finish this so as not to block the callback thread
 		UE::Tasks::Launch(UE_SOURCE_LOCATION, [this, RequestId, Download] 
 		{
@@ -203,6 +205,7 @@ namespace BuildPatchServices
 				else if (ShouldRetry(Download))
 				{
 					++RetryCount;
+					UE_LOG(LogOptimisedDelta, Log, TEXT("Failed to load %s, retrying %i/%i"), *RelativeDeltaFilePath, RetryCount, DeltaRetries);
 					CloudDirIdx = (CloudDirIdx + RetryCount) % Configuration.CloudDirectories.Num();
 					DownloadService->RequestFile(Configuration.CloudDirectories[CloudDirIdx] / RelativeDeltaFilePath, ChunkDeltaComplete, ChunkDeltaProgress);
 				}
@@ -214,6 +217,7 @@ namespace BuildPatchServices
 			else if (ShouldRetry(Download))
 			{
 				++RetryCount;
+				UE_LOG(LogOptimisedDelta, Log, TEXT("Failed to download %s, retrying %i/%i"), *RelativeDeltaFilePath, RetryCount, DeltaRetries);
 				CloudDirIdx = (CloudDirIdx + RetryCount) % Configuration.CloudDirectories.Num();
 				DownloadService->RequestFile(Configuration.CloudDirectories[CloudDirIdx] / RelativeDeltaFilePath, ChunkDeltaComplete, ChunkDeltaProgress);
 			}
