Fix the issue on certain PC, the http thread hang forever because of the incorrectly calculated elapsed time.

diff --git a/Engine/Source/Runtime/Online/HTTP/Private/HttpThread.cpp b/Engine/Source/Runtime/Online/HTTP/Private/HttpThread.cpp
--- a/Engine/Source/Runtime/Online/HTTP/Private/HttpThread.cpp
+++ b/Engine/Source/Runtime/Online/HTTP/Private/HttpThread.cpp
@@ -365,6 +365,7 @@ bool FLegacyHttpThread::Init()
 	return FHttpThreadBase::Init();
 }
 
+UE_DISABLE_OPTIMIZATION_SHIP
 uint32 FLegacyHttpThread::Run()
 {
 	// Arrays declared outside of loop to re-use memory
@@ -393,6 +394,16 @@ uint32 FLegacyHttpThread::Run()
 				{
 					SCOPE_CYCLE_COUNTER(STAT_HTTPThread_ActiveSleep);
 					double InnerLoopTime = InnerLoopEnd - InnerLoopBegin;
+
+					// On Windows when optimization enabled, seems InnerLoopEnd can get a value without adding the 
+					// const value 16777216.0 from FWindowsPlatformTime::Seoncds(), it could be caused by [REDACTED_PROFILE_URL]/openssl/issues/21522
+					// Until we upgrade to new openssl to confirm the fix, keep this along with PRAGMA_DISABLE_OPTIMIZATION
+					// as an additional step to be safe
+					if (InnerLoopTime < 0.0)
+					{
+						InnerLoopTime = 0.0;
+					}
+
 					double InnerSleep = FMath::Max(HttpThreadActiveFrameTimeInSeconds - InnerLoopTime, HttpThreadActiveMinimumSleepTimeInSeconds);
 					FPlatformProcess::SleepNoStats(InnerSleep);
 				}
@@ -413,6 +424,7 @@ uint32 FLegacyHttpThread::Run()
 	}
 	return 0;
 }
+UE_ENABLE_OPTIMIZATION_SHIP
 
 void FLegacyHttpThread::Stop()
 {
