Fix the p4 auto merge error which caused compile error.

diff --git a/Engine/Source/Runtime/Online/HTTP/Private/Curl/CurlHttp.cpp b/Engine/Source/Runtime/Online/HTTP/Private/Curl/CurlHttp.cpp
--- a/Engine/Source/Runtime/Online/HTTP/Private/Curl/CurlHttp.cpp
+++ b/Engine/Source/Runtime/Online/HTTP/Private/Curl/CurlHttp.cpp
@@ -983,7 +983,8 @@ bool FCurlHttpRequest::ProcessRequest()
 
 	if (!PreCheck() || !SetupRequest())
 	{
-		return FinishRequestNotInHttpManager();
+		FinishRequestNotInHttpManager();
+		return false;
 	}
 
 	// Clear the info cache log so we don't output messages from previous requests when reusing/retrying a request
@@ -1265,7 +1266,7 @@ void FCurlHttpRequest::FinishRequest()
 		{
 			UE_LOG(LogHttp, Warning, TEXT("%p: request failed, libcurl multi error: %d (%s)"), this, (int32)CurlAddToMultiResult, ANSI_TO_TCHAR(curl_multi_strerror(CurlAddToMultiResult)));
 		}
-		else
+		else if (CurlCompletionResult != CURLE_OK)
 		{
 			UE_LOG(LogHttp, Warning, TEXT("%p: request failed, libcurl error: %d (%s)"), this, (int32)CurlCompletionResult, ANSI_TO_TCHAR(curl_easy_strerror(CurlCompletionResult)));
 		}
