Updating BackgroundHTTP logging verbosity to reduce spam and make debugging easier by default

#tests Android Dev Client

diff --git a/Engine/Source/Runtime/Online/BackgroundHTTP/Private/BackgroundHttpManagerImpl.cpp b/Engine/Source/Runtime/Online/BackgroundHTTP/Private/BackgroundHttpManagerImpl.cpp
--- a/Engine/Source/Runtime/Online/BackgroundHTTP/Private/BackgroundHttpManagerImpl.cpp
+++ b/Engine/Source/Runtime/Online/BackgroundHTTP/Private/BackgroundHttpManagerImpl.cpp
@@ -298,7 +298,7 @@ void FBackgroundHttpManagerImpl::SetMaxActiveDownloads(int InMaxActiveDownloads)
 
 void FBackgroundHttpManagerImpl::AddRequest(const FBackgroundHttpRequestPtr Request)
 {
-	UE_LOG(LogBackgroundHttpManager, Log, TEXT("AddRequest Called - RequestID:%s"), *Request->GetRequestID());
+	UE_LOG(LogBackgroundHttpManager, VeryVerbose, TEXT("AddRequest Called - RequestID:%s"), *Request->GetRequestID());
 
 	//If we don't associate with any existing requests, go into our pending list. These will be moved into the ActiveRequest list during our Tick
 	if (!AssociateWithAnyExistingRequest(Request))
@@ -306,7 +306,7 @@ void FBackgroundHttpManagerImpl::AddRequest(const FBackgroundHttpRequestPtr Requ
 		FRWScopeLock ScopeLock(PendingRequestLock, SLT_Write);
 		PendingStartRequests.Add(Request);
 
-		UE_LOG(LogBackgroundHttpManager, Log, TEXT("Adding BackgroundHttpRequest to PendingStartRequests - RequestID:%s"), *Request->GetRequestID());
+		UE_LOG(LogBackgroundHttpManager, Verbose, TEXT("Adding BackgroundHttpRequest to PendingStartRequests - RequestID:%s"), *Request->GetRequestID());
 	}
 }
 
@@ -333,7 +333,7 @@ void FBackgroundHttpManagerImpl::RemoveRequest(const FBackgroundHttpRequestPtr R
 		NumRequestsRemoved = PendingStartRequests.Remove(Request);
 	}
 	
-	UE_LOG(LogBackgroundHttpManager, Log, TEXT("FGenericPlatformBackgroundHttpManager::RemoveRequest Called - RequestID:%s | NumRequestsActuallyRemoved:%d | NumCurrentlyActiveRequests:%d"), *Request->GetRequestID(), NumRequestsRemoved, NumCurrentlyActiveRequests);
+	UE_LOG(LogBackgroundHttpManager, Verbose, TEXT("FGenericPlatformBackgroundHttpManager::RemoveRequest Called - RequestID:%s | NumRequestsActuallyRemoved:%d | NumCurrentlyActiveRequests:%d"), *Request->GetRequestID(), NumRequestsRemoved, NumCurrentlyActiveRequests);
 }
 
 void FBackgroundHttpManagerImpl::CleanUpDataAfterCompletingRequest(const FBackgroundHttpRequestPtr Request)
@@ -366,7 +366,7 @@ bool FBackgroundHttpManagerImpl::AssociateWithAnyExistingRequest(const FBackgrou
 		if (ensureAlwaysMsgf(NewResponseWithExistingFile.IsValid(), TEXT("Failure to create FBackgroundHttpResponsePtr in FPlatformBackgroundHttp::ConstructBackgroundResponse! Can not associate new download with found finished download!")))
 		{
 			bDidAssociateWithExistingRequest = true;
-			UE_LOG(LogBackgroundHttpManager, Log, TEXT("Found existing background task to associate with! RequestID:%s | ExistingFileSize:%lld | ExistingFilePath:%s"), *Request->GetRequestID(), ExistingFileSize, *ExistingFilePath);
+			UE_LOG(LogBackgroundHttpManager, Display, TEXT("Found existing background task to associate with! RequestID:%s | ExistingFileSize:%lld | ExistingFilePath:%s"), *Request->GetRequestID(), ExistingFileSize, *ExistingFilePath);
 
 			//First send progress update for the file size so anything monitoring this download knows we are about to update this progress
 			Request->OnProgressUpdated().ExecuteIfBound(Request, ExistingFileSize, ExistingFileSize);
@@ -463,7 +463,7 @@ void FBackgroundHttpManagerImpl::ActivatePendingRequests()
 
 	if (HighestPriorityRequestToStart.IsValid())
 	{
-		UE_LOG(LogBackgroundHttpManager, Verbose, TEXT("Starting Request: %s Priority:%s"), *HighestPriorityRequestToStart->GetRequestID(), LexToString(HighestRequestPriority));
+		UE_LOG(LogBackgroundHttpManager, Display, TEXT("Activating Request: %s Priority:%s"), *HighestPriorityRequestToStart->GetRequestID(), LexToString(HighestRequestPriority));
 
 		//Actually move request to Active list now
 		FRWScopeLock ActiveScopeLock(ActiveRequestLock, SLT_Write);
@@ -480,6 +480,11 @@ void FBackgroundHttpManagerImpl::ActivatePendingRequests()
 
 FString FBackgroundHttpManagerImpl::GetTempFileLocationForURL(const FString& URL)
 {
-	const FString& TempLocation = GetFileHashHelper()->FindOrAddTempFilenameMappingForURL(URL);
-	return FBackgroundHttpFileHashHelper::GetFullPathOfTempFilename(TempLocation);
+	if (ensureAlwaysMsgf(IsInGameThread(), TEXT("Should only call GetTempFileLocationForURL from the GameThread as this is not thread-safe otherwise!")))
+	{
+		const FString& TempLocation = GetFileHashHelper()->FindOrAddTempFilenameMappingForURL(URL);
+		return FBackgroundHttpFileHashHelper::GetFullPathOfTempFilename(TempLocation);
+	}
+	
+	return TEXT("");
 }
diff --git a/Engine/Source/Runtime/Online/BackgroundHTTP/Private/BackgroundHttpRequestImpl.cpp b/Engine/Source/Runtime/Online/BackgroundHTTP/Private/BackgroundHttpRequestImpl.cpp
--- a/Engine/Source/Runtime/Online/BackgroundHTTP/Private/BackgroundHttpRequestImpl.cpp
+++ b/Engine/Source/Runtime/Online/BackgroundHTTP/Private/BackgroundHttpRequestImpl.cpp
@@ -21,7 +21,7 @@ FBackgroundHttpRequestImpl::FBackgroundHttpRequestImpl()
 
 bool FBackgroundHttpRequestImpl::ProcessRequest()
 {
-	UE_LOG(LogBackgroundHttpRequest, Display, TEXT("Processing Request - RequestID:%s"), *GetRequestID());
+	UE_LOG(LogBackgroundHttpRequest, Verbose, TEXT("Processing Request - RequestID:%s"), *GetRequestID());
 	FBackgroundHttpModule::Get().GetBackgroundHttpManager()->AddRequest(SharedThis(this));
 
 	return true;
@@ -43,7 +43,7 @@ void FBackgroundHttpRequestImpl::PauseRequest()
 void FBackgroundHttpRequestImpl::ResumeRequest()
 {
     //for now a resume is just wrapping a restart in the general case
-    UE_LOG(LogBackgroundHttpRequest, Display, TEXT("Pausing Request (through restart) - RequestID:%s"), *GetRequestID());
+    UE_LOG(LogBackgroundHttpRequest, Display, TEXT("Resuming Request (through restart) - RequestID:%s"), *GetRequestID());
     ProcessRequest();
 }
 
@@ -80,7 +80,12 @@ void FBackgroundHttpRequestImpl::NotifyNotificationObjectOfComplete(bool bWasSuc
 	if (DownloadCompleteNotificationObject.IsValid())
 	{
 		const int32 SharedRefCount = DownloadCompleteNotificationObject.GetSharedReferenceCount();
-		UE_LOG(LogBackgroundHttpRequest, Display, TEXT("Removing Reference to DownloadCompleteNotificationObject - bWasSuccess:%d | CurrentSharedRefCount:%d"), (int)(bWasSuccess), SharedRefCount);
+		UE_LOG(LogBackgroundHttpRequest, VeryVerbose, TEXT("Removing Reference to DownloadCompleteNotificationObject - bWasSuccess:%d | CurrentSharedRefCount:%d"), (int)(bWasSuccess), SharedRefCount);
+
+		if (SharedRefCount == 1)
+		{
+			UE_LOG(LogBackgroundHttpRequest, Display, TEXT("Removing Final Reference to DownloadCompleteNotificationObject!"));
+		}
 
 		DownloadCompleteNotificationObject->NotifyOfDownloadResult(bWasSuccess);
 		DownloadCompleteNotificationObject.Reset();
@@ -113,7 +118,7 @@ void FBackgroundHttpRequestImpl::CompleteWithExistingResponseData(FBackgroundHtt
 	const bool bHasValidResponse = Response.IsValid();
 	const FString ResponseTempLocation = Response.IsValid() ? Response->GetTempContentFilePath() : TEXT("None");
 
-	UE_LOG(LogBackgroundHttpRequest, Display, TEXT("Completing Download With Existing Response Data - RequestID:%s | bHasValidResponse:%d | ResponseTempDownloadLocation:%s"), *GetRequestID(), (int)(bHasValidResponse), *ResponseTempLocation);
+	UE_LOG(LogBackgroundHttpRequest, Verbose, TEXT("Completing Download With Existing Response Data - RequestID:%s | bHasValidResponse:%d | ResponseTempDownloadLocation:%s"), *GetRequestID(), (int)(bHasValidResponse), *ResponseTempLocation);
 
 	OnBackgroundDownloadComplete();
 }
