Force using the parameter to work around a code generation bug of Microsoft ASAN compiler.

#rnx

diff --git a/Engine/Source/Runtime/Online/BuildPatchServices/Private/Installer/DownloadService.cpp b/Engine/Source/Runtime/Online/BuildPatchServices/Private/Installer/DownloadService.cpp
--- a/Engine/Source/Runtime/Online/BuildPatchServices/Private/Installer/DownloadService.cpp
+++ b/Engine/Source/Runtime/Online/BuildPatchServices/Private/Installer/DownloadService.cpp
@@ -212,6 +212,11 @@ namespace BuildPatchServices
 			
 			HttpRequest->OnRequestProgress64().BindSPLambda(this, [this, RequestId](FHttpRequestPtr Request, uint64 BytesSent, uint64 BytesReceived)
 			{
+				#if USING_ADDRESS_SANITISER
+					// Force use the Request parameter to work around some MSVC-specific compiler issue in ASan.
+					UE_LOG(LogDownloadService, Verbose, TEXT("[FDownloadServiceImpl::RequestFile] %s %s received %d bytes"), *Request->GetVerb(), *Request->GetURL(), BytesReceived);
+				#endif
+
 				{
 					FScopeLock ScopeLock(&RequestDelegatesCS);
 					if (FDownloadDelegates* MaybeDelegates = RequestDelegates.Find(RequestId))
