Remove unused includes

diff --git a/Engine/Source/Runtime/Online/XMPP/Public/XmppConnection.h b/Engine/Source/Runtime/Online/XMPP/Public/XmppConnection.h
--- a/Engine/Source/Runtime/Online/XMPP/Public/XmppConnection.h
+++ b/Engine/Source/Runtime/Online/XMPP/Public/XmppConnection.h
@@ -3,11 +3,9 @@
 #pragma once
 
 #include "Containers/UnrealString.h"
-#include "CoreMinimal.h"
 #include "CoreTypes.h"
 #include "Delegates/Delegate.h"
 #include "Templates/SharedPointer.h"
-#include "Templates/UnrealTemplate.h"
 
 class IXmppChat;
 class IXmppMessages;
