Fix the warning in GenericPlatformBackgroundHttpRequest.cpp: Expression 'ByteDifference >= 0' is always true. Unsigned type value is always >= 0.

diff --git a/Engine/Source/Runtime/Online/BackgroundHTTP/Private/GenericPlatform/GenericPlatformBackgroundHttpRequest.cpp b/Engine/Source/Runtime/Online/BackgroundHTTP/Private/GenericPlatform/GenericPlatformBackgroundHttpRequest.cpp
--- a/Engine/Source/Runtime/Online/BackgroundHTTP/Private/GenericPlatform/GenericPlatformBackgroundHttpRequest.cpp
+++ b/Engine/Source/Runtime/Online/BackgroundHTTP/Private/GenericPlatform/GenericPlatformBackgroundHttpRequest.cpp
@@ -172,8 +172,8 @@ void FGenericPlatformBackgroundHttpRequest::FGenericPlatformBackgroundHttpWrappe
 {
 	if (ensureAlwaysMsgf(OriginalRequest.IsValid(), TEXT("Received UpdateHttpProgress callback with invalid OriginalRequest pointer! Can not update request progress!")))
 	{
-		const uint64 ByteDifference = (LastProgressUpdateBytes > 0) ? BytesReceived - LastProgressUpdateBytes : BytesReceived;
-		ensureAlwaysMsgf((ByteDifference >= 0), TEXT("Invalid Byte Difference in UpdateHttpProgress -- ByteDifference:%llu | LastProgressUpdateBytes:%llu | BytesSent:%llu | BytesReceived:%llu"), ByteDifference, LastProgressUpdateBytes, BytesSent, BytesReceived);
+		ensureAlwaysMsgf((BytesReceived >= LastProgressUpdateBytes), TEXT("Invalid Byte Difference in UpdateHttpProgress -- LastProgressUpdateBytes:%llu | BytesSent:%llu | BytesReceived:%llu"), LastProgressUpdateBytes, BytesSent, BytesReceived);
+		const uint64 ByteDifference = BytesReceived - LastProgressUpdateBytes;
 		LastProgressUpdateBytes = BytesReceived;
 
 		UE_LOG(LogBackgroundHttpRequest, VeryVerbose, TEXT("HttpRequest Progress Update- RequestID:%s | BytesSent: %llu | BytesReceived:%llu | BytesReceivedSinceLastUpdate:%llu"), *OriginalRequest.Pin()->GetRequestID(), BytesSent, BytesReceived, ByteDifference);
diff --git a/Engine/Source/Runtime/Online/BackgroundHTTP/Public/GenericPlatform/GenericPlatformBackgroundHttpRequest.h b/Engine/Source/Runtime/Online/BackgroundHTTP/Public/GenericPlatform/GenericPlatformBackgroundHttpRequest.h
--- a/Engine/Source/Runtime/Online/BackgroundHTTP/Public/GenericPlatform/GenericPlatformBackgroundHttpRequest.h
+++ b/Engine/Source/Runtime/Online/BackgroundHTTP/Public/GenericPlatform/GenericPlatformBackgroundHttpRequest.h
@@ -67,7 +67,7 @@ protected:
 		int32 MaxRetries;
 
 		//How many bytes we had last time we sent progress updates
-		int32 LastProgressUpdateBytes;
+		uint64 LastProgressUpdateBytes;
 	};
 
 	TUniquePtr<FGenericPlatformBackgroundHttpWrapper> RequestWrapper;
