Make sure it's safe to retry http request in http thread, by keep the shared ptr of retry manager in the request.  (Caught by the random failure of a test case in WebTests)

NOTE: This change requires existing holder of FHttpRetrySystem::FManager be converted to shared ptr.

diff --git a/Engine/Source/Programs/WebTests/Private/TestHttp.cpp b/Engine/Source/Programs/WebTests/Private/TestHttp.cpp
--- a/Engine/Source/Programs/WebTests/Private/TestHttp.cpp
+++ b/Engine/Source/Programs/WebTests/Private/TestHttp.cpp
@@ -30,6 +30,7 @@
 #define HTTP_TIME_DIFF_TOLERANCE 0.5f
 
 extern TAutoConsoleVariable<bool> CVarHttpInsecureProtocolEnabled;
+extern TAutoConsoleVariable<bool> CVarHttpRetrySystemNonGameThreadSupportEnabled;
 
 class FHttpModuleTestFixture
 {
@@ -38,11 +39,13 @@ public:
 		: WebServerIp(TEXT("127.0.0.1"))
 		, WebServerHttpPort(8000)
 		, bRunHeavyTests(false)
-		, bRetryEnabled(false)
+		, bRetryEnabled(true)
 		, OldVerbosity(LogHttp.GetVerbosity())
 	{
 		ParseSettingsFromCommandLine();
 
+		bRetryEnabled &= CVarHttpRetrySystemNonGameThreadSupportEnabled.GetValueOnAnyThread();
+
 		HttpModule = new FHttpModule();
 		IModuleInterface* Module = HttpModule;
 		Module->StartupModule();
@@ -123,6 +126,12 @@ public:
 	using FHttpRetrySystem::FManager::RequestList;
 	using FHttpRetrySystem::FManager::FHttpRetryRequestEntry;
 	using FHttpRetrySystem::FManager::RetryTimeoutRelativeSecondsDefault;
+
+	bool IsEmpty()
+	{
+		FScopeLock ScopeLock(&RequestListLock);
+		return RequestList.IsEmpty();
+	}
 };
 
 class FWaitUntilCompleteHttpFixture : public FHttpModuleTestFixture
@@ -160,7 +169,7 @@ public:
 
 	void WaitUntilAllHttpRequestsComplete()
 	{
-		while (OngoingRequests != 0)
+		while (OngoingRequests != 0 || (bRetryEnabled && !HttpRetryManager->IsEmpty()) )
 		{
 			HttpModule->GetHttpManager().Tick(TickFrequency);
 			FPlatformProcess::Sleep(TickFrequency);
@@ -179,7 +188,7 @@ public:
 	std::atomic<int32> OngoingRequests = 0;
 	float TickFrequency = 1.0f / 60; /*60 FPS*/;
 
-	uint32 RetryLimitCount = 2;
+	uint32 RetryLimitCount = 0;
 	TSharedPtr<FMockRetryManager> HttpRetryManager;
 };
 
diff --git a/Engine/Source/Runtime/Online/HTTP/Private/HttpRetrySystem.cpp b/Engine/Source/Runtime/Online/HTTP/Private/HttpRetrySystem.cpp
--- a/Engine/Source/Runtime/Online/HTTP/Private/HttpRetrySystem.cpp
+++ b/Engine/Source/Runtime/Online/HTTP/Private/HttpRetrySystem.cpp
@@ -70,7 +70,7 @@ namespace FHttpRetrySystem
 }
 
 FHttpRetrySystem::FRequest::FRequest(
-	FManager& InManager,
+	TSharedRef<FManager> InManager,
 	const TSharedRef<IHttpRequest, ESPMode::ThreadSafe>& HttpRequest, 
 	const FHttpRetrySystem::FRetryLimitCountSetting& InRetryLimitCountOverride,
 	const FHttpRetrySystem::FRetryTimeoutRelativeSecondsSetting& InRetryTimeoutRelativeSecondsOverride,
@@ -107,8 +107,8 @@ FHttpRetrySystem::FRequest::FRequest(
 }
 
 bool FHttpRetrySystem::FRequest::ProcessRequest()
-{ 
-	TSharedRef<FRequest, ESPMode::ThreadSafe> RetryRequest = StaticCastSharedRef<FRequest>(AsShared());
+{
+	TSharedRef<FRequest> RetryRequest = StaticCastSharedRef<FRequest>(AsShared());
 
 	OriginalUrl = HttpRequest->GetURL();
 	if (RetryDomains.IsValid())
@@ -120,7 +120,16 @@ bool FHttpRetrySystem::FRequest::ProcessRequest()
 	HttpRequest->OnProcessRequestComplete().BindThreadSafeSP(RetryRequest, &FHttpRetrySystem::FRequest::HttpOnProcessRequestComplete);
 	HttpRequest->OnHeaderReceived().BindThreadSafeSP(RetryRequest, &FHttpRetrySystem::FRequest::HttpOnHeaderReceived);
 
-	return RetryManager.ProcessRequest(RetryRequest);
+	TSharedPtr<FManager> RetryManagerPtr = RetryManager.Pin();
+
+	if (ensure(RetryManagerPtr))
+	{
+		return RetryManagerPtr->ProcessRequest(RetryRequest);
+	}
+	else
+	{
+		return false;
+	}
 }
 
 void FHttpRetrySystem::FRequest::SetUrlFromRetryDomains()
@@ -149,7 +158,14 @@ void FHttpRetrySystem::FRequest::CancelRequest()
 { 
 	TSharedRef<FRequest, ESPMode::ThreadSafe> RetryRequest = StaticCastSharedRef<FRequest>(AsShared());
 
-	RetryManager.CancelRequest(RetryRequest);
+	if (TSharedPtr<FManager> RetryManagerPtr = RetryManager.Pin())
+	{
+		RetryManagerPtr->CancelRequest(RetryRequest);
+	}
+	else
+	{
+		HttpRequest->CancelRequest();
+	}
 }
 
 void FHttpRetrySystem::FRequest::HttpOnRequestProgress(FHttpRequestPtr InHttpRequest, uint64 BytesSent, uint64 BytesRcv)
@@ -167,55 +183,65 @@ void FHttpRetrySystem::FRequest::HttpOnProcessRequestComplete(FHttpRequestPtr In
 		return;
 	}
 
-	TSharedRef<FRequest> SelfPtr = StaticCastSharedRef<FRequest>(AsShared()); // In case no ref after removing from RetryManager
+	TSharedPtr<FManager> RetryManagerPtr = RetryManager.Pin();
+	if (!RetryManagerPtr)
+	{
+		return;
+	}
 
-	FScopeLock ScopeLock(&RetryManager.RequestListLock);
+	TSharedRef<FRequest> SelfPtr = StaticCastSharedRef<FRequest>(AsShared()); // In case no ref after removing from RetryManager
 
-	uint32 EntryIndex = RetryManager.RequestList.IndexOfByPredicate([this](const FManager::FHttpRetryRequestEntry& Entry) { return Entry.Request == AsShared(); });
-	check(EntryIndex != INDEX_NONE);
-	FManager::FHttpRetryRequestEntry* HttpRetryRequestEntry = &RetryManager.RequestList[EntryIndex];
 
-	if (RetryStatus == FHttpRetrySystem::FRequest::EStatus::Cancelled)
 	{
-		// Do nothing here
-	}
-	if (GetStatus() == EHttpRequestStatus::Failed_ConnectionError || GetStatus() == EHttpRequestStatus::Failed)
-	{
-		if (GetStatus() == EHttpRequestStatus::Failed_ConnectionError && RetryDomains.IsValid())
+		FScopeLock ScopeLock(&RetryManagerPtr->RequestListLock);
+
+		uint32 EntryIndex = RetryManagerPtr->RequestList.IndexOfByPredicate([this](const FManager::FHttpRetryRequestEntry& Entry) { return Entry.Request == AsShared(); });
+		check(EntryIndex != INDEX_NONE);
+		FManager::FHttpRetryRequestEntry* HttpRetryRequestEntry = &RetryManagerPtr->RequestList[EntryIndex];
+
+
+		if (RetryStatus == FHttpRetrySystem::FRequest::EStatus::Cancelled)
 		{
-			MoveToNextRetryDomain();
+			// Do nothing here
 		}
-
-		if (RetryManager.ShouldRetry(*HttpRetryRequestEntry) && RetryManager.CanRetry(*HttpRetryRequestEntry))
+		if (GetStatus() == EHttpRequestStatus::Failed_ConnectionError || GetStatus() == EHttpRequestStatus::Failed)
 		{
-			const double NowAbsoluteSeconds = FPlatformTime::Seconds();
-			float LockoutPeriod = RetryManager.GetLockoutPeriodSeconds(*HttpRetryRequestEntry);
-
-			RetryStatus = FHttpRetrySystem::FRequest::EStatus::ProcessingLockout;
+			if (GetStatus() == EHttpRequestStatus::Failed_ConnectionError && RetryDomains.IsValid())
+			{
+				MoveToNextRetryDomain();
+			}
 
+			if (RetryManagerPtr->ShouldRetry(*HttpRetryRequestEntry) && RetryManagerPtr->CanRetry(*HttpRetryRequestEntry))
 			{
-				QUICK_SCOPE_CYCLE_COUNTER(STAT_FHttpRetrySystem_FManager_Update_OnRequestWillRetry);
-				OnRequestWillRetry().ExecuteIfBound(HttpRetryRequestEntry->Request, GetResponse(), LockoutPeriod);
+				const double NowAbsoluteSeconds = FPlatformTime::Seconds();
+				float LockoutPeriod = RetryManagerPtr->GetLockoutPeriodSeconds(*HttpRetryRequestEntry);
+
+				RetryStatus = FHttpRetrySystem::FRequest::EStatus::ProcessingLockout;
+
+				{
+					QUICK_SCOPE_CYCLE_COUNTER(STAT_FHttpRetrySystem_FManager_Update_OnRequestWillRetry);
+					OnRequestWillRetry().ExecuteIfBound(HttpRetryRequestEntry->Request, GetResponse(), LockoutPeriod);
+				}
+
+				RetryManagerPtr->RetryHttpRequestWithDelay(HttpRetryRequestEntry->Request, LockoutPeriod);
+				return;
 			}
 
-			RetryManager.RetryHttpRequestWithDelay(HttpRetryRequestEntry->Request, LockoutPeriod);
-			return;
+			RetryStatus = FHttpRetrySystem::FRequest::EStatus::FailedRetry;
+		}
+		else
+		{
+			RetryStatus = FHttpRetrySystem::FRequest::EStatus::Succeeded;
 		}
 
-		RetryStatus = FHttpRetrySystem::FRequest::EStatus::FailedRetry;
-	}
-	else
-	{
-		RetryStatus = FHttpRetrySystem::FRequest::EStatus::Succeeded;
-	}
+		if (HttpRetryRequestEntry->CurrentRetryCount > 0)
+		{
+			FHttpRetrySystem::FManager::FHttpLogVerbosityTracker::Get().DecrementRetriedRequests();
+		}
 
-	if (HttpRetryRequestEntry->CurrentRetryCount > 0)
-	{
-		FHttpRetrySystem::FManager::FHttpLogVerbosityTracker::Get().DecrementRetriedRequests();
+		RetryManagerPtr->RequestList.RemoveAtSwap(EntryIndex);
 	}
 
-	RetryManager.RequestList.RemoveAtSwap(EntryIndex);
-
 	OnProcessRequestComplete().ExecuteIfBound(SelfPtr, HttpResponse, bSucceeded);
 }
 
@@ -234,7 +260,9 @@ FHttpRetrySystem::FManager::FManager(const FRetryLimitCountSetting& InRetryLimit
     : RandomFailureRate(FRandomFailureRateSetting())
     , RetryLimitCountDefault(InRetryLimitCountDefault)
 	, RetryTimeoutRelativeSecondsDefault(InRetryTimeoutRelativeSecondsDefault)
-{}
+{
+	check(FHttpModule::Get().GetHttpManager().GetThread());
+}
 
 FHttpRetrySystem::FManager::~FManager()
 {
@@ -258,7 +286,7 @@ TSharedRef<FHttpRetrySystem::FRequest, ESPMode::ThreadSafe> FHttpRetrySystem::FM
 	const FRetryDomainsPtr& InRetryDomains)
 {
 	return MakeShareable(new FRequest(
-		*this,
+		AsShared(),
 		FHttpModule::Get().CreateRequest(),
 		InRetryLimitCountOverride,
 		InRetryTimeoutRelativeSecondsOverride,
@@ -384,14 +412,21 @@ void FHttpRetrySystem::FManager::RetryHttpRequest(FHttpRetryRequestEntry& Reques
 
 void FHttpRetrySystem::FManager::RetryHttpRequestWithDelay(const TSharedRef<FRequest>& Request, float InDelay)
 {
-	FHttpModule::Get().GetHttpManager().AddHttpThreadTask([this, Request]() {
-		FScopeLock ScopeLock(&RequestListLock);
-		// Check if it's still there in case it has been cancelled during the delay period
-		uint32 EntryIndex = RequestList.IndexOfByPredicate([Request](const FManager::FHttpRetryRequestEntry& Entry) { return Entry.Request == Request; });
-		if (EntryIndex != INDEX_NONE)
+	TWeakPtr<FRequest> RequestWeakPtr(Request);
+	FHttpModule::Get().GetHttpManager().AddHttpThreadTask([RequestWeakPtr]() {
+		if (TSharedPtr<FRequest> RequestPtr = RequestWeakPtr.Pin())
 		{
-			FManager::FHttpRetryRequestEntry* HttpRetryRequestEntry = &RequestList[EntryIndex];
-			RetryHttpRequest(*HttpRetryRequestEntry);
+			if (TSharedPtr<FManager> RetryManagerPtr = RequestPtr->RetryManager.Pin())
+			{
+				FScopeLock ScopeLock(&RetryManagerPtr->RequestListLock);
+				// Check if it's still there in case it has been cancelled during the delay period
+				uint32 EntryIndex = RetryManagerPtr->RequestList.IndexOfByPredicate([RequestPtr](const FManager::FHttpRetryRequestEntry& Entry) { return Entry.Request == RequestPtr; });
+				if (EntryIndex != INDEX_NONE)
+				{
+					FManager::FHttpRetryRequestEntry* HttpRetryRequestEntry = &RetryManagerPtr->RequestList[EntryIndex];
+					RetryManagerPtr->RetryHttpRequest(*HttpRetryRequestEntry);
+				}
+			}
 		}
 	}, InDelay);
 }
diff --git a/Engine/Source/Runtime/Online/HTTP/Public/HttpRetrySystem.h b/Engine/Source/Runtime/Online/HTTP/Public/HttpRetrySystem.h
--- a/Engine/Source/Runtime/Online/HTTP/Public/HttpRetrySystem.h
+++ b/Engine/Source/Runtime/Online/HTTP/Public/HttpRetrySystem.h
@@ -104,8 +104,8 @@ namespace FHttpRetrySystem
 		friend class FManager;
 
 		HTTP_API FRequest(
-			class FManager& InManager,
-			const TSharedRef<IHttpRequest, ESPMode::ThreadSafe>& HttpRequest,
+			TSharedRef<FManager> InManager,
+			const TSharedRef<IHttpRequest>& HttpRequest,
 			const FRetryLimitCountSetting& InRetryLimitCountOverride = FRetryLimitCountSetting(),
 			const FRetryTimeoutRelativeSecondsSetting& InRetryTimeoutRelativeSecondsOverride = FRetryTimeoutRelativeSecondsSetting(),
 			const FRetryResponseCodes& InRetryResponseCodes = FRetryResponseCodes(),
@@ -134,14 +134,14 @@ namespace FHttpRetrySystem
 		/** The original URL before replacing anything from RetryDomains */
 		FString								 OriginalUrl;
 
-		FManager& RetryManager;
+		TWeakPtr<FManager>					 RetryManager;
 	};
 }
 
 namespace FHttpRetrySystem
 {
 
-class FManager
+class FManager : public TSharedFromThis<FManager>
 {
 public:
 	// FManager
