Remove legacy apple http impl.

diff --git a/Engine/Source/Runtime/Online/HTTP/Private/Apple/AppleEventLoopHttpThread.cpp b/Engine/Source/Runtime/Online/HTTP/Private/Apple/AppleEventLoopHttpThread.cpp
--- a/Engine/Source/Runtime/Online/HTTP/Private/Apple/AppleEventLoopHttpThread.cpp
+++ b/Engine/Source/Runtime/Online/HTTP/Private/Apple/AppleEventLoopHttpThread.cpp
@@ -2,12 +2,12 @@
 
 #include "Apple/AppleEventLoopHttpThread.h"
 
-#include "Apple/AppleHTTPNSUrlSession.h" 
+#include "Apple/AppleHttp.h" 
 
 bool FAppleEventLoopHttpThread::StartThreadedRequest(IHttpThreadedRequest* Request)
 {
 	FHttpResponsePtr Response = Request->GetResponse();
-	auto AppleResponse = StaticCastSharedPtr<FAppleHttpNSUrlSessionResponse>(Response);
+	auto AppleResponse = StaticCastSharedPtr<FAppleHttpResponse>(Response);
 	AppleResponse->SetNewAppleHttpEventDelegate(FNewAppleHttpEventDelegate::CreateLambda([IOAccess = EventLoop->GetIOAccess()]() mutable
 	{
 		IOAccess.Notify();
diff --git a/Engine/Source/Runtime/Online/HTTP/Private/Apple/AppleHTTPNSUrlConnection.cpp b/Engine/Source/Runtime/Online/HTTP/Private/Apple/AppleHTTPNSUrlConnection.cpp
deleted file mode 100644
diff --git a/Engine/Source/Runtime/Online/HTTP/Private/Apple/AppleHTTPNSUrlConnection.h b/Engine/Source/Runtime/Online/HTTP/Private/Apple/AppleHTTPNSUrlConnection.h
deleted file mode 100644
diff --git a/Engine/Source/Runtime/Online/HTTP/Private/Apple/AppleHTTPNSUrlSession.cpp b/Engine/Source/Runtime/Online/HTTP/Private/Apple/AppleHTTPNSUrlSession.cpp
deleted file mode 100644
diff --git a/Engine/Source/Runtime/Online/HTTP/Private/Apple/AppleHTTPNSUrlSession.h b/Engine/Source/Runtime/Online/HTTP/Private/Apple/AppleHTTPNSUrlSession.h
deleted file mode 100644
diff --git a/Engine/Source/Runtime/Online/HTTP/Private/Apple/AppleHttp.cpp b/Engine/Source/Runtime/Online/HTTP/Private/Apple/AppleHttp.cpp
new file mode 100644
--- /dev/null
+++ b/Engine/Source/Runtime/Online/HTTP/Private/Apple/AppleHttp.cpp
@@ -0,0 +1,970 @@
+// Copyright Epic Games, Inc. All Rights Reserved.
+
+
+#include "AppleHttp.h"
+#include "HttpManager.h"
+#include "Misc/EngineVersion.h"
+#include "Misc/App.h"
+#include "Misc/Base64.h"
+#include "HAL/PlatformTime.h"
+#include "Http.h"
+#include "HttpModule.h"
+
+/**
+ * State of a request's response
+ */
+enum class EAppleHttpRequestResponseState: uint8
+{
+	NotReady,
+	Success,
+	Error,
+	ConnectionError
+};
+
+/**
+ * Class to hold data from delegate implementation notifications.
+ */
+
+@interface FAppleHttpResponseDelegate : NSObject<NSURLSessionDataDelegate>
+{
+	/** Holds the payload as we receive it. */
+	@public TArray<uint8> Payload;
+	
+	/** This stream is shared between the request and the delegate. This field is accessed through a thread out of out control 
+	 * from the delegates methods and through mehthods in FAppleHttpRequest. Once we cancel or have an error this can 
+	 * is nullified as soon as possible to avoid receiving more data after completion delegates were triggered */
+	TSharedPtr<FArchive> ResponseBodyReceiveStream;
+	
+	/** critical section to properly clear stream */
+	FCriticalSection ResponseStreamLock;
+
+	/** flag meant to reduce locking on ResponseStreamLock*/
+	@public BOOL bInitializedWithValidStream;
+
+	/** Delegate invoked after processing URLSession:dataTask:didReceiveData or URLSession:task:didCompleteWithError:*/
+	@public FNewAppleHttpEventDelegate NewAppleHttpEventDelegate;
+}
+
+/** A handle for the response */
+@property(retain) NSHTTPURLResponse* Response;
+/** The total number of bytes written out during the request/response */
+@property uint64 BytesWritten;
+/** The total number of bytes received out during the request/response */
+@property uint64 BytesReceived;
+/** Response state */
+@property EAppleHttpRequestResponseState ResponseState;
+
+/** NSURLSessionDataDelegate delegate methods. Those are called from a thread controlled by the NSURLSession */
+
+/** Sent periodically to notify the delegate of upload progress. */
+- (void)URLSession:(NSURLSession *)session task:(NSURLSessionTask *)task didSendBodyData:(int64_t)bytesSent totalBytesSent:(int64_t)totalBytesSent totalBytesExpectedToSend:(int64_t)totalBytesExpectedToSend;
+/** The task has received a response and no further messages will be received until the completion block is called. */
+- (void)URLSession:(NSURLSession *)session dataTask:(NSURLSessionDataTask *)dataTask didReceiveResponse:(NSURLResponse *)response completionHandler:(void (^)(NSURLSessionResponseDisposition disposition))completionHandler;
+/** Sent when data is available for the delegate to consume. Data may be discontiguous */
+- (void)URLSession:(NSURLSession *)session dataTask:(NSURLSessionDataTask *)dataTask didReceiveData:(NSData *)data;
+/** Sent as the last message related to a specific task.  A nil Error implies that no error occurred and this task is complete. */
+- (void)URLSession:(NSURLSession *)session task:(NSURLSessionTask *)task didCompleteWithError:(nullable NSError *)error;
+/** Asks the delegate if it needs to store responses in the cache. */
+- (void)URLSession:(NSURLSession *)session dataTask:(NSURLSessionDataTask *)dataTask willCacheResponse:(NSCachedURLResponse *)proposedResponse completionHandler:(void (^)(NSCachedURLResponse *cachedResponse))completionHandler;
+@end
+
+@implementation FAppleHttpResponseDelegate
+@synthesize Response;
+@synthesize ResponseState;
+@synthesize BytesWritten;
+@synthesize BytesReceived;
+
+- (FAppleHttpResponseDelegate*)initWithResponseStream:(TSharedPtr<FArchive>)ResponseStream
+{
+	self = [super init];
+	
+	BytesWritten = 0;
+	BytesReceived = 0;
+	ResponseState = EAppleHttpRequestResponseState::NotReady;
+	ResponseBodyReceiveStream = ResponseStream;
+	bInitializedWithValidStream = (ResponseStream != nullptr);
+	
+	return self;
+}
+
+- (void)ClearResponseStream
+{
+	if (bInitializedWithValidStream)
+	{
+	    FScopeLock Lock(&ResponseStreamLock);
+		ResponseBodyReceiveStream = nullptr;
+	}
+}
+
+- (void)dealloc
+{
+	[Response release];
+	[super dealloc];
+}
+
+- (void)URLSession:(NSURLSession *)session task:(NSURLSessionTask *)task didSendBodyData:(int64_t)bytesSent totalBytesSent:(int64_t)totalBytesSent totalBytesExpectedToSend:(int64_t)totalBytesExpectedToSend
+{
+	UE_LOG(LogHttp, Verbose, TEXT("URLSession:task:didSendBodyData:totalBytesSent:totalBytesExpectedToSend: totalBytesSent = %lld, totalBytesSent = %lld: %p"), totalBytesSent, totalBytesExpectedToSend, self);
+	self.BytesWritten = totalBytesSent;
+}
+
+- (void)URLSession:(NSURLSession *)session dataTask:(NSURLSessionDataTask *)dataTask didReceiveResponse:(NSURLResponse *)response completionHandler:(void (^)(NSURLSessionResponseDisposition disposition))completionHandler
+{
+	UE_LOG(LogHttp, Verbose, TEXT("URLSession:dataTask:didReceiveResponse:completionHandler"));
+	
+	self.Response = (NSHTTPURLResponse*)response;
+	uint64 ExpectedResponseLength = response.expectedContentLength;
+	if(!bInitializedWithValidStream && ExpectedResponseLength != NSURLResponseUnknownLength)
+	{
+		Payload.Empty(ExpectedResponseLength);
+	}
+	UE_LOG(LogHttp, Verbose, TEXT("URLSession:dataTask:didReceiveResponse:completionHandler: expectedContentLength = %lld. Length = %llu: %p"), ExpectedResponseLength, Payload.Max(), self);
+	completionHandler(NSURLSessionResponseAllow);
+}
+
+- (void)URLSession:(NSURLSession *)session dataTask:(NSURLSessionDataTask *)dataTask didReceiveData:(NSData *)data
+{
+	__block int64 NewBytesReceived = 0;
+	if (bInitializedWithValidStream)
+	{
+		if (FScopeLock Lock(&ResponseStreamLock); ResponseBodyReceiveStream)
+		{
+			__block bool bHadError = false;
+			[data enumerateByteRangesUsingBlock:^(const void *bytes, NSRange byteRange, BOOL *stop) {
+				NewBytesReceived += byteRange.length;
+				ResponseBodyReceiveStream->Serialize(const_cast<void*>(bytes), byteRange.length);
+				bHadError = ResponseBodyReceiveStream->GetError();
+				*stop = bHadError? YES : NO;
+			}];
+			
+			if (bHadError)
+			{
+				[dataTask cancel];
+				ResponseBodyReceiveStream = nullptr;
+			}
+		}
+	}
+	else
+	{
+		[data enumerateByteRangesUsingBlock:^(const void *bytes, NSRange byteRange, BOOL *stop) {
+			NewBytesReceived += byteRange.length;
+			Payload.Append((const uint8*)bytes, byteRange.length);
+		}];
+	}
+	// Keep BytesReceived as a separated value to avoid concurrent accesses to Payload
+	self.BytesReceived += NewBytesReceived;
+	UE_LOG(LogHttp, Verbose, TEXT("URLSession:dataTask:didReceiveData with %llu bytes. After Append, Payload Length = %llu: %p"), NewBytesReceived, self.BytesReceived, self);
+	
+	NewAppleHttpEventDelegate.ExecuteIfBound();
+}
+
+- (void)URLSession:(NSURLSession *)session task:(NSURLSessionTask *)task didCompleteWithError:(nullable NSError *)error
+{
+	if (error == nil)
+	{
+		UE_LOG(LogHttp, Verbose, TEXT("URLSession:task:didCompleteWithError. Http request succeeded: %p"), self);
+		self.ResponseState = EAppleHttpRequestResponseState::Success;
+	}
+	else
+	{
+		UE_LOG(LogHttp, Warning, TEXT("URLSession:task:didCompleteWithError. Http request failed - %s %s: %p"),
+			   *FString([error localizedDescription]),
+			   *FString([[error userInfo] objectForKey:NSURLErrorFailingURLStringErrorKey]),
+			   self);
+		
+		// Determine if the specific error was failing to connect to the host.
+		switch ([error code])
+		{
+			case NSURLErrorTimedOut:
+			case NSURLErrorCannotFindHost:
+			case NSURLErrorCannotConnectToHost:
+			case NSURLErrorDNSLookupFailed:
+				self.ResponseState = EAppleHttpRequestResponseState::ConnectionError;
+				break;
+			default:
+				self.ResponseState = EAppleHttpRequestResponseState::Error;
+		}
+		// Log more details if verbose logging is enabled and this is an SSL error
+		if (UE_LOG_ACTIVE(LogHttp, Verbose))
+		{
+			SecTrustRef PeerTrustInfo = reinterpret_cast<SecTrustRef>([[error userInfo] objectForKey:NSURLErrorFailingURLPeerTrustErrorKey]);
+			if (PeerTrustInfo != nullptr)
+			{
+				SecTrustResultType TrustResult = kSecTrustResultInvalid;
+				SecTrustGetTrustResult(PeerTrustInfo, &TrustResult);
+				
+				FString TrustResultString;
+				switch (TrustResult)
+				{
+#define MAP_TO_RESULTSTRING(Constant) case Constant: TrustResultString = TEXT(#Constant); break;
+						MAP_TO_RESULTSTRING(kSecTrustResultInvalid)
+						MAP_TO_RESULTSTRING(kSecTrustResultProceed)
+						MAP_TO_RESULTSTRING(kSecTrustResultDeny)
+						MAP_TO_RESULTSTRING(kSecTrustResultUnspecified)
+						MAP_TO_RESULTSTRING(kSecTrustResultRecoverableTrustFailure)
+						MAP_TO_RESULTSTRING(kSecTrustResultFatalTrustFailure)
+						MAP_TO_RESULTSTRING(kSecTrustResultOtherError)
+#undef MAP_TO_RESULTSTRING
+					default:
+						TrustResultString = TEXT("unknown");
+						break;
+				}
+				UE_LOG(LogHttp, Verbose, TEXT("URLSession:task:didCompleteWithError. SSL trust result: %s (%d)"), *TrustResultString, TrustResult);
+			}
+		}
+	}
+	NewAppleHttpEventDelegate.ExecuteIfBound();
+}
+
+- (void)URLSession:(NSURLSession *)session dataTask:(NSURLSessionDataTask *)dataTask willCacheResponse:(NSCachedURLResponse *)proposedResponse completionHandler:(void (^)(NSCachedURLResponse *cachedResponse))completionHandler
+{
+	// All FAppleHttpRequest use NSURLRequestReloadIgnoringLocalCacheData
+	// NSURLRequestReloadIgnoringLocalCacheData disables loading of data from cache, but responses can still be stored in cache
+	// Passing nil to this handler disables caching the responses
+	completionHandler(nil);
+}
+@end
+
+
+/**
+ * NSInputStream subclass to send streamed FArchive contents
+ */
+@interface FNSInputStreamFromArchive : NSInputStream<NSStreamDelegate>
+{
+	TSharedPtr<FArchive> Archive;
+	int64 AlreadySentContent;
+	NSStreamStatus StreamStatus;
+	id<NSStreamDelegate> Delegate;
+}
+@end
+
+@implementation FNSInputStreamFromArchive
+
++(FNSInputStreamFromArchive*)initWithArchive:(TSharedRef<FArchive>) Archive
+{
+	FNSInputStreamFromArchive* Ret = [[[FNSInputStreamFromArchive alloc] init] autorelease];
+	Ret->Archive = MoveTemp(Archive);
+	return Ret;
+}
+
+- (id)init
+{
+	self = [super init];
+	if (self)
+	{
+		StreamStatus = NSStreamStatusNotOpen;
+
+		// Docs say it is good practice that streams are it's own delegates by default
+		Delegate = self;
+	}
+	
+	return self;
+}
+
+/** NSStream implementation */
+- (void)dealloc
+{
+	[super dealloc];
+}
+
+- (void)open
+{
+	AlreadySentContent = 0;
+	StreamStatus = NSStreamStatusOpen;
+}
+
+- (void)close
+{
+	StreamStatus = NSStreamStatusClosed;
+}
+
+- (NSStreamStatus)streamStatus
+{
+	return StreamStatus;
+}
+
+- (NSError *)streamError
+{
+	return nil;
+}
+
+- (id<NSStreamDelegate>)delegate
+{
+	return Delegate;
+}
+
+- (void)setDelegate:(id<NSStreamDelegate>)InDelegate
+{
+	if (InDelegate == nil)
+	{
+		InDelegate = self;
+	}
+	else
+	{
+		Delegate = InDelegate;
+	}
+}
+
+- (id)propertyForKey:(NSString *)key
+{
+	return nil;
+}
+
+- (BOOL)setProperty:(id)property forKey:(NSString *)key
+{
+	return NO;
+}
+
+- (void)scheduleInRunLoop:(NSRunLoop *)aRunLoop forMode:(NSString *)mode
+{
+	// There is no need to scheduled anything. Data is always available until end is reached
+}
+
+- (void)removeFromRunLoop:(NSRunLoop *)aRunLoop forMode:(NSString *)mode
+{
+	// There is no need to be descheduled since we didn't schedule
+}
+
+/** NSStreamDelegate implementation */
+- (void)stream:(NSStream *)stream handleEvent:(NSStreamEvent)eventCode
+{
+	// Won't update local data
+}
+
+/** NSInputStream implementation. Those methods are invoked in a worker thread out of our control */
+
+// Reads up to 'len' bytes into 'buffer'. Returns the actual number of bytes read.
+- (NSInteger)read:(uint8_t *)buffer maxLength:(NSUInteger)len
+{
+	const int64 ContentLength = Archive->TotalSize();
+	check(AlreadySentContent <= ContentLength);
+	const int64 SizeToSend = ContentLength - AlreadySentContent;
+	const int64 SizeToSendThisTime = FMath::Min(SizeToSend, static_cast<int64>(len));
+	if (SizeToSendThisTime != 0)
+	{
+		if (Archive->Tell() != AlreadySentContent)
+		{
+			Archive->Seek(AlreadySentContent);
+		}
+		Archive->Serialize((uint8*)buffer, SizeToSendThisTime);
+		AlreadySentContent += SizeToSendThisTime;
+	}
+	return SizeToSendThisTime;
+}
+
+// return NO because getting the internal buffer is not appropriate for this subclass
+- (BOOL)getBuffer:(uint8_t **)buffer length:(NSUInteger *)len
+{
+	return NO;
+}
+
+// returns YES to always force reads
+- (BOOL)hasBytesAvailable
+{
+	return YES;
+}
+@end
+/****************************************************************************
+ * FAppleHttpRequest implementation
+ ***************************************************************************/
+
+FAppleHttpRequest::FAppleHttpRequest(NSURLSession* InSession)
+:   Session([InSession retain])
+,   Task(nil)
+,	bIsPayloadFile(false)
+,	ContentBytesLength(0)
+,	ElapsedTime(0.0f)
+,	LastReportedBytesWritten(0)
+,	LastReportedBytesRead(0)
+{
+	UE_LOG(LogHttp, Verbose, TEXT("FAppleHttpRequest::FAppleHttpRequest()"));
+	Request = [[NSMutableURLRequest alloc] init];
+	Request.timeoutInterval = FHttpModule::Get().GetHttpTimeout();
+
+	// Disable cache to mimic WinInet behavior
+	Request.cachePolicy = NSURLRequestReloadIgnoringLocalCacheData;
+
+	// Add default headers
+	const TMap<FString, FString>& DefaultHeaders = FHttpModule::Get().GetDefaultHeaders();
+	for (TMap<FString, FString>::TConstIterator It(DefaultHeaders); It; ++It)
+	{
+		SetHeader(It.Key(), It.Value());
+	}
+}
+
+FAppleHttpRequest::~FAppleHttpRequest()
+{
+	UE_LOG(LogHttp, Verbose, TEXT("FAppleHttpRequest::~FAppleHttpRequest()"));
+	CleanupRequest();
+	[Request release];
+    [Session release];
+}
+
+FString FAppleHttpRequest::GetURL() const
+{
+	SCOPED_AUTORELEASE_POOL;
+	NSURL* URL = Request.URL;
+	if (URL != nullptr)
+	{
+		FString ConvertedURL(URL.absoluteString);
+		UE_LOG(LogHttp, Verbose, TEXT("FAppleHttpRequest::GetURL() - %s"), *ConvertedURL);
+		return ConvertedURL;
+	}
+	else
+	{
+		UE_LOG(LogHttp, Verbose, TEXT("FAppleHttpRequest::GetURL() - NULL"));
+		return FString();
+	}
+}
+
+void FAppleHttpRequest::SetURL(const FString& URL)
+{
+	SCOPED_AUTORELEASE_POOL;
+	UE_LOG(LogHttp, Verbose, TEXT("FAppleHttpRequest::SetURL() - %s"), *URL);
+	Request.URL = [NSURL URLWithString: URL.GetNSString()];
+}
+
+FString FAppleHttpRequest::GetHeader(const FString& HeaderName) const
+{
+	SCOPED_AUTORELEASE_POOL;
+	FString Header([Request valueForHTTPHeaderField:HeaderName.GetNSString()]);
+	UE_LOG(LogHttp, Verbose, TEXT("FAppleHttpRequest::GetHeader() - %s"), *Header);
+	return Header;
+}
+
+
+void FAppleHttpRequest::SetHeader(const FString& HeaderName, const FString& HeaderValue)
+{
+	SCOPED_AUTORELEASE_POOL;
+	UE_LOG(LogHttp, Verbose, TEXT("FAppleHttpRequest::SetHeader() - %s / %s"), *HeaderName, *HeaderValue );
+	[Request setValue: HeaderValue.GetNSString() forHTTPHeaderField: HeaderName.GetNSString()];
+}
+
+void FAppleHttpRequest::AppendToHeader(const FString& HeaderName, const FString& AdditionalHeaderValue)
+{
+    if (!HeaderName.IsEmpty() && !AdditionalHeaderValue.IsEmpty())
+    {
+        NSDictionary* Headers = [Request allHTTPHeaderFields];
+        NSString* PreviousHeaderValuePtr = [Headers objectForKey: HeaderName.GetNSString()];
+        FString PreviousValue(PreviousHeaderValuePtr);
+		FString NewValue;
+		if (!PreviousValue.IsEmpty())
+		{
+			NewValue = PreviousValue + TEXT(", ");
+		}
+		NewValue += AdditionalHeaderValue;
+
+        SetHeader(HeaderName, NewValue);
+	}
+}
+
+TArray<FString> FAppleHttpRequest::GetAllHeaders() const
+{
+	SCOPED_AUTORELEASE_POOL;
+	UE_LOG(LogHttp, Verbose, TEXT("FAppleHttpRequest::GetAllHeaders()"));
+	NSDictionary* Headers = Request.allHTTPHeaderFields;
+	TArray<FString> Result;
+	Result.Reserve(Headers.count);
+	for (NSString* Key in Headers.allKeys)
+	{
+		FString ConvertedValue(Headers[Key]);
+		FString ConvertedKey(Key);
+		UE_LOG(LogHttp, Verbose, TEXT("Header= %s, Key= %s"), *ConvertedValue, *ConvertedKey);
+
+		Result.Add( FString::Printf( TEXT("%s: %s"), *ConvertedKey, *ConvertedValue ) );
+	}
+	return Result;
+}
+
+const TArray<uint8>& FAppleHttpRequest::GetContent() const
+{
+	UE_LOG(LogHttp, Verbose, TEXT("FAppleHttpRequest::GetContent()"));
+	StorageForGetContent.Empty();
+	if (bIsPayloadFile)
+	{
+		UE_LOG(LogHttp, Warning, TEXT("FAppleHttpRequest::GetContent() called on a request that is set up for streaming a file. Return value is an empty buffer"));
+	}
+	else
+	{
+		SCOPED_AUTORELEASE_POOL;
+		NSData* Body = Request.HTTPBody; // accessing HTTPBody will call retain autorelease on the value, increasing its retain count
+		StorageForGetContent.Append((const uint8*)Body.bytes, Body.length);
+	}
+	return StorageForGetContent;
+}
+
+void FAppleHttpRequest::SetContent(const TArray<uint8>& ContentPayload)
+{
+	if (CompletionStatus == EHttpRequestStatus::Processing)
+	{
+		UE_LOG(LogHttp, Warning, TEXT("FAppleHttpRequest::SetContent() - attempted to set content on a request that is inflight"));
+		return;
+	}
+
+	UE_LOG(LogHttp, Verbose, TEXT("FAppleHttpRequest::SetContent()"));
+	Request.HTTPBodyStream = nil;
+	Request.HTTPBody = [NSData dataWithBytes:ContentPayload.GetData() length:ContentPayload.Num()];
+	ContentBytesLength = ContentPayload.Num();
+	bIsPayloadFile = false;
+}
+
+void FAppleHttpRequest::SetContent(TArray<uint8>&& ContentPayload)
+{
+	if (CompletionStatus == EHttpRequestStatus::Processing)
+	{
+		UE_LOG(LogHttp, Warning, TEXT("FAppleHttpRequest::SetContent() - attempted to set content on a request that is inflight"));
+		return;
+	}
+
+	UE_LOG(LogHttp, Verbose, TEXT("FAppleHttpRequest::SetContent()"));
+
+	Request.HTTPBodyStream = nil;
+	// We cannot use NSData dataWithBytesNoCopy:length:freeWhenDone: and keep the data in this instance because we don't have control
+	// over the lifetime of the request copy that NSURLSessionTask keeps
+	Request.HTTPBody = [NSData dataWithBytes:ContentPayload.GetData() length:ContentPayload.Num()];
+	ContentBytesLength = ContentPayload.Num();
+	bIsPayloadFile = false;
+
+	// Clear argument content since client code probably expects that
+	ContentPayload.Empty();
+}
+
+FString FAppleHttpRequest::GetContentType() const
+{
+	FString ContentType = GetHeader(TEXT("Content-Type"));
+	UE_LOG(LogHttp, Verbose, TEXT("FAppleHttpRequest::GetContentType() - %s"), *ContentType);
+	return ContentType;
+}
+
+uint64 FAppleHttpRequest::GetContentLength() const
+{
+	UE_LOG(LogHttp, Verbose, TEXT("FAppleHttpRequest::GetContentLength() - %i"), ContentBytesLength);
+	return ContentBytesLength;
+}
+
+void FAppleHttpRequest::SetContentAsString(const FString& ContentString)
+{
+	if (CompletionStatus == EHttpRequestStatus::Processing)
+	{
+		UE_LOG(LogHttp, Warning, TEXT("FAppleHttpRequest::SetContentAsString() - attempted to set content on a request that is inflight"));
+		return;
+	}
+
+	UE_LOG(LogHttp, Verbose, TEXT("FAppleHttpRequest::SetContentAsString() - %s"), *ContentString);
+	FTCHARToUTF8 Converter(*ContentString);
+
+	Request.HTTPBodyStream = nil;
+	// The extra length computation here is unfortunate, but it's technically not safe to assume the length is the same.
+	Request.HTTPBody = [NSData dataWithBytes:(ANSICHAR*)Converter.Get() length:Converter.Length()];
+	ContentBytesLength = Converter.Length();
+	bIsPayloadFile = false;
+}
+
+bool FAppleHttpRequest::SetContentAsStreamedFile(const FString& Filename)
+{
+	SCOPED_AUTORELEASE_POOL;
+    UE_LOG(LogHttp, Verbose, TEXT("FAppleHttpRequest::SetContentAsStreamedFile() - %s"), *Filename);
+
+	if (CompletionStatus == EHttpRequestStatus::Processing)
+	{
+		UE_LOG(LogHttp, Warning, TEXT("FAppleHttpRequest::SetContentAsStreamedFile() - attempted to set content on a request that is inflight"));
+		return false;
+	}
+
+	NSString* PlatformFilename = Filename.GetNSString();
+
+	Request.HTTPBody = nil;
+
+	struct stat FileAttrs = { 0 };
+	if (stat(PlatformFilename.fileSystemRepresentation, &FileAttrs) == 0)
+	{
+		UE_LOG(LogHttp, VeryVerbose, TEXT("FAppleHttpRequest::SetContentAsStreamedFile succeeded in getting the file size - %lld"), FileAttrs.st_size);
+		// Under the hood, the Foundation framework unsets HTTPBody, and takes over as the stream delegate.
+		// The stream itself should be unopened when passed to setHTTPBodyStream.
+		Request.HTTPBodyStream = [NSInputStream inputStreamWithFileAtPath: PlatformFilename];
+		ContentBytesLength = FileAttrs.st_size;
+		bIsPayloadFile = true;
+	}
+	else
+	{
+		UE_LOG(LogHttp, Warning, TEXT("FAppleHttpRequest::SetContentAsStreamedFile failed to get file size"));
+		Request.HTTPBodyStream = nil;
+		ContentBytesLength = 0;
+		bIsPayloadFile = false;
+	}
+
+	return bIsPayloadFile;
+}
+
+bool FAppleHttpRequest::SetContentFromStream(TSharedRef<FArchive, ESPMode::ThreadSafe> Stream)
+{
+	SCOPED_AUTORELEASE_POOL;
+	if (CompletionStatus == EHttpRequestStatus::Processing)
+	{
+		UE_LOG(LogHttp, Warning, TEXT("FAppleHttpRequest::SetContentFromStream() - attempted to set content on a request that is inflight"));
+		return false;
+	}
+
+	Request.HTTPBody = nil;
+
+	Request.HTTPBodyStream = [FNSInputStreamFromArchive initWithArchive: Stream];
+	ContentBytesLength = Stream->TotalSize();
+	bIsPayloadFile = true;
+
+	return true;
+}
+
+bool FAppleHttpRequest::SetResponseBodyReceiveStream(TSharedRef<FArchive> Stream)
+{
+	ResponseBodyReceiveStream = Stream;
+	return true;
+}
+
+FString FAppleHttpRequest::GetVerb() const
+{
+	FString ConvertedVerb(Request.HTTPMethod);
+	UE_LOG(LogHttp, Verbose, TEXT("FAppleHttpRequest::GetVerb() - %s"), *ConvertedVerb);
+	return ConvertedVerb;
+}
+
+void FAppleHttpRequest::SetVerb(const FString& Verb)
+{
+	SCOPED_AUTORELEASE_POOL;
+	UE_LOG(LogHttp, Verbose, TEXT("FAppleHttpRequest::SetVerb() - %s"), *Verb);
+	Request.HTTPMethod = Verb.GetNSString();
+}
+
+void FAppleHttpRequest::SetTimeout(float InTimeoutSecs)
+{
+	Request.timeoutInterval = InTimeoutSecs;
+}
+
+void FAppleHttpRequest::ClearTimeout()
+{
+	Request.timeoutInterval = FHttpModule::Get().GetHttpTimeout();
+}
+
+TOptional<float> FAppleHttpRequest::GetTimeout() const
+{
+	return TOptional<float>(Request.timeoutInterval);
+}
+
+bool FAppleHttpRequest::ProcessRequest()
+{
+	SCOPED_AUTORELEASE_POOL;
+	UE_LOG(LogHttp, Verbose, TEXT("FAppleHttpRequest::ProcessRequest()"));
+
+	if (!PreCheck() || !StartRequest())
+	{
+		FinishRequestNotInHttpManager();
+		return false;
+	}
+
+	return true;
+}
+
+bool FAppleHttpRequest::StartRequest()
+{
+	SCOPED_AUTORELEASE_POOL;
+	UE_LOG(LogHttp, Verbose, TEXT("FAppleHttpRequest::StartRequest()"));
+	bool bStarted = false;
+
+	// set the content-length and user-agent (it is possible that the OS ignores this value)
+	if(GetContentLength() > 0)
+	{
+		[Request setValue:[NSString stringWithFormat:@"%llu", GetContentLength()] forHTTPHeaderField:@"Content-Length"];
+	}
+
+	const FString UserAgent = GetHeader("User-Agent");
+	if(UserAgent.IsEmpty())
+	{
+		NSString* Tag = FPlatformHttp::GetDefaultUserAgent().GetNSString();
+		[Request setValue:Tag forHTTPHeaderField:@"User-Agent"];
+	}
+
+	CleanupRequest();
+
+	LastReportedBytesWritten = 0;
+	LastReportedBytesRead = 0;
+	ElapsedTime = 0.0f;
+	Response = nullptr;
+
+	Task = [Session dataTaskWithRequest: Request];
+	
+	if (Task != nil)
+	{
+		bStarted = true;
+
+		CompletionStatus = EHttpRequestStatus::Processing;
+
+		Response = MakeShared<FAppleHttpResponse>(*this);
+
+		// Both Task and Response keep a strong reference to the delegate
+		Task.delegate = Response->ResponseDelegate;
+
+		//Setup delegates before starting the request
+		FHttpModule::Get().GetHttpManager().AddThreadedRequest(SharedThis(this));
+
+		[[Task retain] resume];
+		UE_LOG(LogHttp, Verbose, TEXT("[NSURLSessionTask resume]"));
+	}
+	else
+	{
+		UE_LOG(LogHttp, Warning, TEXT("ProcessRequest failed. Could not initialize Internet connection."));
+		CompletionStatus = EHttpRequestStatus::Failed_ConnectionError;
+	}
+
+	return bStarted;
+}
+
+void FAppleHttpRequest::FinishRequest()
+{
+	UE_LOG(LogHttp, Verbose, TEXT("FAppleHttpRequest::FinishRequest()"));
+
+	// Clean up session/request handles that may have been created
+	CleanupRequest();
+
+	bool bSuccess = false;
+	if (Response.IsValid() && Response->IsReady() && !Response->HadError())
+	{
+		bSuccess = true;
+		UE_LOG(LogHttp, Verbose, TEXT("Request succeeded"));
+		CompletionStatus = EHttpRequestStatus::Succeeded;
+
+		// TODO: Try to broadcast OnHeaderReceived when we receive headers instead of here at the end
+		BroadcastResponseHeadersReceived();
+	}
+	else
+	{
+		UE_LOG(LogHttp, Verbose, TEXT("Request failed"));
+		FString URL([[Request URL] absoluteString]);
+		CompletionStatus = EHttpRequestStatus::Failed;
+		if (Response.IsValid() && Response->HadConnectionError())
+		{
+			CompletionStatus = EHttpRequestStatus::Failed_ConnectionError;
+			Response = nullptr;
+		}
+	}
+
+	OnProcessRequestComplete().ExecuteIfBound(SharedThis(this), Response, bSuccess);
+}
+
+void FAppleHttpRequest::CleanupRequest()
+{
+	UE_LOG(LogHttp, Verbose, TEXT("FAppleHttpRequest::CleanupRequest()"));
+	
+	if (Response != nullptr)
+	{
+		Response->CleanSharedObjects();
+	}
+
+	if(Task != nil)
+	{
+		if (CompletionStatus == EHttpRequestStatus::Processing)
+		{
+			[Task cancel];
+		}
+		[Task release];
+		Task = nil;
+	}
+}
+
+void FAppleHttpRequest::CancelRequest()
+{
+	UE_LOG(LogHttp, Verbose, TEXT("FAppleHttpRequest::CancelRequest()"));
+
+	if (Task != nil)
+	{
+		[Task cancel];
+	}
+}
+
+const FHttpResponsePtr FAppleHttpRequest::GetResponse() const
+{
+	return Response;
+}
+
+void FAppleHttpRequest::Tick(float DeltaSeconds)
+{
+	if (DelegateThreadPolicy == EHttpRequestDelegateThreadPolicy::CompleteOnGameThread)
+	{
+		CheckProgressDelegate();
+	}
+}
+
+void FAppleHttpRequest::CheckProgressDelegate()
+{
+	if (Response.IsValid() && (CompletionStatus == EHttpRequestStatus::Processing || Response->HadError()))
+	{
+		const uint64 BytesWritten = Response->GetNumBytesWritten();
+		const uint64 BytesRead = Response->GetNumBytesReceived();
+		if (BytesWritten != LastReportedBytesWritten || BytesRead != LastReportedBytesRead)
+		{
+PRAGMA_DISABLE_DEPRECATION_WARNINGS
+			OnRequestProgress().ExecuteIfBound(SharedThis(this), BytesWritten, BytesRead);
+PRAGMA_ENABLE_DEPRECATION_WARNINGS
+			OnRequestProgress64().ExecuteIfBound(SharedThis(this), BytesWritten, BytesRead);
+			LastReportedBytesWritten = BytesWritten;
+			LastReportedBytesRead = BytesRead;
+		}
+	}
+}
+
+float FAppleHttpRequest::GetElapsedTime() const
+{
+	return ElapsedTime;
+}
+
+bool FAppleHttpRequest::StartThreadedRequest()
+{
+	return true;
+}
+
+bool FAppleHttpRequest::IsThreadedRequestComplete()
+{
+	TSharedPtr<FAppleHttpResponse> AliveResponse = Response;
+	return (AliveResponse.IsValid() && AliveResponse->IsReady());
+}
+
+void FAppleHttpRequest::TickThreadedRequest(float DeltaSeconds)
+{
+	ElapsedTime += DeltaSeconds;
+
+	if (DelegateThreadPolicy == EHttpRequestDelegateThreadPolicy::CompleteOnHttpThread)
+	{
+		CheckProgressDelegate();
+	}
+}
+
+/****************************************************************************
+ * FAppleHttpResponse implementation
+ **************************************************************************/
+
+FAppleHttpResponse::FAppleHttpResponse(const FAppleHttpRequest& InRequest)
+	: FHttpResponseCommon(InRequest.GetURL())
+{
+	UE_LOG(LogHttp, Verbose, TEXT("FAppleHttpResponse::FAppleHttpResponse()"));
+	ResponseDelegate = [[FAppleHttpResponseDelegate alloc] initWithResponseStream: InRequest.ResponseBodyReceiveStream];
+}
+
+FAppleHttpResponse::~FAppleHttpResponse()
+{
+	UE_LOG(LogHttp, Verbose, TEXT("FAppleHttpResponse::~FAppleHttpResponse()"));
+	
+	[ResponseDelegate release];
+	ResponseDelegate = nil;
+}
+
+void FAppleHttpResponse::SetNewAppleHttpEventDelegate(FNewAppleHttpEventDelegate&& Delegate)
+{	
+	ResponseDelegate->NewAppleHttpEventDelegate = MoveTemp(Delegate);
+}
+
+void FAppleHttpResponse::CleanSharedObjects()
+{
+	[ResponseDelegate ClearResponseStream];
+}
+
+FString FAppleHttpResponse::GetHeader(const FString& HeaderName) const
+{
+	SCOPED_AUTORELEASE_POOL;
+	UE_LOG(LogHttp, Verbose, TEXT("FAppleHttpResponse::GetHeader()"));
+	NSString* ConvertedHeaderName = HeaderName.GetNSString();
+	return FString([ResponseDelegate.Response.allHeaderFields objectForKey:ConvertedHeaderName]);
+}
+
+TArray<FString> FAppleHttpResponse::GetAllHeaders() const
+{
+	UE_LOG(LogHttp, Verbose, TEXT("FAppleHttpResponse::GetAllHeaders()"));
+
+	NSDictionary* Headers = ResponseDelegate.Response.allHeaderFields;
+	TArray<FString> Result;
+	Result.Reserve([Headers count]);
+	for (NSString* Key in [Headers allKeys])
+	{
+		FString ConvertedValue([Headers objectForKey:Key]);
+		FString ConvertedKey(Key);
+		Result.Add( FString::Printf( TEXT("%s: %s"), *ConvertedKey, *ConvertedValue ) );
+	}
+	return Result;
+}
+
+FString FAppleHttpResponse::GetContentType() const
+{
+	UE_LOG(LogHttp, Verbose, TEXT("FAppleHttpResponse::GetContentType()"));
+
+	return GetHeader( TEXT( "Content-Type" ) );
+}
+
+uint64 FAppleHttpResponse::GetContentLength() const
+{
+	UE_LOG(LogHttp, Verbose, TEXT("FAppleHttpResponse::GetContentLength()"));
+	
+	return ResponseDelegate.Response.expectedContentLength;
+}
+
+const TArray<uint8>& FAppleHttpResponse::GetContent() const
+{
+	if( !IsReady() )
+	{
+		UE_LOG(LogHttp, Warning, TEXT("Payload is incomplete. Response still processing. %s"), *GetURL());
+	}
+	else
+	{
+		UE_LOG(LogHttp, Verbose, TEXT("FAppleHttpResponse::GetContent() - Num: %i"), ResponseDelegate->Payload.Num());
+	}
+	return ResponseDelegate->Payload;
+}
+
+FString FAppleHttpResponse::GetContentAsString() const
+{
+	UE_LOG(LogHttp, Verbose, TEXT("FAppleHttpResponse::GetContentAsString()"));
+
+	// Fill in our data.
+	const TArray<uint8>& Payload = GetContent();
+
+	TArray<uint8> ZeroTerminatedPayload;
+	ZeroTerminatedPayload.AddZeroed( Payload.Num() + 1 );
+	FMemory::Memcpy( ZeroTerminatedPayload.GetData(), Payload.GetData(), Payload.Num() );
+
+	return UTF8_TO_TCHAR( ZeroTerminatedPayload.GetData() );
+}
+
+int32 FAppleHttpResponse::GetResponseCode() const
+{
+	UE_LOG(LogHttp, Verbose, TEXT("FAppleHttpResponse::GetResponseCode()"));
+
+	return ResponseDelegate.Response.statusCode;
+}
+
+bool FAppleHttpResponse::IsReady() const
+{
+	return (ResponseDelegate.ResponseState != EAppleHttpRequestResponseState::NotReady);
+}
+
+bool FAppleHttpResponse::HadError() const
+{
+	switch(ResponseDelegate.ResponseState)
+	{
+		case EAppleHttpRequestResponseState::Error:
+		case EAppleHttpRequestResponseState::ConnectionError:
+			return true;
+		default:
+			return false;
+	}
+}
+
+bool FAppleHttpResponse::HadConnectionError() const
+{
+	return (ResponseDelegate.ResponseState == EAppleHttpRequestResponseState::ConnectionError);
+}
+
+const uint64 FAppleHttpResponse::GetNumBytesReceived() const
+{
+	return ResponseDelegate.BytesReceived;
+}
+
+const uint64 FAppleHttpResponse::GetNumBytesWritten() const
+{
+	return ResponseDelegate.BytesWritten;
+}
diff --git a/Engine/Source/Runtime/Online/HTTP/Private/Apple/AppleHttp.h b/Engine/Source/Runtime/Online/HTTP/Private/Apple/AppleHttp.h
new file mode 100644
--- /dev/null
+++ b/Engine/Source/Runtime/Online/HTTP/Private/Apple/AppleHttp.h
@@ -0,0 +1,212 @@
+// Copyright Epic Games, Inc. All Rights Reserved.
+
+#pragma once
+
+#include "CoreMinimal.h"
+#include "GenericPlatform/HttpResponseCommon.h"
+#include "IHttpThreadedRequest.h"
+#include "PlatformHttp.h"
+
+
+/**
+ * Delegate invoked when in progress Task completes. It is invoked in an out of our control thread
+ *
+ */
+DECLARE_DELEGATE(FNewAppleHttpEventDelegate);
+
+/**
+ * Apple implementation of an Http request
+ */
+class FAppleHttpRequest : public IHttpThreadedRequest
+{
+public:
+	// implementation friends
+	friend class FAppleHttpResponse;
+
+
+	//~ Begin IHttpBase Interface
+	virtual FString GetURL() const override;
+	virtual FString GetHeader(const FString& HeaderName) const override;
+	virtual TArray<FString> GetAllHeaders() const override;	
+	virtual FString GetContentType() const override;
+	virtual uint64 GetContentLength() const override;
+	virtual const TArray<uint8>& GetContent() const override;
+	//~ End IHttpBase Interface
+
+	//~ Begin IHttpRequest Interface
+	virtual FString GetVerb() const override;
+	virtual void SetVerb(const FString& Verb) override;
+	virtual void SetURL(const FString& URL) override;
+	virtual void SetContent(const TArray<uint8>& ContentPayload) override;
+	virtual void SetContent(TArray<uint8>&& ContentPayload) override;
+	virtual void SetContentAsString(const FString& ContentString) override;
+    virtual bool SetContentAsStreamedFile(const FString& Filename) override;
+	virtual bool SetContentFromStream(TSharedRef<FArchive, ESPMode::ThreadSafe> Stream) override;
+	virtual bool SetResponseBodyReceiveStream(TSharedRef<FArchive> Stream) override;
+	virtual void SetHeader(const FString& HeaderName, const FString& HeaderValue) override;
+	virtual void AppendToHeader(const FString& HeaderName, const FString& AdditionalHeaderValue) override;
+	virtual void SetTimeout(float InTimeoutSecs) override;
+	virtual void ClearTimeout() override;
+	virtual TOptional<float> GetTimeout() const override;
+	virtual bool ProcessRequest() override;
+	virtual void CancelRequest() override;
+	virtual const FHttpResponsePtr GetResponse() const override;
+	virtual void Tick(float DeltaSeconds) override;
+	virtual float GetElapsedTime() const override;
+	//~ End IHttpRequest Interface
+
+	//~ Begin IHttpRequestThreaded Interface
+	virtual bool StartThreadedRequest() override;
+	virtual void FinishRequest() override;
+	virtual bool IsThreadedRequestComplete() override;
+	virtual void TickThreadedRequest(float DeltaSeconds) override;
+	//~ End IHttpRequestThreaded Interface
+
+	/**
+	 * Constructor
+	 *
+	 * @param InSession - NSURLSession session used to create NSURLSessionTask to retrieve the response
+	 */
+	explicit FAppleHttpRequest(NSURLSession* InSession);
+
+	/**
+	 * Destructor. Clean up any connection/request handles
+	 */
+	virtual ~FAppleHttpRequest();
+
+
+private:
+	/**
+	 * Trigger the request progress delegate if progress has changed
+	 */
+	void CheckProgressDelegate();
+
+	/**
+	 * Create the session connection and initiate the web request
+	 *
+	 * @return true if the request was started
+	 */
+	bool StartRequest();
+
+	/**
+	 * Close session/request handles and unregister callbacks
+	 */
+	void CleanupRequest();
+
+private:
+	/** This is the NSMutableURLRequest, all our Apple functionality will deal with this. */
+	NSMutableURLRequest* Request;
+
+    /** This is the session our request belongs to */
+    NSURLSession* Session;
+	
+	/** This is the Task associated to the sessionin charge of our request */
+	NSURLSessionTask* Task;
+    
+	/** Flag whether the request payload source is a file */
+	bool bIsPayloadFile;
+
+	/** The request payload length in bytes. This must be tracked separately for a file stream */
+	uint64 ContentBytesLength;
+
+	/** The response object which we will use to pair with this request */
+	TSharedPtr<class FAppleHttpResponse,ESPMode::ThreadSafe> Response;
+
+	/** Array used to retrieve back content set on the ObjC request when calling GetContent*/
+	mutable TArray<uint8> StorageForGetContent;
+
+	/** The stream to receive response body */
+	TSharedPtr<FArchive> ResponseBodyReceiveStream;
+
+	/** Time taken to complete/cancel the request. */
+	float ElapsedTime;
+
+	/** Last reported bytes written */
+	int32 LastReportedBytesWritten;
+
+	/** Last reported bytesread */
+	int32 LastReportedBytesRead;
+};
+
+@class FAppleHttpResponseDelegate;
+
+/**
+ * Apple implementation of an Http response
+ */
+class FAppleHttpResponse : public FHttpResponseCommon
+{
+private:
+	// Delegate implementation. Keeps the response state and data
+	FAppleHttpResponseDelegate* ResponseDelegate;
+
+public:
+	// implementation friends
+	friend class FAppleHttpRequest;
+
+
+	//~ Begin IHttpBase Interface
+	virtual FString GetHeader(const FString& HeaderName) const override;
+	virtual TArray<FString> GetAllHeaders() const override;	
+	virtual FString GetContentType() const override;
+	virtual uint64 GetContentLength() const override;
+	virtual const TArray<uint8>& GetContent() const override;
+	//~ End IHttpBase Interface
+
+	//~ Begin IHttpResponse Interface
+	virtual int32 GetResponseCode() const override;
+	virtual FString GetContentAsString() const override;
+	//~ End IHttpResponse Interface
+
+	/**
+	 * Check whether headers are available.
+	 */
+	bool AreHeadersAvailable() const;
+
+	/**
+	 * Check whether a response is ready or not.
+	 */
+	bool IsReady() const;
+	
+	/**
+	 * Check whether a response had an error.
+	 */
+	bool HadError() const;
+
+	/**
+	 * Check whether a response had a connection error.
+	 */
+	bool HadConnectionError() const;
+
+	/**
+	 * Get the number of bytes received so far
+	 */
+	const uint64 GetNumBytesReceived() const;
+
+	/**
+	* Get the number of bytes sent so far
+	*/
+	const uint64 GetNumBytesWritten() const;
+
+	/**
+	 * Cleans internal shared objects between request and response
+	 */
+	void CleanSharedObjects();
+
+	/**
+	 * Sets delegate invoked when  URLSession:dataTask:didReceiveData or URLSession:task:didCompleteWithError: are triggered
+	 * Should be set right before task is started 
+	*/
+	void SetNewAppleHttpEventDelegate(FNewAppleHttpEventDelegate&& Delegate);
+
+	/**
+	 * Constructor
+	 *
+	 * @param InRequest - original request that created this response
+	 */
+	FAppleHttpResponse(const FAppleHttpRequest& InRequest);
+
+	/**
+	 * Destructor
+	 */
+	virtual ~FAppleHttpResponse();
+};
diff --git a/Engine/Source/Runtime/Online/HTTP/Private/Apple/ApplePlatformHttp.cpp b/Engine/Source/Runtime/Online/HTTP/Private/Apple/ApplePlatformHttp.cpp
--- a/Engine/Source/Runtime/Online/HTTP/Private/Apple/ApplePlatformHttp.cpp
+++ b/Engine/Source/Runtime/Online/HTTP/Private/Apple/ApplePlatformHttp.cpp
@@ -3,8 +3,7 @@
 #include "Apple/ApplePlatformHttp.h"
 #include "Http.h"
 #include "AppleHttpManager.h"
-#include "AppleHTTPNSUrlConnection.h"
-#include "AppleHTTPNSUrlSession.h"
+#include "AppleHttp.h"
 #include "Apple/CFRef.h"
 #include "Misc/Base64.h"
 #include "Misc/CommandLine.h"
@@ -223,25 +222,12 @@ void FApplePlatformHttp::Init()
 	FSslModule::Get();
 #endif
 
-	bool bUseNSUrlConnection = FParse::Param(FCommandLine::Get(), TEXT("UseNSUrlConnection"));
-	if (bUseNSUrlConnection)
-	{
-		UE_LOG(LogHttp, Warning, TEXT("UseNSUrlConnection command line argument is deprecated. It will be removed in UE 5.4"));
-	}
-	bUseNSUrlSession = !bUseNSUrlConnection;
-
-	if (bUseNSUrlSession)
-	{
-		InitWithNSUrlSession();
-	}
+	InitWithNSUrlSession();
 }
 
 void FApplePlatformHttp::Shutdown()
 {
-	if (bUseNSUrlSession)
-	{
-		ShutdownWithNSUrlSession();
-	}
+	ShutdownWithNSUrlSession();
 }
 
 void FApplePlatformHttp::InitWithNSUrlSession()
@@ -276,30 +262,15 @@ void FApplePlatformHttp::ShutdownWithNSUrlSession()
 
 FHttpManager* FApplePlatformHttp::CreatePlatformHttpManager()
 {
-	if(bUseNSUrlSession)
-	{
-		return new FAppleHttpManager();
-	}
-	else
-	{
-		// Event based http manager does not support FAppleHTTPNSURLConnection 
-		return nullptr;
-	}
+	return new FAppleHttpManager();
 }
 
 IHttpRequest* FApplePlatformHttp::ConstructRequest()
 {
-	if(bUseNSUrlSession)
-	{
-		return new FAppleHttpNSUrlSessionRequest(Session);
-	}
-	else
-	{
-		return new FAppleHttpNSUrlConnectionRequest();
-	}
+	return new FAppleHttpRequest(Session);
 }
 
 bool FApplePlatformHttp::UsesThreadedHttp()
 {
-	return bUseNSUrlSession;
+	return true;
 }
diff --git a/Engine/Source/Runtime/Online/HTTP/Public/Apple/ApplePlatformHttp.h b/Engine/Source/Runtime/Online/HTTP/Public/Apple/ApplePlatformHttp.h
--- a/Engine/Source/Runtime/Online/HTTP/Public/Apple/ApplePlatformHttp.h
+++ b/Engine/Source/Runtime/Online/HTTP/Public/Apple/ApplePlatformHttp.h
@@ -42,9 +42,6 @@ public:
 	static bool UsesThreadedHttp();
 
 private:
-	/** Flag to allow fall back to use NSUrlConnection instead of NSUrlSession. Assigned from commandline */
-	static inline bool bUseNSUrlSession = false;
-
     /** Session used to create Apple based requests */
     static inline NSURLSession* Session = nil;
 
