Sync http request status into http response status, so user code can write different handling cases for different status, even if user keep the response after request get destroyed.

#tests Tried with game client, also passed auto test project WebTests.

diff --git a/Engine/Source/Runtime/Online/BuildPatchServices/Private/Tests/Mock/HttpResponse.mock.h b/Engine/Source/Runtime/Online/BuildPatchServices/Private/Tests/Mock/HttpResponse.mock.h
--- a/Engine/Source/Runtime/Online/BuildPatchServices/Private/Tests/Mock/HttpResponse.mock.h
+++ b/Engine/Source/Runtime/Online/BuildPatchServices/Private/Tests/Mock/HttpResponse.mock.h
@@ -70,6 +70,13 @@ namespace BuildPatchServices
 			static TArray<uint8> None;
 			return None;
 		}
+
+		HTTP_API virtual EHttpRequestStatus::Type GetStatus() const override
+		{
+			MOCK_FUNC_NOT_IMPLEMENTED("FMockHttpRequest::GetStatus");
+			return EHttpRequestStatus::Failed;
+		}
+
 	};
 }
 
diff --git a/Engine/Source/Runtime/Online/HTTP/Private/Apple/AppleHttp.cpp b/Engine/Source/Runtime/Online/HTTP/Private/Apple/AppleHttp.cpp
--- a/Engine/Source/Runtime/Online/HTTP/Private/Apple/AppleHttp.cpp
+++ b/Engine/Source/Runtime/Online/HTTP/Private/Apple/AppleHttp.cpp
@@ -695,7 +695,7 @@ bool FAppleHttpRequest::StartRequest()
 	{
 		bStarted = true;
 
-		CompletionStatus = EHttpRequestStatus::Processing;
+		SetStatus(EHttpRequestStatus::Processing);
 
 		Response = MakeShared<FAppleHttpResponse>(*this);
 
@@ -711,7 +711,7 @@ bool FAppleHttpRequest::StartRequest()
 	else
 	{
 		UE_LOG(LogHttp, Warning, TEXT("ProcessRequest failed. Could not initialize Internet connection."));
-		CompletionStatus = EHttpRequestStatus::Failed_ConnectionError;
+		SetStatus(EHttpRequestStatus::Failed_ConnectionError);
 	}
 
 	return bStarted;
@@ -729,7 +729,7 @@ void FAppleHttpRequest::FinishRequest()
 	{
 		bSuccess = true;
 		UE_LOG(LogHttp, Verbose, TEXT("Request succeeded"));
-		CompletionStatus = EHttpRequestStatus::Succeeded;
+		SetStatus(EHttpRequestStatus::Succeeded);
 
 		// TODO: Try to broadcast OnHeaderReceived when we receive headers instead of here at the end
 		BroadcastResponseHeadersReceived();
@@ -738,10 +738,10 @@ void FAppleHttpRequest::FinishRequest()
 	{
 		UE_LOG(LogHttp, Verbose, TEXT("Request failed"));
 		FString URL([[Request URL] absoluteString]);
-		CompletionStatus = EHttpRequestStatus::Failed;
+		SetStatus(EHttpRequestStatus::Failed);
 		if (Response.IsValid() && Response->HadConnectionError())
 		{
-			CompletionStatus = EHttpRequestStatus::Failed_ConnectionError;
+			SetStatus(EHttpRequestStatus::Failed_ConnectionError);
 			Response = nullptr;
 		}
 	}
@@ -841,7 +841,7 @@ void FAppleHttpRequest::TickThreadedRequest(float DeltaSeconds)
  **************************************************************************/
 
 FAppleHttpResponse::FAppleHttpResponse(const FAppleHttpRequest& InRequest)
-	: FHttpResponseCommon(InRequest.GetURL())
+	: FHttpResponseCommon(InRequest)
 {
 	UE_LOG(LogHttp, Verbose, TEXT("FAppleHttpResponse::FAppleHttpResponse()"));
 	ResponseDelegate = [[FAppleHttpResponseDelegate alloc] initWithResponseStream: InRequest.ResponseBodyReceiveStream];
diff --git a/Engine/Source/Runtime/Online/HTTP/Private/Curl/CurlHttp.cpp b/Engine/Source/Runtime/Online/HTTP/Private/Curl/CurlHttp.cpp
--- a/Engine/Source/Runtime/Online/HTTP/Private/Curl/CurlHttp.cpp
+++ b/Engine/Source/Runtime/Online/HTTP/Private/Curl/CurlHttp.cpp
@@ -998,7 +998,7 @@ bool FCurlHttpRequest::ProcessRequest()
 
 	QUICK_SCOPE_CYCLE_COUNTER(STAT_CurlHttpAddThreadedRequest);
 	// Mark as in-flight to prevent overlapped requests using the same object
-	CompletionStatus = EHttpRequestStatus::Processing;
+	SetStatus(EHttpRequestStatus::Processing);
 	// Add to global list while being processed so that the ref counted request does not get deleted
 	FHttpModule::Get().GetHttpManager().AddThreadedRequest(SharedThis(this));
 
@@ -1252,7 +1252,8 @@ void FCurlHttpRequest::FinishRequest()
 		}
 
 		// Mark last request attempt as completed successfully
-		CompletionStatus = EHttpRequestStatus::Succeeded;
+		SetStatus(EHttpRequestStatus::Succeeded);
+
 		// Call delegate with valid request/response objects
 		OnProcessRequestComplete().ExecuteIfBound(SharedThis(this),Response,true);
 	}
@@ -1286,7 +1287,7 @@ void FCurlHttpRequest::FinishRequest()
 		// Mark last request attempt as completed but failed
 		if (bCanceled)
 		{
-			CompletionStatus = EHttpRequestStatus::Failed;
+			SetStatus(EHttpRequestStatus::Failed);
 		}
 		else if (bCurlRequestCompleted)
 		{
@@ -1297,21 +1298,21 @@ void FCurlHttpRequest::FinishRequest()
 			case CURLE_COULDNT_RESOLVE_HOST:
 			case CURLE_SSL_CONNECT_ERROR:
 				// report these as connection errors (safe to retry)
-				CompletionStatus = EHttpRequestStatus::Failed_ConnectionError;
+				SetStatus(EHttpRequestStatus::Failed_ConnectionError);
 				break;
 			default:
-				CompletionStatus = EHttpRequestStatus::Failed;
+				SetStatus(EHttpRequestStatus::Failed);
 			}
 		}
 		else
 		{
 			if (bAnyHttpActivity)
 			{
-				CompletionStatus = EHttpRequestStatus::Failed;
+				SetStatus(EHttpRequestStatus::Failed);
 			}
 			else
 			{
-				CompletionStatus = EHttpRequestStatus::Failed_ConnectionError;
+				SetStatus(EHttpRequestStatus::Failed_ConnectionError);
 			}
 		}
 		// Call delegate with failure
@@ -1331,7 +1332,7 @@ float FCurlHttpRequest::GetElapsedTime() const
 // FCurlHttpRequest
 
 FCurlHttpResponse::FCurlHttpResponse(const FCurlHttpRequest& InRequest)
-	: FHttpResponseCommon(InRequest.GetURL())
+	: FHttpResponseCommon(InRequest)
 	, TotalBytesRead(0)
 	, HttpCode(EHttpResponseCodes::Unknown)
 	, ContentLength(0)
diff --git a/Engine/Source/Runtime/Online/HTTP/Private/GenericPlatform/HttpRequestCommon.cpp b/Engine/Source/Runtime/Online/HTTP/Private/GenericPlatform/HttpRequestCommon.cpp
--- a/Engine/Source/Runtime/Online/HTTP/Private/GenericPlatform/HttpRequestCommon.cpp
+++ b/Engine/Source/Runtime/Online/HTTP/Private/GenericPlatform/HttpRequestCommon.cpp
@@ -1,6 +1,7 @@
 // Copyright Epic Games, Inc. All Rights Reserved.
 
 #include "GenericPlatform/HttpRequestCommon.h"
+#include "GenericPlatform/HttpResponseCommon.h"
 #include "Http.h"
 #include "HttpManager.h"
 
@@ -68,3 +69,14 @@ EHttpRequestDelegateThreadPolicy FHttpRequestCommon::GetDelegateThreadPolicy() c
 	return DelegateThreadPolicy; 
 }
 
+void FHttpRequestCommon::SetStatus(EHttpRequestStatus::Type InCompletionStatus)
+{
+	CompletionStatus = InCompletionStatus;
+
+	if (FHttpResponsePtr Response = GetResponse())
+	{
+		TSharedPtr<FHttpResponseCommon> ResponseCommon = StaticCastSharedPtr<FHttpResponseCommon>(Response);
+		ResponseCommon->SetRequestStatus(InCompletionStatus);
+	}
+}
+
diff --git a/Engine/Source/Runtime/Online/HTTP/Private/GenericPlatform/HttpResponseCommon.cpp b/Engine/Source/Runtime/Online/HTTP/Private/GenericPlatform/HttpResponseCommon.cpp
--- a/Engine/Source/Runtime/Online/HTTP/Private/GenericPlatform/HttpResponseCommon.cpp
+++ b/Engine/Source/Runtime/Online/HTTP/Private/GenericPlatform/HttpResponseCommon.cpp
@@ -1,10 +1,12 @@
 // Copyright Epic Games, Inc. All Rights Reserved.
 
 #include "GenericPlatform/HttpResponseCommon.h"
+#include "GenericPlatform/HttpRequestCommon.h"
 #include "GenericPlatform/GenericPlatformHttp.h"
 
-FHttpResponseCommon::FHttpResponseCommon(const FString& InURL)
-	: URL(InURL)
+FHttpResponseCommon::FHttpResponseCommon(const FHttpRequestCommon& HttpRequest)
+	: URL(HttpRequest.GetURL())
+	, CompletionStatus(HttpRequest.GetStatus())
 {
 }
 
@@ -22,3 +24,13 @@ FString FHttpResponseCommon::GetURL() const
 {
 	return URL;
 }
+
+void FHttpResponseCommon::SetRequestStatus(EHttpRequestStatus::Type InCompletionStatus)
+{
+	CompletionStatus = InCompletionStatus;
+}
+
+EHttpRequestStatus::Type FHttpResponseCommon::GetStatus() const
+{
+	return CompletionStatus;
+}
diff --git a/Engine/Source/Runtime/Online/HTTP/Public/GenericPlatform/HttpRequestCommon.h b/Engine/Source/Runtime/Online/HTTP/Public/GenericPlatform/HttpRequestCommon.h
--- a/Engine/Source/Runtime/Online/HTTP/Public/GenericPlatform/HttpRequestCommon.h
+++ b/Engine/Source/Runtime/Online/HTTP/Public/GenericPlatform/HttpRequestCommon.h
@@ -23,6 +23,7 @@ protected:
 	 * Check if this request is valid or allowed, before actually process the request
 	 */
 	HTTP_API bool PreCheck() const;
+	HTTP_API void SetStatus(EHttpRequestStatus::Type InCompletionStatus);
 
 protected:
 	/** Current status of request being processed */
diff --git a/Engine/Source/Runtime/Online/HTTP/Public/GenericPlatform/HttpResponseCommon.h b/Engine/Source/Runtime/Online/HTTP/Public/GenericPlatform/HttpResponseCommon.h
--- a/Engine/Source/Runtime/Online/HTTP/Public/GenericPlatform/HttpResponseCommon.h
+++ b/Engine/Source/Runtime/Online/HTTP/Public/GenericPlatform/HttpResponseCommon.h
@@ -4,18 +4,26 @@
 
 #include "Interfaces/IHttpResponse.h"
 
+class FHttpRequestCommon;
+
 /**
  * Contains implementation of some common functions that don't vary between implementations of different platforms
  */
 class FHttpResponseCommon : public IHttpResponse
 {
+	friend FHttpRequestCommon;
+
 public:
-	HTTP_API FHttpResponseCommon(const FString& InURL);
+	HTTP_API FHttpResponseCommon(const FHttpRequestCommon& HttpRequest);
 
 	// IHttpBase
 	HTTP_API virtual FString GetURLParameter(const FString& ParameterName) const override;
 	HTTP_API virtual FString GetURL() const override;
+	HTTP_API virtual EHttpRequestStatus::Type GetStatus() const override;
 
 protected:
+	HTTP_API void SetRequestStatus(EHttpRequestStatus::Type InCompletionStatus);
+
 	FString URL;
+	EHttpRequestStatus::Type CompletionStatus;
 };
diff --git a/Engine/Source/Runtime/Online/HTTP/Public/Interfaces/IHttpBase.h b/Engine/Source/Runtime/Online/HTTP/Public/Interfaces/IHttpBase.h
--- a/Engine/Source/Runtime/Online/HTTP/Public/Interfaces/IHttpBase.h
+++ b/Engine/Source/Runtime/Online/HTTP/Public/Interfaces/IHttpBase.h
@@ -4,6 +4,61 @@
 
 #include "CoreMinimal.h"
 
+namespace EHttpRequestStatus
+{
+	/**
+	 * Enumerates the current state of an Http request
+	 */
+	enum Type
+	{
+		/** Has not been started via ProcessRequest() */
+		NotStarted,
+		/** Currently being ticked and processed */
+		Processing,
+		/** Finished but failed */
+		Failed,
+		/** Failed because it was unable to connect (safe to retry) */
+		Failed_ConnectionError,
+		/** Finished and was successful */
+		Succeeded
+	};
+
+	/** @return the stringified version of the enum passed in */
+	inline const TCHAR* ToString(EHttpRequestStatus::Type EnumVal)
+	{
+		switch (EnumVal)
+		{
+			case NotStarted:
+			{
+				return TEXT("NotStarted");
+			}
+			case Processing:
+			{
+				return TEXT("Processing");
+			}
+			case Failed:
+			{
+				return TEXT("Failed");
+			}
+			case Failed_ConnectionError:
+			{
+				return TEXT("ConnectionError");
+			}
+			case Succeeded:
+			{
+				return TEXT("Succeeded");
+			}
+		}
+		return TEXT("");
+	}
+
+	inline bool IsFinished(const EHttpRequestStatus::Type Value)
+	{
+		return Value == Failed || Value == Failed_ConnectionError || Value == Succeeded;
+	}
+}
+
+
 /**
  * Base interface for Http Requests and Responses.
  */
@@ -18,6 +73,13 @@ public:
 	 */
 	virtual FString GetURL() const = 0;
 
+	/**
+	 * Get the current status of the request being processed
+	 *
+	 * @return the current status
+	 */
+	virtual EHttpRequestStatus::Type GetStatus() const = 0;
+
 	/** 
 	 * Gets an URL parameter.
 	 * expected format is ?Key=Value&Key=Value...
diff --git a/Engine/Source/Runtime/Online/HTTP/Public/Interfaces/IHttpRequest.h b/Engine/Source/Runtime/Online/HTTP/Public/Interfaces/IHttpRequest.h
--- a/Engine/Source/Runtime/Online/HTTP/Public/Interfaces/IHttpRequest.h
+++ b/Engine/Source/Runtime/Online/HTTP/Public/Interfaces/IHttpRequest.h
@@ -8,60 +8,6 @@
 class IHttpRequest;
 class IHttpResponse;
 
-namespace EHttpRequestStatus
-{
-	/**
-	 * Enumerates the current state of an Http request
-	 */
-	enum Type
-	{
-		/** Has not been started via ProcessRequest() */
-		NotStarted,
-		/** Currently being ticked and processed */
-		Processing,
-		/** Finished but failed */
-		Failed,
-		/** Failed because it was unable to connect (safe to retry) */
-		Failed_ConnectionError,
-		/** Finished and was successful */
-		Succeeded
-	};
-
-	/** @return the stringified version of the enum passed in */
-	inline const TCHAR* ToString(EHttpRequestStatus::Type EnumVal)
-	{
-		switch (EnumVal)
-		{
-			case NotStarted:
-			{
-				return TEXT("NotStarted");
-			}
-			case Processing:
-			{
-				return TEXT("Processing");
-			}
-			case Failed:
-			{
-				return TEXT("Failed");
-			}
-			case Failed_ConnectionError:
-			{
-				return TEXT("ConnectionError");
-			}
-			case Succeeded:
-			{
-				return TEXT("Succeeded");
-			}
-		}
-		return TEXT("");
-	}
-
-	inline bool IsFinished(const EHttpRequestStatus::Type Value)
-	{
-		return Value == Failed || Value == Failed_ConnectionError || Value == Succeeded;
-	}
-}
-
 /**
  * Enumerates thread policy about which thread to complete the http request
  */
@@ -357,13 +303,6 @@ public:
 	 */
 	virtual void CancelRequest() = 0;
 
-	/**
-	 * Get the current status of the request being processed
-	 *
-	 * @return the current status
-	 */
-	virtual EHttpRequestStatus::Type GetStatus() const = 0;
-
 	/**
 	 * Get the associated Response
 	 *
