Fix the editor link error related to http status change.

diff --git a/Engine/Source/Runtime/Online/BuildPatchServices/Private/Tests/Mock/HttpResponse.mock.h b/Engine/Source/Runtime/Online/BuildPatchServices/Private/Tests/Mock/HttpResponse.mock.h
--- a/Engine/Source/Runtime/Online/BuildPatchServices/Private/Tests/Mock/HttpResponse.mock.h
+++ b/Engine/Source/Runtime/Online/BuildPatchServices/Private/Tests/Mock/HttpResponse.mock.h
@@ -71,7 +71,7 @@ namespace BuildPatchServices
 			return None;
 		}
 
-		HTTP_API virtual EHttpRequestStatus::Type GetStatus() const override
+		virtual EHttpRequestStatus::Type GetStatus() const override
 		{
 			MOCK_FUNC_NOT_IMPLEMENTED("FMockHttpRequest::GetStatus");
 			return EHttpRequestStatus::Failed;
