Remove XMPP room password from logging.

diff --git a/Engine/Source/Runtime/Online/XMPP/Private/XmppStrophe/XmppMultiUserChatStrophe.cpp b/Engine/Source/Runtime/Online/XMPP/Private/XmppStrophe/XmppMultiUserChatStrophe.cpp
--- a/Engine/Source/Runtime/Online/XMPP/Private/XmppStrophe/XmppMultiUserChatStrophe.cpp
+++ b/Engine/Source/Runtime/Online/XMPP/Private/XmppStrophe/XmppMultiUserChatStrophe.cpp
@@ -548,7 +548,7 @@ bool FXmppMultiUserChatStrophe::JoinPublicRoom(const FXmppRoomId& RoomId, const
 
 bool FXmppMultiUserChatStrophe::JoinPrivateRoom(const FXmppRoomId& RoomId, const FString& Nickname, const FString& Password)
 {
-	UE_LOG(LogXmpp, Verbose, TEXT("MUC: JoinPrivateRoom RoomId=%s Nickname=%s Password=%s"), *RoomId, *Nickname, *Password);
+	UE_LOG(LogXmpp, Verbose, TEXT("MUC: JoinPrivateRoom RoomId=%s Nickname=%s"), *RoomId, *Nickname);
 
 	bool bSuccess = false;
 	FString ErrorStr;
