Move FCurlHttpRequest::MarkAsCompleted function into cpp impl for easier debugging.

diff --git a/Engine/Source/Runtime/Online/HTTP/Private/Curl/CurlHttp.cpp b/Engine/Source/Runtime/Online/HTTP/Private/Curl/CurlHttp.cpp
--- a/Engine/Source/Runtime/Online/HTTP/Private/Curl/CurlHttp.cpp
+++ b/Engine/Source/Runtime/Online/HTTP/Private/Curl/CurlHttp.cpp
@@ -1151,6 +1151,12 @@ void FCurlHttpRequest::BroadcastNewlyReceivedHeader(const FString& HeaderKey, co
 	OnHeaderReceived().ExecuteIfBound(SharedThis(this), HeaderKey, HeaderValue);
 }
 
+void FCurlHttpRequest::MarkAsCompleted(CURLcode InCurlCompletionResult)
+{
+	CurlCompletionResult = InCurlCompletionResult;
+	bCurlRequestCompleted = true;
+}
+
 void FCurlHttpRequest::FinishRequest()
 {
 	QUICK_SCOPE_CYCLE_COUNTER(STAT_FCurlHttpRequest_FinishRequest);
diff --git a/Engine/Source/Runtime/Online/HTTP/Private/Curl/CurlHttp.h b/Engine/Source/Runtime/Online/HTTP/Private/Curl/CurlHttp.h
--- a/Engine/Source/Runtime/Online/HTTP/Private/Curl/CurlHttp.h
+++ b/Engine/Source/Runtime/Online/HTTP/Private/Curl/CurlHttp.h
@@ -203,11 +203,7 @@ public:
 	 *
 	 * @param CurlCompletionResult Operation result code as returned by libcurl
 	 */
-	inline void MarkAsCompleted(CURLcode InCurlCompletionResult)
-	{
-		CurlCompletionResult = InCurlCompletionResult;
-		bCurlRequestCompleted = true;
-	}
+	void MarkAsCompleted(CURLcode InCurlCompletionResult);
 	
 	/** 
 	 * Set the result for adding the easy handle to curl multi
