Add http fail reason interface in http request and response, instead of using the status to check if it's connection failure or any other reason. It's to prepare for migrating retry system time out flow into http request itself by adding failure reason of timeout.

diff --git a/Engine/Plugins/Experimental/Web/WebAPI/Source/WebAPI/Private/WebAPISubsystem.cpp b/Engine/Plugins/Experimental/Web/WebAPI/Source/WebAPI/Private/WebAPISubsystem.cpp
--- a/Engine/Plugins/Experimental/Web/WebAPI/Source/WebAPI/Private/WebAPISubsystem.cpp
+++ b/Engine/Plugins/Experimental/Web/WebAPI/Source/WebAPI/Private/WebAPISubsystem.cpp
@@ -157,34 +157,37 @@ void UWebAPISubsystem::Tick(float DeltaTime)
 		{
 			Request->ProcessRequest();
 		}
-		else if(RequestStatus == EHttpRequestStatus::Failed_ConnectionError)
-		{			
-			Request->ProcessRequest();
-		}
 		else if(RequestStatus == EHttpRequestStatus::Failed)
 		{
-			// If there's a response that was denied, keep for retry after auth
-			if(Request->GetResponse().IsValid())
+			if (Request->GetFailureReason() == EHttpFailureReason::ConnectionError)
+			{
+				Request->ProcessRequest();
+			}
+			else
 			{
-				if(Request->GetResponse()->GetResponseCode() == EHttpResponseCodes::Denied)
+				// If there's a response that was denied, keep for retry after auth
+				if(Request->GetResponse().IsValid())
 				{
-					FString Host = UWebAPIUtilities::GetHostFromUrl(Request->GetURL());
-
-					TArray<TSharedRef<IHttpRequest>>& HostRequests = HostRequestBuffer.FindOrAdd(Host);
-					// Prevent filling this buffer due to never being authorized, etc. - discards requests if buffer full
-					if(HostRequests.Num() <= CVarMaxDeniedHttpRetryCount->GetInt())
+					if(Request->GetResponse()->GetResponseCode() == EHttpResponseCodes::Denied)
 					{
-						HostRequests.Add(Request);
-						HostRequestBuffer.Add(Host);
-					}
+						FString Host = UWebAPIUtilities::GetHostFromUrl(Request->GetURL());
 
-					RequestBuffer.Remove(Request);
+						TArray<TSharedRef<IHttpRequest>>& HostRequests = HostRequestBuffer.FindOrAdd(Host);
+						// Prevent filling this buffer due to never being authorized, etc. - discards requests if buffer full
+						if(HostRequests.Num() <= CVarMaxDeniedHttpRetryCount->GetInt())
+						{
+							HostRequests.Add(Request);
+							HostRequestBuffer.Add(Host);
+						}
 
-					continue;
+						RequestBuffer.Remove(Request);
+
+						continue;
+					}
 				}
+				
+				RequestBuffer.Remove(Request);			
 			}
-			
-			RequestBuffer.Remove(Request);			
 		}
 	}
 }
diff --git a/Engine/Source/Runtime/Online/BuildPatchServices/Private/Tests/Mock/HttpRequest.mock.h b/Engine/Source/Runtime/Online/BuildPatchServices/Private/Tests/Mock/HttpRequest.mock.h
--- a/Engine/Source/Runtime/Online/BuildPatchServices/Private/Tests/Mock/HttpRequest.mock.h
+++ b/Engine/Source/Runtime/Online/BuildPatchServices/Private/Tests/Mock/HttpRequest.mock.h
@@ -177,6 +177,12 @@ namespace BuildPatchServices
 			return EHttpRequestStatus::Type();
 		}
 
+		virtual EHttpFailureReason GetFailureReason() const override
+		{
+			MOCK_FUNC_NOT_IMPLEMENTED("FMockHttpRequest::GetFailureReason");
+			return EHttpFailureReason::None;
+		}
+
 		virtual const FHttpResponsePtr GetResponse() const override
 		{
 			MOCK_FUNC_NOT_IMPLEMENTED("FMockHttpRequest::GetResponse");
diff --git a/Engine/Source/Runtime/Online/BuildPatchServices/Private/Tests/Mock/HttpResponse.mock.h b/Engine/Source/Runtime/Online/BuildPatchServices/Private/Tests/Mock/HttpResponse.mock.h
--- a/Engine/Source/Runtime/Online/BuildPatchServices/Private/Tests/Mock/HttpResponse.mock.h
+++ b/Engine/Source/Runtime/Online/BuildPatchServices/Private/Tests/Mock/HttpResponse.mock.h
@@ -77,6 +77,11 @@ namespace BuildPatchServices
 			return EHttpRequestStatus::Failed;
 		}
 
+		virtual EHttpFailureReason GetFailureReason() const override
+		{
+			MOCK_FUNC_NOT_IMPLEMENTED("FMockHttpRequest::GetFailureReason");
+			return EHttpFailureReason::Other;
+		}
 	};
 }
 
diff --git a/Engine/Source/Runtime/Online/HTTP/Private/Apple/AppleHttp.cpp b/Engine/Source/Runtime/Online/HTTP/Private/Apple/AppleHttp.cpp
--- a/Engine/Source/Runtime/Online/HTTP/Private/Apple/AppleHttp.cpp
+++ b/Engine/Source/Runtime/Online/HTTP/Private/Apple/AppleHttp.cpp
@@ -10,17 +10,6 @@
 #include "Http.h"
 #include "HttpModule.h"
 
-/**
- * State of a request's response
- */
-enum class EAppleHttpRequestResponseState: uint8
-{
-	NotReady,
-	Success,
-	Error,
-	ConnectionError
-};
-
 /**
  * Class to hold data from delegate implementation notifications.
  */
@@ -51,8 +40,10 @@ enum class EAppleHttpRequestResponseState: uint8
 @property uint64 BytesWritten;
 /** The total number of bytes received out during the request/response */
 @property uint64 BytesReceived;
-/** Response state */
-@property EAppleHttpRequestResponseState ResponseState;
+/** Request status */
+@property EHttpRequestStatus::Type RequestStatus;
+/** Reason of failure */
+@property EHttpFailureReason FailureReason;
 
 /** NSURLSessionDataDelegate delegate methods. Those are called from a thread controlled by the NSURLSession */
 
@@ -70,7 +61,8 @@ enum class EAppleHttpRequestResponseState: uint8
 
 @implementation FAppleHttpResponseDelegate
 @synthesize Response;
-@synthesize ResponseState;
+@synthesize RequestStatus;
+@synthesize FailureReason;
 @synthesize BytesWritten;
 @synthesize BytesReceived;
 
@@ -80,7 +72,8 @@ enum class EAppleHttpRequestResponseState: uint8
 	
 	BytesWritten = 0;
 	BytesReceived = 0;
-	ResponseState = EAppleHttpRequestResponseState::NotReady;
+	RequestStatus = EHttpRequestStatus::NotStarted;
+	FailureReason = EHttpFailureReason::None;
 	ResponseBodyReceiveStream = ResponseStream;
 	bInitializedWithValidStream = (ResponseStream != nullptr);
 	
@@ -163,7 +156,7 @@ enum class EAppleHttpRequestResponseState: uint8
 	if (error == nil)
 	{
 		UE_LOG(LogHttp, Verbose, TEXT("URLSession:task:didCompleteWithError. Http request succeeded: %p"), self);
-		self.ResponseState = EAppleHttpRequestResponseState::Success;
+		self.RequestStatus = EHttpRequestStatus::Succeeded;
 	}
 	else
 	{
@@ -172,6 +165,7 @@ enum class EAppleHttpRequestResponseState: uint8
 			   *FString([[error userInfo] objectForKey:NSURLErrorFailingURLStringErrorKey]),
 			   self);
 		
+		self.RequestStatus = EHttpRequestStatus::Failed;
 		// Determine if the specific error was failing to connect to the host.
 		switch ([error code])
 		{
@@ -179,10 +173,14 @@ enum class EAppleHttpRequestResponseState: uint8
 			case NSURLErrorCannotFindHost:
 			case NSURLErrorCannotConnectToHost:
 			case NSURLErrorDNSLookupFailed:
-				self.ResponseState = EAppleHttpRequestResponseState::ConnectionError;
+				self.FailureReason = EHttpFailureReason::ConnectionError;
+				break;
+			case NSURLErrorCancelled:
+				self.FailureReason = EHttpFailureReason::Cancelled;
 				break;
 			default:
-				self.ResponseState = EAppleHttpRequestResponseState::Error;
+				self.FailureReason = EHttpFailureReason::Other;
+				break;
 		}
 		// Log more details if verbose logging is enabled and this is an SSL error
 		if (UE_LOG_ACTIVE(LogHttp, Verbose))
@@ -696,6 +694,7 @@ bool FAppleHttpRequest::StartRequest()
 		bStarted = true;
 
 		SetStatus(EHttpRequestStatus::Processing);
+		SetFailureReason(EHttpFailureReason::None);
 
 		Response = MakeShared<FAppleHttpResponse>(*this);
 
@@ -711,7 +710,8 @@ bool FAppleHttpRequest::StartRequest()
 	else
 	{
 		UE_LOG(LogHttp, Warning, TEXT("ProcessRequest failed. Could not initialize Internet connection."));
-		SetStatus(EHttpRequestStatus::Failed_ConnectionError);
+		SetStatus(EHttpRequestStatus::Failed);
+		SetFailureReason(EHttpFailureReason::ConnectionError);
 	}
 
 	return bStarted;
@@ -724,29 +724,26 @@ void FAppleHttpRequest::FinishRequest()
 	// Clean up session/request handles that may have been created
 	CleanupRequest();
 
-	bool bSuccess = false;
-	if (Response.IsValid() && Response->IsReady() && !Response->HadError())
-	{
-		bSuccess = true;
-		UE_LOG(LogHttp, Verbose, TEXT("Request succeeded"));
-		SetStatus(EHttpRequestStatus::Succeeded);
+	bool bSucceeded = (Response && Response->GetStatusFromDelegate() == EHttpRequestStatus::Succeeded);
+	UE_LOG(LogHttp, Verbose, TEXT("Request %s"), bSucceeded ? TEXT("succeeded") : TEXT("failed"));
+	SetStatus(bSucceeded ? EHttpRequestStatus::Succeeded : EHttpRequestStatus::Failed);
 
-		// TODO: Try to broadcast OnHeaderReceived when we receive headers instead of here at the end
-		BroadcastResponseHeadersReceived();
-	}
-	else
+	if (!bSucceeded)
 	{
-		UE_LOG(LogHttp, Verbose, TEXT("Request failed"));
-		FString URL([[Request URL] absoluteString]);
-		SetStatus(EHttpRequestStatus::Failed);
-		if (Response.IsValid() && Response->HadConnectionError())
+		SetFailureReason(Response ? Response->GetFailureReasonFromDelegate() : EHttpFailureReason::Other);
+
+		if (GetFailureReason() == EHttpFailureReason::ConnectionError)
 		{
-			SetStatus(EHttpRequestStatus::Failed_ConnectionError);
 			Response = nullptr;
 		}
 	}
+	else
+	{
+		// TODO: Try to broadcast OnHeaderReceived when we receive headers instead of here at the end
+		BroadcastResponseHeadersReceived();
+	}
 
-	OnProcessRequestComplete().ExecuteIfBound(SharedThis(this), Response, bSuccess);
+	OnProcessRequestComplete().ExecuteIfBound(SharedThis(this), Response, bSucceeded);
 }
 
 void FAppleHttpRequest::CleanupRequest()
@@ -794,7 +791,7 @@ void FAppleHttpRequest::Tick(float DeltaSeconds)
 
 void FAppleHttpRequest::CheckProgressDelegate()
 {
-	if (Response.IsValid() && (CompletionStatus == EHttpRequestStatus::Processing || Response->HadError()))
+	if (Response.IsValid() && (CompletionStatus == EHttpRequestStatus::Processing || Response->GetStatusFromDelegate() == EHttpRequestStatus::Failed))
 	{
 		const uint64 BytesWritten = Response->GetNumBytesWritten();
 		const uint64 BytesRead = Response->GetNumBytesReceived();
@@ -939,24 +936,17 @@ int32 FAppleHttpResponse::GetResponseCode() const
 
 bool FAppleHttpResponse::IsReady() const
 {
-	return (ResponseDelegate.ResponseState != EAppleHttpRequestResponseState::NotReady);
+	return EHttpRequestStatus::IsFinished(ResponseDelegate.RequestStatus);
 }
 
-bool FAppleHttpResponse::HadError() const
+EHttpRequestStatus::Type FAppleHttpResponse::GetStatusFromDelegate() const
 {
-	switch(ResponseDelegate.ResponseState)
-	{
-		case EAppleHttpRequestResponseState::Error:
-		case EAppleHttpRequestResponseState::ConnectionError:
-			return true;
-		default:
-			return false;
-	}
+	return ResponseDelegate.RequestStatus;
 }
 
-bool FAppleHttpResponse::HadConnectionError() const
+EHttpFailureReason FAppleHttpResponse::GetFailureReasonFromDelegate() const
 {
-	return (ResponseDelegate.ResponseState == EAppleHttpRequestResponseState::ConnectionError);
+	return ResponseDelegate.FailureReason;
 }
 
 const uint64 FAppleHttpResponse::GetNumBytesReceived() const
diff --git a/Engine/Source/Runtime/Online/HTTP/Private/Apple/AppleHttp.h b/Engine/Source/Runtime/Online/HTTP/Private/Apple/AppleHttp.h
--- a/Engine/Source/Runtime/Online/HTTP/Private/Apple/AppleHttp.h
+++ b/Engine/Source/Runtime/Online/HTTP/Private/Apple/AppleHttp.h
@@ -74,7 +74,6 @@ public:
 	 */
 	virtual ~FAppleHttpRequest();
 
-
 private:
 	/**
 	 * Trigger the request progress delegate if progress has changed
@@ -135,15 +134,7 @@ private:
  */
 class FAppleHttpResponse : public FHttpResponseCommon
 {
-private:
-	// Delegate implementation. Keeps the response state and data
-	FAppleHttpResponseDelegate* ResponseDelegate;
-
 public:
-	// implementation friends
-	friend class FAppleHttpRequest;
-
-
 	//~ Begin IHttpBase Interface
 	virtual FString GetHeader(const FString& HeaderName) const override;
 	virtual TArray<FString> GetAllHeaders() const override;	
@@ -157,25 +148,10 @@ public:
 	virtual FString GetContentAsString() const override;
 	//~ End IHttpResponse Interface
 
-	/**
-	 * Check whether headers are available.
-	 */
-	bool AreHeadersAvailable() const;
-
 	/**
 	 * Check whether a response is ready or not.
 	 */
 	bool IsReady() const;
-	
-	/**
-	 * Check whether a response had an error.
-	 */
-	bool HadError() const;
-
-	/**
-	 * Check whether a response had a connection error.
-	 */
-	bool HadConnectionError() const;
 
 	/**
 	 * Get the number of bytes received so far
@@ -209,4 +185,21 @@ public:
 	 * Destructor
 	 */
 	virtual ~FAppleHttpResponse();
+
+private:
+	// implementation friends
+	friend class FAppleHttpRequest;
+
+	/**
+	 * Get status from the internal delegate
+	 */
+	EHttpRequestStatus::Type GetStatusFromDelegate() const;
+
+	/**
+	 * Get reason of failure from the internal delegate
+	 */
+	EHttpReasonOfFailure GetReasonOfFailureFromDelegate() const;
+
+	// Delegate implementation. Keeps the response state and data
+	FAppleHttpResponseDelegate* ResponseDelegate;
 };
diff --git a/Engine/Source/Runtime/Online/HTTP/Private/Curl/CurlHttp.cpp b/Engine/Source/Runtime/Online/HTTP/Private/Curl/CurlHttp.cpp
--- a/Engine/Source/Runtime/Online/HTTP/Private/Curl/CurlHttp.cpp
+++ b/Engine/Source/Runtime/Online/HTTP/Private/Curl/CurlHttp.cpp
@@ -999,6 +999,7 @@ bool FCurlHttpRequest::ProcessRequest()
 	QUICK_SCOPE_CYCLE_COUNTER(STAT_CurlHttpAddThreadedRequest);
 	// Mark as in-flight to prevent overlapped requests using the same object
 	SetStatus(EHttpRequestStatus::Processing);
+	SetFailureReason(EHttpFailureReason::None);
 	// Add to global list while being processed so that the ref counted request does not get deleted
 	FHttpModule::Get().GetHttpManager().AddThreadedRequest(SharedThis(this));
 
@@ -1290,10 +1291,12 @@ void FCurlHttpRequest::FinishRequest()
 			}
 		}
 
+		SetStatus(EHttpRequestStatus::Failed);
+
 		// Mark last request attempt as completed but failed
 		if (bCanceled)
 		{
-			SetStatus(EHttpRequestStatus::Failed);
+			SetFailureReason(EHttpFailureReason::Cancelled);
 		}
 		else if (bCurlRequestCompleted)
 		{
@@ -1304,21 +1307,21 @@ void FCurlHttpRequest::FinishRequest()
 			case CURLE_COULDNT_RESOLVE_HOST:
 			case CURLE_SSL_CONNECT_ERROR:
 				// report these as connection errors (safe to retry)
-				SetStatus(EHttpRequestStatus::Failed_ConnectionError);
+				SetFailureReason(EHttpFailureReason::ConnectionError);
 				break;
 			default:
-				SetStatus(EHttpRequestStatus::Failed);
+				SetFailureReason(EHttpFailureReason::Other);
 			}
 		}
 		else
 		{
 			if (bAnyHttpActivity)
 			{
-				SetStatus(EHttpRequestStatus::Failed);
+				SetFailureReason(EHttpFailureReason::Other);
 			}
 			else
 			{
-				SetStatus(EHttpRequestStatus::Failed_ConnectionError);
+				SetFailureReason(EHttpFailureReason::ConnectionError);
 			}
 		}
 		// Call delegate with failure
diff --git a/Engine/Source/Runtime/Online/HTTP/Private/GenericPlatform/GenericPlatformHttp.cpp b/Engine/Source/Runtime/Online/HTTP/Private/GenericPlatform/GenericPlatformHttp.cpp
--- a/Engine/Source/Runtime/Online/HTTP/Private/GenericPlatform/GenericPlatformHttp.cpp
+++ b/Engine/Source/Runtime/Online/HTTP/Private/GenericPlatform/GenericPlatformHttp.cpp
@@ -118,6 +118,7 @@ public:
 	virtual bool ProcessRequest() override { return false; }
 	virtual void CancelRequest() override {}
 	virtual EHttpRequestStatus::Type GetStatus() const override { return EHttpRequestStatus::NotStarted; }
+	virtual EHttpFailureReason GetFailureReason() const override { return EHttpFailureReason::None; }
 	virtual const FHttpResponsePtr GetResponse() const override { return nullptr; }
 	virtual void Tick(float DeltaSeconds) override {}
 	virtual float GetElapsedTime() const override { return 0.0f; }
diff --git a/Engine/Source/Runtime/Online/HTTP/Private/GenericPlatform/HttpRequestCommon.cpp b/Engine/Source/Runtime/Online/HTTP/Private/GenericPlatform/HttpRequestCommon.cpp
--- a/Engine/Source/Runtime/Online/HTTP/Private/GenericPlatform/HttpRequestCommon.cpp
+++ b/Engine/Source/Runtime/Online/HTTP/Private/GenericPlatform/HttpRequestCommon.cpp
@@ -20,6 +20,11 @@ EHttpRequestStatus::Type FHttpRequestCommon::GetStatus() const
 	return CompletionStatus;
 }
 
+EHttpFailureReason FHttpRequestCommon::GetFailureReason() const
+{
+	return FailureReason;
+}
+
 bool FHttpRequestCommon::PreCheck() const
 {
 	// Prevent overlapped requests using the same instance
@@ -72,3 +77,13 @@ void FHttpRequestCommon::SetStatus(EHttpRequestStatus::Type InCompletionStatus)
 	}
 }
 
+void FHttpRequestCommon::SetFailureReason(EHttpFailureReason InFailureReason)
+{
+	FailureReason = InFailureReason;
+
+	if (FHttpResponsePtr Response = GetResponse())
+	{
+		TSharedPtr<FHttpResponseCommon> ResponseCommon = StaticCastSharedPtr<FHttpResponseCommon>(Response);
+		ResponseCommon->SetRequestFailureReason(InFailureReason);
+	}
+}
\ No newline at end of file
diff --git a/Engine/Source/Runtime/Online/HTTP/Private/GenericPlatform/HttpResponseCommon.cpp b/Engine/Source/Runtime/Online/HTTP/Private/GenericPlatform/HttpResponseCommon.cpp
--- a/Engine/Source/Runtime/Online/HTTP/Private/GenericPlatform/HttpResponseCommon.cpp
+++ b/Engine/Source/Runtime/Online/HTTP/Private/GenericPlatform/HttpResponseCommon.cpp
@@ -7,6 +7,7 @@
 FHttpResponseCommon::FHttpResponseCommon(const FHttpRequestCommon& HttpRequest)
 	: URL(HttpRequest.GetURL())
 	, CompletionStatus(HttpRequest.GetStatus())
+	, FailureReason(HttpRequest.GetFailureReason())
 {
 }
 
@@ -34,3 +35,14 @@ EHttpRequestStatus::Type FHttpResponseCommon::GetStatus() const
 {
 	return CompletionStatus;
 }
+
+void FHttpResponseCommon::SetRequestFailureReason(EHttpFailureReason InFailureReason)
+{
+	FailureReason = InFailureReason;
+}
+
+EHttpFailureReason FHttpResponseCommon::GetFailureReason() const
+{
+	return FailureReason;
+}
+
diff --git a/Engine/Source/Runtime/Online/HTTP/Private/HttpRetrySystem.cpp b/Engine/Source/Runtime/Online/HTTP/Private/HttpRetrySystem.cpp
--- a/Engine/Source/Runtime/Online/HTTP/Private/HttpRetrySystem.cpp
+++ b/Engine/Source/Runtime/Online/HTTP/Private/HttpRetrySystem.cpp
@@ -204,9 +204,9 @@ void FHttpRetrySystem::FRequest::HttpOnProcessRequestComplete(FHttpRequestPtr In
 		{
 			// Do nothing here
 		}
-		if (GetStatus() == EHttpRequestStatus::Failed_ConnectionError || GetStatus() == EHttpRequestStatus::Failed)
+		else if (GetStatus() == EHttpRequestStatus::Failed)
 		{
-			if (GetStatus() == EHttpRequestStatus::Failed_ConnectionError && RetryDomains.IsValid())
+			if (GetFailureReason() == EHttpFailureReason::ConnectionError && RetryDomains.IsValid())
 			{
 				MoveToNextRetryDomain();
 			}
@@ -227,7 +227,14 @@ void FHttpRetrySystem::FRequest::HttpOnProcessRequestComplete(FHttpRequestPtr In
 				return;
 			}
 
-			RetryStatus = FHttpRetrySystem::FRequest::EStatus::FailedRetry;
+			if (GetFailureReason() == EHttpFailureReason::TimedOut)
+			{
+				RetryStatus = FHttpRetrySystem::FRequest::EStatus::FailedTimeout;
+			}
+			else
+			{
+				RetryStatus = FHttpRetrySystem::FRequest::EStatus::FailedRetry;
+			}
 		}
 		else
 		{
@@ -242,6 +249,7 @@ void FHttpRetrySystem::FRequest::HttpOnProcessRequestComplete(FHttpRequestPtr In
 		RetryManagerPtr->RequestList.RemoveAtSwap(EntryIndex);
 	}
 
+	LLM_SCOPE_BYTAG(HTTP);
 	OnProcessRequestComplete().ExecuteIfBound(SelfPtr, HttpResponse, bSucceeded);
 }
 
@@ -305,27 +313,29 @@ bool FHttpRetrySystem::FManager::ShouldRetry(const FHttpRetryRequestEntry& HttpR
 	if (!Response.IsValid())
 	{
 		// ONLY retry bad responses if they are connection errors (NOT protocol errors or unknown) otherwise request may be sent (and processed!) twice
-		EHttpRequestStatus::Type Status = HttpRetryRequestEntry.Request->GetStatus();
-		if (Status == EHttpRequestStatus::Failed_ConnectionError)
+		if (HttpRetryRequestEntry.Request->GetStatus() == EHttpRequestStatus::Failed)
 		{
-			bResult = true;
-		}
-		else if (Status == EHttpRequestStatus::Failed)
-		{
-			const FName Verb = FName(*HttpRetryRequestEntry.Request->GetVerb());
-
-			// Be default, we will also allow retry for GET and HEAD requests even if they may duplicate on the server
-			static const TSet<FName> DefaultRetryVerbs(TArray<FName>({ FName(TEXT("GET")), FName(TEXT("HEAD")) }));
-
-			const bool bIsRetryVerbsEmpty = HttpRetryRequestEntry.Request->RetryVerbs.Num() == 0;
-			if (bIsRetryVerbsEmpty && DefaultRetryVerbs.Contains(Verb))
+			if (HttpRetryRequestEntry.Request->GetFailureReason() == EHttpFailureReason::ConnectionError)
 			{
 				bResult = true;
 			}
-			// If retry verbs are specified, only allow retrying the specified list of verbs
-			else if (HttpRetryRequestEntry.Request->RetryVerbs.Contains(Verb))
+			else
 			{
-				bResult = true;
+				const FName Verb = FName(*HttpRetryRequestEntry.Request->GetVerb());
+
+				// Be default, we will also allow retry for GET and HEAD requests even if they may duplicate on the server
+				static const TSet<FName> DefaultRetryVerbs(TArray<FName>({ FName(TEXT("GET")), FName(TEXT("HEAD")) }));
+
+				const bool bIsRetryVerbsEmpty = HttpRetryRequestEntry.Request->RetryVerbs.Num() == 0;
+				if (bIsRetryVerbsEmpty && DefaultRetryVerbs.Contains(Verb))
+				{
+					bResult = true;
+				}
+				// If retry verbs are specified, only allow retrying the specified list of verbs
+				else if (HttpRetryRequestEntry.Request->RetryVerbs.Contains(Verb))
+				{
+					bResult = true;
+				}
 			}
 		}
 	}
@@ -444,7 +454,7 @@ float FHttpRetrySystem::FManager::GetLockoutPeriodSeconds(const FHttpRetryReques
 	{
 		if (LockoutPeriod <= 0.0f)
 		{
-			const bool bFailedToConnect = HttpRetryRequestEntry.Request->GetStatus() == EHttpRequestStatus::Failed_ConnectionError;
+			const bool bFailedToConnect = (HttpRetryRequestEntry.Request->GetStatus() == EHttpRequestStatus::Failed && HttpRetryRequestEntry.Request->GetFailureReason() == EHttpFailureReason::ConnectionError);
 			const bool bHasRetryDomains = HttpRetryRequestEntry.Request->RetryDomains.IsValid();
 			// Skip the lockout period if we failed to connect to a domain and we have other domains to try
 			const bool bSkipLockoutPeriod = (bFailedToConnect && bHasRetryDomains);
@@ -570,7 +580,7 @@ bool FHttpRetrySystem::FManager::Update(uint32* FileCount, uint32* FailingCount,
 					}
 
 					// If we failed to connect, try the next domain in the list
-					if (RequestStatus == EHttpRequestStatus::Failed_ConnectionError)
+					if (HttpRetryRequest->GetStatus() == EHttpRequestStatus::Failed && HttpRetryRequest->GetFailureReason() == EHttpFailureReason::ConnectionError)
 					{
 						if (HttpRetryRequest->RetryDomains.IsValid())
 						{
@@ -580,13 +590,13 @@ bool FHttpRetrySystem::FManager::Update(uint32* FileCount, uint32* FailingCount,
 					// Save these for failure case retry checks if we hit a completion state
 					bool bShouldRetry = false;
 					bool bCanRetry = false;
-					if (RequestStatus == EHttpRequestStatus::Failed || RequestStatus == EHttpRequestStatus::Failed_ConnectionError || RequestStatus == EHttpRequestStatus::Succeeded)
+					if (EHttpRequestStatus::IsFinished(RequestStatus))
 					{
 						bShouldRetry = ShouldRetry(*HttpRetryRequestEntry);
 						bCanRetry = CanRetry(*HttpRetryRequestEntry);
 					}
 
-					if (RequestStatus == EHttpRequestStatus::Failed || RequestStatus == EHttpRequestStatus::Failed_ConnectionError || forceFail || (bShouldRetry && bCanRetry))
+					if (RequestStatus == EHttpRequestStatus::Failed || forceFail || (bShouldRetry && bCanRetry))
 					{
 						bIsGreen = false;
 
diff --git a/Engine/Source/Runtime/Online/HTTP/Private/NullHttp.cpp b/Engine/Source/Runtime/Online/HTTP/Private/NullHttp.cpp
--- a/Engine/Source/Runtime/Online/HTTP/Private/NullHttp.cpp
+++ b/Engine/Source/Runtime/Online/HTTP/Private/NullHttp.cpp
@@ -155,6 +155,11 @@ EHttpRequestStatus::Type FNullHttpRequest::GetStatus() const
 	return CompletionStatus;
 }
 
+EHttpFailureReason FNullHttpRequest::GetFailureReason() const
+{
+	return FailureReason;
+}
+
 const FHttpResponsePtr FNullHttpRequest::GetResponse() const
 {
 	return FHttpResponsePtr(nullptr);
@@ -184,6 +189,7 @@ float FNullHttpRequest::GetElapsedTime() const
 void FNullHttpRequest::FinishedRequest()
 {
 	CompletionStatus = EHttpRequestStatus::Failed;
+	FailureReason = EHttpFailureReason::Other;
 
 	UE_LOG(LogHttp, Log, TEXT("Finished request %p. no response %s url=%s elapsed=%.3f"),
 		this, *GetVerb(), *GetURL(), ElapsedTime);
diff --git a/Engine/Source/Runtime/Online/HTTP/Private/NullHttp.h b/Engine/Source/Runtime/Online/HTTP/Private/NullHttp.h
--- a/Engine/Source/Runtime/Online/HTTP/Private/NullHttp.h
+++ b/Engine/Source/Runtime/Online/HTTP/Private/NullHttp.h
@@ -36,6 +36,7 @@ public:
 	virtual bool ProcessRequest() override;
 	virtual void CancelRequest() override;
 	virtual EHttpRequestStatus::Type GetStatus() const override;
+	virtual EHttpFailureReason GetFailureReason() const override;
 	virtual const FHttpResponsePtr GetResponse() const override;
 	virtual void Tick(float DeltaSeconds) override;
 	virtual float GetElapsedTime() const override;
@@ -44,6 +45,7 @@ public:
 
 	FNullHttpRequest()
 		: CompletionStatus(EHttpRequestStatus::NotStarted)
+		, FailureReason(EHttpFailureReason::None)
 		, ElapsedTime(0)
 	{}
 	virtual ~FNullHttpRequest() {}
@@ -55,6 +57,7 @@ private:
 	FString Verb;
 	TArray<uint8> Payload;
 	EHttpRequestStatus::Type CompletionStatus;
+	EHttpFailureReason FailureReason;
 	TMap<FString, FString> Headers;
 	float ElapsedTime;
 };
diff --git a/Engine/Source/Runtime/Online/HTTP/Private/WinHttp/Support/WinHttpConnectionHttp.cpp b/Engine/Source/Runtime/Online/HTTP/Private/WinHttp/Support/WinHttpConnectionHttp.cpp
--- a/Engine/Source/Runtime/Online/HTTP/Private/WinHttp/Support/WinHttpConnectionHttp.cpp
+++ b/Engine/Source/Runtime/Online/HTTP/Private/WinHttp/Support/WinHttpConnectionHttp.cpp
@@ -312,7 +312,7 @@ void FWinHttpConnectionHttp::PumpStates()
 			{
 				if (!SendRequest())
 				{
-					FinishRequest(EHttpRequestStatus::Failed_ConnectionError);
+					FinishRequest(EHttpRequestStatus::Failed);
 					return;
 				}
 				continue;
@@ -1149,7 +1149,7 @@ void FWinHttpConnectionHttp::HandleRequestError(const uint32 ErrorApiId, const u
 	{
 		if (!bConnectedToServer && IsErrorCodeAConnectionError(ErrorCode))
 		{
-			FinishRequest(EHttpRequestStatus::Failed_ConnectionError);
+			FinishRequest(EHttpRequestStatus::Failed);
 		}
 		else
 		{
diff --git a/Engine/Source/Runtime/Online/HTTP/Public/GenericPlatform/HttpRequestCommon.h b/Engine/Source/Runtime/Online/HTTP/Public/GenericPlatform/HttpRequestCommon.h
--- a/Engine/Source/Runtime/Online/HTTP/Public/GenericPlatform/HttpRequestCommon.h
+++ b/Engine/Source/Runtime/Online/HTTP/Public/GenericPlatform/HttpRequestCommon.h
@@ -15,6 +15,7 @@ public:
 
 	// IHttpRequest
 	HTTP_API virtual EHttpRequestStatus::Type GetStatus() const override;
+	HTTP_API virtual EHttpFailureReason GetFailureReason() const override;
 	HTTP_API virtual void SetDelegateThreadPolicy(EHttpRequestDelegateThreadPolicy InDelegateThreadPolicy) override;
 	HTTP_API virtual EHttpRequestDelegateThreadPolicy GetDelegateThreadPolicy() const override;
 
@@ -24,11 +25,15 @@ protected:
 	 */
 	HTTP_API bool PreCheck() const;
 	HTTP_API void SetStatus(EHttpRequestStatus::Type InCompletionStatus);
+	HTTP_API void SetFailureReason(EHttpFailureReason InFailureReason);
 
 protected:
 	/** Current status of request being processed */
 	EHttpRequestStatus::Type CompletionStatus = EHttpRequestStatus::NotStarted;
 
+	/** Reason of failure of the HTTP request */
+	EHttpFailureReason FailureReason = EHttpFailureReason::None;
+
 	/** Thread policy about which thread to complete this request */
 	EHttpRequestDelegateThreadPolicy DelegateThreadPolicy = EHttpRequestDelegateThreadPolicy::CompleteOnGameThread;
 };
diff --git a/Engine/Source/Runtime/Online/HTTP/Public/GenericPlatform/HttpResponseCommon.h b/Engine/Source/Runtime/Online/HTTP/Public/GenericPlatform/HttpResponseCommon.h
--- a/Engine/Source/Runtime/Online/HTTP/Public/GenericPlatform/HttpResponseCommon.h
+++ b/Engine/Source/Runtime/Online/HTTP/Public/GenericPlatform/HttpResponseCommon.h
@@ -20,10 +20,13 @@ public:
 	HTTP_API virtual FString GetURLParameter(const FString& ParameterName) const override;
 	HTTP_API virtual FString GetURL() const override;
 	HTTP_API virtual EHttpRequestStatus::Type GetStatus() const override;
+	HTTP_API virtual EHttpFailureReason GetFailureReason() const override;
 
 protected:
 	HTTP_API void SetRequestStatus(EHttpRequestStatus::Type InCompletionStatus);
+	HTTP_API void SetRequestFailureReason(EHttpFailureReason InFailureReason);
 
 	FString URL;
 	EHttpRequestStatus::Type CompletionStatus;
+	EHttpFailureReason FailureReason;
 };
diff --git a/Engine/Source/Runtime/Online/HTTP/Public/HttpModule.h b/Engine/Source/Runtime/Online/HTTP/Public/HttpModule.h
--- a/Engine/Source/Runtime/Online/HTTP/Public/HttpModule.h
+++ b/Engine/Source/Runtime/Online/HTTP/Public/HttpModule.h
@@ -348,7 +348,7 @@ private:
 	void OnConfigSectionsChanged(const FString& IniFilename, const TSet<FString>& SectionNames);
 
 	/** Keeps track of Http requests while they are being processed */
-	FHttpManager* HttpManager;
+	FHttpManager* HttpManager = nullptr;
 	/** timeout in seconds for the entire http request to complete. 0 is no timeout */
 	float HttpTimeout;
 	/** timeout in seconds to establish the connection. -1 for system defaults, 0 is no timeout */
diff --git a/Engine/Source/Runtime/Online/HTTP/Public/HttpRequestAdapter.h b/Engine/Source/Runtime/Online/HTTP/Public/HttpRequestAdapter.h
--- a/Engine/Source/Runtime/Online/HTTP/Public/HttpRequestAdapter.h
+++ b/Engine/Source/Runtime/Online/HTTP/Public/HttpRequestAdapter.h
@@ -18,33 +18,34 @@ public:
     {}
 
 	// IHttpRequest interface
-    virtual FString                       GetURL() const override                                                  { return HttpRequest->GetURL(); }
-	virtual FString                       GetURLParameter(const FString& ParameterName) const override             { return HttpRequest->GetURLParameter(ParameterName); }
-	virtual FString                       GetHeader(const FString& HeaderName) const override                      { return HttpRequest->GetHeader(HeaderName); }
-	virtual TArray<FString>               GetAllHeaders() const override                                           { return HttpRequest->GetAllHeaders(); }
-	virtual FString                       GetContentType() const override                                          { return HttpRequest->GetContentType(); }
-	virtual uint64                        GetContentLength() const override                                        { return HttpRequest->GetContentLength(); }
-	virtual const TArray<uint8>&          GetContent() const override                                              { return HttpRequest->GetContent(); }
-	virtual FString                       GetVerb() const override                                                 { return HttpRequest->GetVerb(); }
-	virtual void                          SetVerb(const FString& Verb) override                                    { HttpRequest->SetVerb(Verb); }
-	virtual void                          SetURL(const FString& URL) override                                      { HttpRequest->SetURL(URL); }
-	virtual void                          SetContent(const TArray<uint8>& ContentPayload) override                 { HttpRequest->SetContent(ContentPayload); }
-	virtual void                          SetContent(TArray<uint8>&& ContentPayload) override                      { HttpRequest->SetContent(MoveTemp(ContentPayload)); }
-	virtual void                          SetContentAsString(const FString& ContentString) override                { HttpRequest->SetContentAsString(ContentString); }
-    virtual bool                          SetContentAsStreamedFile(const FString& Filename) override               { return HttpRequest->SetContentAsStreamedFile(Filename); }
-	virtual bool                          SetContentFromStream(TSharedRef<FArchive, ESPMode::ThreadSafe> Stream) override { return HttpRequest->SetContentFromStream(Stream); }
-	virtual bool						  SetResponseBodyReceiveStream(TSharedRef<FArchive> Stream) override       { return HttpRequest->SetResponseBodyReceiveStream(Stream); }
-	virtual void                          SetHeader(const FString& HeaderName, const FString& HeaderValue) override { HttpRequest->SetHeader(HeaderName, HeaderValue); }
-	virtual void                          AppendToHeader(const FString& HeaderName, const FString& AdditionalHeaderValue) override { HttpRequest->AppendToHeader(HeaderName, AdditionalHeaderValue); }
-	virtual void                          SetTimeout(float InTimeoutSecs) override                                 { HttpRequest->SetTimeout(InTimeoutSecs); }
-	virtual void                          ClearTimeout() override                                                  { HttpRequest->ClearTimeout(); }
-	virtual TOptional<float>              GetTimeout() const override                                              { return HttpRequest->GetTimeout(); }
-	virtual const FHttpResponsePtr        GetResponse() const override                                             { return HttpRequest->GetResponse(); }
-	virtual float                         GetElapsedTime() const override                                          { return HttpRequest->GetElapsedTime(); }
-	virtual EHttpRequestStatus::Type	  GetStatus() const override                                               { return HttpRequest->GetStatus(); }
-	virtual void                          Tick(float DeltaSeconds) override                                        { HttpRequest->Tick(DeltaSeconds); }
-	virtual void                          SetDelegateThreadPolicy(EHttpRequestDelegateThreadPolicy InThreadPolicy) override{ HttpRequest->SetDelegateThreadPolicy(InThreadPolicy); }
-	virtual EHttpRequestDelegateThreadPolicy GetDelegateThreadPolicy() const override                                 { return HttpRequest->GetDelegateThreadPolicy(); }
+	virtual FString GetURL() const override { return HttpRequest->GetURL(); }
+	virtual FString GetURLParameter(const FString& ParameterName) const override { return HttpRequest->GetURLParameter(ParameterName); }
+	virtual FString GetHeader(const FString& HeaderName) const override { return HttpRequest->GetHeader(HeaderName); }
+	virtual TArray<FString> GetAllHeaders() const override { return HttpRequest->GetAllHeaders(); }
+	virtual FString GetContentType() const override { return HttpRequest->GetContentType(); }
+	virtual uint64 GetContentLength() const override { return HttpRequest->GetContentLength(); }
+	virtual const TArray<uint8>& GetContent() const override { return HttpRequest->GetContent(); }
+	virtual FString GetVerb() const override { return HttpRequest->GetVerb(); }
+	virtual void SetVerb(const FString& Verb) override { HttpRequest->SetVerb(Verb); }
+	virtual void SetURL(const FString& URL) override { HttpRequest->SetURL(URL); }
+	virtual void SetContent(const TArray<uint8>& ContentPayload) override { HttpRequest->SetContent(ContentPayload); }
+	virtual void SetContent(TArray<uint8>&& ContentPayload) override { HttpRequest->SetContent(MoveTemp(ContentPayload)); }
+	virtual void SetContentAsString(const FString& ContentString) override { HttpRequest->SetContentAsString(ContentString); }
+	virtual bool SetContentAsStreamedFile(const FString& Filename) override { return HttpRequest->SetContentAsStreamedFile(Filename); }
+	virtual bool SetContentFromStream(TSharedRef<FArchive, ESPMode::ThreadSafe> Stream) override { return HttpRequest->SetContentFromStream(Stream); }
+	virtual bool SetResponseBodyReceiveStream(TSharedRef<FArchive> Stream) override { return HttpRequest->SetResponseBodyReceiveStream(Stream); }
+	virtual void SetHeader(const FString& HeaderName, const FString& HeaderValue) override { HttpRequest->SetHeader(HeaderName, HeaderValue); }
+	virtual void AppendToHeader(const FString& HeaderName, const FString& AdditionalHeaderValue) override { HttpRequest->AppendToHeader(HeaderName, AdditionalHeaderValue); }
+	virtual void SetTimeout(float InTimeoutSecs) override { HttpRequest->SetTimeout(InTimeoutSecs); }
+	virtual void ClearTimeout() override { HttpRequest->ClearTimeout(); }
+	virtual TOptional<float> GetTimeout() const override { return HttpRequest->GetTimeout(); }
+	virtual const FHttpResponsePtr GetResponse() const override { return HttpRequest->GetResponse(); }
+	virtual float GetElapsedTime() const override { return HttpRequest->GetElapsedTime(); }
+	virtual EHttpRequestStatus::Type GetStatus() const override { return HttpRequest->GetStatus(); }
+	virtual EHttpFailureReason GetFailureReason() const override { return HttpRequest->GetFailureReason(); }
+	virtual void Tick(float DeltaSeconds) override { HttpRequest->Tick(DeltaSeconds); }
+	virtual void SetDelegateThreadPolicy(EHttpRequestDelegateThreadPolicy InThreadPolicy) override { HttpRequest->SetDelegateThreadPolicy(InThreadPolicy); }
+	virtual EHttpRequestDelegateThreadPolicy GetDelegateThreadPolicy() const override { return HttpRequest->GetDelegateThreadPolicy(); }
 
 protected:
     TSharedRef<IHttpRequest, ESPMode::ThreadSafe> HttpRequest;
diff --git a/Engine/Source/Runtime/Online/HTTP/Public/Interfaces/IHttpBase.h b/Engine/Source/Runtime/Online/HTTP/Public/Interfaces/IHttpBase.h
--- a/Engine/Source/Runtime/Online/HTTP/Public/Interfaces/IHttpBase.h
+++ b/Engine/Source/Runtime/Online/HTTP/Public/Interfaces/IHttpBase.h
@@ -18,7 +18,7 @@ namespace EHttpRequestStatus
 		/** Finished but failed */
 		Failed,
 		/** Failed because it was unable to connect (safe to retry) */
-		Failed_ConnectionError,
+		Failed_ConnectionError UE_DEPRECATED(5.4, "Failed_ConnectionError has been deprecated, use Failed + EHttpFailureReason::ConnectionError instead"),
 		/** Finished and was successful */
 		Succeeded
 	};
@@ -26,6 +26,7 @@ namespace EHttpRequestStatus
 	/** @return the stringified version of the enum passed in */
 	inline const TCHAR* ToString(EHttpRequestStatus::Type EnumVal)
 	{
+PRAGMA_DISABLE_DEPRECATION_WARNINGS
 		switch (EnumVal)
 		{
 			case NotStarted:
@@ -50,14 +51,31 @@ namespace EHttpRequestStatus
 			}
 		}
 		return TEXT("");
+PRAGMA_ENABLE_DEPRECATION_WARNINGS
 	}
 
 	inline bool IsFinished(const EHttpRequestStatus::Type Value)
 	{
-		return Value == Failed || Value == Failed_ConnectionError || Value == Succeeded;
+		return Value != NotStarted && Value != Processing;
 	}
 }
 
+/**
+ * The reason of the failure when HTTP request failed
+ */
+enum class EHttpFailureReason : uint8
+{
+	None = 0,
+	ConnectionError,
+	Cancelled,
+	// This acts differently than the platform timeout. It's the time out of the entire request including 
+	// retries, which configured by the user. The platform timeout has different meanings/APIs on different 
+	// platforms, such as resolve host timeout, connect timeout, send timeout, receive timeout, entire 
+	// request timeout etc. It's not practical to unify to same behavior between them through their APIs. 
+	// When user's configured time is up, it's TimedOut no matter which step this Http request is in.
+	TimedOut,
+	Other
+};
 
 /**
  * Base interface for Http Requests and Responses.
@@ -80,6 +98,13 @@ public:
 	 */
 	virtual EHttpRequestStatus::Type GetStatus() const = 0;
 
+	/**
+	 * Get the reason of th failure if GetStatus returns Failed
+	 *
+	 * @return the reason of the failure
+	 */
+	virtual EHttpFailureReason GetFailureReason() const = 0;
+
 	/** 
 	 * Gets an URL parameter.
 	 * expected format is ?Key=Value&Key=Value...
