Fix apple http compile error due to missed find-replace.

diff --git a/Engine/Source/Runtime/Online/HTTP/Private/Apple/AppleHttp.h b/Engine/Source/Runtime/Online/HTTP/Private/Apple/AppleHttp.h
--- a/Engine/Source/Runtime/Online/HTTP/Private/Apple/AppleHttp.h
+++ b/Engine/Source/Runtime/Online/HTTP/Private/Apple/AppleHttp.h
@@ -198,7 +198,7 @@ private:
 	/**
 	 * Get reason of failure from the internal delegate
 	 */
-	EHttpReasonOfFailure GetReasonOfFailureFromDelegate() const;
+	EHttpFailureReason GetFailureReasonFromDelegate() const;
 
 	// Delegate implementation. Keeps the response state and data
 	FAppleHttpResponseDelegate* ResponseDelegate;
