Fix the static analysis warning in WinHttpConnectionHttp.cpp: The 'then' statement is equivalent to the 'else' statement.

diff --git a/Engine/Source/Runtime/Online/HTTP/Private/WinHttp/Support/WinHttpConnectionHttp.cpp b/Engine/Source/Runtime/Online/HTTP/Private/WinHttp/Support/WinHttpConnectionHttp.cpp
--- a/Engine/Source/Runtime/Online/HTTP/Private/WinHttp/Support/WinHttpConnectionHttp.cpp
+++ b/Engine/Source/Runtime/Online/HTTP/Private/WinHttp/Support/WinHttpConnectionHttp.cpp
@@ -1147,14 +1147,7 @@ void FWinHttpConnectionHttp::HandleRequestError(const uint32 ErrorApiId, const u
 
 	if (!FinalState.IsSet())
 	{
-		if (!bConnectedToServer && IsErrorCodeAConnectionError(ErrorCode))
-		{
-			FinishRequest(EHttpRequestStatus::Failed);
-		}
-		else
-		{
-			FinishRequest(EHttpRequestStatus::Failed);
-		}
+		FinishRequest(EHttpRequestStatus::Failed);
 	}
 
 	// If the request was cancelled, we can release the payload memory
