Refactor to prepare for retry migration from retry system to http request, should have no behavior change:
- Move timeout related functions into FHttpRequestCommon, make sure those functions are not exposed to FHttpRequestAdapterBase;
- Rename PreCheck to PreProcess, also clear internal states in PreProcess to make sure the request is clean for retry.

diff --git a/Engine/Source/Runtime/Online/HTTP/Private/Apple/AppleHttp.cpp b/Engine/Source/Runtime/Online/HTTP/Private/Apple/AppleHttp.cpp
--- a/Engine/Source/Runtime/Online/HTTP/Private/Apple/AppleHttp.cpp
+++ b/Engine/Source/Runtime/Online/HTTP/Private/Apple/AppleHttp.cpp
@@ -652,19 +652,18 @@ bool FAppleHttpRequest::ProcessRequest()
 	SCOPED_AUTORELEASE_POOL;
 	UE_LOG(LogHttp, Verbose, TEXT("FAppleHttpRequest::ProcessRequest()"));
 
-	if (!PreCheck() || !StartRequest())
+	if (!PreProcess())
 	{
-		FinishRequestNotInHttpManager();
 		return false;
 	}
 
 	return true;
 }
 
-bool FAppleHttpRequest::StartRequest()
+bool FAppleHttpRequest::SetupRequest()
 {
 	SCOPED_AUTORELEASE_POOL;
-	UE_LOG(LogHttp, Verbose, TEXT("FAppleHttpRequest::StartRequest()"));
+	UE_LOG(LogHttp, Verbose, TEXT("FAppleHttpRequest::SetupRequest()"));
 	bool bStarted = false;
 
 	// set the content-length and user-agent (it is possible that the OS ignores this value)
diff --git a/Engine/Source/Runtime/Online/HTTP/Private/Apple/AppleHttp.h b/Engine/Source/Runtime/Online/HTTP/Private/Apple/AppleHttp.h
--- a/Engine/Source/Runtime/Online/HTTP/Private/Apple/AppleHttp.h
+++ b/Engine/Source/Runtime/Online/HTTP/Private/Apple/AppleHttp.h
@@ -85,7 +85,7 @@ private:
 	 *
 	 * @return true if the request was started
 	 */
-	bool StartRequest();
+	virtual bool SetupRequest() override;
 
 	/**
 	 * Close session/request handles and unregister callbacks
diff --git a/Engine/Source/Runtime/Online/HTTP/Private/Curl/CurlHttp.cpp b/Engine/Source/Runtime/Online/HTTP/Private/Curl/CurlHttp.cpp
--- a/Engine/Source/Runtime/Online/HTTP/Private/Curl/CurlHttp.cpp
+++ b/Engine/Source/Runtime/Online/HTTP/Private/Curl/CurlHttp.cpp
@@ -981,9 +981,8 @@ bool FCurlHttpRequest::ProcessRequest()
 	Response = nullptr;
 	LastReportedBytesRead = 0;
 
-	if (!PreCheck() || !SetupRequest())
+	if (!PreProcess())
 	{
-		FinishRequestNotInHttpManager();
 		return false;
 	}
 
@@ -999,7 +998,7 @@ bool FCurlHttpRequest::ProcessRequest()
 	QUICK_SCOPE_CYCLE_COUNTER(STAT_CurlHttpAddThreadedRequest);
 	// Mark as in-flight to prevent overlapped requests using the same object
 	SetStatus(EHttpRequestStatus::Processing);
-	SetFailureReason(EHttpFailureReason::None);
+
 	// Add to global list while being processed so that the ref counted request does not get deleted
 	FHttpModule::Get().GetHttpManager().AddThreadedRequest(SharedThis(this));
 
diff --git a/Engine/Source/Runtime/Online/HTTP/Private/Curl/CurlHttp.h b/Engine/Source/Runtime/Online/HTTP/Private/Curl/CurlHttp.h
--- a/Engine/Source/Runtime/Online/HTTP/Private/Curl/CurlHttp.h
+++ b/Engine/Source/Runtime/Online/HTTP/Private/Curl/CurlHttp.h
@@ -337,7 +337,7 @@ private:
 	 *
 	 * @return true if the request was successfully setup
 	 */
-	bool SetupRequest();
+	virtual bool SetupRequest() override;
 
 	/**
 	 * Trigger the request progress delegate if progress has changed
diff --git a/Engine/Source/Runtime/Online/HTTP/Private/GenericPlatform/GenericPlatformHttp.cpp b/Engine/Source/Runtime/Online/HTTP/Private/GenericPlatform/GenericPlatformHttp.cpp
--- a/Engine/Source/Runtime/Online/HTTP/Private/GenericPlatform/GenericPlatformHttp.cpp
+++ b/Engine/Source/Runtime/Online/HTTP/Private/GenericPlatform/GenericPlatformHttp.cpp
@@ -124,6 +124,9 @@ public:
 	virtual float GetElapsedTime() const override { return 0.0f; }
 	virtual void SetDelegateThreadPolicy(EHttpRequestDelegateThreadPolicy InDelegateThreadPolicy) override {}
 	virtual EHttpRequestDelegateThreadPolicy GetDelegateThreadPolicy() const override { return EHttpRequestDelegateThreadPolicy::CompleteOnGameThread; }
+	virtual void SetTimeout(float InTimeoutSecs) override {}
+	virtual void ClearTimeout() override {}
+	virtual TOptional<float> GetTimeout() const override { return TOptional<float>(); }
 };
 
 FDefaultUserAgentBuilder::FDefaultUserAgentBuilder()
diff --git a/Engine/Source/Runtime/Online/HTTP/Private/GenericPlatform/HttpRequestCommon.cpp b/Engine/Source/Runtime/Online/HTTP/Private/GenericPlatform/HttpRequestCommon.cpp
--- a/Engine/Source/Runtime/Online/HTTP/Private/GenericPlatform/HttpRequestCommon.cpp
+++ b/Engine/Source/Runtime/Online/HTTP/Private/GenericPlatform/HttpRequestCommon.cpp
@@ -56,6 +56,58 @@ bool FHttpRequestCommon::PreCheck() const
 	return true;
 }
 
+bool FHttpRequestCommon::PreProcess()
+{
+	ClearInCaseOfRetry();
+
+	if (!PreCheck() || !SetupRequest())
+	{
+		FinishRequestNotInHttpManager();
+		return false;
+	}
+
+	return true;
+}
+
+void FHttpRequestCommon::ClearInCaseOfRetry()
+{
+	// TODO: clear response shared ptr here as well after moving it from child class to this class
+
+	FailureReason = EHttpFailureReason::None;
+}
+
+void FHttpRequestCommon::FinishRequestNotInHttpManager()
+{
+	if (IsInGameThread())
+	{
+		if (DelegateThreadPolicy == EHttpRequestDelegateThreadPolicy::CompleteOnGameThread)
+		{
+			FinishRequest();
+		}
+		else
+		{
+			FHttpModule::Get().GetHttpManager().AddHttpThreadTask([StrongThis = StaticCastSharedRef<FHttpRequestCommon>(AsShared())]()
+			{
+				StrongThis->FinishRequest();
+			});
+		}
+	}
+	else
+	{
+		if (DelegateThreadPolicy == EHttpRequestDelegateThreadPolicy::CompleteOnHttpThread)
+		{
+			FinishRequest();
+		}
+		else
+		{
+			FHttpModule::Get().GetHttpManager().AddGameThreadTask([StrongThis = StaticCastSharedRef<FHttpRequestCommon>(AsShared())]()
+			{
+				StrongThis->FinishRequest();
+			});
+		}
+	}
+}
+
 void FHttpRequestCommon::SetDelegateThreadPolicy(EHttpRequestDelegateThreadPolicy InDelegateThreadPolicy)
 { 
 	DelegateThreadPolicy = InDelegateThreadPolicy; 
@@ -79,6 +131,7 @@ void FHttpRequestCommon::SetStatus(EHttpRequestStatus::Type InCompletionStatus)
 
 void FHttpRequestCommon::SetFailureReason(EHttpFailureReason InFailureReason)
 {
+	check(FailureReason == EHttpFailureReason::None);
 	FailureReason = InFailureReason;
 
 	if (FHttpResponsePtr Response = GetResponse())
@@ -86,4 +139,25 @@ void FHttpRequestCommon::SetFailureReason(EHttpFailureReason InFailureReason)
 		TSharedPtr<FHttpResponseCommon> ResponseCommon = StaticCastSharedPtr<FHttpResponseCommon>(Response);
 		ResponseCommon->SetRequestFailureReason(InFailureReason);
 	}
-}
\ No newline at end of file
+}
+
+void FHttpRequestCommon::SetTimeout(float InTimeoutSecs)
+{
+	TimeoutSecs = InTimeoutSecs;
+}
+
+void FHttpRequestCommon::ClearTimeout()
+{
+	TimeoutSecs.Reset();
+}
+
+TOptional<float> FHttpRequestCommon::GetTimeout() const
+{
+	return TimeoutSecs;
+}
+
+float FHttpRequestCommon::GetTimeoutOrDefault() const
+{
+	return GetTimeout().Get(FHttpModule::Get().GetHttpTimeout());
+}
+
diff --git a/Engine/Source/Runtime/Online/HTTP/Private/GenericPlatform/HttpRequestImpl.cpp b/Engine/Source/Runtime/Online/HTTP/Private/GenericPlatform/HttpRequestImpl.cpp
--- a/Engine/Source/Runtime/Online/HTTP/Private/GenericPlatform/HttpRequestImpl.cpp
+++ b/Engine/Source/Runtime/Online/HTTP/Private/GenericPlatform/HttpRequestImpl.cpp
@@ -31,26 +31,6 @@ FHttpRequestWillRetryDelegate& FHttpRequestImpl::OnRequestWillRetry()
 	return OnRequestWillRetryDelegate;
 }
 
-void FHttpRequestImpl::SetTimeout(float InTimeoutSecs)
-{
-	TimeoutSecs = InTimeoutSecs;
-}
-
-void FHttpRequestImpl::ClearTimeout()
-{
-	TimeoutSecs.Reset();
-}
-
-TOptional<float> FHttpRequestImpl::GetTimeout() const
-{
-	return TimeoutSecs;
-}
-
-float FHttpRequestImpl::GetTimeoutOrDefault() const
-{
-	return GetTimeout().Get(FHttpModule::Get().GetHttpTimeout());
-}
-
 void FHttpRequestImpl::BroadcastResponseHeadersReceived()
 {
 	if (OnHeaderReceived().IsBound())
diff --git a/Engine/Source/Runtime/Online/HTTP/Private/IHttpThreadedRequest.cpp b/Engine/Source/Runtime/Online/HTTP/Private/IHttpThreadedRequest.cpp
--- a/Engine/Source/Runtime/Online/HTTP/Private/IHttpThreadedRequest.cpp
+++ b/Engine/Source/Runtime/Online/HTTP/Private/IHttpThreadedRequest.cpp
@@ -4,35 +4,3 @@
 #include "HttpModule.h"
 #include "HttpManager.h"
 
-void IHttpThreadedRequest::FinishRequestNotInHttpManager()
-{
-	if (IsInGameThread())
-	{
-		if (DelegateThreadPolicy == EHttpRequestDelegateThreadPolicy::CompleteOnGameThread)
-		{
-			FinishRequest();
-		}
-		else
-		{
-			FHttpModule::Get().GetHttpManager().AddHttpThreadTask([StrongThis = StaticCastSharedRef<IHttpThreadedRequest>(AsShared())]()
-			{
-				StrongThis->FinishRequest();
-			});
-		}
-	}
-	else
-	{
-		if (DelegateThreadPolicy == EHttpRequestDelegateThreadPolicy::CompleteOnHttpThread)
-		{
-			FinishRequest();
-		}
-		else
-		{
-			FHttpModule::Get().GetHttpManager().AddGameThreadTask([StrongThis = StaticCastSharedRef<IHttpThreadedRequest>(AsShared())]()
-			{
-				StrongThis->FinishRequest();
-			});
-		}
-	}
-}
-
diff --git a/Engine/Source/Runtime/Online/HTTP/Private/IHttpThreadedRequest.h b/Engine/Source/Runtime/Online/HTTP/Private/IHttpThreadedRequest.h
--- a/Engine/Source/Runtime/Online/HTTP/Private/IHttpThreadedRequest.h
+++ b/Engine/Source/Runtime/Online/HTTP/Private/IHttpThreadedRequest.h
@@ -12,13 +12,4 @@ public:
 	virtual bool StartThreadedRequest() = 0;
 	virtual bool IsThreadedRequestComplete() = 0;
 	virtual void TickThreadedRequest(float DeltaSeconds) = 0;
-
-	// Can be called on game thread or http thread depend on the delegate thread policy
-	virtual void FinishRequest() = 0;
-
-protected:
-	/**
-	 * Finish the request when it's not in http manager
-	 */
-	void FinishRequestNotInHttpManager();
 };
diff --git a/Engine/Source/Runtime/Online/HTTP/Private/NullHttp.cpp b/Engine/Source/Runtime/Online/HTTP/Private/NullHttp.cpp
--- a/Engine/Source/Runtime/Online/HTTP/Private/NullHttp.cpp
+++ b/Engine/Source/Runtime/Online/HTTP/Private/NullHttp.cpp
@@ -170,7 +170,7 @@ void FNullHttpRequest::Tick(float DeltaSeconds)
 	if (CompletionStatus == EHttpRequestStatus::Processing)
 	{
 		ElapsedTime += DeltaSeconds;
-		const float HttpTimeout = GetTimeoutOrDefault();
+		const float HttpTimeout = GetTimeout().Get(FHttpModule::Get().GetHttpTimeout());
 		if (HttpTimeout > 0 && ElapsedTime >= HttpTimeout)
 		{
 			UE_LOG(LogHttp, Warning, TEXT("Timeout processing Http request. %p"),
@@ -206,6 +206,21 @@ EHttpRequestDelegateThreadPolicy FNullHttpRequest::GetDelegateThreadPolicy() con
 	return EHttpRequestDelegateThreadPolicy::CompleteOnGameThread;
 }
 
+void FNullHttpRequest::SetTimeout(float InTimeoutSecs) 
+{
+	TimeoutSecs = InTimeoutSecs;
+}
+
+void FNullHttpRequest::ClearTimeout() 
+{
+	TimeoutSecs.Reset();
+}
+
+TOptional<float> FNullHttpRequest::GetTimeout() const 
+{ 
+	return TimeoutSecs; 
+}
+
 // FNullHttpResponse
 
 FString FNullHttpResponse::GetURL() const
@@ -252,5 +267,3 @@ FString FNullHttpResponse::GetContentAsString() const
 {
 	return FString();
 }
-
-
diff --git a/Engine/Source/Runtime/Online/HTTP/Private/NullHttp.h b/Engine/Source/Runtime/Online/HTTP/Private/NullHttp.h
--- a/Engine/Source/Runtime/Online/HTTP/Private/NullHttp.h
+++ b/Engine/Source/Runtime/Online/HTTP/Private/NullHttp.h
@@ -42,6 +42,9 @@ public:
 	virtual float GetElapsedTime() const override;
 	virtual void SetDelegateThreadPolicy(EHttpRequestDelegateThreadPolicy InThreadPolicy) override;
 	virtual EHttpRequestDelegateThreadPolicy GetDelegateThreadPolicy() const override;
+	virtual void SetTimeout(float InTimeoutSecs) override;
+	virtual void ClearTimeout() override;
+	virtual TOptional<float> GetTimeout() const override;
 
 	FNullHttpRequest()
 		: CompletionStatus(EHttpRequestStatus::NotStarted)
@@ -60,6 +63,7 @@ private:
 	EHttpFailureReason FailureReason;
 	TMap<FString, FString> Headers;
 	float ElapsedTime;
+	TOptional<float> TimeoutSecs;
 };
 
 /**
diff --git a/Engine/Source/Runtime/Online/HTTP/Public/GenericPlatform/HttpRequestCommon.h b/Engine/Source/Runtime/Online/HTTP/Public/GenericPlatform/HttpRequestCommon.h
--- a/Engine/Source/Runtime/Online/HTTP/Public/GenericPlatform/HttpRequestCommon.h
+++ b/Engine/Source/Runtime/Online/HTTP/Public/GenericPlatform/HttpRequestCommon.h
@@ -19,14 +19,31 @@ public:
 	HTTP_API virtual void SetDelegateThreadPolicy(EHttpRequestDelegateThreadPolicy InDelegateThreadPolicy) override;
 	HTTP_API virtual EHttpRequestDelegateThreadPolicy GetDelegateThreadPolicy() const override;
 
+	HTTP_API virtual void SetTimeout(float InTimeoutSecs) override;
+	HTTP_API virtual void ClearTimeout() override;
+	HTTP_API virtual TOptional<float> GetTimeout() const override;
+	HTTP_API float GetTimeoutOrDefault() const;
+
+	// Can be called on game thread or http thread depend on the delegate thread policy
+	HTTP_API virtual void FinishRequest() = 0;
+
 protected:
 	/**
 	 * Check if this request is valid or allowed, before actually process the request
 	 */
+	HTTP_API bool PreProcess();
+	HTTP_API virtual bool SetupRequest() = 0;
 	HTTP_API bool PreCheck() const;
+	HTTP_API void ClearInCaseOfRetry();
+
 	HTTP_API void SetStatus(EHttpRequestStatus::Type InCompletionStatus);
 	HTTP_API void SetFailureReason(EHttpFailureReason InFailureReason);
 
+	/**
+	 * Finish the request when it's not in http manager
+	 */
+	HTTP_API void FinishRequestNotInHttpManager();
+
 protected:
 	/** Current status of request being processed */
 	EHttpRequestStatus::Type CompletionStatus = EHttpRequestStatus::NotStarted;
@@ -36,4 +53,10 @@ protected:
 
 	/** Thread policy about which thread to complete this request */
 	EHttpRequestDelegateThreadPolicy DelegateThreadPolicy = EHttpRequestDelegateThreadPolicy::CompleteOnGameThread;
+
+	/** Timeout in seconds for the entire HTTP request to complete */
+	TOptional<float> TimeoutSecs;
+
+	/** Record when this request started */
+	double RequestStartTimeAbsoluteSeconds;
 };
diff --git a/Engine/Source/Runtime/Online/HTTP/Public/GenericPlatform/HttpRequestImpl.h b/Engine/Source/Runtime/Online/HTTP/Public/GenericPlatform/HttpRequestImpl.h
--- a/Engine/Source/Runtime/Online/HTTP/Public/GenericPlatform/HttpRequestImpl.h
+++ b/Engine/Source/Runtime/Online/HTTP/Public/GenericPlatform/HttpRequestImpl.h
@@ -18,12 +18,6 @@ public:
 	HTTP_API virtual FHttpRequestHeaderReceivedDelegate& OnHeaderReceived() override;
 	HTTP_API virtual FHttpRequestWillRetryDelegate& OnRequestWillRetry() override;
 
-	HTTP_API virtual void SetTimeout(float InTimeoutSecs) override;
-	HTTP_API virtual void ClearTimeout() override;
-	HTTP_API virtual TOptional<float> GetTimeout() const override;
-
-	HTTP_API float GetTimeoutOrDefault() const;
-
 protected:
 	/** 
 	 * Broadcast all of our response's headers as having been received
@@ -46,7 +40,4 @@ protected:
 	
 	/** Delegate that will get called when request will be retried */
 	FHttpRequestWillRetryDelegate OnRequestWillRetryDelegate;
-
-	/** Timeout in seconds for the entire HTTP request to complete */
-	TOptional<float> TimeoutSecs;
 };
