PR #10520: Fix issues with non blocking sockets in SocketBSD preventing http server from serving large files on linux

#tests Reproduced and fixed it in WebTests project by a test case.

diff --git a/Engine/Source/Programs/WebTests/Private/TestHttp.cpp b/Engine/Source/Programs/WebTests/Private/TestHttp.cpp
--- a/Engine/Source/Programs/WebTests/Private/TestHttp.cpp
+++ b/Engine/Source/Programs/WebTests/Private/TestHttp.cpp
@@ -9,6 +9,10 @@
 #include "HttpManager.h"
 #include "HttpRetrySystem.h"
 #include "Http.h"
+#include "HttpPath.h"
+#include "HttpRouteHandle.h"
+#include "IHttpRouter.h"
+#include "HttpServerModule.h"
 #include "Misc/CommandLine.h"
 #include "TestHarness.h"
 #include "Serialization/JsonSerializerMacros.h"
@@ -175,7 +179,7 @@ public:
 
 	void WaitUntilAllHttpRequestsComplete()
 	{
-		while (OngoingRequests != 0 || (bRetryEnabled && !HttpRetryManager->IsEmpty()) )
+		while (HasOngoingRequest())
 		{
 			HttpModule->GetHttpManager().Tick(TickFrequency);
 			FPlatformProcess::Sleep(TickFrequency);
@@ -186,6 +190,11 @@ public:
 		HttpModule->GetHttpManager().Tick(TickFrequency);
 	}
 
+	bool HasOngoingRequest() const
+	{
+		return OngoingRequests != 0 || (bRetryEnabled && !HttpRetryManager->IsEmpty());
+	}
+
 	TSharedRef<IHttpRequest> CreateRequest()
 	{
 		return bRetryEnabled ? HttpRetryManager->CreateRequest() : HttpModule->CreateRequest();
@@ -1089,7 +1098,74 @@ TEST_CASE_METHOD(FWaitUntilCompleteHttpFixture, "Scheme besides http and https c
 	HttpRequest->ProcessRequest();
 }
 
-#endif
+class FLocalHttpServerFixture : public FWaitUntilCompleteHttpFixture
+{
+public:
+	FLocalHttpServerFixture()
+	{
+		HttpServerModule = new FHttpServerModule();
+		IModuleInterface* Module = HttpServerModule;
+		Module->StartupModule();
+
+		HttpRouter = HttpServerModule->GetHttpRouter(LocalHttpServerPort);
+		CHECK(HttpRouter.IsValid());
+	}
+
+	void StartServerWithHandler(const FHttpPath& HttpPath, EHttpServerRequestVerbs Verb, FHttpRequestHandler RequestHandler)
+	{
+		CHECK(HttpRouteHandle == nullptr);
+		HttpRouteHandle = HttpRouter->BindRoute(HttpPath, Verb, RequestHandler);
+		HttpServerModule->StartAllListeners();
+	}
+
+	~FLocalHttpServerFixture()
+	{
+		while (HasOngoingRequest())
+		{
+			HttpServerModule->Tick(TickFrequency);
+			HttpModule->GetHttpManager().Tick(TickFrequency);
+			FPlatformProcess::Sleep(TickFrequency);
+		}
+
+		HttpRouter->UnbindRoute(HttpRouteHandle);
+		HttpRouter.Reset();
+
+		IModuleInterface* Module = HttpServerModule;
+		Module->ShutdownModule();
+		delete Module;
+	}
+
+	TSharedPtr<IHttpRouter> HttpRouter;
+	FHttpRouteHandle HttpRouteHandle;
+	FHttpServerModule* HttpServerModule = nullptr;
+	uint32 LocalHttpServerPort = 9000;
+};
+
+TEST_CASE_METHOD(FLocalHttpServerFixture, "Local http server can serve large file", HTTP_TAG)
+{
+	const uint32 FileSize = 100 * 1024 * 1024; // 100 MB seems good enough to repro SE_EWOULDBLOCK or SE_TRY_AGAIN on Mac
+	StartServerWithHandler(FHttpPath(TEXT("/large_file")), EHttpServerRequestVerbs::VERB_GET, FHttpRequestHandler::CreateLambda([FileSize](const FHttpServerRequest& Request, const FHttpResultCallback& OnComplete) {
+		TArray<uint8> ResultData;
+		ResultData.SetNum(FileSize);
+		FMemory::Memset(ResultData.GetData(), 'd', FileSize);
+
+		OnComplete(FHttpServerResponse::Create(MoveTemp(ResultData), TEXT("text/text")));
+		return true;
+	}));
+
+	// Start client request
+	TSharedRef<IHttpRequest> HttpRequest = CreateRequest();
+	HttpRequest->SetURL(TEXT("http://localhost:9000/large_file"));
+	HttpRequest->SetVerb(TEXT("GET"));
+	HttpRequest->OnProcessRequestComplete().BindLambda([FileSize](FHttpRequestPtr HttpRequest, FHttpResponsePtr HttpResponse, bool bSucceeded) {
+		CHECK(bSucceeded);
+		REQUIRE(HttpResponse != nullptr);
+		CHECK(HttpResponse->GetContentLength() == FileSize);
+	});
+	HttpRequest->ProcessRequest();
+}
+
+#endif // (PLATFORM_WINDOWS && !WITH_CURL_XCURL) || PLATFORM_MAC || PLATFORM_UNIX
 
 // TODO: Add cancel test, with multiple cancel calls
 
diff --git a/Engine/Source/Programs/WebTests/WebTests.Build.cs b/Engine/Source/Programs/WebTests/WebTests.Build.cs
--- a/Engine/Source/Programs/WebTests/WebTests.Build.cs
+++ b/Engine/Source/Programs/WebTests/WebTests.Build.cs
@@ -11,6 +11,7 @@ public class WebTests : TestModuleRules
 				"ApplicationCore",
 				"Core",
 				"HTTP",
+				"HTTPServer",
 				"WebSockets",
 				"Json"
 			});
diff --git a/Engine/Source/Runtime/Online/HTTPServer/Private/HttpConnectionResponseWriteContext.cpp b/Engine/Source/Runtime/Online/HTTPServer/Private/HttpConnectionResponseWriteContext.cpp
--- a/Engine/Source/Runtime/Online/HTTPServer/Private/HttpConnectionResponseWriteContext.cpp
+++ b/Engine/Source/Runtime/Online/HTTPServer/Private/HttpConnectionResponseWriteContext.cpp
@@ -5,6 +5,7 @@
 #include "HttpServerHttpVersion.h"
 #include "HttpServerConstantsPrivate.h"
 #include "Sockets.h"
+#include "SocketSubsystem.h"
 
 DEFINE_LOG_CATEGORY(LogHttpConnectionResponseWriteContext);
 
@@ -79,8 +80,17 @@ bool FHttpConnectionResponseWriteContext::WriteBytes(const uint8* Bytes, int32 B
 	bool bWriteSuccess = Socket->Send(Bytes, BytesLen, OutBytesWritten);
 	if (!bWriteSuccess)
 	{
-		UE_LOG(LogHttpConnectionResponseWriteContext, Warning,
-			TEXT("WriteBytes sent %d/%d bytes"), OutBytesWritten, BytesLen);
+		ISocketSubsystem* SocketSubsystem = ISocketSubsystem::Get(PLATFORM_SOCKETSUBSYSTEM);
+		const ESocketErrors Err = SocketSubsystem->GetLastErrorCode();
+		if (Err == SE_EWOULDBLOCK || Err == SE_TRY_AGAIN)
+		{
+			bWriteSuccess = true;
+			OutBytesWritten = 0;
+		}
+		else
+		{
+			UE_LOG(LogHttpConnectionResponseWriteContext, Warning, TEXT("WriteBytes sent %d/%d bytes"), OutBytesWritten, BytesLen);
+		}
 	}
 
 	if (OutBytesWritten > 0)
diff --git a/Engine/Source/Runtime/Online/HTTPServer/Private/HttpListener.cpp b/Engine/Source/Runtime/Online/HTTPServer/Private/HttpListener.cpp
--- a/Engine/Source/Runtime/Online/HTTPServer/Private/HttpListener.cpp
+++ b/Engine/Source/Runtime/Online/HTTPServer/Private/HttpListener.cpp
@@ -108,7 +108,7 @@ bool FHttpListener::StartListening()
 	NewSocket->SetSendBufferSize(Config.BufferSize, ActualBufferSize);
 	if (ActualBufferSize < Config.BufferSize)
 	{
-		UE_LOG(LogHttpListener, Warning, 
+		UE_LOG(LogHttpListener, Log, 
 			TEXT("HttpListener unable to set desired buffer size (%d): Limited to %d"),
 			Config.BufferSize, ActualBufferSize);
 	}
diff --git a/Engine/Source/Runtime/Online/HTTPServer/Private/HttpServerModule.cpp b/Engine/Source/Runtime/Online/HTTPServer/Private/HttpServerModule.cpp
--- a/Engine/Source/Runtime/Online/HTTPServer/Private/HttpServerModule.cpp
+++ b/Engine/Source/Runtime/Online/HTTPServer/Private/HttpServerModule.cpp
@@ -5,6 +5,16 @@
 #include "Modules/ModuleManager.h"
 #include "Stats/Stats.h"
 
+class FHttpServerModuleImpl
+{
+public:
+	/** The association of port bindings and respective HTTP listeners */
+	TMap<uint32, TUniquePtr<FHttpListener>> Listeners;
+
+	/** Whether listeners can be started */
+	bool bHttpListenersEnabled = false;
+};
+
 DEFINE_LOG_CATEGORY(LogHttpServerModule);
 
 // FHttpServerModule 
@@ -12,6 +22,17 @@ IMPLEMENT_MODULE(FHttpServerModule, HTTPServer);
 
 FHttpServerModule* FHttpServerModule::Singleton = nullptr;
 
+FHttpServerModule::FHttpServerModule()
+	: Impl(new FHttpServerModuleImpl)
+{
+}
+
+FHttpServerModule::~FHttpServerModule()
+{
+	delete Impl;
+	Impl = nullptr;
+}
+
 void FHttpServerModule::StartupModule()
 {
 	Singleton = this;
@@ -23,19 +44,19 @@ void FHttpServerModule::ShutdownModule()
 	StopAllListeners();
 
 	// destroy all listeners
-	Listeners.Empty();
+	Impl->Listeners.Empty();
 
 	Singleton = nullptr;
 }
 
 void FHttpServerModule::StartAllListeners()
 {
-	bHttpListenersEnabled = true;
+	Impl->bHttpListenersEnabled = true;
 
 	UE_LOG(LogHttpServerModule, Log,
 		TEXT("Starting all listeners..."));
 
-	for (const auto& Listener : Listeners)
+	for (const auto& Listener : Impl->Listeners)
 	{
 		if (!Listener.Value->IsListening())
 		{
@@ -48,14 +69,12 @@ void FHttpServerModule::StartAllListeners()
 
 void FHttpServerModule::StopAllListeners()
 {
-	bHttpListenersEnabled = false;
-
 	UE_LOG(LogHttpServerModule, Log,
 		TEXT("Stopping all listeners..."));
 
-	bHttpListenersEnabled = false;
+	Impl->bHttpListenersEnabled = false;
 
-	for (const auto& Listener : Listeners)
+	for (const auto& Listener : Impl->Listeners)
 	{
 		if (Listener.Value->IsListening())
 		{
@@ -70,7 +89,7 @@ void FHttpServerModule::StopAllListeners()
 
 bool FHttpServerModule::HasPendingListeners() const
 {
-	for (const auto& Listener : Listeners)
+	for (const auto& Listener : Impl->Listeners)
 	{
 		if (Listener.Value->HasPendingConnections())
 		{
@@ -103,10 +122,10 @@ TSharedPtr<IHttpRouter> FHttpServerModule::GetHttpRouter(uint32 Port, bool bFail
 	bool bFailedToListen = false;
 
 	// We may already have a listener for this port
-	TUniquePtr<FHttpListener>* ExistingListener = Listeners.Find(Port);
+	TUniquePtr<FHttpListener>* ExistingListener = Impl->Listeners.Find(Port);
 	if (ExistingListener)
 	{
-		if (bHttpListenersEnabled)
+		if (Impl->bHttpListenersEnabled)
 		{
 			// if listeners are enabled, the existing listener for this port
 			// should always be listening (IsListening() will only be true
@@ -121,7 +140,7 @@ TSharedPtr<IHttpRouter> FHttpServerModule::GetHttpRouter(uint32 Port, bool bFail
 				// get rid of it and create a new one for now
 				UE_LOG(LogHttpServerModule, Error, TEXT("[%s] the existing listener for port %d is not listening/bound and listeners are still enabled"),
 					ANSI_TO_TCHAR(__FUNCTION__), Port);
-				Listeners.Remove(Port);
+				Impl->Listeners.Remove(Port);
 			}
 		}
 		else
@@ -136,7 +155,7 @@ TSharedPtr<IHttpRouter> FHttpServerModule::GetHttpRouter(uint32 Port, bool bFail
 	TUniquePtr<FHttpListener> NewListener = MakeUnique<FHttpListener>(Port);
 
 	// Try to start this listener now
-	if (bHttpListenersEnabled)
+	if (Impl->bHttpListenersEnabled)
 	{
 		if (!NewListener->StartListening())
 		{
@@ -155,7 +174,7 @@ TSharedPtr<IHttpRouter> FHttpServerModule::GetHttpRouter(uint32 Port, bool bFail
 	else
 	{
 		// the legacy behavior returns the router regardless of listener success
-		const auto& NewListenerRef = Listeners.Add(Port, MoveTemp(NewListener));
+		const auto& NewListenerRef = Impl->Listeners.Add(Port, MoveTemp(NewListener));
 		return NewListenerRef->GetRouter();
 	}
 }
@@ -165,9 +184,9 @@ bool FHttpServerModule::Tick(float DeltaTime)
 	QUICK_SCOPE_CYCLE_COUNTER(STAT_FHttpServerModule_Tick);
 	check(Singleton == this);
 
-	if (bHttpListenersEnabled)
+	if (Impl->bHttpListenersEnabled)
 	{
-		for (const auto& Listener : Listeners)
+		for (const auto& Listener : Impl->Listeners)
 		{
 			Listener.Value->Tick(DeltaTime);
 		}
diff --git a/Engine/Source/Runtime/Online/HTTPServer/Public/HttpServerModule.h b/Engine/Source/Runtime/Online/HTTPServer/Public/HttpServerModule.h
--- a/Engine/Source/Runtime/Online/HTTPServer/Public/HttpServerModule.h
+++ b/Engine/Source/Runtime/Online/HTTPServer/Public/HttpServerModule.h
@@ -14,6 +14,7 @@
 
 class FHttpListener;
 class IHttpRouter;
+class FHttpServerModuleImpl;
 
 DECLARE_LOG_CATEGORY_EXTERN(LogHttpServerModule, Log, All);
 
@@ -26,6 +27,8 @@ class FHttpServerModule :
 {
 
 public:
+	FHttpServerModule();
+	~FHttpServerModule();
 
 	/**
 	 * Checks to see if this module is loaded and ready.  It is only valid to call Get() if IsAvailable() returns true.
@@ -97,10 +100,6 @@ private:
 	/** Singleton Instance */
 	static FHttpServerModule* Singleton;
 
-	/** The association of port bindings and respective HTTP listeners */
-	TMap<uint32, TUniquePtr<FHttpListener>> Listeners;
-
-	/** Whether listeners can be started */
-	bool bHttpListenersEnabled = false;
+	FHttpServerModuleImpl* Impl;
 };
 
