Move some http request start logs and http disable option from curl http impl into http request common.

diff --git a/Engine/Source/Runtime/Online/HTTP/Private/Curl/CurlHttp.cpp b/Engine/Source/Runtime/Online/HTTP/Private/Curl/CurlHttp.cpp
--- a/Engine/Source/Runtime/Online/HTTP/Private/Curl/CurlHttp.cpp
+++ b/Engine/Source/Runtime/Online/HTTP/Private/Curl/CurlHttp.cpp
@@ -742,12 +742,6 @@ bool FCurlHttpRequest::SetupRequest()
 	QUICK_SCOPE_CYCLE_COUNTER(STAT_FCurlHttpRequest_SetupRequest);
 	check(EasyHandle);
 
-	// Disabled http request processing
-	if (!FHttpModule::Get().IsHttpEnabled())
-	{
-		UE_LOG(LogHttp, Verbose, TEXT("Http disabled. Skipping request. url=%s"), *GetURL());
-		return false;
-	}
 	if ((GetVerb().IsEmpty() || GetVerb().Equals(TEXT("GET"), ESearchCase::IgnoreCase))
 		&& (RequestPayload.IsValid() && RequestPayload->GetContentLength() > 0))
 	{
@@ -768,14 +762,6 @@ bool FCurlHttpRequest::SetupRequest()
 	CurlAddToMultiResult = CURLM_OK;
 	LastReportedBytesSent = 0;
 
-	// default no verb to a GET
-	if (Verb.IsEmpty())
-	{
-		Verb = TEXT("GET");
-	}
-
-	UE_LOG(LogHttp, Verbose, TEXT("%p: URL='%s'"), this, *URL);
-	UE_LOG(LogHttp, Verbose, TEXT("%p: Verb='%s'"), this, *Verb);
 	UE_LOG(LogHttp, Verbose, TEXT("%p: Custom headers are %s"), this, Headers.Num() ? TEXT("present") : TEXT("NOT present"));
 	UE_LOG(LogHttp, Verbose, TEXT("%p: Payload size=%llu"), this, RequestPayload->GetContentLength());
 
diff --git a/Engine/Source/Runtime/Online/HTTP/Private/GenericPlatform/HttpRequestCommon.cpp b/Engine/Source/Runtime/Online/HTTP/Private/GenericPlatform/HttpRequestCommon.cpp
--- a/Engine/Source/Runtime/Online/HTTP/Private/GenericPlatform/HttpRequestCommon.cpp
+++ b/Engine/Source/Runtime/Online/HTTP/Private/GenericPlatform/HttpRequestCommon.cpp
@@ -27,6 +27,13 @@ EHttpFailureReason FHttpRequestCommon::GetFailureReason() const
 
 bool FHttpRequestCommon::PreCheck() const
 {
+	// Disabled http request processing
+	if (!FHttpModule::Get().IsHttpEnabled())
+	{
+		UE_LOG(LogHttp, Verbose, TEXT("Http disabled. Skipping request. url=%s"), *GetURL());
+		return false;
+	}
+
 	// Prevent overlapped requests using the same instance
 	if (CompletionStatus == EHttpRequestStatus::Processing)
 	{
@@ -66,6 +73,8 @@ bool FHttpRequestCommon::PreProcess()
 		return false;
 	}
 
+	UE_LOG(LogHttp, Verbose, TEXT("%p: Verb='%s' URL='%s'"), this, *GetVerb(), *GetURL());
+
 	return true;
 }
 
