Add analytics stats for http.

diff --git a/Engine/Source/Runtime/Online/HTTP/Private/Curl/CurlHttp.cpp b/Engine/Source/Runtime/Online/HTTP/Private/Curl/CurlHttp.cpp
--- a/Engine/Source/Runtime/Online/HTTP/Private/Curl/CurlHttp.cpp
+++ b/Engine/Source/Runtime/Online/HTTP/Private/Curl/CurlHttp.cpp
@@ -5,6 +5,7 @@
 #include "Curl/CurlHttp.h"
 #include "Stats/Stats.h"
 #include "Misc/App.h"
+#include "HAL/PlatformTime.h"
 #include "HttpModule.h"
 #include "Http.h"
 #include "Misc/EngineVersion.h"
@@ -729,6 +730,16 @@ size_t FCurlHttpRequest::DebugCallback(CURL * Handle, curl_infotype DebugInfoTyp
 		case CURLINFO_SSL_DATA_OUT:
 			TimeSinceLastResponse = 0.0f;
 			bAnyHttpActivity = true;
+#if WITH_CURL_XCURL
+			// Unlike libCurl, currently there is an issue in xCurl that it triggers CURLINFO_HEADER_OUT even if can't 
+			// connect. Had to disable this code, make sure not to treat that event as connected
+			if (ConnectTime < 0 && DebugInfoType != CURLINFO_HEADER_OUT)
+			{
+				ConnectTime = FPlatformTime::Seconds() - StartProcessTime;
+			}
+#else
+			curl_easy_getinfo(EasyHandle, CURLINFO_CONNECT_TIME, &ConnectTime);
+#endif
 			break;
 		default:
 			break;
@@ -986,6 +997,10 @@ bool FCurlHttpRequest::ProcessRequest()
 	// Add to global list while being processed so that the ref counted request does not get deleted
 	FHttpModule::Get().GetHttpManager().AddThreadedRequest(SharedThis(this));
 
+#if WITH_CURL_XCURL
+	StartProcessTime = FPlatformTime::Seconds();
+#endif
+
 	UE_LOG(LogHttp, Verbose, TEXT("%p: request (easy handle:%p) has been added to threaded queue for processing"), this, EasyHandle);
 	return true;
 }
@@ -1241,11 +1256,7 @@ void FCurlHttpRequest::FinishRequest()
 			}
 		}
 
-		// Mark last request attempt as completed successfully
-		SetStatus(EHttpRequestStatus::Succeeded);
-
-		// Call delegate with valid request/response objects
-		OnProcessRequestComplete().ExecuteIfBound(SharedThis(this),Response,true);
+		HandleRequestSucceed(Response);
 	}
 	else
 	{
diff --git a/Engine/Source/Runtime/Online/HTTP/Private/GenericPlatform/HttpRequestCommon.cpp b/Engine/Source/Runtime/Online/HTTP/Private/GenericPlatform/HttpRequestCommon.cpp
--- a/Engine/Source/Runtime/Online/HTTP/Private/GenericPlatform/HttpRequestCommon.cpp
+++ b/Engine/Source/Runtime/Online/HTTP/Private/GenericPlatform/HttpRequestCommon.cpp
@@ -127,6 +127,13 @@ EHttpRequestDelegateThreadPolicy FHttpRequestCommon::GetDelegateThreadPolicy() c
 	return DelegateThreadPolicy; 
 }
 
+void FHttpRequestCommon::HandleRequestSucceed(TSharedPtr<IHttpResponse> Response)
+{
+	SetStatus(EHttpRequestStatus::Succeeded);
+	OnProcessRequestComplete().ExecuteIfBound(SharedThis(this), Response, true);
+	FHttpModule::Get().GetHttpManager().RecordStatTimeToConnect(ConnectTime);
+}
+
 void FHttpRequestCommon::SetStatus(EHttpRequestStatus::Type InCompletionStatus)
 {
 	CompletionStatus = InCompletionStatus;
@@ -169,4 +176,3 @@ float FHttpRequestCommon::GetTimeoutOrDefault() const
 {
 	return GetTimeout().Get(FHttpModule::Get().GetHttpTimeout());
 }
-
diff --git a/Engine/Source/Runtime/Online/HTTP/Private/HttpManager.cpp b/Engine/Source/Runtime/Online/HTTP/Private/HttpManager.cpp
--- a/Engine/Source/Runtime/Online/HTTP/Private/HttpManager.cpp
+++ b/Engine/Source/Runtime/Online/HTTP/Private/HttpManager.cpp
@@ -552,3 +552,13 @@ FHttpThreadBase* FHttpManager::GetThread()
 {
 	return Thread;
 }
+
+void FHttpManager::RecordStatTimeToConnect(float Duration)
+{
+	HttpStats.MaxTimeToConnect = FGenericPlatformMath::Max(Duration, HttpStats.MaxTimeToConnect);
+}
+
+void FHttpManager::RecordStatRequestsInQueue(uint32 RequestsInQueue)
+{
+	HttpStats.MaxRequestsInQueue = FGenericPlatformMath::Max(RequestsInQueue, HttpStats.MaxRequestsInQueue);
+}
diff --git a/Engine/Source/Runtime/Online/HTTP/Private/HttpThread.cpp b/Engine/Source/Runtime/Online/HTTP/Private/HttpThread.cpp
--- a/Engine/Source/Runtime/Online/HTTP/Private/HttpThread.cpp
+++ b/Engine/Source/Runtime/Online/HTTP/Private/HttpThread.cpp
@@ -250,6 +250,11 @@ void FHttpThreadBase::Process(TArray<IHttpThreadedRequest*>& RequestsToCancel, T
 		}
 	}
 
+	if (!RateLimitedThreadedRequests.IsEmpty())
+	{
+		FHttpModule::Get().GetHttpManager().RecordStatRequestsInQueue(RateLimitedThreadedRequests.Num());
+	}
+
 	{
 		SCOPE_CYCLE_COUNTER(STAT_HTTPThread_HttpThreadTick);
 
diff --git a/Engine/Source/Runtime/Online/HTTP/Public/GenericPlatform/HttpRequestCommon.h b/Engine/Source/Runtime/Online/HTTP/Public/GenericPlatform/HttpRequestCommon.h
--- a/Engine/Source/Runtime/Online/HTTP/Public/GenericPlatform/HttpRequestCommon.h
+++ b/Engine/Source/Runtime/Online/HTTP/Public/GenericPlatform/HttpRequestCommon.h
@@ -44,6 +44,8 @@ protected:
 	 */
 	HTTP_API void FinishRequestNotInHttpManager();
 
+	HTTP_API void HandleRequestSucceed(TSharedPtr<IHttpResponse> Response);
+
 protected:
 	/** Current status of request being processed */
 	EHttpRequestStatus::Type CompletionStatus = EHttpRequestStatus::NotStarted;
@@ -59,4 +61,10 @@ protected:
 
 	/** Record when this request started */
 	double RequestStartTimeAbsoluteSeconds;
+
+	/** Record when the request start to process */
+	double StartProcessTime = 0.0;
+
+	/** Record how long it take to connect to the endpoint */
+	double ConnectTime = -1.0;
 };
diff --git a/Engine/Source/Runtime/Online/HTTP/Public/HttpManager.h b/Engine/Source/Runtime/Online/HTTP/Public/HttpManager.h
--- a/Engine/Source/Runtime/Online/HTTP/Public/HttpManager.h
+++ b/Engine/Source/Runtime/Online/HTTP/Public/HttpManager.h
@@ -49,6 +49,19 @@ DECLARE_DELEGATE_OneParam(FHttpManagerRequestAddedDelegate, const FHttpRequestRe
  */
 DECLARE_DELEGATE_OneParam(FHttpManagerRequestCompletedDelegate, const FHttpRequestRef& /*Request*/);
 
+struct FHttpStats
+{
+	/** The max time to successfully connect the backend */
+	float MaxTimeToConnect = -1.0f;
+	/** The max waiting queue in http manager */
+	uint32 MaxRequestsInQueue = 0;
+
+	bool operator==(const FHttpStats& Other) const
+	{
+		return MaxRequestsInQueue == Other.MaxRequestsInQueue && FMath::IsNearlyEqual(MaxTimeToConnect, Other.MaxTimeToConnect);
+	}
+};
+
 /**
  * Manages Http request that are currently being processed
  */
@@ -244,6 +257,8 @@ public:
 	 */
 	void SetURLRequestFilter(const UE::Core::FURLRequestFilter& InURLRequestFilter) { URLRequestFilter = InURLRequestFilter; }
 
+	FHttpStats GetHttpStats() const { return HttpStats; }
+
 protected:
 	/** 
 	 * Create HTTP thread object
@@ -306,6 +321,8 @@ protected:
 
 	TMap<EHttpFlushReason, FHttpFlushTimeLimit> FlushTimeLimitsMap;
 
+	FHttpStats HttpStats;
+
 PACKAGE_SCOPE:
 
 	/** Used to lock access to add/remove/find requests */
@@ -321,4 +338,10 @@ PACKAGE_SCOPE:
 	 * Access http thread of http manager for internal usage
 	 */
 	HTTP_API FHttpThreadBase* GetThread();
+
+	/** Record the time to connect, to have a general idea how long the client usually take to connect for success requests, to adjust the connection timeout */
+	HTTP_API void RecordStatTimeToConnect(float Duration);
+
+	/** Record the requests waiting in queue, to have an idea if there are too many requests or if request number limit is too small */
+	HTTP_API void RecordStatRequestsInQueue(uint32 RequestsInQueue);
 };
