[HTTP]
* Removed always-on logging for when http request succeeds but returns code 503. This is a normal behavior from HordeServer when it has no resources available.. and since we keep asking the hordeserver for resources every x seconds it becomes a terrible spam in the log

diff --git a/Engine/Source/Runtime/Online/HTTP/Private/Curl/CurlHttp.cpp b/Engine/Source/Runtime/Online/HTTP/Private/Curl/CurlHttp.cpp
--- a/Engine/Source/Runtime/Online/HTTP/Private/Curl/CurlHttp.cpp
+++ b/Engine/Source/Runtime/Online/HTTP/Private/Curl/CurlHttp.cpp
@@ -1221,7 +1221,7 @@ void FCurlHttpRequest::FinishRequest()
 	if (Response.IsValid() &&
 		Response->bSucceeded)
 	{
-		const bool bDebugServerResponse = Response->GetResponseCode() >= 500 && Response->GetResponseCode() <= 503;
+		const bool bDebugServerResponse = false;// Response->GetResponseCode() >= 500 && Response->GetResponseCode() <= 503;
 
 		// log info about error responses to identify failed downloads
 		if (UE_LOG_ACTIVE(LogHttp, Verbose) ||
