Make sure the Tick method in game thread don't access Response in CurlHttp impl.

#localization none
#tests passed automation test project WebTests, tried launch the pc version

diff --git a/Engine/Source/Runtime/Online/HTTP/Private/Curl/CurlHttp.cpp b/Engine/Source/Runtime/Online/HTTP/Private/Curl/CurlHttp.cpp
--- a/Engine/Source/Runtime/Online/HTTP/Private/Curl/CurlHttp.cpp
+++ b/Engine/Source/Runtime/Online/HTTP/Private/Curl/CurlHttp.cpp
@@ -76,6 +76,7 @@ FCurlHttpRequest::FCurlHttpRequest()
 	, bAnyHttpActivity(false)
 	, BytesSent(0)
 	, TotalBytesSent(0)
+	, TotalBytesRead(0)
 	, LastReportedBytesRead(0)
 	, LastReportedBytesSent(0)
 	, LeastRecentlyCachedInfoMessageIndex(0)
@@ -497,13 +498,26 @@ size_t FCurlHttpRequest::ReceiveResponseHeaderCallback(void* Ptr, size_t SizeInB
 					Response->ContentLength = FCString::Atoi64(*HeaderValue);
 				}
 
+				const constexpr FStringView Seperator(TEXTVIEW(", "));
+
+				FString NewValue;
+				FString* PreviousValue = Response->Headers.Find(HeaderKey);
+				if (PreviousValue != nullptr && !PreviousValue->IsEmpty())
+				{
+					NewValue = MoveTemp(*PreviousValue);
+					NewValue.Reserve(NewValue.Len() + Seperator.Len() + HeaderValue.Len());
+					NewValue += Seperator;
+				}
+				NewValue += HeaderValue;
+				Response->Headers.Add(HeaderKey, MoveTemp(NewValue));
+
 				if (DelegateThreadPolicy == EHttpRequestDelegateThreadPolicy::CompleteOnHttpThread)
 				{
-					BroadcastNewlyReceivedHeader(HeaderKey, HeaderValue);
+					OnHeaderReceived().ExecuteIfBound(SharedThis(this), HeaderKey, HeaderValue);
 				}
-				else
+				else if (OnHeaderReceived().IsBound())
 				{
-					Response->NewlyReceivedHeaders.Enqueue(TPair<FString, FString>(MoveTemp(HeaderKey), MoveTemp(HeaderValue)));
+					NewlyReceivedHeaders.Enqueue(TPair<FString, FString>(MoveTemp(HeaderKey), MoveTemp(HeaderValue)));
 				}
 			}
 		}
@@ -544,9 +558,9 @@ size_t FCurlHttpRequest::ReceiveResponseBodyCallback(void* Ptr, size_t SizeInBlo
 	  
 	uint64 SizeToDownload = SizeInBlocks * BlockSizeInBytes;
 
-	UE_LOG(LogHttp, Verbose, TEXT("%p: ReceiveResponseBodyCallback: %llu bytes out of %llu received. (SizeInBlocks=%llu, BlockSizeInBytes=%llu, Response->TotalBytesRead=%llu, Response->GetContentLength()=%llu, SizeToDownload=%llu (<-this will get returned from the callback))"),
-		this, Response->TotalBytesRead.load() + SizeToDownload, Response->GetContentLength(),
-		SizeInBlocks, BlockSizeInBytes, Response->TotalBytesRead.load(), Response->GetContentLength(), SizeToDownload);
+	UE_LOG(LogHttp, Verbose, TEXT("%p: ReceiveResponseBodyCallback: %llu bytes out of %llu received. (SizeInBlocks=%llu, BlockSizeInBytes=%llu, TotalBytesRead=%llu, Response->GetContentLength()=%llu, SizeToDownload=%llu (<-this will get returned from the callback))"),
+		this, TotalBytesRead.load() + SizeToDownload, Response->GetContentLength(),
+		SizeInBlocks, BlockSizeInBytes, TotalBytesRead.load(), Response->GetContentLength(), SizeToDownload);
 
 	// note that we can be passed 0 bytes if file transmitted has 0 length
 	if (SizeToDownload == 0)
@@ -566,12 +580,12 @@ size_t FCurlHttpRequest::ReceiveResponseBodyCallback(void* Ptr, size_t SizeInBlo
 	else
 	{
 		Response->Payload.AddUninitialized(SizeToDownload);
-		FMemory::Memcpy(static_cast<uint8*>(Response->Payload.GetData()) + Response->TotalBytesRead.load(), Ptr, SizeToDownload);
+		FMemory::Memcpy(static_cast<uint8*>(Response->Payload.GetData()) + TotalBytesRead.load(), Ptr, SizeToDownload);
 
 		NumberOfBytesProcessed = SizeToDownload;
 	}
 
-	Response->TotalBytesRead += NumberOfBytesProcessed;
+	TotalBytesRead += NumberOfBytesProcessed;
 
 	return NumberOfBytesProcessed;
 }
@@ -1095,20 +1109,17 @@ void FCurlHttpRequest::Tick(float DeltaSeconds)
 void FCurlHttpRequest::CheckProgressDelegate()
 {
 	QUICK_SCOPE_CYCLE_COUNTER(STAT_FCurlHttpRequest_CheckProgressDelegate);
-	const uint64 CurrentBytesRead = Response.IsValid() ? Response->TotalBytesRead.load() : 0;
+	const uint64 CurrentBytesRead = TotalBytesRead.load();
 	const uint64 CurrentBytesSent = BytesSent.load();
 
 	const bool bProcessing = CompletionStatus == EHttpRequestStatus::Processing;
 	const bool bBytesSentChanged = (CurrentBytesSent != LastReportedBytesSent);
-	const bool bBytesReceivedChanged = (Response.IsValid() && CurrentBytesRead != LastReportedBytesRead);
+	const bool bBytesReceivedChanged = CurrentBytesRead != LastReportedBytesRead;
 	const bool bProgressChanged = bBytesSentChanged || bBytesReceivedChanged;
 	if (bProcessing && bProgressChanged)
 	{
 		LastReportedBytesSent = CurrentBytesSent;
-		if (Response.IsValid())
-		{
-			LastReportedBytesRead = CurrentBytesRead;
-		}
+		LastReportedBytesRead = CurrentBytesRead;
 PRAGMA_DISABLE_DEPRECATION_WARNINGS
 		// Update response progress
 		OnRequestProgress().ExecuteIfBound(SharedThis(this), LastReportedBytesSent, LastReportedBytesRead);
@@ -1119,35 +1130,12 @@ PRAGMA_ENABLE_DEPRECATION_WARNINGS
 
 void FCurlHttpRequest::BroadcastNewlyReceivedHeaders()
 {
-	if (Response.IsValid())
+	// Process the headers received on the HTTP thread and merge them into the response's list of headers and then broadcast the new headers
+	TPair<FString, FString> NewHeader;
+	while (NewlyReceivedHeaders.Dequeue(NewHeader))
 	{
-		// Process the headers received on the HTTP thread and merge them into the response's list of headers and then broadcast the new headers
-		TPair<FString, FString> NewHeader;
-		while (Response->NewlyReceivedHeaders.Dequeue(NewHeader))
-		{
-			BroadcastNewlyReceivedHeader(NewHeader.Key, NewHeader.Value);
-		}
-	}
-}
-
-void FCurlHttpRequest::BroadcastNewlyReceivedHeader(const FString& HeaderKey, const FString& HeaderValue)
-{
-	check(Response);
-
-	const constexpr FStringView Seperator(TEXTVIEW(", "));
-
-	FString NewValue;
-	FString* PreviousValue = Response->Headers.Find(HeaderKey);
-	if (PreviousValue != nullptr && !PreviousValue->IsEmpty())
-	{
-		NewValue = MoveTemp(*PreviousValue);
-		NewValue.Reserve(NewValue.Len() + Seperator.Len() + HeaderValue.Len());
-		NewValue += Seperator;
+		OnHeaderReceived().ExecuteIfBound(SharedThis(this), NewHeader.Key, NewHeader.Value);
 	}
-	NewValue += HeaderValue;
-	Response->Headers.Add(HeaderKey, MoveTemp(NewValue));
-
-	OnHeaderReceived().ExecuteIfBound(SharedThis(this), HeaderKey, HeaderValue);
 }
 
 void FCurlHttpRequest::MarkAsCompleted(CURLcode InCurlCompletionResult)
@@ -1193,14 +1181,14 @@ void FCurlHttpRequest::FinishRequest()
 				else
 				{
 					// If curl did not know how much we downloaded, or we were missing a Content-Length header (Chunked request), set our ContentLength as the amount we downloaded
-					Response->ContentLength = Response->TotalBytesRead;
+					Response->ContentLength = TotalBytesRead;
 				}
 			}
 
 			if (Response->HttpCode <= 0 && URL.StartsWith(TEXT("Http"), ESearchCase::IgnoreCase))
 			{
 				UE_LOG(LogHttp, Warning, TEXT("%p: invalid HTTP response code received. URL: %s, HTTP code: %d, content length: %llu, actual payload size: %llu"),
-					this, *GetURL(), Response->HttpCode, Response->ContentLength, Response->TotalBytesRead.load());
+					this, *GetURL(), Response->HttpCode, Response->ContentLength, TotalBytesRead.load());
 				Response->bSucceeded = false;
 			}
 		}
@@ -1230,12 +1218,12 @@ void FCurlHttpRequest::FinishRequest()
 			if (bDebugServerResponse)
 			{
 				UE_LOG(LogHttp, Warning, TEXT("%p: request has been successfully processed. URL: %s, HTTP code: %d, content length: %llu, actual payload size: %llu, elapsed: %.2fs"),
-					this, *GetURL(), Response->HttpCode, Response->ContentLength, Response->TotalBytesRead.load(), ElapsedTime);
+					this, *GetURL(), Response->HttpCode, Response->ContentLength, TotalBytesRead.load(), ElapsedTime);
 			}
 			else
 			{
 				UE_LOG(LogHttp, Log, TEXT("%p: request has been successfully processed. URL: %s, HTTP code: %d, content length: %llu, actual payload size: %llu, elapsed: %.2fs"),
-					this, *GetURL(), Response->HttpCode, Response->ContentLength, Response->TotalBytesRead.load(), ElapsedTime);
+					this, *GetURL(), Response->HttpCode, Response->ContentLength, TotalBytesRead.load(), ElapsedTime);
 			}
 
 			TArray<FString> AllHeaders = Response->GetAllHeaders();
@@ -1340,7 +1328,6 @@ float FCurlHttpRequest::GetElapsedTime() const
 
 FCurlHttpResponse::FCurlHttpResponse(const FCurlHttpRequest& InRequest)
 	: FHttpResponseCommon(InRequest)
-	, TotalBytesRead(0)
 	, HttpCode(EHttpResponseCodes::Unknown)
 	, ContentLength(0)
 	, bIsReady(0)
diff --git a/Engine/Source/Runtime/Online/HTTP/Private/Curl/CurlHttp.h b/Engine/Source/Runtime/Online/HTTP/Private/Curl/CurlHttp.h
--- a/Engine/Source/Runtime/Online/HTTP/Private/Curl/CurlHttp.h
+++ b/Engine/Source/Runtime/Online/HTTP/Private/Curl/CurlHttp.h
@@ -346,7 +346,6 @@ private:
 
 	/** Broadcast newly received headers */
 	void BroadcastNewlyReceivedHeaders();
-	void BroadcastNewlyReceivedHeader(const FString& HeaderKey, const FString& HeaderValue);
 
 	/** Combine a header's key/value in the format "Key: Value" */
 	static FString CombineHeaderKeyValue(const FString& HeaderKey, const FString& HeaderValue);
@@ -387,10 +386,14 @@ private:
 	float TimeSinceLastResponse;
 	/** Have we had any HTTP activity with the host? Sending headers, SSL handshake, etc */
 	bool bAnyHttpActivity;
+	/** Newly received headers we need to inform listeners about */
+	TSpscQueue<TPair<FString, FString>> NewlyReceivedHeaders;
 	/** Number of bytes sent already */
 	std::atomic<int64> BytesSent;
 	/** Total number of bytes sent already (includes data re-sent by seek attempts) */
 	std::atomic<int64> TotalBytesSent;
+	/** Caches how many bytes of the response we've read so far */
+	std::atomic<int64> TotalBytesRead;
 	/** Last bytes read reported to progress delegate */
 	uint64 LastReportedBytesRead;
 	/** Last bytes sent reported to progress delegate */
@@ -439,12 +442,8 @@ private:
 	TArray<uint8> Payload;
 	/** The stream to receive response body */
 	TSharedPtr<FArchive> ResponseBodyReceiveStream;
-	/** Caches how many bytes of the response we've read so far */
-	std::atomic<int64> TotalBytesRead;
 	/** Cached key/value header pairs. Parsed once request completes. Only accessible on the game thread. */
 	TMap<FString, FString> Headers;
-	/** Newly received headers we need to inform listeners about */
-	TSpscQueue<TPair<FString, FString>> NewlyReceivedHeaders;
 	/** Cached code from completed response */
 	int32 HttpCode;
 	/** Cached content length from completed response */
