Fix crash in CurlHttp
#rnx

diff --git a/Engine/Source/Runtime/Online/HTTP/Private/Curl/CurlHttp.cpp b/Engine/Source/Runtime/Online/HTTP/Private/Curl/CurlHttp.cpp
--- a/Engine/Source/Runtime/Online/HTTP/Private/Curl/CurlHttp.cpp
+++ b/Engine/Source/Runtime/Online/HTTP/Private/Curl/CurlHttp.cpp
@@ -467,6 +467,7 @@ size_t FCurlHttpRequest::ReceiveResponseHeaderCallback(void* Ptr, size_t SizeInB
 	if (!Response.IsValid())
 	{
 		Response = MakeShared<FCurlHttpResponse, ESPMode::ThreadSafe>(*this);
+		TotalBytesRead = 0;
 	}
 
 	TimeSinceLastResponse = 0.0f;
@@ -504,8 +505,8 @@ size_t FCurlHttpRequest::ReceiveResponseHeaderCallback(void* Ptr, size_t SizeInB
 				FString* PreviousValue = Response->Headers.Find(HeaderKey);
 				if (PreviousValue != nullptr && !PreviousValue->IsEmpty())
 				{
-					NewValue = MoveTemp(*PreviousValue);
-					NewValue.Reserve(NewValue.Len() + Seperator.Len() + HeaderValue.Len());
+					NewValue.Reserve(PreviousValue->Len() + Seperator.Len() + HeaderValue.Len());
+					NewValue = *PreviousValue;
 					NewValue += Seperator;
 				}
 				NewValue += HeaderValue;
@@ -580,6 +581,7 @@ size_t FCurlHttpRequest::ReceiveResponseBodyCallback(void* Ptr, size_t SizeInBlo
 	else
 	{
 		Response->Payload.AddUninitialized(SizeToDownload);
+		check(TotalBytesRead.load() + SizeToDownload <= Response->Payload.Num());
 		FMemory::Memcpy(static_cast<uint8*>(Response->Payload.GetData()) + TotalBytesRead.load(), Ptr, SizeToDownload);
 
 		NumberOfBytesProcessed = SizeToDownload;
@@ -989,6 +991,7 @@ bool FCurlHttpRequest::ProcessRequest()
 	// Clear out response. If this is a re-used request, Response could point to a stale response until SetupRequestHttpThread is called
 	Response = nullptr;
 	LastReportedBytesRead = 0;
+	TotalBytesRead = 0;
 
 	if (!PreProcess())
 	{
@@ -1315,6 +1318,7 @@ void FCurlHttpRequest::FinishRequest()
 
 		//Delegate needs to know about the errors -- so clear out Response (since connection failed) afterwards...
 		Response = nullptr;
+		TotalBytesRead = 0;
 	}
 }
 
