Add runtime args for enabling/disabling http event loop.

diff --git a/Engine/Source/Runtime/Online/HTTP/Private/Apple/AppleHttpManager.cpp b/Engine/Source/Runtime/Online/HTTP/Private/Apple/AppleHttpManager.cpp
--- a/Engine/Source/Runtime/Online/HTTP/Private/Apple/AppleHttpManager.cpp
+++ b/Engine/Source/Runtime/Online/HTTP/Private/Apple/AppleHttpManager.cpp
@@ -5,6 +5,7 @@
 #include "Apple/AppleEventLoopHttpThread.h"
 #include "HAL/IConsoleManager.h"
 #include "Math/UnrealMathUtility.h"
+#include "Misc/CommandLine.h"
 #include "HttpModule.h"
 #include "Http.h"
 
@@ -18,6 +19,11 @@ TAutoConsoleVariable<int32> CVarAppleEventLoopEnableChance(
 FHttpThreadBase* FAppleHttpManager::CreateHttpThread()
 {
 	bool bUseEventLoop = !(FMath::RandRange(0, 99) < CVarAppleEventLoopEnableChance.GetValueOnGameThread());
+
+	// Also support to change it through runtime args.
+	// Can't set cvar CVarCurlEventLoopEnableChance through runtime args or .ini files because http module initialized too early
+	FParse::Bool(FCommandLine::Get(), TEXT("useeventloop="), bUseEventLoop);
+
 	if (bUseEventLoop)
 	{
 		UE_LOG(LogHttp, Log, TEXT("CreateHttpThread using FAppleEventLoopHttpThread"));
diff --git a/Engine/Source/Runtime/Online/HTTP/Private/Curl/CurlHttpManager.cpp b/Engine/Source/Runtime/Online/HTTP/Private/Curl/CurlHttpManager.cpp
--- a/Engine/Source/Runtime/Online/HTTP/Private/Curl/CurlHttpManager.cpp
+++ b/Engine/Source/Runtime/Online/HTTP/Private/Curl/CurlHttpManager.cpp
@@ -453,6 +453,11 @@ void FCurlHttpManager::UpdateConfigs()
 FHttpThreadBase* FCurlHttpManager::CreateHttpThread()
 {
 	bool bUseEventLoop = (FMath::RandRange(0, 99) < CVarCurlEventLoopEnableChance.GetValueOnGameThread());
+
+	// Also support to change it through runtime args.
+	// Can't set cvar CVarCurlEventLoopEnableChance through runtime args or .ini files because http module initialized too early
+	FParse::Bool(FCommandLine::Get(), TEXT("useeventloop="), bUseEventLoop);
+
 	if (bUseEventLoop)
 	{
 #if WITH_CURL_MULTIPOLL
