Add status code received callback for http and implement it for all platforms.

#tests Tested through new test cases in WebTests project

diff --git a/Engine/Source/Programs/WebTests/Private/TestHttp.cpp b/Engine/Source/Programs/WebTests/Private/TestHttp.cpp
--- a/Engine/Source/Programs/WebTests/Private/TestHttp.cpp
+++ b/Engine/Source/Programs/WebTests/Private/TestHttp.cpp
@@ -276,7 +276,7 @@ TEST_CASE_METHOD(FWaitUntilCompleteHttpFixture, "Http request connect timeout",
 		CHECK(HttpRequest->GetStatus() == EHttpRequestStatus::Failed);
 		CHECK(HttpRequest->GetFailureReason() == EHttpFailureReason::ConnectionError);
 		const double DurationInSeconds  = FPlatformTime::Seconds() - StartTime;
-		double HttpTimeDiffTolerance = 1.0;
+		double HttpTimeDiffTolerance = 1.5;
 #if WITH_CURL_XCURL
 		HttpTimeDiffTolerance += 3.0; // It seems xCurl takes up to 3 more seconds for connect timeout
 #endif
@@ -968,6 +968,54 @@ TEST_CASE_METHOD(FValidateHeaderReceiveOrderFixture, "Http request header receiv
 	HttpRequest->ProcessRequest();
 }
 
+class FValidateStatusCodeReceiveOrderFixture : public FWaitUntilCompleteHttpFixture
+{
+public:
+	~FValidateStatusCodeReceiveOrderFixture()
+	{
+		WaitUntilAllHttpRequestsComplete();
+	}
+
+	std::atomic<bool> bStatusCodeReceived = false;
+	std::atomic<bool> bCompleteCallbackTriggered = false;
+};
+
+TEST_CASE_METHOD(FValidateStatusCodeReceiveOrderFixture, "Http request status code received callback will be called by thread policy", HTTP_TAG)
+{
+	TSharedRef<IHttpRequest> HttpRequest = CreateRequest();
+	HttpRequest->SetURL(UrlStreamDownload(20/*Chunks*/, 1024*1024/*ChunkSize*/));
+	HttpRequest->SetVerb(TEXT("GET"));
+
+	SECTION("in http thread")
+	{
+		HttpRequest->SetDelegateThreadPolicy(EHttpRequestDelegateThreadPolicy::CompleteOnHttpThread);
+		HttpRequest->OnStatusCodeReceived().BindLambda([this](FHttpRequestPtr Request, int32 StatusCode) {
+			CHECK(StatusCode == 200);
+			CHECK(!bCompleteCallbackTriggered);
+			CHECK(!IsInGameThread());
+			bStatusCodeReceived = true;
+		});
+	}
+	SECTION("in game thread")
+	{
+		HttpRequest->SetDelegateThreadPolicy(EHttpRequestDelegateThreadPolicy::CompleteOnGameThread);
+		HttpRequest->OnStatusCodeReceived().BindLambda([this](FHttpRequestPtr Request, int32 StatusCode) {
+			CHECK(StatusCode == 200);
+			CHECK(!bCompleteCallbackTriggered);
+			CHECK(IsInGameThread());
+			bStatusCodeReceived = true;
+		});
+	}
+
+	HttpRequest->OnProcessRequestComplete().BindLambda([this](FHttpRequestPtr /*HttpRequest*/, FHttpResponsePtr /*HttpResponse */, bool bSucceeded) {
+		CHECK(bStatusCodeReceived);
+		bCompleteCallbackTriggered = true;
+		CHECK(bSucceeded);
+	});
+
+	HttpRequest->ProcessRequest();
+}
+
 namespace UE
 {
 namespace TestHttp
diff --git a/Engine/Source/Runtime/Online/BuildPatchServices/Private/Tests/Mock/HttpRequest.mock.h b/Engine/Source/Runtime/Online/BuildPatchServices/Private/Tests/Mock/HttpRequest.mock.h
--- a/Engine/Source/Runtime/Online/BuildPatchServices/Private/Tests/Mock/HttpRequest.mock.h
+++ b/Engine/Source/Runtime/Online/BuildPatchServices/Private/Tests/Mock/HttpRequest.mock.h
@@ -161,6 +161,11 @@ namespace BuildPatchServices
 			return HttpRequestProgressDelegate64;
 		}
 
+		virtual FHttpRequestStatusCodeReceivedDelegate& OnStatusCodeReceived() override
+		{
+			return HttpStatusCodeReceivedDelegate;
+		}
+
 		virtual FHttpRequestHeaderReceivedDelegate& OnHeaderReceived() override
 		{
 			return HttpHeaderReceivedDelegate;
@@ -215,6 +220,7 @@ namespace BuildPatchServices
 		FHttpRequestProgressDelegate HttpRequestProgressDelegate;
 		FHttpRequestProgressDelegate64 HttpRequestProgressDelegate64;
 		FHttpRequestCompleteDelegate HttpRequestCompleteDelegate;
+		FHttpRequestStatusCodeReceivedDelegate HttpStatusCodeReceivedDelegate;
 		FHttpRequestHeaderReceivedDelegate HttpHeaderReceivedDelegate;
 		FHttpRequestWillRetryDelegate HttpRequestWillRetryDelegate;
 
diff --git a/Engine/Source/Runtime/Online/HTTP/Private/Apple/AppleHttp.cpp b/Engine/Source/Runtime/Online/HTTP/Private/Apple/AppleHttp.cpp
--- a/Engine/Source/Runtime/Online/HTTP/Private/Apple/AppleHttp.cpp
+++ b/Engine/Source/Runtime/Online/HTTP/Private/Apple/AppleHttp.cpp
@@ -44,6 +44,8 @@
 @property EHttpRequestStatus::Type RequestStatus;
 /** Reason of failure */
 @property EHttpFailureReason FailureReason;
+/** Associated request. Cleared when canceled */
+@property TWeakPtr<FAppleHttpRequest> SourceRequest;
 
 /** NSURLSessionDataDelegate delegate methods. Those are called from a thread controlled by the NSURLSession */
 
@@ -65,8 +67,9 @@
 @synthesize FailureReason;
 @synthesize BytesWritten;
 @synthesize BytesReceived;
+@synthesize SourceRequest;
 
-- (FAppleHttpResponseDelegate*)initWithResponseStream:(TSharedPtr<FArchive>)ResponseStream
+- (FAppleHttpResponseDelegate*)initWithRequest:(FAppleHttpRequest&) Request
 {
 	self = [super init];
 	
@@ -74,14 +77,16 @@
 	BytesReceived = 0;
 	RequestStatus = EHttpRequestStatus::NotStarted;
 	FailureReason = EHttpFailureReason::None;
-	ResponseBodyReceiveStream = ResponseStream;
-	bInitializedWithValidStream = (ResponseStream != nullptr);
+	SourceRequest = StaticCastWeakPtr<FAppleHttpRequest>(TWeakPtr<IHttpRequest>(Request.AsShared()));
+	ResponseBodyReceiveStream = Request.GetResponseBodyReceiveStream();
+	bInitializedWithValidStream = (ResponseBodyReceiveStream != nullptr);
 	
 	return self;
 }
 
-- (void)ClearResponseStream
+- (void)CleanSharedObjects
 {
+	self.SourceRequest = {};
 	if (bInitializedWithValidStream)
 	{
 	    FScopeLock Lock(&ResponseStreamLock);
@@ -95,6 +100,14 @@
 	[super dealloc];
 }
 
+- (void) HandleStatusCodeReceived:(int32) StatusCode
+{
+	if (TSharedPtr<FAppleHttpRequest> Request = SourceRequest.Pin())
+	{
+		Request->HandleStatusCodeReceived(StatusCode);
+	}
+}
+
 - (void)URLSession:(NSURLSession *)session task:(NSURLSessionTask *)task didSendBodyData:(int64_t)bytesSent totalBytesSent:(int64_t)totalBytesSent totalBytesExpectedToSend:(int64_t)totalBytesExpectedToSend
 {
 	UE_LOG(LogHttp, Verbose, TEXT("URLSession:task:didSendBodyData:totalBytesSent:totalBytesExpectedToSend: totalBytesSent = %lld, totalBytesSent = %lld: %p"), totalBytesSent, totalBytesExpectedToSend, self);
@@ -106,6 +119,8 @@
 	UE_LOG(LogHttp, Verbose, TEXT("URLSession:dataTask:didReceiveResponse:completionHandler"));
 	
 	self.Response = (NSHTTPURLResponse*)response;
+	int32 StatusCode = [self.Response statusCode];
+	[self HandleStatusCodeReceived: StatusCode];
 	uint64 ExpectedResponseLength = response.expectedContentLength;
 	if(!bInitializedWithValidStream && ExpectedResponseLength != NSURLResponseUnknownLength)
 	{
@@ -400,6 +415,16 @@ FAppleHttpRequest::~FAppleHttpRequest()
     [Session release];
 }
 
+const TSharedPtr<FArchive> FAppleHttpRequest::GetResponseBodyReceiveStream() const
+{
+	return ResponseBodyReceiveStream;
+}
+
+void FAppleHttpRequest::HandleStatusCodeReceived(int32 StatusCode)
+{
+	TriggerStatusCodeReceivedDelegate(StatusCode);
+}
+
 FString FAppleHttpRequest::GetURL() const
 {
 	SCOPED_AUTORELEASE_POOL;
@@ -823,11 +848,11 @@ void FAppleHttpRequest::TickThreadedRequest(float DeltaSeconds)
  * FAppleHttpResponse implementation
  **************************************************************************/
 
-FAppleHttpResponse::FAppleHttpResponse(const FAppleHttpRequest& InRequest)
+FAppleHttpResponse::FAppleHttpResponse(FAppleHttpRequest& InRequest)
 	: FHttpResponseCommon(InRequest)
 {
 	UE_LOG(LogHttp, Verbose, TEXT("FAppleHttpResponse::FAppleHttpResponse()"));
-	ResponseDelegate = [[FAppleHttpResponseDelegate alloc] initWithResponseStream: InRequest.ResponseBodyReceiveStream];
+	ResponseDelegate = [[FAppleHttpResponseDelegate alloc] initWithRequest: InRequest];
 }
 
 FAppleHttpResponse::~FAppleHttpResponse()
@@ -845,7 +870,7 @@ void FAppleHttpResponse::SetNewAppleHttpEventDelegate(FNewAppleHttpEventDelegate
 
 void FAppleHttpResponse::CleanSharedObjects()
 {
-	[ResponseDelegate ClearResponseStream];
+	[ResponseDelegate CleanSharedObjects];
 }
 
 FString FAppleHttpResponse::GetHeader(const FString& HeaderName) const
diff --git a/Engine/Source/Runtime/Online/HTTP/Private/Apple/AppleHttp.h b/Engine/Source/Runtime/Online/HTTP/Private/Apple/AppleHttp.h
--- a/Engine/Source/Runtime/Online/HTTP/Private/Apple/AppleHttp.h
+++ b/Engine/Source/Runtime/Online/HTTP/Private/Apple/AppleHttp.h
@@ -71,6 +71,10 @@ public:
 	 */
 	virtual ~FAppleHttpRequest();
 
+	const TSharedPtr<FArchive> GetResponseBodyReceiveStream() const;
+
+	void HandleStatusCodeReceived(int32 StatusCode);
+
 private:
 	/**
 	 * Trigger the request progress delegate if progress has changed
@@ -176,7 +180,7 @@ public:
 	 *
 	 * @param InRequest - original request that created this response
 	 */
-	FAppleHttpResponse(const FAppleHttpRequest& InRequest);
+	FAppleHttpResponse(FAppleHttpRequest& InRequest);
 
 	/**
 	 * Destructor
diff --git a/Engine/Source/Runtime/Online/HTTP/Private/Curl/CurlHttp.cpp b/Engine/Source/Runtime/Online/HTTP/Private/Curl/CurlHttp.cpp
--- a/Engine/Source/Runtime/Online/HTTP/Private/Curl/CurlHttp.cpp
+++ b/Engine/Source/Runtime/Online/HTTP/Private/Curl/CurlHttp.cpp
@@ -528,6 +528,11 @@ size_t FCurlHttpRequest::ReceiveResponseHeaderCallback(void* Ptr, size_t SizeInB
 			if (CURLE_OK == curl_easy_getinfo(EasyHandle, CURLINFO_RESPONSE_CODE, &HttpCode))
 			{
 				bRedirected = (HttpCode >= 300 && HttpCode < 400);
+
+				if (!bRedirected)
+				{
+					TriggerStatusCodeReceivedDelegate(HttpCode);
+				}
 			}
 		}
 		return HeaderSize;
diff --git a/Engine/Source/Runtime/Online/HTTP/Private/GenericPlatform/HttpRequestCommon.cpp b/Engine/Source/Runtime/Online/HTTP/Private/GenericPlatform/HttpRequestCommon.cpp
--- a/Engine/Source/Runtime/Online/HTTP/Private/GenericPlatform/HttpRequestCommon.cpp
+++ b/Engine/Source/Runtime/Online/HTTP/Private/GenericPlatform/HttpRequestCommon.cpp
@@ -176,3 +176,18 @@ float FHttpRequestCommon::GetTimeoutOrDefault() const
 {
 	return GetTimeout().Get(FHttpModule::Get().GetHttpTimeout());
 }
+
+void FHttpRequestCommon::TriggerStatusCodeReceivedDelegate(int32 StatusCode)
+{
+	if (DelegateThreadPolicy == EHttpRequestDelegateThreadPolicy::CompleteOnHttpThread)
+	{
+		OnStatusCodeReceived().ExecuteIfBound(SharedThis(this), StatusCode);
+	}
+	else if (OnStatusCodeReceived().IsBound())
+	{
+		FHttpModule::Get().GetHttpManager().AddGameThreadTask([StrongThis = AsShared(), StatusCode]()
+		{
+			StrongThis->OnStatusCodeReceived().ExecuteIfBound(StrongThis, StatusCode);
+		});
+	}
+}
diff --git a/Engine/Source/Runtime/Online/HTTP/Private/GenericPlatform/HttpRequestImpl.cpp b/Engine/Source/Runtime/Online/HTTP/Private/GenericPlatform/HttpRequestImpl.cpp
--- a/Engine/Source/Runtime/Online/HTTP/Private/GenericPlatform/HttpRequestImpl.cpp
+++ b/Engine/Source/Runtime/Online/HTTP/Private/GenericPlatform/HttpRequestImpl.cpp
@@ -21,6 +21,11 @@ FHttpRequestProgressDelegate64& FHttpRequestImpl::OnRequestProgress64()
 	return RequestProgressDelegate64;
 }
 
+FHttpRequestStatusCodeReceivedDelegate& FHttpRequestImpl::OnStatusCodeReceived()
+{
+	return StatusCodeReceivedDelegate;
+}
+
 FHttpRequestHeaderReceivedDelegate& FHttpRequestImpl::OnHeaderReceived()
 {
 	return HeaderReceivedDelegate;
diff --git a/Engine/Source/Runtime/Online/HTTP/Private/HttpManager.cpp b/Engine/Source/Runtime/Online/HTTP/Private/HttpManager.cpp
--- a/Engine/Source/Runtime/Online/HTTP/Private/HttpManager.cpp
+++ b/Engine/Source/Runtime/Online/HTTP/Private/HttpManager.cpp
@@ -89,6 +89,7 @@ PRAGMA_DISABLE_DEPRECATION_WARNINGS
 		Request->OnRequestProgress().Unbind();
 PRAGMA_ENABLE_DEPRECATION_WARNINGS
 		Request->OnRequestProgress64().Unbind();
+		Request->OnStatusCodeReceived().Unbind();
 		Request->OnHeaderReceived().Unbind();
 
 		// Don't emit these tracking logs in commandlet runs. Build system traps warnings during cook, and these are not truly fatal, but useful for tracking down shutdown issues.
diff --git a/Engine/Source/Runtime/Online/HTTP/Private/HttpRetrySystem.cpp b/Engine/Source/Runtime/Online/HTTP/Private/HttpRetrySystem.cpp
--- a/Engine/Source/Runtime/Online/HTTP/Private/HttpRetrySystem.cpp
+++ b/Engine/Source/Runtime/Online/HTTP/Private/HttpRetrySystem.cpp
@@ -118,6 +118,7 @@ bool FHttpRetrySystem::FRequest::ProcessRequest()
 
 	HttpRequest->OnRequestProgress64().BindThreadSafeSP(RetryRequest, &FHttpRetrySystem::FRequest::HttpOnRequestProgress);
 	HttpRequest->OnProcessRequestComplete().BindThreadSafeSP(RetryRequest, &FHttpRetrySystem::FRequest::HttpOnProcessRequestComplete);
+	HttpRequest->OnStatusCodeReceived().BindThreadSafeSP(RetryRequest, &FHttpRetrySystem::FRequest::HttpOnStatusCodeReceived);
 	HttpRequest->OnHeaderReceived().BindThreadSafeSP(RetryRequest, &FHttpRetrySystem::FRequest::HttpOnHeaderReceived);
 
 	TSharedPtr<FManager> RetryManagerPtr = RetryManager.Pin();
@@ -253,6 +254,12 @@ void FHttpRetrySystem::FRequest::HttpOnProcessRequestComplete(FHttpRequestPtr In
 	OnProcessRequestComplete().ExecuteIfBound(SelfPtr, HttpResponse, bSucceeded);
 }
 
+void FHttpRetrySystem::FRequest::HttpOnStatusCodeReceived(FHttpRequestPtr Request, int32 StatusCode)
+{
+	TSharedRef<FRequest> SelfPtr = StaticCastSharedRef<FRequest>(AsShared());
+	OnStatusCodeReceived().ExecuteIfBound(SelfPtr, StatusCode);
+}
+
 void FHttpRetrySystem::FRequest::HttpOnHeaderReceived(FHttpRequestPtr Request, const FString& HeaderName, const FString& NewHeaderValue)
 {
 	if (!CVarHttpRetrySystemNonGameThreadSupportEnabled.GetValueOnAnyThread())
diff --git a/Engine/Source/Runtime/Online/HTTP/Public/GenericPlatform/HttpRequestCommon.h b/Engine/Source/Runtime/Online/HTTP/Public/GenericPlatform/HttpRequestCommon.h
--- a/Engine/Source/Runtime/Online/HTTP/Public/GenericPlatform/HttpRequestCommon.h
+++ b/Engine/Source/Runtime/Online/HTTP/Public/GenericPlatform/HttpRequestCommon.h
@@ -46,6 +46,8 @@ protected:
 
 	HTTP_API void HandleRequestSucceed(TSharedPtr<IHttpResponse> Response);
 
+	HTTP_API void TriggerStatusCodeReceivedDelegate(int32 StatusCode);
+
 protected:
 	/** Current status of request being processed */
 	EHttpRequestStatus::Type CompletionStatus = EHttpRequestStatus::NotStarted;
diff --git a/Engine/Source/Runtime/Online/HTTP/Public/GenericPlatform/HttpRequestImpl.h b/Engine/Source/Runtime/Online/HTTP/Public/GenericPlatform/HttpRequestImpl.h
--- a/Engine/Source/Runtime/Online/HTTP/Public/GenericPlatform/HttpRequestImpl.h
+++ b/Engine/Source/Runtime/Online/HTTP/Public/GenericPlatform/HttpRequestImpl.h
@@ -15,6 +15,7 @@ public:
 	HTTP_API virtual FHttpRequestCompleteDelegate& OnProcessRequestComplete() override;
 	HTTP_API virtual FHttpRequestProgressDelegate& OnRequestProgress() override;
 	HTTP_API virtual FHttpRequestProgressDelegate64& OnRequestProgress64() override;
+	HTTP_API virtual FHttpRequestStatusCodeReceivedDelegate& OnStatusCodeReceived() override;
 	HTTP_API virtual FHttpRequestHeaderReceivedDelegate& OnHeaderReceived() override;
 	HTTP_API virtual FHttpRequestWillRetryDelegate& OnRequestWillRetry() override;
 
@@ -35,6 +36,9 @@ protected:
 	/** Delegate that will get called once per tick with bytes downloaded so far */
 	FHttpRequestProgressDelegate64 RequestProgressDelegate64;
 
+	/** Delegate that will get called when status code received */
+	FHttpRequestStatusCodeReceivedDelegate StatusCodeReceivedDelegate;
+
 	/** Delegate that will get called for each new header received */
 	FHttpRequestHeaderReceivedDelegate HeaderReceivedDelegate;
 	
diff --git a/Engine/Source/Runtime/Online/HTTP/Public/HttpRetrySystem.h b/Engine/Source/Runtime/Online/HTTP/Public/HttpRetrySystem.h
--- a/Engine/Source/Runtime/Online/HTTP/Public/HttpRetrySystem.h
+++ b/Engine/Source/Runtime/Online/HTTP/Public/HttpRetrySystem.h
@@ -115,6 +115,7 @@ namespace FHttpRetrySystem
 
 		void HttpOnRequestProgress(FHttpRequestPtr InHttpRequest, uint64 BytesSent, uint64 BytesRcv);
 		void HttpOnProcessRequestComplete(FHttpRequestPtr HttpRequest, FHttpResponsePtr HttpResponse, bool bSucceeded);
+		void HttpOnStatusCodeReceived(FHttpRequestPtr Request, int32 StatusCode);
 		void HttpOnHeaderReceived(FHttpRequestPtr Request, const FString& HeaderName, const FString& NewHeaderValue);
 
 		/** Update our HTTP request's URL's domain from our RetryDomains */
diff --git a/Engine/Source/Runtime/Online/HTTP/Public/Interfaces/IHttpRequest.h b/Engine/Source/Runtime/Online/HTTP/Public/Interfaces/IHttpRequest.h
--- a/Engine/Source/Runtime/Online/HTTP/Public/Interfaces/IHttpRequest.h
+++ b/Engine/Source/Runtime/Online/HTTP/Public/Interfaces/IHttpRequest.h
@@ -35,6 +35,14 @@ typedef TSharedRef<IHttpResponse, ESPMode::ThreadSafe> FHttpResponseRef;
  */
 using FHttpRequestCompleteDelegate = TTSDelegate<void(FHttpRequestPtr /*Request*/, FHttpResponsePtr /*Response*/, bool /*bConnectedSuccessfully*/)>;
 
+/**
+ * Delegate called when an Http request receives status code
+ *
+ * @param Request original Http request that started things
+ * @param status code
+ */
+using FHttpRequestStatusCodeReceivedDelegate = TTSDelegate<void(FHttpRequestPtr /*Request*/, int32 /*StatusCode*/)>;
+
 /**
  * Delegate called when an Http request receives a header
  *
@@ -298,6 +306,11 @@ public:
 	 */
 	virtual FHttpRequestHeaderReceivedDelegate& OnHeaderReceived() = 0;
 
+	/** 
+	 * Delegate called to signal the receipt of a header.  See FHttpRequestStatusCodeReceivedDelegate
+	 */
+	virtual FHttpRequestStatusCodeReceivedDelegate& OnStatusCodeReceived() = 0;
+
 	/**
 	 * Called to cancel a request that is still being processed
 	 */
@@ -326,7 +339,7 @@ public:
 
 	/**
 	 * Set thread policy about which thread to trigger the delegates, set by FHttpManager::SetRequestCompletedDelegate, 
-	 * IHttpRequest::OnHeaderReceived, IHttpRequest::OnRequestProgress64 and IHttpRequest::OnProcessRequestComplete.
+	 * IHttpRequest::OnStatusCodeReceived, IHttpRequest::OnHeaderReceived, IHttpRequest::OnRequestProgress64 and IHttpRequest::OnProcessRequestComplete.
 	 *
 	 * Note that when set it as CompleteOnHttpThread, the thread to trigger delegates could be any thread 
 	 * depends on the implementation. User code should make the delegate thread-safe and shouldn't assume 
