Add editor version of "[HTTP.HttpThread]:RunningThreadedRequestLimit" configuration and increase value for editor to fix limit warning for UBA shader compilation.

#rnx

diff --git a/Engine/Config/BaseEditor.ini b/Engine/Config/BaseEditor.ini
--- a/Engine/Config/BaseEditor.ini
+++ b/Engine/Config/BaseEditor.ini
@@ -676,3 +676,6 @@ bUseWidgetTemplateSelector=True
 +WidgetClassesToHide=/Script/UMG.SpinBox
 +WidgetClassesToHide=/Script/UMG.Throbber
 +WidgetClassesToHide=/Script/UMG.TreeView
+
+[HTTP.HttpThread]
+RunningThreadedRequestLimitEditor=20
diff --git a/Engine/Source/Runtime/Online/HTTP/Private/HttpThread.cpp b/Engine/Source/Runtime/Online/HTTP/Private/HttpThread.cpp
--- a/Engine/Source/Runtime/Online/HTTP/Private/HttpThread.cpp
+++ b/Engine/Source/Runtime/Online/HTTP/Private/HttpThread.cpp
@@ -141,7 +141,14 @@ bool FHttpThreadBase::NeedsSingleThreadTick() const
 void FHttpThreadBase::UpdateConfigs()
 {
 	int32 LocalRunningThreadedRequestLimit = -1;
-	if (GConfig->GetInt(TEXT("HTTP.HttpThread"), TEXT("RunningThreadedRequestLimit"), LocalRunningThreadedRequestLimit, GEngineIni))
+	const bool bFoundLocalRunningThreadedRequestLimit =
+	(
+#if WITH_EDITOR
+		GConfig->GetInt(TEXT("HTTP.HttpThread"), TEXT("RunningThreadedRequestLimitEditor"), LocalRunningThreadedRequestLimit, GEditorIni) ||
+#endif
+		GConfig->GetInt(TEXT("HTTP.HttpThread"), TEXT("RunningThreadedRequestLimit"), LocalRunningThreadedRequestLimit, GEngineIni)
+	);
+	if (bFoundLocalRunningThreadedRequestLimit)
 	{
 		if (LocalRunningThreadedRequestLimit < 1)
 		{
