Make activity timeout can be turned off by hotfixing to 0;
Change default activity timeout to 30s.

diff --git a/Engine/Source/Runtime/Online/HTTP/Private/GenericPlatform/HttpRequestCommon.cpp b/Engine/Source/Runtime/Online/HTTP/Private/GenericPlatform/HttpRequestCommon.cpp
--- a/Engine/Source/Runtime/Online/HTTP/Private/GenericPlatform/HttpRequestCommon.cpp
+++ b/Engine/Source/Runtime/Online/HTTP/Private/GenericPlatform/HttpRequestCommon.cpp
@@ -241,7 +241,11 @@ void FHttpRequestCommon::StartActivityTimeoutTimer()
 	}
 
 	float HttpActivityTimeout = FHttpModule::Get().GetHttpActivityTimeout();
-	check(HttpActivityTimeout > 0);
+	if (HttpActivityTimeout == 0)
+	{
+		return;
+	}
+
 	StartActivityTimeoutTimerBy(HttpActivityTimeout);
 
 	ResetActivityTimeoutTimer(TEXTVIEW("Connected"));
@@ -292,6 +296,11 @@ void FHttpRequestCommon::ResetActivityTimeoutTimer(FStringView Reason)
 		return;
 	}
 
+	if (!ActivityTimeoutHttpTaskTimerHandle)
+	{
+		return;
+	}
+
 	ActivityTimeoutAt = FPlatformTime::Seconds() + FHttpModule::Get().GetHttpActivityTimeout();
 	UE_LOG(LogHttp, VeryVerbose, TEXT("Request [%p] reset response timeout timer at %s: %s"), this, *FDateTime::Now().ToString(TEXT("%H:%M:%S:%s")), Reason.GetData());
 }
@@ -305,11 +314,13 @@ void FHttpRequestCommon::StopActivityTimeoutTimer()
 		return;
 	}
 
-	if (ActivityTimeoutHttpTaskTimerHandle)
+	if (!ActivityTimeoutHttpTaskTimerHandle)
 	{
-		FHttpModule::Get().GetHttpManager().RemoveHttpThreadTask(ActivityTimeoutHttpTaskTimerHandle);
-		ActivityTimeoutHttpTaskTimerHandle.Reset();
+		return;
 	}
+
+	FHttpModule::Get().GetHttpManager().RemoveHttpThreadTask(ActivityTimeoutHttpTaskTimerHandle);
+	ActivityTimeoutHttpTaskTimerHandle.Reset();
 }
 
 void FHttpRequestCommon::StartTotalTimeoutTimer()
diff --git a/Engine/Source/Runtime/Online/HTTP/Private/HttpModule.cpp b/Engine/Source/Runtime/Online/HTTP/Private/HttpModule.cpp
--- a/Engine/Source/Runtime/Online/HTTP/Private/HttpModule.cpp
+++ b/Engine/Source/Runtime/Online/HTTP/Private/HttpModule.cpp
@@ -81,7 +81,7 @@ void FHttpModule::StartupModule()
 	MaxReadBufferSize = 256 * 1024;
 	HttpTotalTimeout = 0.0f;
 	HttpConnectionTimeout = 60.0f;
-	HttpActivityTimeout = 300.0f;
+	HttpActivityTimeout = 30.0f;
 	HttpReceiveTimeout = HttpConnectionTimeout;
 	HttpSendTimeout = HttpConnectionTimeout;
 	HttpMaxConnectionsPerServer = 16;
