Turn off the server response debug log by default;
Get rid of the http request cancelled warning log.

diff --git a/Engine/Source/Runtime/Online/HTTP/Private/Curl/CurlHttp.cpp b/Engine/Source/Runtime/Online/HTTP/Private/Curl/CurlHttp.cpp
--- a/Engine/Source/Runtime/Online/HTTP/Private/Curl/CurlHttp.cpp
+++ b/Engine/Source/Runtime/Online/HTTP/Private/Curl/CurlHttp.cpp
@@ -12,13 +12,21 @@
 #include "Curl/CurlHttpManager.h"
 #include "Misc/ScopeLock.h"
 #include "HAL/FileManager.h"
+#include "HAL/IConsoleManager.h"
 #include "Internationalization/Regex.h"
 
 #if WITH_SSL
 #include "Ssl.h"
-
 #include <openssl/ssl.h>
+#endif
 
+TAutoConsoleVariable<bool> CVarCurlDebugServerResponseEnabled(
+	TEXT("http.CurlDebugServerResponseEnabled"),
+	false,
+	TEXT("Enable debugging of server response")
+);
+
+#if WITH_SSL
 static int SslCertVerify(int PreverifyOk, X509_STORE_CTX* Context)
 {
 	if (PreverifyOk == 1)
@@ -1221,14 +1229,12 @@ void FCurlHttpRequest::FinishRequest()
 		Response->bIsReady = true;
 	}
 
-	if (Response.IsValid() &&
-		Response->bSucceeded)
+	if (Response.IsValid() && Response->bSucceeded)
 	{
-		const bool bDebugServerResponse = false;// Response->GetResponseCode() >= 500 && Response->GetResponseCode() <= 503;
+		bool bDebugServerResponse = CVarCurlDebugServerResponseEnabled.GetValueOnAnyThread() && (Response->GetResponseCode() >= 500 && Response->GetResponseCode() <= 503);
 
 		// log info about error responses to identify failed downloads
-		if (UE_LOG_ACTIVE(LogHttp, Verbose) ||
-			bDebugServerResponse)
+		if (UE_LOG_ACTIVE(LogHttp, Verbose) || bDebugServerResponse)
 		{
 			if (bDebugServerResponse)
 			{
@@ -1263,11 +1269,7 @@ void FCurlHttpRequest::FinishRequest()
 	}
 	else
 	{
-		if (bCanceled)
-		{
-			UE_LOG(LogHttp, Warning, TEXT("%p: request was cancelled"), this);
-		}
-		else if (CurlAddToMultiResult != CURLM_OK)
+		if (CurlAddToMultiResult != CURLM_OK)
 		{
 			UE_LOG(LogHttp, Warning, TEXT("%p: request failed, libcurl multi error: %d (%s)"), this, (int32)CurlAddToMultiResult, ANSI_TO_TCHAR(curl_multi_strerror(CurlAddToMultiResult)));
 		}
