Fix the http dead lock when Flush while AbortRequest;
Add the ProcessRequestUntilComplete sync call to IHttpRequest.

#tests Passed tests in WebTests.

diff --git a/Engine/Source/Programs/WebTests/Private/TestHttp.cpp b/Engine/Source/Programs/WebTests/Private/TestHttp.cpp
--- a/Engine/Source/Programs/WebTests/Private/TestHttp.cpp
+++ b/Engine/Source/Programs/WebTests/Private/TestHttp.cpp
@@ -274,7 +274,7 @@ TEST_CASE_METHOD(FWaitUntilCompleteHttpFixture, "Http Methods", HTTP_TAG)
 
 TEST_CASE_METHOD(FWaitUntilCompleteHttpFixture, "Can process https request", HTTP_TAG)
 {
-	TSharedRef<IHttpRequest> HttpRequest = HttpModule->CreateRequest();
+	TSharedRef<IHttpRequest> HttpRequest = CreateRequest();
 	HttpRequest->SetVerb(TEXT("GET"));
 	HttpRequest->SetURL(TEXT("https://www.unrealengine.com/"));
 	HttpRequest->OnProcessRequestComplete().BindLambda([](FHttpRequestPtr HttpRequest, FHttpResponsePtr HttpResponse, bool bSucceeded) {
@@ -284,6 +284,18 @@ TEST_CASE_METHOD(FWaitUntilCompleteHttpFixture, "Can process https request", HTT
 	HttpRequest->ProcessRequest();
 }
 
+TEST_CASE_METHOD(FWaitUntilCompleteHttpFixture, "Can do blocking call", HTTP_TAG)
+{
+	TSharedRef<IHttpRequest> HttpRequest = CreateRequest();
+	HttpRequest->SetURL(UrlToTestMethods());
+	HttpRequest->OnProcessRequestComplete().BindLambda([](FHttpRequestPtr HttpRequest, FHttpResponsePtr HttpResponse, bool bSucceeded) {
+		CHECK(bSucceeded);
+		REQUIRE(HttpResponse != nullptr);
+	});
+	HttpRequest->ProcessRequestUntilComplete();
+	CHECK(HttpRequest->GetStatus() == EHttpRequestStatus::Succeeded);
+}
+
 TEST_CASE_METHOD(FWaitUntilCompleteHttpFixture, "Get large response content without chunks", HTTP_TAG)
 {
 	TSharedRef<IHttpRequest> HttpRequest = CreateRequest();
@@ -1343,6 +1355,23 @@ TEST_CASE_METHOD(FThreadedBatchRequestsFixture, "Retry manager and http manager
 	BlockUntilFlushed();
 }
 
+TEST_CASE_METHOD(FWaitUntilCompleteHttpFixture, "Flush while activity timeout shouldn't dead lock", HTTP_TAG)
+{
+	DisableWarningsInThisTest();
+
+	HttpModule->HttpActivityTimeout = 2.0f;
+
+	TSharedPtr<IHttpRequest> HttpRequest = CreateRequest();
+	HttpRequest->SetURL(UrlStreamDownload(3/*Chunks*/, HTTP_TEST_TIMEOUT_CHUNK_SIZE, 5/*ChunkLatency*/));
+	HttpRequest->OnProcessRequestComplete().BindLambda([](FHttpRequestPtr HttpRequest, FHttpResponsePtr HttpResponse, bool bSucceeded) {
+		CHECK(HttpRequest->GetStatus() == EHttpRequestStatus::Failed);
+		CHECK(HttpRequest->GetFailureReason() == EHttpFailureReason::ConnectionError);
+	});
+	HttpRequest->ProcessRequest();
+
+	HttpModule->GetHttpManager().Flush(EHttpFlushReason::FullFlush);
+}
+
 #if (PLATFORM_WINDOWS && !WITH_CURL_XCURL) || PLATFORM_MAC || PLATFORM_UNIX
 
 TEST_CASE_METHOD(FWaitUntilCompleteHttpFixture, "Scheme besides http and https can work if allowed by settings", HTTP_TAG)
diff --git a/Engine/Source/Runtime/Online/BuildPatchServices/Private/Tests/Mock/HttpRequest.mock.h b/Engine/Source/Runtime/Online/BuildPatchServices/Private/Tests/Mock/HttpRequest.mock.h
--- a/Engine/Source/Runtime/Online/BuildPatchServices/Private/Tests/Mock/HttpRequest.mock.h
+++ b/Engine/Source/Runtime/Online/BuildPatchServices/Private/Tests/Mock/HttpRequest.mock.h
@@ -140,6 +140,10 @@ namespace BuildPatchServices
 			return TOptional<float>();
 		}
 
+		virtual void ProcessRequestUntilComplete() override
+		{
+		}
+
 		virtual bool ProcessRequest() override
 		{
 			++RxProcessRequest;
diff --git a/Engine/Source/Runtime/Online/HTTP/Private/Curl/CurlHttp.cpp b/Engine/Source/Runtime/Online/HTTP/Private/Curl/CurlHttp.cpp
--- a/Engine/Source/Runtime/Online/HTTP/Private/Curl/CurlHttp.cpp
+++ b/Engine/Source/Runtime/Online/HTTP/Private/Curl/CurlHttp.cpp
@@ -1107,6 +1107,11 @@ void FCurlHttpRequest::TickThreadedRequest(float DeltaSeconds)
 
 void FCurlHttpRequest::AbortRequest()
 {
+	if (bCurlRequestCompleted)
+	{
+		return;
+	}
+
 	FHttpManager& HttpManager = FHttpModule::Get().GetHttpManager();
 	if (HttpManager.IsValidRequest(this))
 	{
diff --git a/Engine/Source/Runtime/Online/HTTP/Private/GenericPlatform/GenericPlatformHttp.cpp b/Engine/Source/Runtime/Online/HTTP/Private/GenericPlatform/GenericPlatformHttp.cpp
--- a/Engine/Source/Runtime/Online/HTTP/Private/GenericPlatform/GenericPlatformHttp.cpp
+++ b/Engine/Source/Runtime/Online/HTTP/Private/GenericPlatform/GenericPlatformHttp.cpp
@@ -127,6 +127,7 @@ public:
 	virtual void SetTimeout(float InTimeoutSecs) override {}
 	virtual void ClearTimeout() override {}
 	virtual TOptional<float> GetTimeout() const override { return TOptional<float>(); }
+	virtual void ProcessRequestUntilComplete() override {}
 };
 
 FDefaultUserAgentBuilder::FDefaultUserAgentBuilder()
diff --git a/Engine/Source/Runtime/Online/HTTP/Private/GenericPlatform/HttpRequestCommon.cpp b/Engine/Source/Runtime/Online/HTTP/Private/GenericPlatform/HttpRequestCommon.cpp
--- a/Engine/Source/Runtime/Online/HTTP/Private/GenericPlatform/HttpRequestCommon.cpp
+++ b/Engine/Source/Runtime/Online/HTTP/Private/GenericPlatform/HttpRequestCommon.cpp
@@ -2,6 +2,7 @@
 
 #include "GenericPlatform/HttpRequestCommon.h"
 #include "GenericPlatform/HttpResponseCommon.h"
+#include "HAL/Event.h"
 #include "Http.h"
 #include "HttpManager.h"
 #include "Misc/CommandLine.h"
@@ -284,7 +285,7 @@ void FHttpRequestCommon::OnActivityTimeoutTimerTaskTrigger()
 	QUICK_SCOPE_CYCLE_COUNTER(STAT_FHttpRequestCommon_AbortRequest);
 	bActivityTimedOut = true;
 	AbortRequest();
-	UE_LOG(LogHttp, Warning, TEXT("Request [%s] timed out at [%s] because of no responding for %0.2f seconds"), *GetURL(), *FDateTime::Now().ToString(TEXT("%H:%M:%S:%s")), FHttpModule::Get().GetHttpActivityTimeout());
+	UE_LOG(LogHttp, Log, TEXT("Request [%s] timed out at [%s] because of no responding for %0.2f seconds"), *GetURL(), *FDateTime::Now().ToString(TEXT("%H:%M:%S:%s")), FHttpModule::Get().GetHttpActivityTimeout());
 }
 
 void FHttpRequestCommon::ResetActivityTimeoutTimer(FStringView Reason)
@@ -393,6 +394,21 @@ void FHttpRequestCommon::Shutdown()
 	StopTotalTimeoutTimer();
 }
 
+void FHttpRequestCommon::ProcessRequestUntilComplete()
+{
+	checkf(!OnProcessRequestComplete().IsBound(), TEXT("OnProcessRequestComplete is not supported for sync call"));
+
+	SetDelegateThreadPolicy(EHttpRequestDelegateThreadPolicy::CompleteOnHttpThread);
+
+	FEvent* Event = FPlatformProcess::GetSynchEventFromPool(true);
+	OnProcessRequestComplete().BindLambda([Event](FHttpRequestPtr HttpRequest, FHttpResponsePtr HttpResponse, bool bSucceeded) {
+		Event->Trigger();
+	});
+	ProcessRequest();
+	Event->Wait();
+	FPlatformProcess::ReturnSynchEventToPool(Event);
+}
+
 void FHttpRequestCommon::TriggerStatusCodeReceivedDelegate(int32 StatusCode)
 {
 	if (DelegateThreadPolicy == EHttpRequestDelegateThreadPolicy::CompleteOnHttpThread)
diff --git a/Engine/Source/Runtime/Online/HTTP/Private/HttpManager.cpp b/Engine/Source/Runtime/Online/HTTP/Private/HttpManager.cpp
--- a/Engine/Source/Runtime/Online/HTTP/Private/HttpManager.cpp
+++ b/Engine/Source/Runtime/Online/HTTP/Private/HttpManager.cpp
@@ -300,8 +300,6 @@ void FHttpManager::Flush(EHttpFlushReason FlushReason)
 {
 	QUICK_SCOPE_CYCLE_COUNTER(STAT_FHttpManager_Flush);
 
-	FScopeLock ScopeLock(&RequestLock);
-
 	checkf(FlushReason != EHttpFlushReason::Shutdown || !HasAnyBoundDelegate(), TEXT("Use Shutdown() instead of Flush(EHttpFlushReason::Shutdown) directly."));
 	
 	// This variable is set to indicate that flush is happening.
@@ -318,15 +316,23 @@ void FHttpManager::Flush(EHttpFlushReason FlushReason)
 
 	UE_CLOG(!IsRunningCommandlet(), LogHttp, Verbose, TEXT("[FHttpManager::Flush] FlushReason [%s] FlushTimeSoftLimitSeconds [%.3fs] FlushTimeHardLimitSeconds [%.3fs] SecondsToSleepForOutstandingThreadedRequests [%.3fs]"), LexToString(FlushReason), FlushTimeSoftLimitSeconds, FlushTimeHardLimitSeconds, SecondsToSleepForOutstandingThreadedRequests);
 
-	UE_CLOG(!IsRunningCommandlet() && Requests.Num(), LogHttp, Verbose, TEXT("[FHttpManager::Flush] Cleanup starts for %d outstanding Http Requests."), Requests.Num());
+	uint32 RequestsNum = 0;
+
+	{
+		FScopeLock ScopeLock(&RequestLock);
+		RequestsNum = Requests.Num();
+	}
+
+	UE_CLOG(!IsRunningCommandlet() && RequestsNum, LogHttp, Verbose, TEXT("[FHttpManager::Flush] Cleanup starts for %d outstanding Http Requests."), RequestsNum);
 
 	double BeginWaitTime = FPlatformTime::Seconds();
 	double LastFlushTickTime = BeginWaitTime;
 	double StallWarnTime = BeginWaitTime + 0.5;
 	double AppTime = FPlatformTime::Seconds();
 
+
 	// For a duration equal to FlushTimeHardLimitSeconds, we wait for ongoing http requests to complete
-	while (Requests.Num() > 0 && (FlushTimeHardLimitSeconds < 0 || (AppTime - BeginWaitTime < FlushTimeHardLimitSeconds)))
+	while (RequestsNum > 0 && (FlushTimeHardLimitSeconds < 0 || (AppTime - BeginWaitTime < FlushTimeHardLimitSeconds)))
 	{
 		SCOPED_ENTER_BACKGROUND_EVENT(STAT_FHttpManager_Flush_Iteration);
 
@@ -334,18 +340,22 @@ void FHttpManager::Flush(EHttpFlushReason FlushReason)
 		if (FlushTimeSoftLimitSeconds >= 0 && (AppTime - BeginWaitTime >= FlushTimeSoftLimitSeconds))
 		{
 			// Don't emit these tracking logs in commandlet runs. Build system traps warnings during cook, and these are not truly fatal, but useful for tracking down shutdown issues.
-			UE_CLOG(ShouldOutputHttpWarnings(), LogHttp, Warning, TEXT("[FHttpManager::Flush] FlushTimeSoftLimitSeconds [%.3fs] exceeded. Cancelling %d outstanding HTTP requests:"), FlushTimeSoftLimitSeconds, Requests.Num());
+			UE_CLOG(ShouldOutputHttpWarnings(), LogHttp, Warning, TEXT("[FHttpManager::Flush] FlushTimeSoftLimitSeconds [%.3fs] exceeded. Cancelling %d outstanding HTTP requests:"), FlushTimeSoftLimitSeconds, RequestsNum);
 
-			for (TArray<FHttpRequestRef>::TIterator It(Requests); It; ++It)
 			{
-				FHttpRequestRef& Request = *It;
+				FScopeLock ScopeLock(&RequestLock);
 
-				// Don't emit these tracking logs in commandlet runs. Build system traps warnings during cook, and these are not truly fatal, but useful for tracking down shutdown issues.
-				UE_CLOG(ShouldOutputHttpWarnings(), LogHttp, Warning, TEXT("	verb=[%s] url=[%s] refs=[%d] status=%s"), *Request->GetVerb(), *Request->GetURL(), Request.GetSharedReferenceCount(), EHttpRequestStatus::ToString(Request->GetStatus()));
+				for (TArray<FHttpRequestRef>::TIterator It(Requests); It; ++It)
+				{
+					FHttpRequestRef& Request = *It;
 
-				FScopedEnterBackgroundEvent(*Request->GetURL());
+					// Don't emit these tracking logs in commandlet runs. Build system traps warnings during cook, and these are not truly fatal, but useful for tracking down shutdown issues.
+					UE_CLOG(ShouldOutputHttpWarnings(), LogHttp, Warning, TEXT("	verb=[%s] url=[%s] refs=[%d] status=%s"), *Request->GetVerb(), *Request->GetURL(), Request.GetSharedReferenceCount(), EHttpRequestStatus::ToString(Request->GetStatus()));
+
+					FScopedEnterBackgroundEvent(*Request->GetURL());
 
-				Request->CancelRequest();
+					Request->CancelRequest();
+				}
 			}
 		}
 
@@ -353,8 +363,13 @@ void FHttpManager::Flush(EHttpFlushReason FlushReason)
 		FlushTick(AppTime - LastFlushTickTime);
 		LastFlushTickTime = AppTime;
 
+		{
+			FScopeLock ScopeLock(&RequestLock);
+			RequestsNum = Requests.Num();
+		}
+
 		// Process threaded Http Requests
-		if (Requests.Num() > 0)
+		if (RequestsNum > 0)
 		{
 			if (Thread)
 			{
@@ -363,7 +378,7 @@ void FHttpManager::Flush(EHttpFlushReason FlushReason)
 					if (AppTime >= StallWarnTime)
 					{
 						// Don't emit these tracking logs in commandlet runs. Build system traps warnings during cook, and these are not truly fatal, but useful for tracking down shutdown issues.
-						UE_CLOG(ShouldOutputHttpWarnings(), LogHttp, Warning, TEXT("	Ticking HTTPThread for %d outstanding Http requests."), Requests.Num());
+						UE_CLOG(ShouldOutputHttpWarnings(), LogHttp, Warning, TEXT("	Ticking HTTPThread for %d outstanding Http requests."), RequestsNum);
 						StallWarnTime = AppTime + 0.5;
 					}
 					Thread->Tick();
@@ -371,7 +386,7 @@ void FHttpManager::Flush(EHttpFlushReason FlushReason)
 				else
 				{
 					// Don't emit these tracking logs in commandlet runs. Build system traps warnings during cook, and these are not truly fatal, but useful for tracking down shutdown issues.
-					UE_CLOG(ShouldOutputHttpWarnings(), LogHttp, Warning, TEXT("	Sleeping %.3fs to wait for %d outstanding Http Requests."), SecondsToSleepForOutstandingThreadedRequests, Requests.Num());
+					UE_CLOG(ShouldOutputHttpWarnings(), LogHttp, Warning, TEXT("	Sleeping %.3fs to wait for %d outstanding Http Requests."), SecondsToSleepForOutstandingThreadedRequests, RequestsNum);
 					FPlatformProcess::Sleep(SecondsToSleepForOutstandingThreadedRequests);
 				}
 			}
@@ -384,13 +399,15 @@ void FHttpManager::Flush(EHttpFlushReason FlushReason)
 		AppTime = FPlatformTime::Seconds();
 	}
 
-	UE_CLOG(!IsRunningCommandlet(), LogHttp, Verbose, TEXT("[FHttpManager::Flush] Cleanup ended after %.3fs. %d outstanding Http Requests."), AppTime - BeginWaitTime, Requests.Num());
+	UE_CLOG(!IsRunningCommandlet(), LogHttp, Verbose, TEXT("[FHttpManager::Flush] Cleanup ended after %.3fs. %d outstanding Http Requests."), AppTime - BeginWaitTime, RequestsNum);
 
 	// Don't emit these tracking logs in commandlet runs. Build system traps warnings during cook, and these are not truly fatal, but useful for tracking down shutdown issues.
-	if (Requests.Num() > 0 && (FlushTimeHardLimitSeconds > 0 && (AppTime - BeginWaitTime > FlushTimeHardLimitSeconds)) && ShouldOutputHttpWarnings())
+	if (RequestsNum > 0 && (FlushTimeHardLimitSeconds > 0 && (AppTime - BeginWaitTime > FlushTimeHardLimitSeconds)) && ShouldOutputHttpWarnings())
 	{
 		UE_LOG(LogHttp, Warning, TEXT("[FHttpManager::Flush] FlushTimeHardLimitSeconds [%.3fs] exceeded. The following requests are being abandoned without being flushed:"), FlushTimeHardLimitSeconds);
 
+		FScopeLock ScopeLock(&RequestLock);
+
 		for (TArray<FHttpRequestRef>::TIterator It(Requests); It; ++It)
 		{
 			FHttpRequestRef& Request = *It;
diff --git a/Engine/Source/Runtime/Online/HTTP/Private/NullHttp.cpp b/Engine/Source/Runtime/Online/HTTP/Private/NullHttp.cpp
--- a/Engine/Source/Runtime/Online/HTTP/Private/NullHttp.cpp
+++ b/Engine/Source/Runtime/Online/HTTP/Private/NullHttp.cpp
@@ -219,6 +219,10 @@ TOptional<float> FNullHttpRequest::GetTimeout() const
 	return TimeoutSecs; 
 }
 
+void FNullHttpRequest::ProcessRequestUntilComplete()
+{
+}
+
 // FNullHttpResponse
 
 FString FNullHttpResponse::GetURL() const
diff --git a/Engine/Source/Runtime/Online/HTTP/Private/NullHttp.h b/Engine/Source/Runtime/Online/HTTP/Private/NullHttp.h
--- a/Engine/Source/Runtime/Online/HTTP/Private/NullHttp.h
+++ b/Engine/Source/Runtime/Online/HTTP/Private/NullHttp.h
@@ -45,6 +45,7 @@ public:
 	virtual void SetTimeout(float InTimeoutSecs) override;
 	virtual void ClearTimeout() override;
 	virtual TOptional<float> GetTimeout() const override;
+	virtual void ProcessRequestUntilComplete() override;
 
 	FNullHttpRequest()
 		: CompletionStatus(EHttpRequestStatus::NotStarted)
diff --git a/Engine/Source/Runtime/Online/HTTP/Public/GenericPlatform/HttpRequestCommon.h b/Engine/Source/Runtime/Online/HTTP/Public/GenericPlatform/HttpRequestCommon.h
--- a/Engine/Source/Runtime/Online/HTTP/Public/GenericPlatform/HttpRequestCommon.h
+++ b/Engine/Source/Runtime/Online/HTTP/Public/GenericPlatform/HttpRequestCommon.h
@@ -35,6 +35,8 @@ public:
 
 	HTTP_API virtual void Shutdown() override;
 
+	HTTP_API virtual void ProcessRequestUntilComplete() override;
+
 protected:
 	/**
 	 * Check if this request is valid or allowed, before actually process the request
diff --git a/Engine/Source/Runtime/Online/HTTP/Public/HttpRequestAdapter.h b/Engine/Source/Runtime/Online/HTTP/Public/HttpRequestAdapter.h
--- a/Engine/Source/Runtime/Online/HTTP/Public/HttpRequestAdapter.h
+++ b/Engine/Source/Runtime/Online/HTTP/Public/HttpRequestAdapter.h
@@ -39,6 +39,7 @@ public:
 	virtual void SetTimeout(float InTimeoutSecs) override { HttpRequest->SetTimeout(InTimeoutSecs); }
 	virtual void ClearTimeout() override { HttpRequest->ClearTimeout(); }
 	virtual TOptional<float> GetTimeout() const override { return HttpRequest->GetTimeout(); }
+	virtual void ProcessRequestUntilComplete() override { return HttpRequest->ProcessRequestUntilComplete(); }
 	virtual const FHttpResponsePtr GetResponse() const override { return HttpRequest->GetResponse(); }
 	virtual float GetElapsedTime() const override { return HttpRequest->GetElapsedTime(); }
 	virtual EHttpRequestStatus::Type GetStatus() const override { return HttpRequest->GetStatus(); }
diff --git a/Engine/Source/Runtime/Online/HTTP/Public/Interfaces/IHttpRequest.h b/Engine/Source/Runtime/Online/HTTP/Public/Interfaces/IHttpRequest.h
--- a/Engine/Source/Runtime/Online/HTTP/Public/Interfaces/IHttpRequest.h
+++ b/Engine/Source/Runtime/Online/HTTP/Public/Interfaces/IHttpRequest.h
@@ -356,6 +356,11 @@ public:
 	 */
 	virtual EHttpRequestDelegateThreadPolicy GetDelegateThreadPolicy() const = 0;
 
+	/**
+	 * Blocking call to wait the request until it's completed, without calling the callback
+	 */
+	virtual void ProcessRequestUntilComplete() = 0;
+
 	/** 
 	 * Destructor for overrides 
 	 */
