Add connection and activity timeout support on AppleHttp so WebTests are passing

#changelist validated
#virtualized

diff --git a/Engine/Source/Programs/WebTests/Private/TestHttp.cpp b/Engine/Source/Programs/WebTests/Private/TestHttp.cpp
--- a/Engine/Source/Programs/WebTests/Private/TestHttp.cpp
+++ b/Engine/Source/Programs/WebTests/Private/TestHttp.cpp
@@ -691,10 +691,11 @@ TEST_CASE_METHOD(FWaitUntilCompleteHttpFixture, "Http request activity timeout",
 		CHECK(HttpRequest->GetFailureReason() == EHttpFailureReason::ConnectionError);
 
 		const double DurationInSeconds  = FPlatformTime::Seconds() - StartTime;
-#if WITH_CURL_XCURL
-		// Unlike libCurl, currently there is an issue in xCurl that it triggers CURLINFO_HEADER_OUT even if can't 
+#if WITH_CURL_XCURL || PLATFORM_MAC || PLATFORM_IOS
+		// Unlike libCurl, currently there is an issue in xCurl that it triggers CURLINFO_HEADER_OUT even if can't
 		// connect. Had to disable that code, make sure not to treat that event as connected
-		// So it takes 5s to receive the first chunk to be considered as connected, then start response timer and 
+		// In a similar way on MacOS/iOS we don't get any notification until some data is received
+		// So it takes 5s to receive the first chunk to be considered as connected, then start response timer and
 		// take 3s to response timeout
 		CHECK(FMath::IsNearlyEqual(DurationInSeconds, ReceiveTimeoutSetting + 5, HTTP_TIME_DIFF_TOLERANCE));
 #else
diff --git a/Engine/Source/Runtime/Online/HTTP/Private/Apple/AppleHttp.cpp b/Engine/Source/Runtime/Online/HTTP/Private/Apple/AppleHttp.cpp
--- a/Engine/Source/Runtime/Online/HTTP/Private/Apple/AppleHttp.cpp
+++ b/Engine/Source/Runtime/Online/HTTP/Private/Apple/AppleHttp.cpp
@@ -23,13 +23,16 @@
 	 * from the delegates methods and through mehthods in FAppleHttpRequest. Once we cancel or have an error this can 
 	 * is nullified as soon as possible to avoid receiving more data after completion delegates were triggered */
 	TSharedPtr<FArchive> ResponseBodyReceiveStream;
-	
+
 	/** critical section to properly clear stream */
 	FCriticalSection ResponseStreamLock;
-
+	
 	/** flag meant to reduce locking on ResponseStreamLock*/
-	@public BOOL bInitializedWithValidStream;
-
+	BOOL bInitializedWithValidStream;
+	
+	/** Have we received any data? */
+	BOOL bAnyHttpActivity;
+	
 	/** Delegate invoked after processing URLSession:dataTask:didReceiveData or URLSession:task:didCompleteWithError:*/
 	@public FNewAppleHttpEventDelegate NewAppleHttpEventDelegate;
 }
@@ -47,6 +50,7 @@
 /** Associated request. Cleared when canceled */
 @property TWeakPtr<FAppleHttpRequest> SourceRequest;
 
+
 /** NSURLSessionDataDelegate delegate methods. Those are called from a thread controlled by the NSURLSession */
 
 /** Sent periodically to notify the delegate of upload progress. */
@@ -77,10 +81,11 @@
 	BytesReceived = 0;
 	RequestStatus = EHttpRequestStatus::NotStarted;
 	FailureReason = EHttpFailureReason::None;
+	bAnyHttpActivity = false;
 	SourceRequest = StaticCastWeakPtr<FAppleHttpRequest>(TWeakPtr<IHttpRequest>(Request.AsShared()));
 	ResponseBodyReceiveStream = Request.GetResponseBodyReceiveStream();
 	bInitializedWithValidStream = (ResponseBodyReceiveStream != nullptr);
-	
+
 	return self;
 }
 
@@ -108,14 +113,44 @@
 	}
 }
 
+- (BOOL) DidValidActivityOcurred:(FStringView) Reason
+{
+	TSharedPtr<FAppleHttpRequest> Request = SourceRequest.Pin();
+
+	if (!Request)
+	{
+		return FALSE;
+	}
+
+	if (!bAnyHttpActivity)
+	{
+		bAnyHttpActivity = true;
+
+		Request->StartActivityTimeoutTimer();
+	}
+
+	Request->ResetActivityTimeoutTimer(Reason);
+	return TRUE;
+}
+
 - (void)URLSession:(NSURLSession *)session task:(NSURLSessionTask *)task didSendBodyData:(int64_t)bytesSent totalBytesSent:(int64_t)totalBytesSent totalBytesExpectedToSend:(int64_t)totalBytesExpectedToSend
 {
+	if (![self DidValidActivityOcurred: TEXTVIEW("Sent body data")])
+	{
+		return;
+	}
 	UE_LOG(LogHttp, Verbose, TEXT("URLSession:task:didSendBodyData:totalBytesSent:totalBytesExpectedToSend: totalBytesSent = %lld, totalBytesSent = %lld: %p"), totalBytesSent, totalBytesExpectedToSend, self);
 	self.BytesWritten = totalBytesSent;
 }
 
 - (void)URLSession:(NSURLSession *)session dataTask:(NSURLSessionDataTask *)dataTask didReceiveResponse:(NSURLResponse *)response completionHandler:(void (^)(NSURLSessionResponseDisposition disposition))completionHandler
 {
+	if (![self DidValidActivityOcurred: TEXTVIEW("Received response")])
+	{
+		completionHandler(NSURLSessionResponseCancel);
+		return;
+	}
+
 	UE_LOG(LogHttp, Verbose, TEXT("URLSession:dataTask:didReceiveResponse:completionHandler"));
 	
 	self.Response = (NSHTTPURLResponse*)response;
@@ -132,6 +167,11 @@
 
 - (void)URLSession:(NSURLSession *)session dataTask:(NSURLSessionDataTask *)dataTask didReceiveData:(NSData *)data
 {
+	if (![self DidValidActivityOcurred: TEXTVIEW("Received data")])
+	{
+		return;
+	}
+	
 	__block int64 NewBytesReceived = 0;
 	if (bInitializedWithValidStream)
 	{
@@ -168,6 +208,14 @@
 
 - (void)URLSession:(NSURLSession *)session task:(NSURLSessionTask *)task didCompleteWithError:(nullable NSError *)error
 {
+	TSharedPtr<FAppleHttpRequest> Request = SourceRequest.Pin();
+
+	if (!Request)
+	{
+		return;
+	}
+
+	self.RequestStatus = EHttpRequestStatus::Failed;
 	if (error == nil)
 	{
 		UE_LOG(LogHttp, Verbose, TEXT("URLSession:task:didCompleteWithError. Http request succeeded: %p"), self);
@@ -226,6 +274,8 @@
 			}
 		}
 	}
+
+	Request->StopActivityTimeoutTimer();
 	NewAppleHttpEventDelegate.ExecuteIfBound();
 }
 
@@ -390,11 +440,15 @@ FAppleHttpRequest::FAppleHttpRequest(NSURLSession* InSession)
 ,	LastReportedBytesWritten(0)
 ,	LastReportedBytesRead(0)
 {
+	bUsePlatformActivityTimeout = false;
+	
 	UE_LOG(LogHttp, Verbose, TEXT("FAppleHttpRequest::FAppleHttpRequest()"));
 	Request = [[NSMutableURLRequest alloc] init];
 	float HttpConnectionTimeout = FHttpModule::Get().GetHttpConnectionTimeout();
 	check(HttpConnectionTimeout > 0.0f);
 	Request.timeoutInterval = HttpConnectionTimeout;
+	
+	UE_CLOG(HttpConnectionTimeout <= FHttpModule::Get().GetHttpActivityTimeout(), LogHttp, Warning, TEXT("HttpConnectionTimeout should be greater than HttpActivityTimeout. Otherwise requests may complete unexpectedly with ConnectionError after HttpConnectionTimeout seconds without activity"));
 
 	// Disable cache to mimic WinInet behavior
 	Request.cachePolicy = NSURLRequestReloadIgnoringLocalCacheData;
@@ -641,6 +695,11 @@ bool FAppleHttpRequest::SetContentFromStream(TSharedRef<FArchive, ESPMode::Threa
 
 bool FAppleHttpRequest::SetResponseBodyReceiveStream(TSharedRef<FArchive> Stream)
 {
+	if (CompletionStatus == EHttpRequestStatus::Processing)
+	{
+		UE_LOG(LogHttp, Warning, TEXT("FCurlHttpRequest::SetContentFromStream() - attempted to set content on a request that is inflight"));
+		return false;
+	}
 	ResponseBodyReceiveStream = Stream;
 	return true;
 }
@@ -744,9 +803,16 @@ void FAppleHttpRequest::FinishRequest()
 		if (Response)
 		{
 			Reason = Response->GetFailureReasonFromDelegate();
-			if (Reason == EHttpFailureReason::Cancelled && bTimedOut)
+			if (Reason == EHttpFailureReason::Cancelled)
 			{
-				Reason = EHttpFailureReason::TimedOut;
+				if (bTimedOut)
+				{
+					Reason = EHttpFailureReason::TimedOut;
+				}
+				else if (bActivityTimedOut)
+				{
+					Reason = EHttpFailureReason::ConnectionError;
+				}
 			}
 		}
 		SetFailureReason(Reason);
diff --git a/Engine/Source/Runtime/Online/HTTP/Private/Apple/AppleHttp.h b/Engine/Source/Runtime/Online/HTTP/Private/Apple/AppleHttp.h
--- a/Engine/Source/Runtime/Online/HTTP/Private/Apple/AppleHttp.h
+++ b/Engine/Source/Runtime/Online/HTTP/Private/Apple/AppleHttp.h
@@ -6,7 +6,7 @@
 #include "GenericPlatform/HttpResponseCommon.h"
 #include "IHttpThreadedRequest.h"
 #include "PlatformHttp.h"
-
+#include "HttpPackage.h"
 
 /**
  * Delegate invoked when in progress Task completes. It is invoked in an out of our control thread
@@ -70,6 +70,11 @@ public:
 	 */
 	virtual ~FAppleHttpRequest();
 
+PACKAGE_SCOPE:
+	using FHttpRequestCommon::StartActivityTimeoutTimer;
+	using FHttpRequestCommon::ResetActivityTimeoutTimer;
+	using FHttpRequestCommon::StopActivityTimeoutTimer;
+
 	const TSharedPtr<FArchive> GetResponseBodyReceiveStream() const;
 
 	void HandleStatusCodeReceived(int32 StatusCode);
diff --git a/Engine/Source/Runtime/Online/HTTP/Private/Apple/ApplePlatformHttp.cpp b/Engine/Source/Runtime/Online/HTTP/Private/Apple/ApplePlatformHttp.cpp
--- a/Engine/Source/Runtime/Online/HTTP/Private/Apple/ApplePlatformHttp.cpp
+++ b/Engine/Source/Runtime/Online/HTTP/Private/Apple/ApplePlatformHttp.cpp
@@ -237,10 +237,6 @@ void FApplePlatformHttp::InitWithNSUrlSession()
 	// Disable cache to mimic WinInet behavior
 	Config.requestCachePolicy = NSURLRequestReloadIgnoringLocalCacheData;
 
-	float HttpActivityTimeout = FHttpModule::Get().GetHttpActivityTimeout();
-	check(HttpActivityTimeout > 0);
-	Config.timeoutIntervalForRequest = HttpActivityTimeout;
-	
 #if WITH_SSL
 	// Load SSL module during HTTP module's StatupModule() to make sure module manager figures out the dependencies correctly
 	// and doesn't unload SSL before unloading HTTP module at exit
