[CIS] Default Http.RetrySystemNonGameThreadSupportEnabled => false
+ Disabling this feature fixes an issue where a retry-able request that returns an HTTP 429 status omits the configured retry flow and ultimately (prematurely) invokes request completion

diff --git a/Engine/Source/Runtime/Online/HTTP/Private/HttpRetrySystem.cpp b/Engine/Source/Runtime/Online/HTTP/Private/HttpRetrySystem.cpp
--- a/Engine/Source/Runtime/Online/HTTP/Private/HttpRetrySystem.cpp
+++ b/Engine/Source/Runtime/Online/HTTP/Private/HttpRetrySystem.cpp
@@ -16,7 +16,7 @@
 
 TAutoConsoleVariable<bool> CVarHttpRetrySystemNonGameThreadSupportEnabled(
 	TEXT("Http.RetrySystemNonGameThreadSupportEnabled"),
-	true,
+	false,
 	TEXT("Enable retry system non-game thread support")
 );
 
