Add more comment about the interface ProcessRequestUntilComplete.

diff --git a/Engine/Source/Runtime/Online/HTTP/Public/Interfaces/IHttpRequest.h b/Engine/Source/Runtime/Online/HTTP/Public/Interfaces/IHttpRequest.h
--- a/Engine/Source/Runtime/Online/HTTP/Public/Interfaces/IHttpRequest.h
+++ b/Engine/Source/Runtime/Online/HTTP/Public/Interfaces/IHttpRequest.h
@@ -357,7 +357,15 @@ public:
 	virtual EHttpRequestDelegateThreadPolicy GetDelegateThreadPolicy() const = 0;
 
 	/**
-	 * Blocking call to wait the request until it's completed, without calling the callback
+	 * Blocking call to wait the request until it's completed
+	 *
+	 * WARNINGS: 
+	 *
+	 * - This is a blocking call, DON'T use this in a time-sensitive context
+	 * - Complete delegate will be used in this function so customized complete delegate is not supported
+	 * - This will force the usage of EHttpRequestDelegateThreadPolicy::CompleteOnHttpThread to make sure the
+	 *   request can complete, when this function get called from main thread. So if any other delegate is
+	 *   bound, make sure the bound delegate can handle the custom logic in a thread-safe way
 	 */
 	virtual void ProcessRequestUntilComplete() = 0;
 
