Add GetEffectiveURL interface to get the effective url after redirection.

#tests Passed WebTests test case, launched the game to main menu.

diff --git a/Engine/Source/Programs/WebTests/Private/TestHttp.cpp b/Engine/Source/Programs/WebTests/Private/TestHttp.cpp
--- a/Engine/Source/Programs/WebTests/Private/TestHttp.cpp
+++ b/Engine/Source/Programs/WebTests/Private/TestHttp.cpp
@@ -840,11 +840,15 @@ TEST_CASE_METHOD(FWaitUntilCompleteHttpFixture, "Streaming http upload from file
 TEST_CASE_METHOD(FWaitUntilCompleteHttpFixture, "Redirect enabled by default and can work well", HTTP_TAG)
 {
 	TSharedRef<IHttpRequest> HttpRequest = CreateRequest();
-	HttpRequest->SetURL(FString::Format(TEXT("{0}/redirect_from"), { *UrlHttpTests() }));
+	FString OriginalURL = FString::Format(TEXT("{0}/redirect_from"), { *UrlHttpTests() });
+	FString ExpectedURL = FString::Format(TEXT("{0}/redirect_to"), { *UrlHttpTests() });
+	HttpRequest->SetURL(OriginalURL);
 	HttpRequest->SetVerb(TEXT("GET"));
-	HttpRequest->OnProcessRequestComplete().BindLambda([](FHttpRequestPtr HttpRequest, FHttpResponsePtr HttpResponse, bool bSucceeded) {
+	HttpRequest->OnProcessRequestComplete().BindLambda([OriginalURL, ExpectedURL](FHttpRequestPtr HttpRequest, FHttpResponsePtr HttpResponse, bool bSucceeded) {
 		CHECK(bSucceeded);
 		CHECK(HttpResponse->GetResponseCode() == 200);
+		CHECK(HttpResponse->GetURL() == OriginalURL);
+		CHECK(HttpResponse->GetEffectiveURL() == ExpectedURL);
 	});
 	HttpRequest->ProcessRequest();
 }
diff --git a/Engine/Source/Runtime/Online/BuildPatchServices/Private/Tests/Mock/HttpRequest.mock.h b/Engine/Source/Runtime/Online/BuildPatchServices/Private/Tests/Mock/HttpRequest.mock.h
--- a/Engine/Source/Runtime/Online/BuildPatchServices/Private/Tests/Mock/HttpRequest.mock.h
+++ b/Engine/Source/Runtime/Online/BuildPatchServices/Private/Tests/Mock/HttpRequest.mock.h
@@ -220,6 +220,12 @@ namespace BuildPatchServices
 			return EHttpRequestDelegateThreadPolicy::CompleteOnGameThread;
 		}
 
+		virtual const FString& GetEffectiveURL() const override
+		{
+			MOCK_FUNC_NOT_IMPLEMENTED("FMockHttpRequest::GetEffectiveURL");
+			return EffectiveURL;
+		}
+
 	public:
 		FHttpRequestProgressDelegate HttpRequestProgressDelegate;
 		FHttpRequestProgressDelegate64 HttpRequestProgressDelegate64;
@@ -232,6 +238,7 @@ namespace BuildPatchServices
 		TArray<FRxSetURL> RxSetURL;
 		int32 RxProcessRequest;
 		int32 RxCancelRequest;
+		FString EffectiveURL;
 	};
 }
 
diff --git a/Engine/Source/Runtime/Online/BuildPatchServices/Private/Tests/Mock/HttpResponse.mock.h b/Engine/Source/Runtime/Online/BuildPatchServices/Private/Tests/Mock/HttpResponse.mock.h
--- a/Engine/Source/Runtime/Online/BuildPatchServices/Private/Tests/Mock/HttpResponse.mock.h
+++ b/Engine/Source/Runtime/Online/BuildPatchServices/Private/Tests/Mock/HttpResponse.mock.h
@@ -82,6 +82,13 @@ namespace BuildPatchServices
 			MOCK_FUNC_NOT_IMPLEMENTED("FMockHttpRequest::GetFailureReason");
 			return EHttpFailureReason::Other;
 		}
+
+		virtual const FString& GetEffectiveURL() const override
+		{
+			MOCK_FUNC_NOT_IMPLEMENTED("FMockHttpRequest::GetEffectiveURL");
+			static FString None;
+			return None;
+		}
 	};
 }
 
diff --git a/Engine/Source/Runtime/Online/HTTP/Private/Apple/AppleHttp.cpp b/Engine/Source/Runtime/Online/HTTP/Private/Apple/AppleHttp.cpp
--- a/Engine/Source/Runtime/Online/HTTP/Private/Apple/AppleHttp.cpp
+++ b/Engine/Source/Runtime/Online/HTTP/Private/Apple/AppleHttp.cpp
@@ -109,7 +109,15 @@
 {
 	if (TSharedPtr<FAppleHttpRequest> Request = SourceRequest.Pin())
 	{
-		Request->HandleStatusCodeReceived(StatusCode);
+		Request->TriggerStatusCodeReceivedDelegate(StatusCode);
+	}
+}
+
+- (void) SaveEffectiveURL:(const FString&) InEffectiveURL
+{
+	if (TSharedPtr<FAppleHttpRequest> Request = SourceRequest.Pin())
+	{
+		Request->SetEffectiveURL(InEffectiveURL);
 	}
 }
 
@@ -154,8 +162,14 @@
 	UE_LOG(LogHttp, Verbose, TEXT("URLSession:dataTask:didReceiveResponse:completionHandler"));
 	
 	self.Response = (NSHTTPURLResponse*)response;
+
 	int32 StatusCode = [self.Response statusCode];
 	[self HandleStatusCodeReceived: StatusCode];
+
+	NSURL* Url = [self.Response URL];
+	FString EffectiveURL([Url absoluteString]);
+	[self SaveEffectiveURL: EffectiveURL];
+
 	uint64 ExpectedResponseLength = response.expectedContentLength;
 	if(!bInitializedWithValidStream && ExpectedResponseLength != NSURLResponseUnknownLength)
 	{
@@ -474,11 +488,6 @@ const TSharedPtr<FArchive> FAppleHttpRequest::GetResponseBodyReceiveStream() con
 	return ResponseBodyReceiveStream;
 }
 
-void FAppleHttpRequest::HandleStatusCodeReceived(int32 StatusCode)
-{
-	TriggerStatusCodeReceivedDelegate(StatusCode);
-}
-
 FString FAppleHttpRequest::GetURL() const
 {
 	SCOPED_AUTORELEASE_POOL;
diff --git a/Engine/Source/Runtime/Online/HTTP/Private/Apple/AppleHttp.h b/Engine/Source/Runtime/Online/HTTP/Private/Apple/AppleHttp.h
--- a/Engine/Source/Runtime/Online/HTTP/Private/Apple/AppleHttp.h
+++ b/Engine/Source/Runtime/Online/HTTP/Private/Apple/AppleHttp.h
@@ -77,7 +77,9 @@ PACKAGE_SCOPE:
 
 	const TSharedPtr<FArchive> GetResponseBodyReceiveStream() const;
 
-	void HandleStatusCodeReceived(int32 StatusCode);
+PACKAGE_SCOPE:
+	using FHttpRequestCommon::TriggerStatusCodeReceivedDelegate;
+	using FHttpRequestCommon::SetEffectiveURL;
 
 private:
 	/**
diff --git a/Engine/Source/Runtime/Online/HTTP/Private/Curl/CurlHttp.cpp b/Engine/Source/Runtime/Online/HTTP/Private/Curl/CurlHttp.cpp
--- a/Engine/Source/Runtime/Online/HTTP/Private/Curl/CurlHttp.cpp
+++ b/Engine/Source/Runtime/Online/HTTP/Private/Curl/CurlHttp.cpp
@@ -532,14 +532,28 @@ size_t FCurlHttpRequest::ReceiveResponseHeaderCallback(void* Ptr, size_t SizeInB
 		}
 		else
 		{
-			long HttpCode = 0;
-			if (CURLE_OK == curl_easy_getinfo(EasyHandle, CURLINFO_RESPONSE_CODE, &HttpCode))
+			if (Header.IsEmpty())
 			{
-				bRedirected = (HttpCode >= 300 && HttpCode < 400);
-
-				if (!bRedirected)
+				char* EffectiveUrlPtr = nullptr;
+				if (curl_easy_getinfo(EasyHandle, CURLINFO_EFFECTIVE_URL, &EffectiveUrlPtr) == CURLE_OK)
 				{
-					TriggerStatusCodeReceivedDelegate(HttpCode);
+					if (EffectiveUrlPtr)
+					{
+						SetEffectiveURL(FString(EffectiveUrlPtr));
+					}
+				}
+			}
+			else
+			{
+				long HttpCode = 0;
+				if (CURLE_OK == curl_easy_getinfo(EasyHandle, CURLINFO_RESPONSE_CODE, &HttpCode))
+				{
+					bRedirected = (HttpCode >= 300 && HttpCode < 400);
+
+					if (!bRedirected)
+					{
+						TriggerStatusCodeReceivedDelegate(HttpCode);
+					}
 				}
 			}
 		}
diff --git a/Engine/Source/Runtime/Online/HTTP/Private/Curl/CurlHttp.h b/Engine/Source/Runtime/Online/HTTP/Private/Curl/CurlHttp.h
--- a/Engine/Source/Runtime/Online/HTTP/Private/Curl/CurlHttp.h
+++ b/Engine/Source/Runtime/Online/HTTP/Private/Curl/CurlHttp.h
@@ -355,7 +355,7 @@ private:
 
 	void OnAnyActivityOccur(FStringView Reason);
 
-	void ClearInCaseOfRetry() override;
+	virtual void ClearInCaseOfRetry() override;
 
 private:
 
diff --git a/Engine/Source/Runtime/Online/HTTP/Private/GenericPlatform/GenericPlatformHttp.cpp b/Engine/Source/Runtime/Online/HTTP/Private/GenericPlatform/GenericPlatformHttp.cpp
--- a/Engine/Source/Runtime/Online/HTTP/Private/GenericPlatform/GenericPlatformHttp.cpp
+++ b/Engine/Source/Runtime/Online/HTTP/Private/GenericPlatform/GenericPlatformHttp.cpp
@@ -119,6 +119,7 @@ public:
 	virtual void CancelRequest() override {}
 	virtual EHttpRequestStatus::Type GetStatus() const override { return EHttpRequestStatus::NotStarted; }
 	virtual EHttpFailureReason GetFailureReason() const override { return EHttpFailureReason::None; }
+	virtual const FString& GetEffectiveURL() const override { return EffectiveURL; }
 	virtual const FHttpResponsePtr GetResponse() const override { return nullptr; }
 	virtual void Tick(float DeltaSeconds) override {}
 	virtual float GetElapsedTime() const override { return 0.0f; }
@@ -128,6 +129,9 @@ public:
 	virtual void ClearTimeout() override {}
 	virtual TOptional<float> GetTimeout() const override { return TOptional<float>(); }
 	virtual void ProcessRequestUntilComplete() override {}
+
+private:
+	FString EffectiveURL;
 };
 
 FDefaultUserAgentBuilder::FDefaultUserAgentBuilder()
@@ -479,7 +483,6 @@ FString FGenericPlatformHttp::GetUrlDomainAndPort(const FStringView Url)
 	return FString(DomainAndPort);
 }
 
-
 FString FGenericPlatformHttp::GetUrlDomain(const FStringView Url)
 {
 	FStringView Protocol;
@@ -503,6 +506,22 @@ FString FGenericPlatformHttp::GetUrlDomain(const FStringView Url)
 	return FString(Domain);
 }
 
+FString FGenericPlatformHttp::GetUrlBase(const FStringView Url)
+{
+	FStringView UrlBase = Url;
+	for (int32 Index = Url.Find(TEXT("://")) + 3; Index < Url.Len(); ++Index)
+	{
+		const TCHAR Character = Url[Index];
+		if (Character == TEXT('/') || Character == TEXT('?') || Character == TEXT('#'))
+		{
+			UrlBase.LeftInline(Index);
+			break;
+		}
+	}
+
+	return FString(UrlBase);
+}
+
 FString FGenericPlatformHttp::GetMimeType(const FString& FilePath)
 {
 	const FString FileExtension = FPaths::GetExtension(FilePath, true);
diff --git a/Engine/Source/Runtime/Online/HTTP/Private/GenericPlatform/HttpRequestCommon.cpp b/Engine/Source/Runtime/Online/HTTP/Private/GenericPlatform/HttpRequestCommon.cpp
--- a/Engine/Source/Runtime/Online/HTTP/Private/GenericPlatform/HttpRequestCommon.cpp
+++ b/Engine/Source/Runtime/Online/HTTP/Private/GenericPlatform/HttpRequestCommon.cpp
@@ -29,6 +29,11 @@ EHttpRequestStatus::Type FHttpRequestCommon::GetStatus() const
 	return CompletionStatus;
 }
 
+const FString& FHttpRequestCommon::GetEffectiveURL() const
+{
+	return EffectiveURL;
+}
+
 EHttpFailureReason FHttpRequestCommon::GetFailureReason() const
 {
 	return FailureReason;
@@ -107,6 +112,7 @@ void FHttpRequestCommon::ClearInCaseOfRetry()
 	bActivityTimedOut = false;
 	FailureReason = EHttpFailureReason::None;
 	bCanceled = false;
+	EffectiveURL = GetURL();
 }
 
 void FHttpRequestCommon::FinishRequestNotInHttpManager()
@@ -423,3 +429,14 @@ void FHttpRequestCommon::TriggerStatusCodeReceivedDelegate(int32 StatusCode)
 		});
 	}
 }
+
+void FHttpRequestCommon::SetEffectiveURL(const FString& InEffectiveURL)
+{
+	EffectiveURL = InEffectiveURL;
+
+	if (FHttpResponsePtr Response = GetResponse())
+	{
+		TSharedPtr<FHttpResponseCommon> ResponseCommon = StaticCastSharedPtr<FHttpResponseCommon>(Response);
+		ResponseCommon->SetEffectiveURL(EffectiveURL);
+	}
+}
diff --git a/Engine/Source/Runtime/Online/HTTP/Private/GenericPlatform/HttpResponseCommon.cpp b/Engine/Source/Runtime/Online/HTTP/Private/GenericPlatform/HttpResponseCommon.cpp
--- a/Engine/Source/Runtime/Online/HTTP/Private/GenericPlatform/HttpResponseCommon.cpp
+++ b/Engine/Source/Runtime/Online/HTTP/Private/GenericPlatform/HttpResponseCommon.cpp
@@ -6,6 +6,7 @@
 
 FHttpResponseCommon::FHttpResponseCommon(const FHttpRequestCommon& HttpRequest)
 	: URL(HttpRequest.GetURL())
+	, EffectiveURL(HttpRequest.GetEffectiveURL())
 	, CompletionStatus(HttpRequest.GetStatus())
 	, FailureReason(HttpRequest.GetFailureReason())
 {
@@ -26,6 +27,11 @@ FString FHttpResponseCommon::GetURL() const
 	return URL;
 }
 
+const FString& FHttpResponseCommon::GetEffectiveURL() const
+{
+	return EffectiveURL;
+}
+
 void FHttpResponseCommon::SetRequestStatus(EHttpRequestStatus::Type InCompletionStatus)
 {
 	CompletionStatus = InCompletionStatus;
@@ -46,3 +52,7 @@ EHttpFailureReason FHttpResponseCommon::GetFailureReason() const
 	return FailureReason;
 }
 
+void FHttpResponseCommon::SetEffectiveURL(const FString& InEffectiveURL)
+{
+	EffectiveURL = InEffectiveURL;
+}
diff --git a/Engine/Source/Runtime/Online/HTTP/Private/NullHttp.cpp b/Engine/Source/Runtime/Online/HTTP/Private/NullHttp.cpp
--- a/Engine/Source/Runtime/Online/HTTP/Private/NullHttp.cpp
+++ b/Engine/Source/Runtime/Online/HTTP/Private/NullHttp.cpp
@@ -158,6 +158,11 @@ EHttpFailureReason FNullHttpRequest::GetFailureReason() const
 	return FailureReason;
 }
 
+const FString& FNullHttpRequest::GetEffectiveURL() const
+{
+	return EffectiveUrl;
+}
+
 const FHttpResponsePtr FNullHttpRequest::GetResponse() const
 {
 	return FHttpResponsePtr(nullptr);
diff --git a/Engine/Source/Runtime/Online/HTTP/Private/NullHttp.h b/Engine/Source/Runtime/Online/HTTP/Private/NullHttp.h
--- a/Engine/Source/Runtime/Online/HTTP/Private/NullHttp.h
+++ b/Engine/Source/Runtime/Online/HTTP/Private/NullHttp.h
@@ -37,6 +37,7 @@ public:
 	virtual void CancelRequest() override;
 	virtual EHttpRequestStatus::Type GetStatus() const override;
 	virtual EHttpFailureReason GetFailureReason() const override;
+	virtual const FString& GetEffectiveURL() const override;
 	virtual const FHttpResponsePtr GetResponse() const override;
 	virtual void Tick(float DeltaSeconds) override;
 	virtual float GetElapsedTime() const override;
@@ -58,6 +59,7 @@ private:
 	void FinishedRequest();
 
 	FString Url;
+	FString EffectiveUrl;
 	FString Verb;
 	TArray<uint8> Payload;
 	EHttpRequestStatus::Type CompletionStatus;
diff --git a/Engine/Source/Runtime/Online/HTTP/Public/GenericPlatform/GenericPlatformHttp.h b/Engine/Source/Runtime/Online/HTTP/Public/GenericPlatform/GenericPlatformHttp.h
--- a/Engine/Source/Runtime/Online/HTTP/Public/GenericPlatform/GenericPlatformHttp.h
+++ b/Engine/Source/Runtime/Online/HTTP/Public/GenericPlatform/GenericPlatformHttp.h
@@ -130,6 +130,13 @@ public:
 	 */
 	static HTTP_API FString GetUrlDomain(const FStringView Url);
 
+	/** 
+	* Returns the base of the URL, e.g., "http://a.b.c:d" of "http://a.b.c:d/e"
+	* @param Url the URL to return the base of
+	* @return the base of the specified URL
+	*/
+	static HTTP_API FString GetUrlBase(const FStringView Url);
+
 	/**
 	 * Get the mime type for the file
 	 * @return the mime type for the file.
diff --git a/Engine/Source/Runtime/Online/HTTP/Public/GenericPlatform/HttpRequestCommon.h b/Engine/Source/Runtime/Online/HTTP/Public/GenericPlatform/HttpRequestCommon.h
--- a/Engine/Source/Runtime/Online/HTTP/Public/GenericPlatform/HttpRequestCommon.h
+++ b/Engine/Source/Runtime/Online/HTTP/Public/GenericPlatform/HttpRequestCommon.h
@@ -19,6 +19,7 @@ public:
 
 	// IHttpRequest
 	HTTP_API virtual EHttpRequestStatus::Type GetStatus() const override;
+	HTTP_API virtual const FString& GetEffectiveURL() const override;
 	HTTP_API virtual EHttpFailureReason GetFailureReason() const override;
 	HTTP_API virtual void SetDelegateThreadPolicy(EHttpRequestDelegateThreadPolicy InDelegateThreadPolicy) override;
 	HTTP_API virtual EHttpRequestDelegateThreadPolicy GetDelegateThreadPolicy() const override;
@@ -72,6 +73,8 @@ protected:
 
 	HTTP_API void TriggerStatusCodeReceivedDelegate(int32 StatusCode);
 
+	HTTP_API void SetEffectiveURL(const FString& InEffectiveURL);
+
 protected:
 	/** Current status of request being processed */
 	EHttpRequestStatus::Type CompletionStatus = EHttpRequestStatus::NotStarted;
@@ -113,4 +116,7 @@ protected:
 
 	/** Record how long it take to connect to the endpoint */
 	double ConnectTime = -1.0;
+
+	/** Cache the effective URL. When redirected, it will be different with original URL */
+	FString EffectiveURL;
 };
diff --git a/Engine/Source/Runtime/Online/HTTP/Public/GenericPlatform/HttpResponseCommon.h b/Engine/Source/Runtime/Online/HTTP/Public/GenericPlatform/HttpResponseCommon.h
--- a/Engine/Source/Runtime/Online/HTTP/Public/GenericPlatform/HttpResponseCommon.h
+++ b/Engine/Source/Runtime/Online/HTTP/Public/GenericPlatform/HttpResponseCommon.h
@@ -19,14 +19,17 @@ public:
 	// IHttpBase
 	HTTP_API virtual FString GetURLParameter(const FString& ParameterName) const override;
 	HTTP_API virtual FString GetURL() const override;
+	HTTP_API virtual const FString& GetEffectiveURL() const override;
 	HTTP_API virtual EHttpRequestStatus::Type GetStatus() const override;
 	HTTP_API virtual EHttpFailureReason GetFailureReason() const override;
 
 protected:
 	HTTP_API void SetRequestStatus(EHttpRequestStatus::Type InCompletionStatus);
 	HTTP_API void SetRequestFailureReason(EHttpFailureReason InFailureReason);
+	HTTP_API void SetEffectiveURL(const FString& InEffectiveURL);
 
 	FString URL;
+	FString EffectiveURL;
 	EHttpRequestStatus::Type CompletionStatus;
 	EHttpFailureReason FailureReason;
 };
diff --git a/Engine/Source/Runtime/Online/HTTP/Public/HttpRequestAdapter.h b/Engine/Source/Runtime/Online/HTTP/Public/HttpRequestAdapter.h
--- a/Engine/Source/Runtime/Online/HTTP/Public/HttpRequestAdapter.h
+++ b/Engine/Source/Runtime/Online/HTTP/Public/HttpRequestAdapter.h
@@ -44,6 +44,7 @@ public:
 	virtual float GetElapsedTime() const override { return HttpRequest->GetElapsedTime(); }
 	virtual EHttpRequestStatus::Type GetStatus() const override { return HttpRequest->GetStatus(); }
 	virtual EHttpFailureReason GetFailureReason() const override { return HttpRequest->GetFailureReason(); }
+	virtual const FString& GetEffectiveURL() const override { return HttpRequest->GetEffectiveURL(); }
 	virtual void Tick(float DeltaSeconds) override { HttpRequest->Tick(DeltaSeconds); }
 	virtual void SetDelegateThreadPolicy(EHttpRequestDelegateThreadPolicy InThreadPolicy) override { HttpRequest->SetDelegateThreadPolicy(InThreadPolicy); }
 	virtual EHttpRequestDelegateThreadPolicy GetDelegateThreadPolicy() const override { return HttpRequest->GetDelegateThreadPolicy(); }
diff --git a/Engine/Source/Runtime/Online/HTTP/Public/Interfaces/IHttpBase.h b/Engine/Source/Runtime/Online/HTTP/Public/Interfaces/IHttpBase.h
--- a/Engine/Source/Runtime/Online/HTTP/Public/Interfaces/IHttpBase.h
+++ b/Engine/Source/Runtime/Online/HTTP/Public/Interfaces/IHttpBase.h
@@ -91,6 +91,13 @@ public:
 	 */
 	virtual FString GetURL() const = 0;
 
+	/**
+	 * Get the effective URL in case of redirected. If not redirected, it's the same as GetURL
+	 *
+	 * @return the effective URL string.
+	 */
+	virtual const FString& GetEffectiveURL() const = 0;
+
 	/**
 	 * Get the current status of the request being processed
 	 *
