Add option to request microphone permission for AndroidVoice module at startup (SplashActivity)

#android

diff --git a/Engine/Config/BaseEngine.ini b/Engine/Config/BaseEngine.ini
--- a/Engine/Config/BaseEngine.ini
+++ b/Engine/Config/BaseEngine.ini
@@ -3053,6 +3053,7 @@ TextureFormatPriority_DXT=0.6
 TextureFormatPriority_ASTC=0.9
 bEnableNewKeyboard=True
 bAndroidVoiceEnabled=false
+bRecordPermissionAtStartupEnabled=false
 bBuildWithHiddenSymbolVisibility=True
 bSaveSymbols=false
 bAllowControllers=True
diff --git a/Engine/Source/Runtime/Android/AndroidRuntimeSettings/Classes/AndroidRuntimeSettings.h b/Engine/Source/Runtime/Android/AndroidRuntimeSettings/Classes/AndroidRuntimeSettings.h
--- a/Engine/Source/Runtime/Android/AndroidRuntimeSettings/Classes/AndroidRuntimeSettings.h
+++ b/Engine/Source/Runtime/Android/AndroidRuntimeSettings/Classes/AndroidRuntimeSettings.h
@@ -371,6 +371,10 @@ public:
 	UPROPERTY(GlobalConfig, EditAnywhere, Category = "Advanced APK Packaging", Meta = (DisplayName = "Add permissions to support Voice chat (RECORD_AUDIO)"))
 	bool bAndroidVoiceEnabled;
 
+	// Request permission at startup (SplashActivity) for Voice chat (RECORD_AUDIO)
+	UPROPERTY(GlobalConfig, EditAnywhere, Category = "Advanced APK Packaging", Meta = (DisplayName = "Request permission at startup for Voice chat (RECORD_AUDIO)", EditCondition = "bAndroidVoiceEnabled"))
+	bool bRecordPermissionAtStartupEnabled;
+
 	// Add required permission and support to allow multicast/broadcast Wi-Fi traffic through network interface
 	UPROPERTY(GlobalConfig, EditAnywhere, Category = "Advanced APK Packaging", Meta = (DisplayName = "Add support for multicast Wi-Fi traffic (CHANGE_WIFI_MULTICAST_STATE)"))
 	bool bEnableMulticastSupport;
diff --git a/Engine/Source/Runtime/Online/Voice/AndroidVoiceImpl_UPL.xml b/Engine/Source/Runtime/Online/Voice/AndroidVoiceImpl_UPL.xml
--- a/Engine/Source/Runtime/Online/Voice/AndroidVoiceImpl_UPL.xml
+++ b/Engine/Source/Runtime/Online/Voice/AndroidVoiceImpl_UPL.xml
@@ -5,17 +5,49 @@
 	<init>
 		<log text="Android Voice init"/>
 		<setBoolFromProperty result="bAndroidVoiceModuleEnabled" ini="Engine" section="/Script/AndroidRuntimeSettings.AndroidRuntimeSettings" property="bAndroidVoiceEnabled" default="false"/>
+		<setBoolFromProperty result="bRecordPermissionAtStartupEnabled" ini="Engine" section="/Script/AndroidRuntimeSettings.AndroidRuntimeSettings" property="bRecordPermissionAtStartupEnabled" default="false"/>
 	</init>
 
 	<!-- updates applied to AndroidManifest.xml -->
 	<androidManifestUpdates>
-  
-	<if condition="bAndroidVoiceModuleEnabled">
-		<true>
-			<addPermission android:name="android.permission.RECORD_AUDIO"/>
-		</true>
-	</if>
+		<if condition="bAndroidVoiceModuleEnabled">
+			<true>
+				<addPermission android:name="android.permission.RECORD_AUDIO"/>
 
+				<if condition="bRecordPermissionAtStartupEnabled">
+					<true>
+						<!-- add RECORD_AUDIO permission to startup permissions for SplashActivity -->
+						<loopElements tag="meta-data">
+							<setStringFromAttribute result="metaName" tag="$" name="android:name"/>
+							<setBoolIsEqual result="bStartupPerms" arg1="$S(metaName)" arg2="com.epicgames.unreal.GameActivity.StartupPermissions"/>
+							<if condition="bStartupPerms">
+								<true>
+									<setStringFromAttribute result="metaValue" tag="$" name="android:value"/>
+									<setIntLength result="metaValueLength" source="$S(metaValue)"/>
+									<setBoolIsGreater result="bNeedComma" arg1="$I(metaValueLength)" arg2="0"/>
+									<if condition="bNeedComma">
+										<false>
+											<setString result="metaComma" value=""/>
+										</false>
+										<true>
+											<setString result="metaComma" value=","/>
+										</true>
+									</if>
+									<setBoolContains result="bHasRecord" source="$S(metaValue)" find="android.permission.RECORD_AUDIO"/>
+									<if condition="bHasRecord">
+										<false>
+											<setString result="metaValue" value="$S(metaValue)$S(metaComma)android.permission.RECORD_AUDIO"/>
+											<setString result="metaComma" value=","/>
+										</false>
+									</if>
+									<addAttribute tag="$" name="android:value" value="$S(metaValue)"/>
+								</true>
+							</if>
+						</loopElements>
+					</true>
+				</if>
+			</true>
+		</if>
 	</androidManifestUpdates>
-	
+
 </root>
