Stop activity timeout timer after canceling request;
Stop total timeout timer after request completed in retry manager;
Move the virtual impls of FHttpRequestAdapter to cpp file.

#tests Repro and fixed in WebTests

diff --git a/Engine/Source/Programs/WebTests/Private/TestHttp.cpp b/Engine/Source/Programs/WebTests/Private/TestHttp.cpp
--- a/Engine/Source/Programs/WebTests/Private/TestHttp.cpp
+++ b/Engine/Source/Programs/WebTests/Private/TestHttp.cpp
@@ -730,6 +730,30 @@ TEST_CASE_METHOD(FWaitUntilCompleteHttpFixture, "Http request activity timeout",
 	HttpRequest->ProcessRequest();
 }
 
+TEST_CASE_METHOD(FWaitUntilCompleteHttpFixture, "Http request won't trigger activity timeout after cancelling", HTTP_TAG)
+{
+	HttpModule->HttpActivityTimeout = 2.0f;
+
+	TSharedPtr<IHttpRequest> HttpRequest = CreateRequest();
+	HttpRequest->SetURL(UrlStreamDownload(3/*Chunks*/, HTTP_TEST_TIMEOUT_CHUNK_SIZE, 5/*ChunkLatency*/));
+	HttpRequest->SetVerb(TEXT("GET"));
+	HttpRequest->SetDelegateThreadPolicy(EHttpRequestDelegateThreadPolicy::CompleteOnHttpThread);
+
+	const double TimeToWaitBeforeCancel = 1.0f;
+	const double StartTime = FPlatformTime::Seconds();
+	HttpRequest->OnProcessRequestComplete().BindLambda([StartTime, TimeToWaitBeforeCancel](FHttpRequestPtr HttpRequest, FHttpResponsePtr HttpResponse, bool bSucceeded) {
+		const double DurationInSeconds  = FPlatformTime::Seconds() - StartTime;
+		CHECK(FMath::IsNearlyEqual(DurationInSeconds, TimeToWaitBeforeCancel, HTTP_TIME_DIFF_TOLERANCE));
+		CHECK(!bSucceeded);
+		CHECK(HttpRequest->GetStatus() == EHttpRequestStatus::Failed);
+		CHECK(HttpRequest->GetFailureReason() == EHttpFailureReason::Cancelled);
+	});
+	HttpRequest->ProcessRequest();
+	FPlatformProcess::Sleep(TimeToWaitBeforeCancel);
+	HttpRequest->CancelRequest();
+	FPlatformProcess::Sleep(3.0f); // Just make sure there is no warning or assert triggered by the activity timeout callback
+}
+
 TEST_CASE_METHOD(FWaitUntilCompleteHttpFixture, "Http request receive won't timeout for streaming request", HTTP_TAG)
 {
 	HttpModule->HttpActivityTimeout = 3.0f;
@@ -754,12 +778,12 @@ TEST_CASE_METHOD(FWaitUntilCompleteHttpFixture, "Http request total timeout with
 	DisableWarningsInThisTest();
 
 	float TotalTimeoutSetting = 3.0f;
-	HttpModule->HttpTotalTimeout = TotalTimeoutSetting;
 	HttpModule->HttpConnectionTimeout = 5.0f;
 
 	TSharedPtr<IHttpRequest> HttpRequest = CreateRequest();
 	HttpRequest->SetURL(UrlMockLatency(10));
 	HttpRequest->SetVerb(TEXT("GET"));
+	HttpRequest->SetTimeout(TotalTimeoutSetting);
 
 	const double StartTime = FPlatformTime::Seconds();
 
diff --git a/Engine/Source/Runtime/Online/HTTP/Private/GenericPlatform/HttpRequestCommon.cpp b/Engine/Source/Runtime/Online/HTTP/Private/GenericPlatform/HttpRequestCommon.cpp
--- a/Engine/Source/Runtime/Online/HTTP/Private/GenericPlatform/HttpRequestCommon.cpp
+++ b/Engine/Source/Runtime/Online/HTTP/Private/GenericPlatform/HttpRequestCommon.cpp
@@ -195,6 +195,7 @@ void FHttpRequestCommon::SetTimeout(float InTimeoutSecs)
 void FHttpRequestCommon::ClearTimeout()
 {
 	TimeoutSecs.Reset();
+	StopTotalTimeoutTimer();
 }
 
 TOptional<float> FHttpRequestCommon::GetTimeout() const
@@ -214,6 +215,8 @@ void FHttpRequestCommon::CancelRequest()
 		return;
 	}
 
+	StopActivityTimeoutTimer();
+
 	bCanceled = true;
 	UE_LOG(LogHttp, Verbose, TEXT("HTTP request canceled. URL=%s"), *GetURL());
 
@@ -392,6 +395,7 @@ void FHttpRequestCommon::StopTotalTimeoutTimer()
 		TotalTimeoutHttpTaskTimerHandle.Reset();
 	}
 }
+
 void FHttpRequestCommon::Shutdown()
 {
 	FHttpRequestImpl::Shutdown();
diff --git a/Engine/Source/Runtime/Online/HTTP/Private/HttpRequestAdapter.cpp b/Engine/Source/Runtime/Online/HTTP/Private/HttpRequestAdapter.cpp
new file mode 100644
--- /dev/null
+++ b/Engine/Source/Runtime/Online/HTTP/Private/HttpRequestAdapter.cpp
@@ -0,0 +1,159 @@
+// Copyright Epic Games, Inc. All Rights Reserved.
+
+#include "HttpRequestAdapter.h"
+
+FHttpRequestAdapterBase::FHttpRequestAdapterBase(const TSharedRef<IHttpRequest>& InHttpRequest)
+	: HttpRequest(InHttpRequest)
+{
+}
+
+FString FHttpRequestAdapterBase::GetURL() const 
+{ 
+	return HttpRequest->GetURL(); 
+}
+
+FString FHttpRequestAdapterBase::GetURLParameter(const FString& ParameterName) const 
+{ 
+	return HttpRequest->GetURLParameter(ParameterName); 
+}
+
+FString FHttpRequestAdapterBase::GetHeader(const FString& HeaderName) const 
+{ 
+	return HttpRequest->GetHeader(HeaderName); 
+}
+
+TArray<FString> FHttpRequestAdapterBase::GetAllHeaders() const 
+{ 
+	return HttpRequest->GetAllHeaders(); 
+}
+
+FString FHttpRequestAdapterBase::GetContentType() const 
+{ 
+	return HttpRequest->GetContentType(); 
+}
+
+uint64 FHttpRequestAdapterBase::GetContentLength() const 
+{ 
+	return HttpRequest->GetContentLength(); 
+}
+
+const TArray<uint8>& FHttpRequestAdapterBase::GetContent() const 
+{ 
+	return HttpRequest->GetContent(); 
+}
+
+FString FHttpRequestAdapterBase::GetVerb() const 
+{ 
+	return HttpRequest->GetVerb(); 
+}
+
+void FHttpRequestAdapterBase::SetVerb(const FString& Verb) 
+{ 
+	HttpRequest->SetVerb(Verb); 
+}
+
+void FHttpRequestAdapterBase::SetURL(const FString& URL) 
+{ 
+	HttpRequest->SetURL(URL); 
+}
+
+void FHttpRequestAdapterBase::SetContent(const TArray<uint8>& ContentPayload) 
+{ 
+	HttpRequest->SetContent(ContentPayload); 
+}
+
+void FHttpRequestAdapterBase::SetContent(TArray<uint8>&& ContentPayload) 
+{ 
+	HttpRequest->SetContent(MoveTemp(ContentPayload)); 
+}
+
+void FHttpRequestAdapterBase::SetContentAsString(const FString& ContentString) 
+{ 
+	HttpRequest->SetContentAsString(ContentString); 
+}
+
+bool FHttpRequestAdapterBase::SetContentAsStreamedFile(const FString& Filename) 
+{ 
+	return HttpRequest->SetContentAsStreamedFile(Filename); 
+}
+
+bool FHttpRequestAdapterBase::SetContentFromStream(TSharedRef<FArchive, ESPMode::ThreadSafe> Stream) 
+{ 
+	return HttpRequest->SetContentFromStream(Stream); 
+}
+
+bool FHttpRequestAdapterBase::SetResponseBodyReceiveStream(TSharedRef<FArchive> Stream) 
+{ 
+	return HttpRequest->SetResponseBodyReceiveStream(Stream); 
+}
+
+void FHttpRequestAdapterBase::SetHeader(const FString& HeaderName, const FString& HeaderValue) 
+{ 
+	HttpRequest->SetHeader(HeaderName, HeaderValue); 
+}
+
+void FHttpRequestAdapterBase::AppendToHeader(const FString& HeaderName, const FString& AdditionalHeaderValue) 
+{ 
+	HttpRequest->AppendToHeader(HeaderName, AdditionalHeaderValue); 
+}
+
+void FHttpRequestAdapterBase::SetTimeout(float InTimeoutSecs) 
+{ 
+	HttpRequest->SetTimeout(InTimeoutSecs); 
+}
+
+void FHttpRequestAdapterBase::ClearTimeout() 
+{ 
+	HttpRequest->ClearTimeout(); 
+}
+
+TOptional<float> FHttpRequestAdapterBase::GetTimeout() const 
+{ 
+	return HttpRequest->GetTimeout(); 
+}
+
+void FHttpRequestAdapterBase::ProcessRequestUntilComplete() 
+{ 
+	return HttpRequest->ProcessRequestUntilComplete(); 
+}
+
+const FHttpResponsePtr FHttpRequestAdapterBase::GetResponse() const 
+{ 
+	return HttpRequest->GetResponse(); 
+}
+
+float FHttpRequestAdapterBase::GetElapsedTime() const 
+{ 
+	return HttpRequest->GetElapsedTime(); 
+}
+
+EHttpRequestStatus::Type FHttpRequestAdapterBase::GetStatus() const 
+{ 
+	return HttpRequest->GetStatus(); 
+}
+
+EHttpFailureReason FHttpRequestAdapterBase::GetFailureReason() const 
+{ 
+	return HttpRequest->GetFailureReason(); 
+}
+
+const FString& FHttpRequestAdapterBase::GetEffectiveURL() const 
+{ 
+	return HttpRequest->GetEffectiveURL(); 
+}
+
+void FHttpRequestAdapterBase::Tick(float DeltaSeconds) 
+{ 
+	HttpRequest->Tick(DeltaSeconds); 
+}
+
+void FHttpRequestAdapterBase::SetDelegateThreadPolicy(EHttpRequestDelegateThreadPolicy InThreadPolicy) 
+{ 
+	HttpRequest->SetDelegateThreadPolicy(InThreadPolicy); 
+}
+
+EHttpRequestDelegateThreadPolicy FHttpRequestAdapterBase::GetDelegateThreadPolicy() const 
+{ 
+	return HttpRequest->GetDelegateThreadPolicy(); 
+}
+
diff --git a/Engine/Source/Runtime/Online/HTTP/Private/HttpRetrySystem.cpp b/Engine/Source/Runtime/Online/HTTP/Private/HttpRetrySystem.cpp
--- a/Engine/Source/Runtime/Online/HTTP/Private/HttpRetrySystem.cpp
+++ b/Engine/Source/Runtime/Online/HTTP/Private/HttpRetrySystem.cpp
@@ -252,6 +252,8 @@ void FHttpRetrySystem::FRequest::HttpOnProcessRequestComplete(FHttpRequestPtr In
 
 	LLM_SCOPE_BYTAG(HTTP);
 	OnProcessRequestComplete().ExecuteIfBound(SelfPtr, HttpResponse, bSucceeded);
+
+	ClearTimeout();
 }
 
 void FHttpRetrySystem::FRequest::HttpOnStatusCodeReceived(FHttpRequestPtr Request, int32 StatusCode)
diff --git a/Engine/Source/Runtime/Online/HTTP/Public/HttpRequestAdapter.h b/Engine/Source/Runtime/Online/HTTP/Public/HttpRequestAdapter.h
--- a/Engine/Source/Runtime/Online/HTTP/Public/HttpRequestAdapter.h
+++ b/Engine/Source/Runtime/Online/HTTP/Public/HttpRequestAdapter.h
@@ -13,43 +13,41 @@
 class FHttpRequestAdapterBase : public FHttpRequestImpl
 {
 public:
-    FHttpRequestAdapterBase(const TSharedRef<IHttpRequest, ESPMode::ThreadSafe>& InHttpRequest) 
-		: HttpRequest(InHttpRequest)
-    {}
+	HTTP_API FHttpRequestAdapterBase(const TSharedRef<IHttpRequest>& InHttpRequest);
 
 	// IHttpRequest interface
-	virtual FString GetURL() const override { return HttpRequest->GetURL(); }
-	virtual FString GetURLParameter(const FString& ParameterName) const override { return HttpRequest->GetURLParameter(ParameterName); }
-	virtual FString GetHeader(const FString& HeaderName) const override { return HttpRequest->GetHeader(HeaderName); }
-	virtual TArray<FString> GetAllHeaders() const override { return HttpRequest->GetAllHeaders(); }
-	virtual FString GetContentType() const override { return HttpRequest->GetContentType(); }
-	virtual uint64 GetContentLength() const override { return HttpRequest->GetContentLength(); }
-	virtual const TArray<uint8>& GetContent() const override { return HttpRequest->GetContent(); }
-	virtual FString GetVerb() const override { return HttpRequest->GetVerb(); }
-	virtual void SetVerb(const FString& Verb) override { HttpRequest->SetVerb(Verb); }
-	virtual void SetURL(const FString& URL) override { HttpRequest->SetURL(URL); }
-	virtual void SetContent(const TArray<uint8>& ContentPayload) override { HttpRequest->SetContent(ContentPayload); }
-	virtual void SetContent(TArray<uint8>&& ContentPayload) override { HttpRequest->SetContent(MoveTemp(ContentPayload)); }
-	virtual void SetContentAsString(const FString& ContentString) override { HttpRequest->SetContentAsString(ContentString); }
-	virtual bool SetContentAsStreamedFile(const FString& Filename) override { return HttpRequest->SetContentAsStreamedFile(Filename); }
-	virtual bool SetContentFromStream(TSharedRef<FArchive, ESPMode::ThreadSafe> Stream) override { return HttpRequest->SetContentFromStream(Stream); }
-	virtual bool SetResponseBodyReceiveStream(TSharedRef<FArchive> Stream) override { return HttpRequest->SetResponseBodyReceiveStream(Stream); }
-	virtual void SetHeader(const FString& HeaderName, const FString& HeaderValue) override { HttpRequest->SetHeader(HeaderName, HeaderValue); }
-	virtual void AppendToHeader(const FString& HeaderName, const FString& AdditionalHeaderValue) override { HttpRequest->AppendToHeader(HeaderName, AdditionalHeaderValue); }
-	virtual void SetTimeout(float InTimeoutSecs) override { HttpRequest->SetTimeout(InTimeoutSecs); }
-	virtual void ClearTimeout() override { HttpRequest->ClearTimeout(); }
-	virtual TOptional<float> GetTimeout() const override { return HttpRequest->GetTimeout(); }
-	virtual void ProcessRequestUntilComplete() override { return HttpRequest->ProcessRequestUntilComplete(); }
-	virtual const FHttpResponsePtr GetResponse() const override { return HttpRequest->GetResponse(); }
-	virtual float GetElapsedTime() const override { return HttpRequest->GetElapsedTime(); }
-	virtual EHttpRequestStatus::Type GetStatus() const override { return HttpRequest->GetStatus(); }
-	virtual EHttpFailureReason GetFailureReason() const override { return HttpRequest->GetFailureReason(); }
-	virtual const FString& GetEffectiveURL() const override { return HttpRequest->GetEffectiveURL(); }
-	virtual void Tick(float DeltaSeconds) override { HttpRequest->Tick(DeltaSeconds); }
-	virtual void SetDelegateThreadPolicy(EHttpRequestDelegateThreadPolicy InThreadPolicy) override { HttpRequest->SetDelegateThreadPolicy(InThreadPolicy); }
-	virtual EHttpRequestDelegateThreadPolicy GetDelegateThreadPolicy() const override { return HttpRequest->GetDelegateThreadPolicy(); }
+	HTTP_API virtual FString GetURL() const override;
+	HTTP_API virtual FString GetURLParameter(const FString& ParameterName) const override;
+	HTTP_API virtual FString GetHeader(const FString& HeaderName) const override;
+	HTTP_API virtual TArray<FString> GetAllHeaders() const override;
+	HTTP_API virtual FString GetContentType() const override;
+	HTTP_API virtual uint64 GetContentLength() const override;
+	HTTP_API virtual const TArray<uint8>& GetContent() const override;
+	HTTP_API virtual FString GetVerb() const override;
+	HTTP_API virtual void SetVerb(const FString& Verb) override;
+	HTTP_API virtual void SetURL(const FString& URL) override;
+	HTTP_API virtual void SetContent(const TArray<uint8>& ContentPayload) override;
+	HTTP_API virtual void SetContent(TArray<uint8>&& ContentPayload) override;
+	HTTP_API virtual void SetContentAsString(const FString& ContentString) override;
+	HTTP_API virtual bool SetContentAsStreamedFile(const FString& Filename) override;
+	HTTP_API virtual bool SetContentFromStream(TSharedRef<FArchive, ESPMode::ThreadSafe> Stream) override;
+	HTTP_API virtual bool SetResponseBodyReceiveStream(TSharedRef<FArchive> Stream) override;
+	HTTP_API virtual void SetHeader(const FString& HeaderName, const FString& HeaderValue) override;
+	HTTP_API virtual void AppendToHeader(const FString& HeaderName, const FString& AdditionalHeaderValue) override;
+	HTTP_API virtual void SetTimeout(float InTimeoutSecs) override;
+	HTTP_API virtual void ClearTimeout() override;
+	HTTP_API virtual TOptional<float> GetTimeout() const override;
+	HTTP_API virtual void ProcessRequestUntilComplete() override;
+	HTTP_API virtual const FHttpResponsePtr GetResponse() const override;
+	HTTP_API virtual float GetElapsedTime() const override;
+	HTTP_API virtual EHttpRequestStatus::Type GetStatus() const override;
+	HTTP_API virtual EHttpFailureReason GetFailureReason() const override;
+	HTTP_API virtual const FString& GetEffectiveURL() const override;
+	HTTP_API virtual void Tick(float DeltaSeconds) override;
+	HTTP_API virtual void SetDelegateThreadPolicy(EHttpRequestDelegateThreadPolicy InThreadPolicy) override;
+	HTTP_API virtual EHttpRequestDelegateThreadPolicy GetDelegateThreadPolicy() const override;
 
 protected:
-    TSharedRef<IHttpRequest, ESPMode::ThreadSafe> HttpRequest;
+	TSharedRef<IHttpRequest> HttpRequest;
 };
 
