Add more test cases about retry behaviors in WebTests;
Make sure when timed out in retry system, during back off or during retrying, it will return the last response instead of empty response.

#tests Tried in WebTests

diff --git a/Engine/Source/Programs/WebTests/Private/TestHttp.cpp b/Engine/Source/Programs/WebTests/Private/TestHttp.cpp
--- a/Engine/Source/Programs/WebTests/Private/TestHttp.cpp
+++ b/Engine/Source/Programs/WebTests/Private/TestHttp.cpp
@@ -124,6 +124,7 @@ public:
 	const FString UrlStreamDownload(uint32 Chunks, uint32 ChunkSize, uint32 ChunkLatency=0) { return FString::Format(TEXT("{0}/streaming_download/{1}/{2}/{3}/"), { *UrlHttpTests(), Chunks, ChunkSize, ChunkLatency }); }
 	const FString UrlStreamUpload() { return FString::Format(TEXT("{0}/streaming_upload_put"), { *UrlHttpTests() }); }
 	const FString UrlMockLatency(uint32 Latency) const { return FString::Format(TEXT("{0}/mock_latency/{1}/"), { *UrlHttpTests(), Latency }); }
+	const FString UrlMockStatus(uint32 StatusCode) const { return FString::Format(TEXT("{0}/mock_status/{1}/"), { *UrlHttpTests(), StatusCode }); }
 
 	FString WebServerIp;
 	uint32 WebServerHttpPort;
@@ -193,6 +194,8 @@ public:
 	{
 		WaitUntilAllHttpRequestsComplete();
 
+		CHECK(ExpectingExtraCallbacks == 0);
+
 		HttpModule->GetHttpManager().SetRequestAddedDelegate(FHttpManagerRequestAddedDelegate());
 		HttpModule->GetHttpManager().SetRequestCompletedDelegate(FHttpManagerRequestCompletedDelegate());
 	}
@@ -235,6 +238,7 @@ public:
 
 	uint32 RetryLimitCount = 0;
 	TSharedPtr<FMockRetryManager> HttpRetryManager;
+	uint32 ExpectingExtraCallbacks = 0;
 };
 
 TEST_CASE_METHOD(FWaitUntilCompleteHttpFixture, "Http Methods", HTTP_TAG)
@@ -1368,6 +1372,168 @@ TEST_CASE_METHOD(FWaitUntilCompleteHttpFixture, "Cancel http request connect bef
 	HttpRequest->CancelRequest(); // Duplicated calls to CancelRequest should be fine
 }
 
+TEST_CASE_METHOD(FWaitUntilCompleteHttpFixture, "Retry respect Retry-After header in response", HTTP_TAG)
+{
+	if (!bRetryEnabled)
+	{
+		return;
+	}
+
+	DisableWarningsInThisTest();
+
+	TSharedRef<IHttpRequest> HttpRequest = HttpRetryManager->CreateRequest(
+		1/*InRetryLimitCountOverride*/,
+		FHttpRetrySystem::FRetryTimeoutRelativeSecondsSetting()/*InRetryTimeoutRelativeSecondsOverride unused*/,
+		{EHttpResponseCodes::TooManyRequests, EHttpResponseCodes::ServiceUnavail}/*InRetryResponseCodes*/
+	);
+
+	SECTION("TooManyRequests")
+	{
+		HttpRequest->SetURL(UrlMockStatus(EHttpResponseCodes::TooManyRequests));
+	}
+	SECTION("ServiceUnavail")
+	{
+		HttpRequest->SetURL(UrlMockStatus(EHttpResponseCodes::ServiceUnavail));
+	}
+
+	uint32 RetryAfter = 4;
+
+	HttpRequest->SetVerb(TEXT("GET"));
+	HttpRequest->SetHeader(TEXT("Retry-After"), FString::Format(TEXT("{0}"), { RetryAfter })); // Will be forwarded back in response
+
+	++ExpectingExtraCallbacks;
+	HttpRequest->OnRequestWillRetry().BindLambda([this, RetryAfter](FHttpRequestPtr /*Request*/, FHttpResponsePtr /*Response*/, float LockoutPeriod) {
+		--ExpectingExtraCallbacks;
+		CHECK(FMath::IsNearlyEqual(LockoutPeriod, (float)(RetryAfter)));
+	});
+
+	const double StartTime = FPlatformTime::Seconds();
+	HttpRequest->OnProcessRequestComplete().BindLambda([RetryAfter, StartTime](FHttpRequestPtr HttpRequest, FHttpResponsePtr HttpResponse, bool bSucceeded) {
+		CHECK(bSucceeded);
+		const double DurationInSeconds  = FPlatformTime::Seconds() - StartTime;
+		CHECK(FMath::IsNearlyEqual(DurationInSeconds, (float)(RetryAfter), HTTP_TIME_DIFF_TOLERANCE));
+	});
+
+	HttpRequest->ProcessRequest();
+}
+
+TEST_CASE_METHOD(FWaitUntilCompleteHttpFixture, "Request can time out during lock out", HTTP_TAG)
+{
+	if (!bRetryEnabled)
+	{
+		return;
+	}
+
+	DisableWarningsInThisTest();
+
+	TSharedRef<IHttpRequest> HttpRequest = HttpRetryManager->CreateRequest(
+		1/*InRetryLimitCountOverride*/,
+		FHttpRetrySystem::FRetryTimeoutRelativeSecondsSetting()/*InRetryTimeoutRelativeSecondsOverride unused*/,
+		{EHttpResponseCodes::TooManyRequests}/*InRetryResponseCodes*/
+	);
+
+	HttpRequest->SetURL(UrlMockStatus(EHttpResponseCodes::TooManyRequests));
+	HttpRequest->SetTimeout(1.0f);
+
+	uint32 RetryAfter = 4;
+
+	// Will be forwarded back in response
+	HttpRequest->SetHeader(TEXT("Retry-After"), FString::Format(TEXT("{0}"), { RetryAfter }));
+
+	const double StartTime = FPlatformTime::Seconds();
+	HttpRequest->OnProcessRequestComplete().BindLambda([StartTime](FHttpRequestPtr HttpRequest, FHttpResponsePtr HttpResponse, bool bSucceeded) {
+		// When timeout during lock out period, it fails with result of last request before lock out
+		CHECK(bSucceeded);
+		REQUIRE(HttpResponse != nullptr);
+		CHECK(HttpResponse->GetFailureReason() == EHttpFailureReason::None);
+		CHECK(HttpResponse->GetResponseCode() == EHttpResponseCodes::TooManyRequests);
+		CHECK(HttpResponse->GetContentLength() > 0);
+		const double DurationInSeconds  = FPlatformTime::Seconds() - StartTime;
+		CHECK(FMath::IsNearlyEqual(DurationInSeconds, 1.0, HTTP_TIME_DIFF_TOLERANCE));
+	});
+
+	HttpRequest->ProcessRequest();
+}
+
+TEST_CASE_METHOD(FWaitUntilCompleteHttpFixture, "Request can time out during retry request", HTTP_TAG)
+{
+	if (!bRetryEnabled)
+	{
+		return;
+	}
+
+	DisableWarningsInThisTest();
+
+	TSharedRef<IHttpRequest> HttpRequest = HttpRetryManager->CreateRequest(
+		1/*InRetryLimitCountOverride*/,
+		FHttpRetrySystem::FRetryTimeoutRelativeSecondsSetting()/*InRetryTimeoutRelativeSecondsOverride unused*/,
+		{EHttpResponseCodes::TooManyRequests}/*InRetryResponseCodes*/
+	);
+
+	HttpRequest->SetURL(UrlMockStatus(EHttpResponseCodes::TooManyRequests));
+	HttpRequest->SetTimeout(3.0f);
+
+	uint32 RetryAfter = 2;
+	// Will be forwarded back in response
+	HttpRequest->SetHeader(TEXT("Retry-After"), FString::Format(TEXT("{0}"), { RetryAfter }));
+
+	HttpRequest->OnRequestWillRetry().BindLambda([this](FHttpRequestPtr Request, FHttpResponsePtr /*Response*/, float LockoutPeriod) {
+		// Now retry with a latency during request
+		Request->SetURL(UrlStreamDownload(2/*Chunks*/, HTTP_TEST_TIMEOUT_CHUNK_SIZE, 2/*ChunkLatency*/));
+	});
+
+	const double StartTime = FPlatformTime::Seconds();
+	HttpRequest->OnProcessRequestComplete().BindLambda([StartTime](FHttpRequestPtr HttpRequest, FHttpResponsePtr HttpResponse, bool bSucceeded) {
+		CHECK(HttpRequest->GetResponse() == nullptr); // Timeout will return response with nullptr in request itself
+		CHECK(HttpRequest->GetStatus() == EHttpRequestStatus::Failed);
+		CHECK(HttpRequest->GetFailureReason() == EHttpFailureReason::TimedOut);
+
+		// When timeout during retrying request, it fails with result of last request before retrying, to 
+		// keep it the same behavior when timeout during lockout
+		CHECK(bSucceeded);
+		REQUIRE(HttpResponse != nullptr);
+		CHECK(HttpResponse->GetFailureReason() == EHttpFailureReason::None);
+		CHECK(HttpResponse->GetResponseCode() == EHttpResponseCodes::TooManyRequests);
+		CHECK(HttpResponse->GetContentLength() > 0);
+		const double DurationInSeconds  = FPlatformTime::Seconds() - StartTime;
+		CHECK(FMath::IsNearlyEqual(DurationInSeconds, 3.0, HTTP_TIME_DIFF_TOLERANCE));
+	});
+
+	HttpRequest->ProcessRequest();
+}
+
+TEST_CASE_METHOD(FWaitUntilCompleteHttpFixture, "Retry immediately without lock out if connect failed and there are alt domains", HTTP_TAG)
+{
+	if (!bRetryEnabled)
+	{
+		return;
+	}
+
+	DisableWarningsInThisTest();
+
+	HttpModule->HttpConnectionTimeout = 1.0f;
+
+	TArray<FString> AltDomains;
+	AltDomains.Add(UrlToTestMethods());
+	FHttpRetrySystem::FRetryDomainsPtr RetryDomains = MakeShared<FHttpRetrySystem::FRetryDomains>(MoveTemp(AltDomains));
+	TSharedRef<IHttpRequest> HttpRequest = HttpRetryManager->CreateRequest(
+		1/*InRetryLimitCountOverride*/,
+		FHttpRetrySystem::FRetryTimeoutRelativeSecondsSetting()/*InRetryTimeoutRelativeSecondsOverride unused*/,
+		{ EHttpResponseCodes::TooManyRequests, EHttpResponseCodes::ServiceUnavail }/*InRetryResponseCodes*/,
+		FHttpRetrySystem::FRetryVerbs(),
+		RetryDomains
+	);
+
+	HttpRequest->SetURL(UrlWithInvalidPortToTestConnectTimeout());
+	HttpRequest->SetVerb(TEXT("GET"));
+	++ExpectingExtraCallbacks;
+	HttpRequest->OnRequestWillRetry().BindLambda([this](FHttpRequestPtr /*Request*/, FHttpResponsePtr /*Response*/, float LockoutPeriod) {
+		--ExpectingExtraCallbacks;
+		CHECK(LockoutPeriod == 0);
+	});
+	HttpRequest->ProcessRequest();
+}
+
 class FThreadedBatchRequestsFixture : public FWaitThreadedHttpFixture
 {
 public:
diff --git a/Engine/Source/Programs/WebTestsServer/httptests/urls.py b/Engine/Source/Programs/WebTestsServer/httptests/urls.py
--- a/Engine/Source/Programs/WebTestsServer/httptests/urls.py
+++ b/Engine/Source/Programs/WebTestsServer/httptests/urls.py
@@ -12,4 +12,5 @@ urlpatterns = [
     path("httptests/redirect_from", views.redirect_from, name="redirect_from"),
     path("httptests/redirect_to", views.redirect_to, name="redirect_to"),
     path("httptests/mock_latency/<int:latency>/", views.mock_latency, name="mock_latency"),
+    path("httptests/mock_status/<int:status_code>/", views.mock_status, name="mock_status"),
 ]
diff --git a/Engine/Source/Programs/WebTestsServer/httptests/views.py b/Engine/Source/Programs/WebTestsServer/httptests/views.py
--- a/Engine/Source/Programs/WebTestsServer/httptests/views.py
+++ b/Engine/Source/Programs/WebTestsServer/httptests/views.py
@@ -8,6 +8,7 @@ from django.http import StreamingHttpResponse
 from rest_framework.decorators import api_view, parser_classes
 from rest_framework.parsers import FileUploadParser
 import time
+import json
 import logging
 
 logger = logging.getLogger('django')
@@ -81,3 +82,19 @@ def redirect_to(request):
 def mock_latency(request, latency):
     time.sleep(latency)
     return JsonResponse({})
+
+@api_view(['GET', 'POST', 'DELETE', 'PUT'])
+def mock_status(request, status_code):
+    json_result = json.dumps({
+        "key_a": "value_a",
+        "key_b": "value_b"
+    })
+
+    headers_to_forward = ['Retry-After',]
+
+    response_headers = {}
+    for header_to_forward in headers_to_forward:
+        if header_to_forward in request.headers:
+            response_headers[header_to_forward] = request.headers[header_to_forward]
+
+    return HttpResponse(json_result, content_type="application/json", status=status_code, headers=response_headers)
diff --git a/Engine/Source/Runtime/Online/HTTP/Private/HttpRetrySystem.cpp b/Engine/Source/Runtime/Online/HTTP/Private/HttpRetrySystem.cpp
--- a/Engine/Source/Runtime/Online/HTTP/Private/HttpRetrySystem.cpp
+++ b/Engine/Source/Runtime/Online/HTTP/Private/HttpRetrySystem.cpp
@@ -1,6 +1,7 @@
 // Copyright Epic Games, Inc. All Rights Reserved.
 
 #include "HttpRetrySystem.h"
+#include "GenericPlatform/HttpRequestCommon.h"
 #include "HAL/ConsoleManager.h"
 #include "HAL/PlatformTime.h"
 #include "HAL/PlatformProcess.h"
@@ -24,49 +25,51 @@ LLM_DEFINE_TAG(HTTP);
 
 namespace FHttpRetrySystem
 {
-	TOptional<double> ReadThrottledTimeFromResponseInSeconds(FHttpResponsePtr Response)
+
+TOptional<double> ReadThrottledTimeFromResponseInSeconds(FHttpResponsePtr Response)
+{
+	TOptional<double> LockoutPeriod;
+	// Check if there was a Retry-After header
+	if (Response.IsValid())
 	{
-		TOptional<double> LockoutPeriod;
-		// Check if there was a Retry-After header
-		if (Response.IsValid())
+		int32 ResponseCode = Response->GetResponseCode();
+		if (ResponseCode == EHttpResponseCodes::TooManyRequests || ResponseCode == EHttpResponseCodes::ServiceUnavail)
 		{
-			int32 ResponseCode = Response->GetResponseCode();
-			if (ResponseCode == EHttpResponseCodes::TooManyRequests || ResponseCode == EHttpResponseCodes::ServiceUnavail)
+			FString RetryAfter = Response->GetHeader(TEXT("Retry-After"));
+			if (!RetryAfter.IsEmpty())
 			{
-				FString RetryAfter = Response->GetHeader(TEXT("Retry-After"));
-				if (!RetryAfter.IsEmpty())
+				if (RetryAfter.IsNumeric())
 				{
-					if (RetryAfter.IsNumeric())
-					{
-						// seconds
-						LockoutPeriod.Emplace(FCString::Atof(*RetryAfter));
-					}
-					else
-					{
-						// http date
-						FDateTime UTCServerTime;
-						if (FDateTime::ParseHttpDate(RetryAfter, UTCServerTime))
-						{
-							const FDateTime UTCNow = FDateTime::UtcNow();
-							LockoutPeriod.Emplace((UTCServerTime - UTCNow).GetTotalSeconds());
-						}
-					}
+					// seconds
+					LockoutPeriod.Emplace(FCString::Atof(*RetryAfter));
 				}
 				else
 				{
-					FString RateLimitReset = Response->GetHeader(TEXT("X-Rate-Limit-Reset"));
-					if (!RateLimitReset.IsEmpty())
+					// http date
+					FDateTime UTCServerTime;
+					if (FDateTime::ParseHttpDate(RetryAfter, UTCServerTime))
 					{
-						// UTC seconds
-						const FDateTime UTCServerTime = FDateTime::FromUnixTimestamp(FCString::Atoi64(*RateLimitReset));
 						const FDateTime UTCNow = FDateTime::UtcNow();
 						LockoutPeriod.Emplace((UTCServerTime - UTCNow).GetTotalSeconds());
 					}
 				}
 			}
+			else
+			{
+				FString RateLimitReset = Response->GetHeader(TEXT("X-Rate-Limit-Reset"));
+				if (!RateLimitReset.IsEmpty())
+				{
+					// UTC seconds
+					const FDateTime UTCServerTime = FDateTime::FromUnixTimestamp(FCString::Atoi64(*RateLimitReset));
+					const FDateTime UTCNow = FDateTime::UtcNow();
+					LockoutPeriod.Emplace((UTCServerTime - UTCNow).GetTotalSeconds());
+				}
+			}
 		}
-		return LockoutPeriod;
 	}
+	return LockoutPeriod;
+}
+
 }
 
 FHttpRetrySystem::FRequest::FRequest(
@@ -198,7 +201,7 @@ void FHttpRetrySystem::FRequest::HttpOnProcessRequestComplete(FHttpRequestPtr In
 
 		uint32 EntryIndex = RetryManagerPtr->RequestList.IndexOfByPredicate([this](const FManager::FHttpRetryRequestEntry& Entry) { return Entry.Request == AsShared(); });
 		check(EntryIndex != INDEX_NONE);
-		FManager::FHttpRetryRequestEntry* HttpRetryRequestEntry = &RetryManagerPtr->RequestList[EntryIndex];
+		FManager::FHttpRetryRequestEntry& HttpRetryRequestEntry = RetryManagerPtr->RequestList[EntryIndex];
 
 
 		if (RetryStatus == FHttpRetrySystem::FRequest::EStatus::Cancelled)
@@ -212,22 +215,6 @@ void FHttpRetrySystem::FRequest::HttpOnProcessRequestComplete(FHttpRequestPtr In
 				MoveToNextRetryDomain();
 			}
 
-			if (RetryManagerPtr->ShouldRetry(*HttpRetryRequestEntry) && RetryManagerPtr->CanRetry(*HttpRetryRequestEntry))
-			{
-				const double NowAbsoluteSeconds = FPlatformTime::Seconds();
-				float LockoutPeriod = RetryManagerPtr->GetLockoutPeriodSeconds(*HttpRetryRequestEntry);
-
-				RetryStatus = FHttpRetrySystem::FRequest::EStatus::ProcessingLockout;
-
-				{
-					QUICK_SCOPE_CYCLE_COUNTER(STAT_FHttpRetrySystem_FManager_Update_OnRequestWillRetry);
-					OnRequestWillRetry().ExecuteIfBound(HttpRetryRequestEntry->Request, GetResponse(), LockoutPeriod);
-				}
-
-				RetryManagerPtr->RetryHttpRequestWithDelay(HttpRetryRequestEntry->Request, LockoutPeriod);
-				return;
-			}
-
 			if (GetFailureReason() == EHttpFailureReason::TimedOut)
 			{
 				RetryStatus = FHttpRetrySystem::FRequest::EStatus::FailedTimeout;
@@ -242,7 +229,26 @@ void FHttpRetrySystem::FRequest::HttpOnProcessRequestComplete(FHttpRequestPtr In
 			RetryStatus = FHttpRetrySystem::FRequest::EStatus::Succeeded;
 		}
 
-		if (HttpRetryRequestEntry->CurrentRetryCount > 0)
+		if (RetryStatus != FHttpRetrySystem::FRequest::EStatus::Cancelled && 
+			RetryStatus != FHttpRetrySystem::FRequest::EStatus::FailedTimeout && 
+			RetryManagerPtr->ShouldRetry(HttpRetryRequestEntry) && 
+			RetryManagerPtr->CanRetry(HttpRetryRequestEntry))
+		{
+			const double NowAbsoluteSeconds = FPlatformTime::Seconds();
+			float LockoutPeriod = RetryManagerPtr->GetLockoutPeriodSeconds(HttpRetryRequestEntry);
+
+			RetryStatus = FHttpRetrySystem::FRequest::EStatus::ProcessingLockout;
+
+			{
+				QUICK_SCOPE_CYCLE_COUNTER(STAT_FHttpRetrySystem_FManager_Update_OnRequestWillRetry);
+				OnRequestWillRetry().ExecuteIfBound(HttpRetryRequestEntry.Request, GetResponse(), LockoutPeriod);
+			}
+
+			RetryManagerPtr->RetryHttpRequestWithDelay(HttpRetryRequestEntry, LockoutPeriod, bSucceeded);
+			return;
+		}
+
+		if (HttpRetryRequestEntry.CurrentRetryCount > 0)
 		{
 			FHttpRetrySystem::FManager::FHttpLogVerbosityTracker::Get().DecrementRetriedRequests();
 		}
@@ -250,8 +256,19 @@ void FHttpRetrySystem::FRequest::HttpOnProcessRequestComplete(FHttpRequestPtr In
 		RetryManagerPtr->RequestList.RemoveAtSwap(EntryIndex);
 	}
 
+	FHttpResponsePtr ResultResponse = HttpResponse;
+	bool bResultSucceeded = bSucceeded;
+	if (RetryStatus == FHttpRetrySystem::FRequest::EStatus::FailedTimeout && LastResponse != nullptr)
+	{
+		// Last response is better than nothing when it's timeout
+		ResultResponse = LastResponse;
+		LastResponse.Reset();
+
+		bResultSucceeded = bLastSucceeded;
+	}
+
 	LLM_SCOPE_BYTAG(HTTP);
-	OnProcessRequestComplete().ExecuteIfBound(SelfPtr, HttpResponse, bSucceeded);
+	OnProcessRequestComplete().ExecuteIfBound(SelfPtr, ResultResponse, bResultSucceeded);
 
 	ClearTimeout();
 }
@@ -429,10 +446,24 @@ void FHttpRetrySystem::FManager::RetryHttpRequest(FHttpRetryRequestEntry& Reques
 	RequestEntry.Request->HttpRequest->ProcessRequest();
 }
 
-void FHttpRetrySystem::FManager::RetryHttpRequestWithDelay(const TSharedRef<FRequest>& Request, float InDelay)
+void FHttpRetrySystem::FManager::RetryHttpRequestWithDelay(FManager::FHttpRetryRequestEntry& RequestEntry, float InDelay, bool bWasSucceeded)
 {
-	TWeakPtr<FRequest> RequestWeakPtr(Request);
-	FHttpModule::Get().GetHttpManager().AddHttpThreadTask([RequestWeakPtr]() {
+	// Timeout during lock out to keep existing behavior
+	float TimeoutOrDefault = RequestEntry.Request->GetTimeout().Get(FHttpModule::Get().GetHttpTotalTimeout());
+	if (TimeoutOrDefault != 0)
+	{
+		float TimeElapsedForTheRequest = FPlatformTime::Seconds() - RequestEntry.RequestStartTimeAbsoluteSeconds;
+		float WillTimeoutInDelay = TimeoutOrDefault - TimeElapsedForTheRequest;
+		if (WillTimeoutInDelay < InDelay)
+		{
+			HttpRequestTimeoutAfterDelay(RequestEntry, bWasSucceeded, TimeoutOrDefault);
+			return;
+		}
+	}
+
+	// Delay and start
+	TWeakPtr<FRequest> RequestWeakPtr(RequestEntry.Request);
+	FHttpModule::Get().GetHttpManager().AddHttpThreadTask([RequestWeakPtr, bWasSucceeded]() {
 		if (TSharedPtr<FRequest> RequestPtr = RequestWeakPtr.Pin())
 		{
 			if (TSharedPtr<FManager> RetryManagerPtr = RequestPtr->RetryManager.Pin())
@@ -443,6 +474,9 @@ void FHttpRetrySystem::FManager::RetryHttpRequestWithDelay(const TSharedRef<FReq
 				if (EntryIndex != INDEX_NONE)
 				{
 					FManager::FHttpRetryRequestEntry* HttpRetryRequestEntry = &RetryManagerPtr->RequestList[EntryIndex];
+					// TODO: Move this into RetryHttpRequest after stablizing the flow with CVarHttpRetrySystemNonGameThreadSupportEnabled on
+					HttpRetryRequestEntry->Request->LastResponse = HttpRetryRequestEntry->Request->GetResponse();
+					HttpRetryRequestEntry->Request->bLastSucceeded = bWasSucceeded;
 					RetryManagerPtr->RetryHttpRequest(*HttpRetryRequestEntry);
 				}
 			}
@@ -450,6 +484,33 @@ void FHttpRetrySystem::FManager::RetryHttpRequestWithDelay(const TSharedRef<FReq
 	}, InDelay);
 }
 
+void FHttpRetrySystem::FManager::HttpRequestTimeoutAfterDelay(FManager::FHttpRetryRequestEntry& RequestEntry, bool bWasSucceeded, float Delay)
+{
+	TWeakPtr<FRequest> RequestWeakPtr(RequestEntry.Request);
+	FHttpModule::Get().GetHttpManager().AddHttpThreadTask([RequestWeakPtr, bWasSucceeded]() {
+		if (TSharedPtr<FRequest> RequestPtr = RequestWeakPtr.Pin())
+		{
+			if (TSharedPtr<FManager> RetryManagerPtr = RequestPtr->RetryManager.Pin())
+			{
+				FScopeLock ScopeLock(&RetryManagerPtr->RequestListLock);
+				uint32 EntryIndex = RetryManagerPtr->RequestList.IndexOfByPredicate([RequestPtr](const FManager::FHttpRetryRequestEntry& Entry) { return Entry.Request == RequestPtr; });
+				if (EntryIndex != INDEX_NONE)
+				{
+					FManager::FHttpRetryRequestEntry* HttpRetryRequestEntry = &RetryManagerPtr->RequestList[EntryIndex];
+					if (HttpRetryRequestEntry->CurrentRetryCount > 0)
+					{
+						FHttpRetrySystem::FManager::FHttpLogVerbosityTracker::Get().DecrementRetriedRequests();
+					}
+					RetryManagerPtr->RequestList.RemoveAtSwap(EntryIndex);
+				}
+			}
+
+			// Same as existing behavior, when timeout during lock out period, it fails with result of last request before lockout
+			RequestPtr->OnProcessRequestComplete().ExecuteIfBound(RequestPtr, RequestPtr->GetResponse(), bWasSucceeded);
+		}
+	}, Delay);
+}
+
 float FHttpRetrySystem::FManager::GetLockoutPeriodSeconds(const FHttpRetryRequestEntry& HttpRetryRequestEntry)
 {
 	float LockoutPeriod = 0.0f;
diff --git a/Engine/Source/Runtime/Online/HTTP/Public/HttpRetrySystem.h b/Engine/Source/Runtime/Online/HTTP/Public/HttpRetrySystem.h
--- a/Engine/Source/Runtime/Online/HTTP/Public/HttpRetrySystem.h
+++ b/Engine/Source/Runtime/Online/HTTP/Public/HttpRetrySystem.h
@@ -123,19 +123,22 @@ namespace FHttpRetrySystem
 		/** Move to the next retry domain from our RetryDomains */
 		void MoveToNextRetryDomain();
 
-		EStatus::Type                        RetryStatus;
+		EStatus::Type RetryStatus;
 
-		FRetryLimitCountSetting              RetryLimitCountOverride;
+		FRetryLimitCountSetting RetryLimitCountOverride;
 		FRetryTimeoutRelativeSecondsSetting  RetryTimeoutRelativeSecondsOverride;
-		FRetryResponseCodes					 RetryResponseCodes;
-		FRetryVerbs                          RetryVerbs;
-		FRetryDomainsPtr					 RetryDomains;
+		FRetryResponseCodes RetryResponseCodes;
+		FRetryVerbs RetryVerbs;
+		FRetryDomainsPtr RetryDomains;
 		/** The current index in RetryDomains we are attempting */
-		int32								 RetryDomainsIndex = 0;
+		int32 RetryDomainsIndex = 0;
 		/** The original URL before replacing anything from RetryDomains */
-		FString								 OriginalUrl;
+		FString OriginalUrl;
 
-		TWeakPtr<FManager>					 RetryManager;
+		TWeakPtr<FManager> RetryManager;
+		/** Save the last response before the retry */
+		FHttpResponsePtr LastResponse;
+		bool bLastSucceeded = false;
 	};
 }
 
@@ -242,10 +245,13 @@ protected:
 
 	/**
 	 * Retry an HTTP request with delay
-	 * @param RequestToRetry request to retry
+	 * @param RequestEntry request retry
 	 * @param InDelay the delay to wait before retrying
+	 * @param bWasSucceeded was the request succeeded before retry
 	 */
-	void RetryHttpRequestWithDelay(const TSharedRef<FRequest>& Request, float InDelay);
+	void RetryHttpRequestWithDelay(FManager::FHttpRetryRequestEntry& RequestEntry, float InDelay, bool bWasSucceeded);
+
+	void HttpRequestTimeoutAfterDelay(FManager::FHttpRetryRequestEntry& RequestEntry, bool bWasSucceeded, float Delay);
 
 	// @return number of seconds to lockout for
 	float GetLockoutPeriodSeconds(const FHttpRetryRequestEntry& HttpRetryRequestEntry);
