Added optional, separated retry limit support for connection error in http retry system.

#tests Passed new test case in WebTests

diff --git a/Engine/Source/Programs/WebTests/Private/TestHttp.cpp b/Engine/Source/Programs/WebTests/Private/TestHttp.cpp
--- a/Engine/Source/Programs/WebTests/Private/TestHttp.cpp
+++ b/Engine/Source/Programs/WebTests/Private/TestHttp.cpp
@@ -170,6 +170,8 @@ public:
 	using FHttpRetrySystem::FManager::RequestList;
 	using FHttpRetrySystem::FManager::FHttpRetryRequestEntry;
 	using FHttpRetrySystem::FManager::RetryTimeoutRelativeSecondsDefault;
+	using FHttpRetrySystem::FManager::RetryLimitCountDefault;
+	using FHttpRetrySystem::FManager::RetryLimitCountForConnectionErrorDefault;
 
 	bool IsEmpty()
 	{
@@ -1536,6 +1538,57 @@ TEST_CASE_METHOD(FWaitUntilCompleteHttpFixture, "Retry immediately without lock
 	HttpRequest->ProcessRequest();
 }
 
+TEST_CASE_METHOD(FWaitUntilCompleteHttpFixture, "Optionally retry limit can be set differently for connection error", HTTP_TAG)
+{
+	if (!bRetryEnabled)
+	{
+		return;
+	}
+
+	DisableWarningsInThisTest();
+
+	TSharedRef<IHttpRequest> HttpRequest = HttpRetryManager->CreateRequest(
+		5/*InRetryLimitCountOverride*/,
+		FHttpRetrySystem::FRetryTimeoutRelativeSecondsSetting()/*InRetryTimeoutRelativeSecondsOverride unused*/,
+		{ EHttpResponseCodes::TooManyRequests, EHttpResponseCodes::ServiceUnavail }/*InRetryResponseCodes*/,
+		FHttpRetrySystem::FRetryVerbs(), /*unused*/
+		FHttpRetrySystem::FRetryDomainsPtr(), /*unused*/
+		1 /*InRetryLimitCountForConnectionErrorOverride*/
+	);
+
+	float ExpectedTimeoutDuration = 0.0f;
+	SECTION("RetryLimitCountDefault:5 will be used so retries in general take long")
+	{
+		HttpRequest->SetURL(UrlMockStatus(EHttpResponseCodes::TooManyRequests));
+		HttpRequest->SetHeader(TEXT("Retry-After"), FString::Format(TEXT("{0}"), { 3 }));
+
+		ExpectedTimeoutDuration = 15.0f; // each request will take about 0s, 5 retry back offs, each back off takes 3s;
+	}
+	SECTION("RetryLimitCountForConnectionErrorDefault:1 will be used so retries for connection error take less time")
+	{
+		HttpRequest->SetURL(UrlStreamDownload(2/*Chunks*/, HTTP_TEST_TIMEOUT_CHUNK_SIZE, 2/*ChunkLatency*/));
+		HttpModule->HttpActivityTimeout = 1.0f;
+
+		ExpectedTimeoutDuration = 2.0f; // each request will take 1s, 1st retry back off takes 0s
+
+#if WITH_CURL_XCURL
+		// Unlike libCurl, currently there is an issue in xCurl that it triggers CURLINFO_HEADER_OUT even if can't 
+		// connect. Had to disable that code, make sure not to treat that event as connected
+		// It takes 2s to receive the first chunk to be considered as connected, then start response timer and 
+		// take 1s to time out. So it takes 3s in total for each request instead of 1s
+		ExpectedTimeoutDuration += 4;
+#endif
+	}
+
+	const double StartTime = FPlatformTime::Seconds();
+	HttpRequest->OnProcessRequestComplete().BindLambda([StartTime, ExpectedTimeoutDuration](FHttpRequestPtr HttpRequest, FHttpResponsePtr HttpResponse, bool bSucceeded) {
+		const double DurationInSeconds  = FPlatformTime::Seconds() - StartTime;
+		float TimeDiffTolerance = 1.0f;
+		CHECK(FMath::IsNearlyEqual(DurationInSeconds, ExpectedTimeoutDuration, TimeDiffTolerance));
+	});
+	HttpRequest->ProcessRequest();
+}
+
 class FThreadedBatchRequestsFixture : public FWaitThreadedHttpFixture
 {
 public:
diff --git a/Engine/Source/Runtime/Online/HTTP/Private/HttpRetrySystem.cpp b/Engine/Source/Runtime/Online/HTTP/Private/HttpRetrySystem.cpp
--- a/Engine/Source/Runtime/Online/HTTP/Private/HttpRetrySystem.cpp
+++ b/Engine/Source/Runtime/Online/HTTP/Private/HttpRetrySystem.cpp
@@ -79,11 +79,13 @@ FHttpRetrySystem::FRequest::FRequest(
 	const FHttpRetrySystem::FRetryTimeoutRelativeSecondsSetting& InRetryTimeoutRelativeSecondsOverride,
 	const FHttpRetrySystem::FRetryResponseCodes& InRetryResponseCodes,
 	const FHttpRetrySystem::FRetryVerbs& InRetryVerbs,
-	const FHttpRetrySystem::FRetryDomainsPtr& InRetryDomains
+	const FHttpRetrySystem::FRetryDomainsPtr& InRetryDomains,
+	const FRetryLimitCountSetting& InRetryLimitCountForConnectionErrorOverride
 	)
     : FHttpRequestAdapterBase(HttpRequest)
     , RetryStatus(FHttpRetrySystem::FRequest::EStatus::NotStarted)
     , RetryLimitCountOverride(InRetryLimitCountOverride)
+    , RetryLimitCountForConnectionErrorOverride(InRetryLimitCountForConnectionErrorOverride)
     , RetryTimeoutRelativeSecondsOverride(InRetryTimeoutRelativeSecondsOverride)
 	, RetryResponseCodes(InRetryResponseCodes)
 	, RetryVerbs(InRetryVerbs)
@@ -290,9 +292,14 @@ void FHttpRetrySystem::FRequest::HttpOnHeaderReceived(FHttpRequestPtr Request, c
 	OnHeaderReceived().ExecuteIfBound(SelfPtr, HeaderName, NewHeaderValue);
 }
 
-FHttpRetrySystem::FManager::FManager(const FRetryLimitCountSetting& InRetryLimitCountDefault, const FRetryTimeoutRelativeSecondsSetting& InRetryTimeoutRelativeSecondsDefault)
+FHttpRetrySystem::FManager::FManager(
+	const FRetryLimitCountSetting& InRetryLimitCountDefault,
+	const FRetryTimeoutRelativeSecondsSetting& InRetryTimeoutRelativeSecondsDefault,
+	const FRetryLimitCountSetting& InRetryLimitCountForConnectionErrorDefault
+)
     : RandomFailureRate(FRandomFailureRateSetting())
     , RetryLimitCountDefault(InRetryLimitCountDefault)
+    , RetryLimitCountForConnectionErrorDefault(InRetryLimitCountForConnectionErrorDefault)
 	, RetryTimeoutRelativeSecondsDefault(InRetryTimeoutRelativeSecondsDefault)
 {
 	check(FHttpModule::Get().GetHttpManager().GetThread());
@@ -317,7 +324,9 @@ TSharedRef<FHttpRetrySystem::FRequest, ESPMode::ThreadSafe> FHttpRetrySystem::FM
 	const FRetryTimeoutRelativeSecondsSetting& InRetryTimeoutRelativeSecondsOverride,
 	const FRetryResponseCodes& InRetryResponseCodes,
 	const FRetryVerbs& InRetryVerbs,
-	const FRetryDomainsPtr& InRetryDomains)
+	const FRetryDomainsPtr& InRetryDomains,
+	const FRetryLimitCountSetting& InRetryLimitCountForConnectionErrorOverride
+)
 {
 	return MakeShareable(new FRequest(
 		AsShared(),
@@ -326,7 +335,8 @@ TSharedRef<FHttpRetrySystem::FRequest, ESPMode::ThreadSafe> FHttpRetrySystem::FM
 		InRetryTimeoutRelativeSecondsOverride,
 		InRetryResponseCodes,
 		InRetryVerbs,
-		InRetryDomains
+		InRetryDomains,
+		InRetryLimitCountForConnectionErrorOverride
 		));
 }
 
@@ -377,32 +387,31 @@ bool FHttpRetrySystem::FManager::ShouldRetry(const FHttpRetryRequestEntry& HttpR
     return bResult;
 }
 
-bool FHttpRetrySystem::FManager::CanRetry(const FHttpRetryRequestEntry& HttpRetryRequestEntry)
+bool FHttpRetrySystem::FManager::RetryLimitForConnectionErrorIsSet(const FHttpRetryRequestEntry& HttpRetryRequestEntry)
 {
-    bool bResult = false;
+	return HttpRetryRequestEntry.Request->RetryLimitCountForConnectionErrorOverride.IsSet() || RetryLimitCountForConnectionErrorDefault.IsSet();
+}
 
-    bool bShouldTestCurrentRetryCount = false;
-    double RetryLimitCount = 0;
-    if (HttpRetryRequestEntry.Request->RetryLimitCountOverride.IsSet())
-    {
-        bShouldTestCurrentRetryCount = true;
-        RetryLimitCount = HttpRetryRequestEntry.Request->RetryLimitCountOverride.GetValue();
-    }
-    else if (RetryLimitCountDefault.IsSet())
-    {
-        bShouldTestCurrentRetryCount = true;
-        RetryLimitCount = RetryLimitCountDefault.GetValue();
-    }
+bool FHttpRetrySystem::FManager::CanRetryForConnectionError(const FHttpRetryRequestEntry& HttpRetryRequestEntry)
+{
+	uint32 RetryLimitForConnectionError = HttpRetryRequestEntry.Request->RetryLimitCountForConnectionErrorOverride.Get(RetryLimitCountForConnectionErrorDefault.Get(0));
+	return HttpRetryRequestEntry.CurrentRetryCountForConnectionError < RetryLimitForConnectionError;
+}
 
-    if (bShouldTestCurrentRetryCount)
-    {
-        if (HttpRetryRequestEntry.CurrentRetryCount < RetryLimitCount)
-        {
-            bResult = true;
-        }
-    }
+bool FHttpRetrySystem::FManager::CanRetryInGeneral(const FHttpRetryRequestEntry& HttpRetryRequestEntry)
+{
+	uint32 RetryLimit = HttpRetryRequestEntry.Request->RetryLimitCountOverride.Get(RetryLimitCountDefault.Get(0));
+	return HttpRetryRequestEntry.CurrentRetryCount < RetryLimit;
+}
 
-    return bResult;
+bool FHttpRetrySystem::FManager::CanRetry(const FHttpRetryRequestEntry& HttpRetryRequestEntry)
+{
+	if (HttpRetryRequestEntry.Request->GetFailureReason() == EHttpFailureReason::ConnectionError && RetryLimitForConnectionErrorIsSet(HttpRetryRequestEntry))
+	{
+		return CanRetryForConnectionError(HttpRetryRequestEntry);
+	}
+
+	return CanRetryInGeneral(HttpRetryRequestEntry);
 }
 
 bool FHttpRetrySystem::FManager::HasTimedOut(const FHttpRetryRequestEntry& HttpRetryRequestEntry, const double NowAbsoluteSeconds)
@@ -441,6 +450,10 @@ void FHttpRetrySystem::FManager::RetryHttpRequest(FHttpRetryRequestEntry& Reques
 		FHttpLogVerbosityTracker::Get().IncrementRetriedRequests();
 	}
 	++RequestEntry.CurrentRetryCount;
+	if (RequestEntry.Request->GetFailureReason() == EHttpFailureReason::ConnectionError)
+	{
+		++RequestEntry.CurrentRetryCountForConnectionError;
+	}
 	RequestEntry.Request->RetryStatus = FRequest::EStatus::Processing;
 	UE_LOG(LogHttp, Warning, TEXT("Retry %d on %s"), RequestEntry.CurrentRetryCount, *(RequestEntry.Request->GetURL()));
 	RequestEntry.Request->HttpRequest->ProcessRequest();
@@ -791,6 +804,7 @@ bool FHttpRetrySystem::FManager::Update(uint32* FileCount, uint32* FailingCount,
 FHttpRetrySystem::FManager::FHttpRetryRequestEntry::FHttpRetryRequestEntry(TSharedRef<FHttpRetrySystem::FRequest, ESPMode::ThreadSafe>& InRequest)
     : bShouldCancel(false)
     , CurrentRetryCount(0)
+    , CurrentRetryCountForConnectionError(0)
 	, RequestStartTimeAbsoluteSeconds(FPlatformTime::Seconds())
 	, Request(InRequest)
 {}
diff --git a/Engine/Source/Runtime/Online/HTTP/Public/HttpRetrySystem.h b/Engine/Source/Runtime/Online/HTTP/Public/HttpRetrySystem.h
--- a/Engine/Source/Runtime/Online/HTTP/Public/HttpRetrySystem.h
+++ b/Engine/Source/Runtime/Online/HTTP/Public/HttpRetrySystem.h
@@ -110,7 +110,8 @@ namespace FHttpRetrySystem
 			const FRetryTimeoutRelativeSecondsSetting& InRetryTimeoutRelativeSecondsOverride = FRetryTimeoutRelativeSecondsSetting(),
 			const FRetryResponseCodes& InRetryResponseCodes = FRetryResponseCodes(),
 			const FRetryVerbs& InRetryVerbs = FRetryVerbs(),
-			const FRetryDomainsPtr& InRetryDomains = FRetryDomainsPtr()
+			const FRetryDomainsPtr& InRetryDomains = FRetryDomainsPtr(),
+			const FRetryLimitCountSetting& InRetryLimitCountForConnectionErrorOverride = FRetryLimitCountSetting()
 		);
 
 		void HttpOnRequestProgress(FHttpRequestPtr InHttpRequest, uint64 BytesSent, uint64 BytesRcv);
@@ -126,6 +127,7 @@ namespace FHttpRetrySystem
 		EStatus::Type RetryStatus;
 
 		FRetryLimitCountSetting RetryLimitCountOverride;
+		FRetryLimitCountSetting RetryLimitCountForConnectionErrorOverride;
 		FRetryTimeoutRelativeSecondsSetting  RetryTimeoutRelativeSecondsOverride;
 		FRetryResponseCodes RetryResponseCodes;
 		FRetryVerbs RetryVerbs;
@@ -149,7 +151,11 @@ class FManager : public TSharedFromThis<FManager>
 {
 public:
 	// FManager
-	HTTP_API FManager(const FRetryLimitCountSetting& InRetryLimitCountDefault, const FRetryTimeoutRelativeSecondsSetting& InRetryTimeoutRelativeSecondsDefault);
+	HTTP_API FManager(
+		const FRetryLimitCountSetting& InRetryLimitCountDefault, 
+		const FRetryTimeoutRelativeSecondsSetting& InRetryTimeoutRelativeSecondsDefault,
+		const FRetryLimitCountSetting& InRetryLimitCountForConnectionErrorDefault = FRetryLimitCountSetting()
+	);
 
 	/**
 	 * Create a new http request with retries
@@ -159,7 +165,8 @@ public:
 		const FRetryTimeoutRelativeSecondsSetting& InRetryTimeoutRelativeSecondsOverride = FRetryTimeoutRelativeSecondsSetting(),
 		const FRetryResponseCodes& InRetryResponseCodes = FRetryResponseCodes(),
 		const FRetryVerbs& InRetryVerbs = FRetryVerbs(),
-		const FRetryDomainsPtr& InRetryDomains = FRetryDomainsPtr()
+		const FRetryDomainsPtr& InRetryDomains = FRetryDomainsPtr(),
+		const FRetryLimitCountSetting& InRetryLimitCountForConnectionErrorOverride = FRetryLimitCountSetting()
 	);
 
 	HTTP_API virtual ~FManager();
@@ -190,10 +197,11 @@ protected:
 	{
 		FHttpRetryRequestEntry(TSharedRef<FRequest, ESPMode::ThreadSafe>& InRequest);
 
-		bool                    bShouldCancel;
-		uint32                  CurrentRetryCount;
-		double                  RequestStartTimeAbsoluteSeconds;
-		double                  LockoutEndTimeAbsoluteSeconds;
+		bool bShouldCancel;
+		uint32 CurrentRetryCount;
+		uint32 CurrentRetryCountForConnectionError;
+		double RequestStartTimeAbsoluteSeconds;
+		double LockoutEndTimeAbsoluteSeconds;
 
 		TSharedRef<FRequest, ESPMode::ThreadSafe>	Request;
 	};
@@ -256,9 +264,14 @@ protected:
 	// @return number of seconds to lockout for
 	float GetLockoutPeriodSeconds(const FHttpRetryRequestEntry& HttpRetryRequestEntry);
 
+	bool RetryLimitForConnectionErrorIsSet(const FHttpRetryRequestEntry& HttpRetryRequestEntry);
+	bool CanRetryForConnectionError(const FHttpRetryRequestEntry& HttpRetryRequestEntry);
+	bool CanRetryInGeneral(const FHttpRetryRequestEntry& HttpRetryRequestEntry);
+
 	// Default configuration for the retry system
 	FRandomFailureRateSetting RandomFailureRate;
 	FRetryLimitCountSetting RetryLimitCountDefault;
+	FRetryLimitCountSetting RetryLimitCountForConnectionErrorDefault;
 	FRetryTimeoutRelativeSecondsSetting RetryTimeoutRelativeSecondsDefault;
 
 	TArray<FHttpRetryRequestEntry> RequestList;
