Make the http request don't output warning if the request get canceled;
Move the default max concurrent http requests limit to compile time decided by different platforms.

#tests Passed WebTests project

diff --git a/Engine/Source/Programs/WebTests/Private/TestHttp.cpp b/Engine/Source/Programs/WebTests/Private/TestHttp.cpp
--- a/Engine/Source/Programs/WebTests/Private/TestHttp.cpp
+++ b/Engine/Source/Programs/WebTests/Private/TestHttp.cpp
@@ -1755,9 +1755,9 @@ TEST_CASE_METHOD(FWaitUntilCompleteHttpFixture, "Test platform request requests
 
 	DYNAMIC_SECTION(" making " << NumRequests << " requests with bCheckCancel=" << bCheckCancel)
 	{
-		if (NumRequests > 100 && !bRunHeavyTests)
+		if (NumRequests > 50 && !bRunHeavyTests)
 		{
-			SKIP("-run_heavy_tests is unset, limiting NumRequests to 100");
+			return;
 		}
 
 		TArray<TSharedRef<IHttpRequest>> Requests;
diff --git a/Engine/Source/Runtime/Online/HTTP/HTTP.Build.cs b/Engine/Source/Runtime/Online/HTTP/HTTP.Build.cs
--- a/Engine/Source/Runtime/Online/HTTP/HTTP.Build.cs
+++ b/Engine/Source/Runtime/Online/HTTP/HTTP.Build.cs
@@ -28,6 +28,7 @@ public class HTTP : ModuleRules
 
 	protected virtual bool bPlatformSupportsCurlMultiWait { get { return false; } }
 	protected virtual bool bPlatformSupportsCurlQuickExit { get { return !bPlatformSupportsXCurl; } }
+	protected virtual int DefaultMaxConcurrentRequests { get { return int.MaxValue; } }
 
 	private bool bPlatformSupportsCurl { get { return bPlatformSupportsLibCurl || bPlatformSupportsXCurl; } }
 
@@ -114,5 +115,7 @@ public class HTTP : ModuleRules
 		{
 			PublicFrameworks.Add("Security");
 		}
+
+		PrivateDefinitions.Add("UE_HTTP_DEFAULT_MAX_CONCURRENT_REQUESTS=" + DefaultMaxConcurrentRequests);
 	}
 }
diff --git a/Engine/Source/Runtime/Online/HTTP/Private/Apple/AppleHttp.cpp b/Engine/Source/Runtime/Online/HTTP/Private/Apple/AppleHttp.cpp
--- a/Engine/Source/Runtime/Online/HTTP/Private/Apple/AppleHttp.cpp
+++ b/Engine/Source/Runtime/Online/HTTP/Private/Apple/AppleHttp.cpp
@@ -237,11 +237,6 @@
 	}
 	else
 	{
-		UE_LOG(LogHttp, Warning, TEXT("URLSession:task:didCompleteWithError. Http request failed - %s %s: %p"),
-			   *FString([error localizedDescription]),
-			   *FString([[error userInfo] objectForKey:NSURLErrorFailingURLStringErrorKey]),
-			   self);
-		
 		self.RequestStatus = EHttpRequestStatus::Failed;
 		// Determine if the specific error was failing to connect to the host.
 		switch ([error code])
@@ -259,6 +254,12 @@
 				self.FailureReason = EHttpFailureReason::Other;
 				break;
 		}
+
+		UE_CLOG(self.FailureReason != EHttpFailureReason::Cancelled, LogHttp, Warning, TEXT("URLSession:task:didCompleteWithError. Http request failed - %s %s: %p"),
+			   *FString([error localizedDescription]),
+			   *FString([[error userInfo] objectForKey:NSURLErrorFailingURLStringErrorKey]),
+			   self);
+
 		// Log more details if verbose logging is enabled and this is an SSL error
 		if (UE_LOG_ACTIVE(LogHttp, Verbose))
 		{
diff --git a/Engine/Source/Runtime/Online/HTTP/Private/HttpThread.cpp b/Engine/Source/Runtime/Online/HTTP/Private/HttpThread.cpp
--- a/Engine/Source/Runtime/Online/HTTP/Private/HttpThread.cpp
+++ b/Engine/Source/Runtime/Online/HTTP/Private/HttpThread.cpp
@@ -45,9 +45,10 @@ private:
 // FHttpThread
 
 FHttpThreadBase::FHttpThreadBase()
-	:	Thread(nullptr)
-	,	bIsSingleThread(false)
-	,	bIsStopped(true)
+	: Thread(nullptr)
+	, bIsSingleThread(false)
+	, bIsStopped(true)
+	, RunningThreadedRequestLimit(UE_HTTP_DEFAULT_MAX_CONCURRENT_REQUESTS)
 {
 }
 
diff --git a/Engine/Source/Runtime/Online/HTTP/Private/HttpThread.h b/Engine/Source/Runtime/Online/HTTP/Private/HttpThread.h
--- a/Engine/Source/Runtime/Online/HTTP/Private/HttpThread.h
+++ b/Engine/Source/Runtime/Online/HTTP/Private/HttpThread.h
@@ -151,7 +151,7 @@ private:
 	bool bIsStopped;
 
 	/** Limit for threaded http requests running at the same time. If not specified through configuration values, there will be no limit */
-	std::atomic<int32> RunningThreadedRequestLimit = INT_MAX;
+	std::atomic<int32> RunningThreadedRequestLimit;
 
 	/** Last time the thread has been processed. Used in the non-game thread. */
 	double LastTime;
