Only output warning when request stay in http thread queue for too long;
Added the MaxTimeToWaitInQueue stat in HttpStats.

diff --git a/Engine/Source/Programs/WebTests/Private/TestHttp.cpp b/Engine/Source/Programs/WebTests/Private/TestHttp.cpp
--- a/Engine/Source/Programs/WebTests/Private/TestHttp.cpp
+++ b/Engine/Source/Programs/WebTests/Private/TestHttp.cpp
@@ -1766,7 +1766,7 @@ TEST_CASE_METHOD(FWaitUntilCompleteHttpFixture, "Test platform request requests
 		{
 			TSharedRef<IHttpRequest> HttpRequest = FHttpModule::Get().CreateRequest();
 			// Requests server to serve 1024b chunks to allow time for cancel to happen
-			HttpRequest->SetURL(UrlStreamDownload(2, 1024, 1));
+			HttpRequest->SetURL(UrlStreamDownload(2, HTTP_TEST_TIMEOUT_CHUNK_SIZE, /*ChunkLatency=*/bCheckCancel ? 1 : 0));
 			HttpRequest->SetVerb(TEXT("GET"));
 			HttpRequest->OnProcessRequestComplete().BindLambda([bCheckCancel](FHttpRequestPtr HttpRequest, FHttpResponsePtr HttpResponse, bool bSucceeded)
 			{
diff --git a/Engine/Source/Runtime/Online/HTTP/Private/HttpManager.cpp b/Engine/Source/Runtime/Online/HTTP/Private/HttpManager.cpp
--- a/Engine/Source/Runtime/Online/HTTP/Private/HttpManager.cpp
+++ b/Engine/Source/Runtime/Online/HTTP/Private/HttpManager.cpp
@@ -589,3 +589,8 @@ void FHttpManager::RecordStatRequestsInQueue(uint32 RequestsInQueue)
 {
 	HttpStats.MaxRequestsInQueue = FGenericPlatformMath::Max(RequestsInQueue, HttpStats.MaxRequestsInQueue);
 }
+
+void FHttpManager::RecordMaxTimeToWaitInQueue(float Duration)
+{
+	HttpStats.MaxTimeToWaitInQueue = FGenericPlatformMath::Max(Duration, HttpStats.MaxTimeToWaitInQueue);
+}
diff --git a/Engine/Source/Runtime/Online/HTTP/Private/HttpThread.cpp b/Engine/Source/Runtime/Online/HTTP/Private/HttpThread.cpp
--- a/Engine/Source/Runtime/Online/HTTP/Private/HttpThread.cpp
+++ b/Engine/Source/Runtime/Online/HTTP/Private/HttpThread.cpp
@@ -207,6 +207,7 @@ void FHttpThreadBase::Process(TArray<IHttpThreadedRequest*>& RequestsToCancel, T
 
 		while (NewThreadedRequests.Dequeue(Request))
 		{
+			Request->StartWaitingInQueue();
 			RateLimitedThreadedRequests.Add(Request);
 		}
 	}
@@ -249,25 +250,28 @@ void FHttpThreadBase::Process(TArray<IHttpThreadedRequest*>& RequestsToCancel, T
 	const int32 LocalRunningThreadedRequestLimit = GetRunningThreadedRequestLimit();
 	if (RunningThreadedRequestsCounter < LocalRunningThreadedRequestLimit)
 	{
-		while(RunningThreadedRequestsCounter < LocalRunningThreadedRequestLimit && RateLimitedThreadedRequests.Num())
+		while(RunningThreadedRequestsCounter < LocalRunningThreadedRequestLimit && !RateLimitedThreadedRequests.IsEmpty())
 		{
 			SCOPE_CYCLE_COUNTER(STAT_HTTPThread_StartThreadedRequest);
 
 			IHttpThreadedRequest* ReadyThreadedRequest = RateLimitedThreadedRequests[0];
 			RateLimitedThreadedRequests.RemoveAt(0);
 
+			float DurationInQueue = FPlatformTime::Seconds() - ReadyThreadedRequest->GetTimeStartedWaitingInQueue();
+			UE_CLOG(DurationInQueue > 10.0f, LogHttp, Warning, TEXT("Request (%p) waited in queue for %.2fs before starting"), ReadyThreadedRequest, DurationInQueue);
+			float StartImmediately = 0.01f;
+			if (DurationInQueue > StartImmediately)
+			{
+				FHttpModule::Get().GetHttpManager().RecordMaxTimeToWaitInQueue(DurationInQueue);
+			}
+
 			if (StartThreadedRequest(ReadyThreadedRequest))
 			{
 				RunningThreadedRequestsCounter++;
 				RunningThreadedRequests.Add(ReadyThreadedRequest);
 				ReadyThreadedRequest->TickThreadedRequest(0.0f);
-				UE_LOG(LogHttp, Verbose, TEXT("Started running threaded request (%p). Running threaded requests (%d) Rate limited threaded requests (%d)"), ReadyThreadedRequest, RunningThreadedRequests.Num(), RateLimitedThreadedRequests.Num());
-#if WITH_SERVER_CODE
-				if (RunningThreadedRequestsCounter == LocalRunningThreadedRequestLimit)
-				{
-					UE_LOG(LogHttp, Warning, TEXT("Reached threaded request limit (%d)"), RunningThreadedRequestsCounter);
-				}
-#endif // WITH_SERVER_CODE
+				UE_LOG(LogHttp, Verbose, TEXT("Started http request in thread (%p). Waited in queue for (%.2fs) Running threaded requests (%d) Rate limited threaded requests (%d)"), 
+					ReadyThreadedRequest, DurationInQueue, RunningThreadedRequests.Num(), RateLimitedThreadedRequests.Num());
 			}
 			else
 			{
diff --git a/Engine/Source/Runtime/Online/HTTP/Private/IHttpThreadedRequest.cpp b/Engine/Source/Runtime/Online/HTTP/Private/IHttpThreadedRequest.cpp
--- a/Engine/Source/Runtime/Online/HTTP/Private/IHttpThreadedRequest.cpp
+++ b/Engine/Source/Runtime/Online/HTTP/Private/IHttpThreadedRequest.cpp
@@ -1,6 +1,17 @@
 // Copyright Epic Games, Inc. All Rights Reserved.
 
 #include "IHttpThreadedRequest.h"
+#include "HAL/PlatformTime.h"
 #include "HttpModule.h"
 #include "HttpManager.h"
 
+void IHttpThreadedRequest::StartWaitingInQueue()
+{
+	TimeStartedWaitingInQueue = FPlatformTime::Seconds();
+}
+
+float IHttpThreadedRequest::GetTimeStartedWaitingInQueue() const
+{
+	check(TimeStartedWaitingInQueue != 0);
+	return TimeStartedWaitingInQueue;
+}
diff --git a/Engine/Source/Runtime/Online/HTTP/Private/IHttpThreadedRequest.h b/Engine/Source/Runtime/Online/HTTP/Private/IHttpThreadedRequest.h
--- a/Engine/Source/Runtime/Online/HTTP/Private/IHttpThreadedRequest.h
+++ b/Engine/Source/Runtime/Online/HTTP/Private/IHttpThreadedRequest.h
@@ -12,4 +12,10 @@ public:
 	virtual bool StartThreadedRequest() = 0;
 	virtual bool IsThreadedRequestComplete() = 0;
 	virtual void TickThreadedRequest(float DeltaSeconds) = 0;
+
+	void StartWaitingInQueue();
+	float GetTimeStartedWaitingInQueue() const;
+
+public:
+	float TimeStartedWaitingInQueue = 0.0f;
 };
diff --git a/Engine/Source/Runtime/Online/HTTP/Public/HttpManager.h b/Engine/Source/Runtime/Online/HTTP/Public/HttpManager.h
--- a/Engine/Source/Runtime/Online/HTTP/Public/HttpManager.h
+++ b/Engine/Source/Runtime/Online/HTTP/Public/HttpManager.h
@@ -56,10 +56,14 @@ struct FHttpStats
 	float MaxTimeToConnect = -1.0f;
 	/** The max waiting queue in http manager */
 	uint32 MaxRequestsInQueue = 0;
+	/** The max waiting time in queue of http manager */
+	float MaxTimeToWaitInQueue = 0.0f;
 
 	bool operator==(const FHttpStats& Other) const
 	{
-		return MaxRequestsInQueue == Other.MaxRequestsInQueue && FMath::IsNearlyEqual(MaxTimeToConnect, Other.MaxTimeToConnect);
+		return MaxRequestsInQueue == Other.MaxRequestsInQueue
+			&& FMath::IsNearlyEqual(MaxTimeToConnect, Other.MaxTimeToConnect)
+			&& FMath::IsNearlyEqual(MaxTimeToWaitInQueue, Other.MaxTimeToWaitInQueue);
 	}
 };
 
@@ -355,4 +359,7 @@ PACKAGE_SCOPE:
 
 	/** Record the requests waiting in queue, to have an idea if there are too many requests or if request number limit is too small */
 	HTTP_API void RecordStatRequestsInQueue(uint32 RequestsInQueue);
+
+	/** Record the time to wait in queue, to have a general idea how long the client usually wait before actually starting, to adjust the requests */
+	HTTP_API void RecordMaxTimeToWaitInQueue(float Duration);
 };
