Extract the Response from http request implementions into HttpRequestCommon;
Fix the test case the response in request is not cleared on some platforms when retrying.

diff --git a/Engine/Source/Runtime/Online/HTTP/Private/Apple/AppleHttp.cpp b/Engine/Source/Runtime/Online/HTTP/Private/Apple/AppleHttp.cpp
--- a/Engine/Source/Runtime/Online/HTTP/Private/Apple/AppleHttp.cpp
+++ b/Engine/Source/Runtime/Online/HTTP/Private/Apple/AppleHttp.cpp
@@ -765,7 +765,6 @@ bool FAppleHttpRequest::SetupRequest()
 	LastReportedBytesWritten = 0;
 	LastReportedBytesRead = 0;
 	ElapsedTime = 0.0f;
-	Response = nullptr;
 
 	Task = [Session dataTaskWithRequest: Request];
 	
@@ -776,7 +775,8 @@ bool FAppleHttpRequest::SetupRequest()
 		SetStatus(EHttpRequestStatus::Processing);
 		SetFailureReason(EHttpFailureReason::None);
 
-		Response = MakeShared<FAppleHttpResponse>(*this);
+		TSharedPtr<FAppleHttpResponse> Response = MakeShared<FAppleHttpResponse>(*this);
+		ResponseCommon = Response;
 
 		// Both Task and Response keep a strong reference to the delegate
 		Task.delegate = Response->ResponseDelegate;
@@ -803,6 +803,7 @@ void FAppleHttpRequest::FinishRequest()
 
 	PostProcess();
 
+	TSharedPtr<FAppleHttpResponse> Response = StaticCastSharedPtr<FAppleHttpResponse>(ResponseCommon);
 	bool bSucceeded = (Response && Response->GetStatusFromDelegate() == EHttpRequestStatus::Succeeded);
 	UE_LOG(LogHttp, Verbose, TEXT("Request %s"), bSucceeded ? TEXT("succeeded") : TEXT("failed"));
 	SetStatus(bSucceeded ? EHttpRequestStatus::Succeeded : EHttpRequestStatus::Failed);
@@ -845,6 +846,7 @@ void FAppleHttpRequest::CleanupRequest()
 {
 	UE_LOG(LogHttp, Verbose, TEXT("FAppleHttpRequest::CleanupRequest()"));
 	
+	TSharedPtr<FAppleHttpResponse> Response = StaticCastSharedPtr<FAppleHttpResponse>(ResponseCommon);
 	if (Response != nullptr)
 	{
 		Response->CleanSharedObjects();
@@ -869,11 +871,6 @@ void FAppleHttpRequest::AbortRequest()
 	}
 }
 
-const FHttpResponsePtr FAppleHttpRequest::GetResponse() const
-{
-	return Response;
-}
-
 void FAppleHttpRequest::Tick(float DeltaSeconds)
 {
 	if (DelegateThreadPolicy == EHttpRequestDelegateThreadPolicy::CompleteOnGameThread)
@@ -884,6 +881,7 @@ void FAppleHttpRequest::Tick(float DeltaSeconds)
 
 void FAppleHttpRequest::CheckProgressDelegate()
 {
+	TSharedPtr<FAppleHttpResponse> Response = StaticCastSharedPtr<FAppleHttpResponse>(ResponseCommon);
 	if (Response.IsValid() && (CompletionStatus == EHttpRequestStatus::Processing || Response->GetStatusFromDelegate() == EHttpRequestStatus::Failed))
 	{
 		const uint64 BytesWritten = Response->GetNumBytesWritten();
@@ -912,8 +910,8 @@ bool FAppleHttpRequest::StartThreadedRequest()
 
 bool FAppleHttpRequest::IsThreadedRequestComplete()
 {
-	TSharedPtr<FAppleHttpResponse> AliveResponse = Response;
-	return (AliveResponse.IsValid() && AliveResponse->IsReady());
+	TSharedPtr<FAppleHttpResponse> Response = StaticCastSharedPtr<FAppleHttpResponse>(ResponseCommon);
+	return (Response.IsValid() && Response->IsReady());
 }
 
 void FAppleHttpRequest::TickThreadedRequest(float DeltaSeconds)
diff --git a/Engine/Source/Runtime/Online/HTTP/Private/Apple/AppleHttp.h b/Engine/Source/Runtime/Online/HTTP/Private/Apple/AppleHttp.h
--- a/Engine/Source/Runtime/Online/HTTP/Private/Apple/AppleHttp.h
+++ b/Engine/Source/Runtime/Online/HTTP/Private/Apple/AppleHttp.h
@@ -46,7 +46,6 @@ public:
 	virtual void SetHeader(const FString& HeaderName, const FString& HeaderValue) override;
 	virtual void AppendToHeader(const FString& HeaderName, const FString& AdditionalHeaderValue) override;
 	virtual bool ProcessRequest() override;
-	virtual const FHttpResponsePtr GetResponse() const override;
 	virtual void Tick(float DeltaSeconds) override;
 	virtual float GetElapsedTime() const override;
 	//~ End IHttpRequest Interface
@@ -117,9 +116,6 @@ private:
 	/** The request payload length in bytes. This must be tracked separately for a file stream */
 	uint64 ContentBytesLength;
 
-	/** The response object which we will use to pair with this request */
-	TSharedPtr<class FAppleHttpResponse,ESPMode::ThreadSafe> Response;
-
 	/** Array used to retrieve back content set on the ObjC request when calling GetContent*/
 	mutable TArray<uint8> StorageForGetContent;
 
diff --git a/Engine/Source/Runtime/Online/HTTP/Private/Curl/CurlHttp.cpp b/Engine/Source/Runtime/Online/HTTP/Private/Curl/CurlHttp.cpp
--- a/Engine/Source/Runtime/Online/HTTP/Private/Curl/CurlHttp.cpp
+++ b/Engine/Source/Runtime/Online/HTTP/Private/Curl/CurlHttp.cpp
@@ -471,9 +471,9 @@ size_t FCurlHttpRequest::ReceiveResponseHeaderCallback(void* Ptr, size_t SizeInB
 {
 	QUICK_SCOPE_CYCLE_COUNTER(STAT_FCurlHttpRequest_ReceiveResponseHeaderCallback);
 
-	if (!Response.IsValid())
+	if (!ResponseCommon.IsValid())
 	{
-		Response = MakeShared<FCurlHttpResponse, ESPMode::ThreadSafe>(*this);
+		ResponseCommon = MakeShared<FCurlHttpResponse>(*this);
 		TotalBytesRead = 0;
 	}
 
@@ -501,6 +501,8 @@ size_t FCurlHttpRequest::ReceiveResponseHeaderCallback(void* Ptr, size_t SizeInB
 			HeaderValue.TrimStartInline();
 			if (!HeaderKey.IsEmpty() && !HeaderValue.IsEmpty() && !bRedirected)
 			{
+				TSharedPtr<FCurlHttpResponse> Response = StaticCastSharedPtr<FCurlHttpResponse>(ResponseCommon);
+
 				//Store the content length so OnRequestProgress64() delegates have something to work with
 				if (HeaderKey == TEXT("Content-Length"))
 				{
@@ -572,9 +574,9 @@ size_t FCurlHttpRequest::ReceiveResponseBodyCallback(void* Ptr, size_t SizeInBlo
 	QUICK_SCOPE_CYCLE_COUNTER(STAT_FCurlHttpRequest_ReceiveResponseBodyCallback);
 	LLM_SCOPE(ELLMTag::Networking);
 
-	if (!Response.IsValid())
+	if (!ResponseCommon.IsValid())
 	{
-		Response = MakeShared<FCurlHttpResponse, ESPMode::ThreadSafe>(*this);
+		ResponseCommon = MakeShared<FCurlHttpResponse>(*this);
 	}
 
 	OnAnyActivityOccur(TEXTVIEW("Received body"));
@@ -586,6 +588,7 @@ size_t FCurlHttpRequest::ReceiveResponseBodyCallback(void* Ptr, size_t SizeInBlo
 	  
 	uint64 SizeToDownload = SizeInBlocks * BlockSizeInBytes;
 
+	TSharedPtr<FCurlHttpResponse> Response = StaticCastSharedPtr<FCurlHttpResponse>(ResponseCommon);
 	UE_LOG(LogHttp, Verbose, TEXT("%p: ReceiveResponseBodyCallback: %llu bytes out of %llu received. (SizeInBlocks=%llu, BlockSizeInBytes=%llu, TotalBytesRead=%llu, Response->GetContentLength()=%llu, SizeToDownload=%llu (<-this will get returned from the callback))"),
 		this, TotalBytesRead.load() + SizeToDownload, Response->GetContentLength(),
 		SizeInBlocks, BlockSizeInBytes, TotalBytesRead.load(), Response->GetContentLength(), SizeToDownload);
@@ -1036,7 +1039,7 @@ bool FCurlHttpRequest::ProcessRequest()
 	check(EasyHandle);
 
 	// Clear out response. If this is a re-used request, Response could point to a stale response until SetupRequestHttpThread is called
-	Response = nullptr;
+	ResponseCommon = nullptr;
 	LastReportedBytesRead = 0;
 
 	if (!PreProcess())
@@ -1064,7 +1067,6 @@ void FCurlHttpRequest::ClearInCaseOfRetry()
 	IHttpThreadedRequest::ClearInCaseOfRetry();
 
 	// Clear out response. If this is a re-used request, Response could point to a stale response until SetupRequestHttpThread is called
-	Response = nullptr;
 	LastReportedBytesRead = 0;
 	TotalBytesRead = 0;
 	bAnyHttpActivity = false;
@@ -1137,11 +1139,6 @@ void FCurlHttpRequest::AbortRequest()
 	}
 }
 
-const FHttpResponsePtr FCurlHttpRequest::GetResponse() const
-{
-	return Response;
-}
-
 void FCurlHttpRequest::Tick(float DeltaSeconds)
 {
 	if (DelegateThreadPolicy == EHttpRequestDelegateThreadPolicy::CompleteOnGameThread)
@@ -1198,6 +1195,9 @@ void FCurlHttpRequest::FinishRequest()
 	PostProcess();
 	
 	CheckProgressDelegate();
+
+	TSharedPtr<FCurlHttpResponse> Response = StaticCastSharedPtr<FCurlHttpResponse>(ResponseCommon);
+
 	// if completed, get more info
 	if (bCurlRequestCompleted)
 	{
@@ -1352,7 +1352,7 @@ void FCurlHttpRequest::FinishRequest()
 		OnProcessRequestComplete().ExecuteIfBound(SharedThis(this), Response, false);
 
 		//Delegate needs to know about the errors -- so clear out Response (since connection failed) afterwards...
-		Response = nullptr;
+		ResponseCommon = nullptr;
 		TotalBytesRead = 0;
 	}
 }
diff --git a/Engine/Source/Runtime/Online/HTTP/Private/Curl/CurlHttp.h b/Engine/Source/Runtime/Online/HTTP/Private/Curl/CurlHttp.h
--- a/Engine/Source/Runtime/Online/HTTP/Private/Curl/CurlHttp.h
+++ b/Engine/Source/Runtime/Online/HTTP/Private/Curl/CurlHttp.h
@@ -165,7 +165,6 @@ public:
 	virtual void SetHeader(const FString& HeaderName, const FString& HeaderValue) override;
 	virtual void AppendToHeader(const FString& HeaderName, const FString& AdditionalHeaderValue) override;
 	virtual bool ProcessRequest() override;
-	virtual const FHttpResponsePtr GetResponse() const override;
 	virtual void Tick(float DeltaSeconds) override;
 	virtual float GetElapsedTime() const override;
 	//~ End IHttpRequest Interface
@@ -375,8 +374,6 @@ private:
 	CURLMcode		CurlAddToMultiResult;
 	/** Operation result code as returned by libcurl */
 	CURLcode		CurlCompletionResult;
-	/** The response object which we will use to pair with this request */
-	TSharedPtr<class FCurlHttpResponse,ESPMode::ThreadSafe> Response;
 	/** Payload to use with the request. Typically for POST, PUT, or PATCH */
 	TUniquePtr<FRequestPayload> RequestPayload;
 	/** Is the request payload seekable? */
diff --git a/Engine/Source/Runtime/Online/HTTP/Private/GenericPlatform/HttpRequestCommon.cpp b/Engine/Source/Runtime/Online/HTTP/Private/GenericPlatform/HttpRequestCommon.cpp
--- a/Engine/Source/Runtime/Online/HTTP/Private/GenericPlatform/HttpRequestCommon.cpp
+++ b/Engine/Source/Runtime/Online/HTTP/Private/GenericPlatform/HttpRequestCommon.cpp
@@ -107,12 +107,11 @@ void FHttpRequestCommon::PostProcess()
 
 void FHttpRequestCommon::ClearInCaseOfRetry()
 {
-	// TODO: clear response shared ptr here as well after moving it from child class to this class
-
 	bActivityTimedOut = false;
 	FailureReason = EHttpFailureReason::None;
 	bCanceled = false;
 	EffectiveURL = GetURL();
+	ResponseCommon.Reset();
 }
 
 void FHttpRequestCommon::FinishRequestNotInHttpManager()
@@ -157,10 +156,10 @@ EHttpRequestDelegateThreadPolicy FHttpRequestCommon::GetDelegateThreadPolicy() c
 	return DelegateThreadPolicy; 
 }
 
-void FHttpRequestCommon::HandleRequestSucceed(TSharedPtr<IHttpResponse> Response)
+void FHttpRequestCommon::HandleRequestSucceed(TSharedPtr<IHttpResponse> InResponse)
 {
 	SetStatus(EHttpRequestStatus::Succeeded);
-	OnProcessRequestComplete().ExecuteIfBound(SharedThis(this), Response, true);
+	OnProcessRequestComplete().ExecuteIfBound(SharedThis(this), InResponse, true);
 	FHttpModule::Get().GetHttpManager().RecordStatTimeToConnect(ConnectTime);
 }
 
@@ -168,9 +167,8 @@ void FHttpRequestCommon::SetStatus(EHttpRequestStatus::Type InCompletionStatus)
 {
 	CompletionStatus = InCompletionStatus;
 
-	if (FHttpResponsePtr Response = GetResponse())
+	if (ResponseCommon)
 	{
-		TSharedPtr<FHttpResponseCommon> ResponseCommon = StaticCastSharedPtr<FHttpResponseCommon>(Response);
 		ResponseCommon->SetRequestStatus(InCompletionStatus);
 	}
 }
@@ -180,9 +178,8 @@ void FHttpRequestCommon::SetFailureReason(EHttpFailureReason InFailureReason)
 	check(FailureReason == EHttpFailureReason::None);
 	FailureReason = InFailureReason;
 
-	if (FHttpResponsePtr Response = GetResponse())
+	if (ResponseCommon)
 	{
-		TSharedPtr<FHttpResponseCommon> ResponseCommon = StaticCastSharedPtr<FHttpResponseCommon>(Response);
 		ResponseCommon->SetRequestFailureReason(InFailureReason);
 	}
 }
@@ -208,6 +205,11 @@ float FHttpRequestCommon::GetTimeoutOrDefault() const
 	return GetTimeout().Get(FHttpModule::Get().GetHttpTotalTimeout());
 }
 
+const FHttpResponsePtr FHttpRequestCommon::GetResponse() const
+{
+	return ResponseCommon;
+}
+
 void FHttpRequestCommon::CancelRequest()
 {
 	if (bCanceled)
@@ -438,9 +440,8 @@ void FHttpRequestCommon::SetEffectiveURL(const FString& InEffectiveURL)
 {
 	EffectiveURL = InEffectiveURL;
 
-	if (FHttpResponsePtr Response = GetResponse())
+	if (ResponseCommon)
 	{
-		TSharedPtr<FHttpResponseCommon> ResponseCommon = StaticCastSharedPtr<FHttpResponseCommon>(Response);
 		ResponseCommon->SetEffectiveURL(EffectiveURL);
 	}
 }
diff --git a/Engine/Source/Runtime/Online/HTTP/Public/GenericPlatform/HttpRequestCommon.h b/Engine/Source/Runtime/Online/HTTP/Public/GenericPlatform/HttpRequestCommon.h
--- a/Engine/Source/Runtime/Online/HTTP/Public/GenericPlatform/HttpRequestCommon.h
+++ b/Engine/Source/Runtime/Online/HTTP/Public/GenericPlatform/HttpRequestCommon.h
@@ -4,6 +4,7 @@
 
 #include "GenericPlatform/HttpRequestImpl.h"
 
+class FHttpResponseCommon;
 class IHttpTaskTimerHandle;
 
 /**
@@ -29,6 +30,8 @@ public:
 	HTTP_API virtual TOptional<float> GetTimeout() const override;
 	HTTP_API float GetTimeoutOrDefault() const;
 
+	HTTP_API virtual const FHttpResponsePtr GetResponse() const override;
+
 	// Can be called on game thread or http thread depend on the delegate thread policy
 	HTTP_API virtual void FinishRequest() = 0;
 
@@ -56,7 +59,7 @@ protected:
 	 */
 	HTTP_API void FinishRequestNotInHttpManager();
 
-	HTTP_API void HandleRequestSucceed(TSharedPtr<IHttpResponse> Response);
+	HTTP_API void HandleRequestSucceed(TSharedPtr<IHttpResponse> InResponse);
 
 	HTTP_API void StartActivityTimeoutTimer();
 	HTTP_API void StartActivityTimeoutTimerBy(double DelayToTrigger);
@@ -119,4 +122,7 @@ protected:
 
 	/** Cache the effective URL. When redirected, it will be different with original URL */
 	FString EffectiveURL;
+
+	/** The response object which we will use to pair with this request */
+	TSharedPtr<FHttpResponseCommon> ResponseCommon;
 };
