Refactor: Split big function ProcessRequest into small functions in HttpThreadBase.

diff --git a/Engine/Source/Runtime/Online/HTTP/Private/HttpThread.cpp b/Engine/Source/Runtime/Online/HTTP/Private/HttpThread.cpp
--- a/Engine/Source/Runtime/Online/HTTP/Private/HttpThread.cpp
+++ b/Engine/Source/Runtime/Online/HTTP/Private/HttpThread.cpp
@@ -191,10 +191,8 @@ void FHttpThreadBase::Exit()
 	// empty
 }
 
-void FHttpThreadBase::Process(TArray<IHttpThreadedRequest*>& RequestsToCancel, TArray<IHttpThreadedRequest*>& RequestsToComplete)
+void FHttpThreadBase::ConsumeCanceledRequestsAndNewRequests(TArray<IHttpThreadedRequest*>& RequestsToCancel, TArray<IHttpThreadedRequest*>& RequestsToComplete)
 {
-	SCOPE_CYCLE_COUNTER(STAT_HTTPThread_Process);
-
 	// cache all cancelled and new requests
 	{
 		IHttpThreadedRequest* Request = nullptr;
@@ -229,20 +227,10 @@ void FHttpThreadBase::Process(TArray<IHttpThreadedRequest*>& RequestsToCancel, T
 			UE_LOG(LogHttp, Log, TEXT("Unable to find request (%p) in HttpThread"), Request);
 		}
 	}
+}
 
-	const double AppTime = FPlatformTime::Seconds();
-	const double ElapsedTime = AppTime - LastTime;
-	LastTime = AppTime;
-
-	// Tick any running requests
-	// as long as they properly finish in HttpThreadTick below they are unaffected by a possibly large ElapsedTime above
-	for (IHttpThreadedRequest* Request : RunningThreadedRequests)
-	{
-		SCOPE_CYCLE_COUNTER(STAT_HTTPThread_TickThreadedRequest);
-
-		Request->TickThreadedRequest(ElapsedTime);
-	}
-
+void FHttpThreadBase::StartRequestsWaitingInQueue(TArray<IHttpThreadedRequest*>& RequestsToComplete)
+{
 	// We'll start rate limited requests until we hit the limit
 	// Tick new requests separately from existing RunningThreadedRequests so they get a chance 
 	// to send unaffected by possibly large ElapsedTime above
@@ -284,6 +272,22 @@ void FHttpThreadBase::Process(TArray<IHttpThreadedRequest*>& RequestsToCancel, T
 	{
 		FHttpModule::Get().GetHttpManager().RecordStatRequestsInQueue(RateLimitedThreadedRequests.Num());
 	}
+}
+
+void FHttpThreadBase::MoveCompletingRequestsToCompletedRequests(TArray<IHttpThreadedRequest*>& RequestsToComplete)
+{
+	const double AppTime = FPlatformTime::Seconds();
+	const double ElapsedTime = AppTime - LastTime;
+	LastTime = AppTime;
+
+	// Tick any running requests
+	// as long as they properly finish in HttpThreadTick below they are unaffected by a possibly large ElapsedTime above
+	for (IHttpThreadedRequest* Request : RunningThreadedRequests)
+	{
+		SCOPE_CYCLE_COUNTER(STAT_HTTPThread_TickThreadedRequest);
+
+		Request->TickThreadedRequest(ElapsedTime);
+	}
 
 	{
 		SCOPE_CYCLE_COUNTER(STAT_HTTPThread_HttpThreadTick);
@@ -308,7 +312,10 @@ void FHttpThreadBase::Process(TArray<IHttpThreadedRequest*>& RequestsToCancel, T
 			UE_LOG(LogHttp, Verbose, TEXT("Threaded request (%p) completed. Running threaded requests (%d)"), Request, RunningThreadedRequests.Num());
 		}
 	}
+}
 
+void FHttpThreadBase::FinishRequestsFromHttpThreadWithCallbacks(TArray<IHttpThreadedRequest*>& RequestsToComplete)
+{
 	if (RequestsToComplete.Num() > 0)
 	{
 		for (IHttpThreadedRequest* Request : RequestsToComplete)
@@ -329,6 +336,19 @@ void FHttpThreadBase::Process(TArray<IHttpThreadedRequest*>& RequestsToCancel, T
 	}
 }
 
+void FHttpThreadBase::Process(TArray<IHttpThreadedRequest*>& RequestsToCancel, TArray<IHttpThreadedRequest*>& RequestsToComplete)
+{
+	SCOPE_CYCLE_COUNTER(STAT_HTTPThread_Process);
+
+	ConsumeCanceledRequestsAndNewRequests(RequestsToCancel, RequestsToComplete);
+
+	MoveCompletingRequestsToCompletedRequests(RequestsToComplete);
+
+	StartRequestsWaitingInQueue(RequestsToComplete);
+
+	FinishRequestsFromHttpThreadWithCallbacks(RequestsToComplete);
+}
+
 FLegacyHttpThread::FLegacyHttpThread()
 {
 	FPlatformHttp::AddDefaultUserAgentProjectComment(TEXT("http-legacy"));
diff --git a/Engine/Source/Runtime/Online/HTTP/Private/HttpThread.h b/Engine/Source/Runtime/Online/HTTP/Private/HttpThread.h
--- a/Engine/Source/Runtime/Online/HTTP/Private/HttpThread.h
+++ b/Engine/Source/Runtime/Online/HTTP/Private/HttpThread.h
@@ -139,6 +139,12 @@ protected:
 	*/
 	virtual class FSingleThreadRunnable* GetSingleThreadInterface() override { return this; }
 
+private:
+	void ConsumeCanceledRequestsAndNewRequests(TArray<IHttpThreadedRequest*>& RequestsToCancel, TArray<IHttpThreadedRequest*>& RequestsToComplete);
+	void MoveCompletingRequestsToCompletedRequests(TArray<IHttpThreadedRequest*>& RequestsToComplete);
+	void StartRequestsWaitingInQueue(TArray<IHttpThreadedRequest*>& RequestsToComplete);
+	void FinishRequestsFromHttpThreadWithCallbacks(TArray<IHttpThreadedRequest*>& RequestsToComplete);
+
 protected:
 	/** Pointer to Runnable Thread */
 	FRunnableThread* Thread;
