Get rid of the legacy http flush in websockets when suspend.

diff --git a/Engine/Config/BaseEngine.ini b/Engine/Config/BaseEngine.ini
--- a/Engine/Config/BaseEngine.ini
+++ b/Engine/Config/BaseEngine.ini
@@ -32,15 +32,12 @@ HttpConnectionTimeout=60
 HttpActivityTimeout=30
 ; Flush time limits for FHttpManager::Flush:
 ; Default = when Flush is called with no overriding context
-; Background = when Flush is called because the app is backgrounding.  This may not be called on backgrounding on all platforms.
 ; Shutdown = when Flush is called because the app is being shut down.  Some use cases may want an infinite wait.  This may not be called on shutdown on all platforms.
 ; FullFlush = explicit request to wait infinitely for all HTTP requests to complete.  Some use cases may want to disable infinite waits.  Typically not used in game code but may be used by applications or commandlets.
 ; Soft Limit = how long to wait in seconds before cancelling active requests.  -1 = never cancel.  0 = cancel immediately.  Must be < Hard Limit to work.  Typically after calling Cancel the HTTP request needs a tick to process the cancel, so it should not be == Hard Limit.
 ; Hard Limit = how long to wait in seconds before stopping waiting altogether.
 FlushSoftTimeLimitDefault=2.0
 FlushHardTimeLimitDefault=4.0
-FlushSoftTimeLimitBackground=2.0
-FlushHardTimeLimitBackground=4.0
 FlushSoftTimeLimitShutdown=2.0
 FlushHardTimeLimitShutdown=4.0
 FlushSoftTimeLimitFullFlush=-1.0
diff --git a/Engine/Config/ConfigRedirects.ini b/Engine/Config/ConfigRedirects.ini
--- a/Engine/Config/ConfigRedirects.ini
+++ b/Engine/Config/ConfigRedirects.ini
@@ -63,3 +63,7 @@ HttpReceiveTimeout=HttpActivityTimeout
 ; The legacy behavior doesn't make use of HttpSendTimeout at all, and there is no such support on some platforms.
 ; HttpActivityTimeout should be enough to decide if there is still any ongoing activity
 HttpSendTimeout=HttpActivityTimeout
+; FlushSoftTimeLimitBackground has been removed and will no longer be used
+FlushSoftTimeLimitBackground=Removed
+; FlushHardTimeLimitBackground has been removed and will no longer be used
+FlushHardTimeLimitBackground=Removed
diff --git a/Engine/Source/Runtime/Online/HTTP/Private/HttpManager.cpp b/Engine/Source/Runtime/Online/HTTP/Private/HttpManager.cpp
--- a/Engine/Source/Runtime/Online/HTTP/Private/HttpManager.cpp
+++ b/Engine/Source/Runtime/Online/HTTP/Private/HttpManager.cpp
@@ -26,7 +26,9 @@ const TCHAR* LexToString(const EHttpFlushReason& FlushReason)
 	switch (FlushReason)
 	{
 	case EHttpFlushReason::Default:		return TEXT("Default");
+PRAGMA_DISABLE_DEPRECATION_WARNINGS
 	case EHttpFlushReason::Background:	return TEXT("Background");
+PRAGMA_ENABLE_DEPRECATION_WARNINGS
 	case EHttpFlushReason::Shutdown:	return TEXT("Shutdown");
 	case EHttpFlushReason::FullFlush:	return TEXT("FullFlush");
 	}
@@ -141,10 +143,6 @@ void FHttpManager::ReloadFlushTimeLimits()
 			GConfig->GetDouble(TEXT("HTTP"), TEXT("FlushSoftTimeLimitDefault"), SoftLimitSeconds, GEngineIni);
 			GConfig->GetDouble(TEXT("HTTP"), TEXT("FlushHardTimeLimitDefault"), HardLimitSeconds, GEngineIni);
 			break;
-		case EHttpFlushReason::Background:
-			GConfig->GetDouble(TEXT("HTTP"), TEXT("FlushSoftTimeLimitBackground"), SoftLimitSeconds, GEngineIni);
-			GConfig->GetDouble(TEXT("HTTP"), TEXT("FlushHardTimeLimitBackground"), HardLimitSeconds, GEngineIni);
-			break;
 		case EHttpFlushReason::Shutdown:
 			GConfig->GetDouble(TEXT("HTTP"), TEXT("FlushSoftTimeLimitShutdown"), SoftLimitSeconds, GEngineIni);
 			GConfig->GetDouble(TEXT("HTTP"), TEXT("FlushHardTimeLimitShutdown"), HardLimitSeconds, GEngineIni);
diff --git a/Engine/Source/Runtime/Online/HTTP/Public/HttpManager.h b/Engine/Source/Runtime/Online/HTTP/Public/HttpManager.h
--- a/Engine/Source/Runtime/Online/HTTP/Public/HttpManager.h
+++ b/Engine/Source/Runtime/Online/HTTP/Public/HttpManager.h
@@ -28,7 +28,7 @@ enum class EHttpFlushReason : uint8
 	/** Reasonable, typically higher time limits */
 	Default,
 	/** Shorter time limits depending on platform requirements */
-	Background,
+	Background UE_DEPRECATED(5.5, "Flush when go into background is no longer used and will be removed"),
 	/** Shorter time limits depending on platform requirements */
 	Shutdown,
 	/** Infinite wait, should only be used in non-game scenarios where longer waits are acceptable */
