Change the default connection timeout to 30s.

diff --git a/Engine/Config/BaseEngine.ini b/Engine/Config/BaseEngine.ini
--- a/Engine/Config/BaseEngine.ini
+++ b/Engine/Config/BaseEngine.ini
@@ -28,7 +28,7 @@ SaveExt=usa
 Port=7777
 
 [HTTP]
-HttpConnectionTimeout=60
+HttpConnectionTimeout=30
 HttpActivityTimeout=30
 ; Flush time limits for FHttpManager::Flush:
 ; Default = when Flush is called with no overriding context
diff --git a/Engine/Source/Runtime/Online/HTTP/Private/HttpModule.cpp b/Engine/Source/Runtime/Online/HTTP/Private/HttpModule.cpp
--- a/Engine/Source/Runtime/Online/HTTP/Private/HttpModule.cpp
+++ b/Engine/Source/Runtime/Online/HTTP/Private/HttpModule.cpp
@@ -80,7 +80,7 @@ void FHttpModule::StartupModule()
 
 	MaxReadBufferSize = 256 * 1024;
 	HttpTotalTimeout = 0.0f;
-	HttpConnectionTimeout = 60.0f;
+	HttpConnectionTimeout = 30.0f;
 	HttpActivityTimeout = 30.0f;
 	HttpReceiveTimeout = HttpConnectionTimeout;
 	HttpSendTimeout = HttpConnectionTimeout;
