Improve default HTTP retry behavior when services do not indicate how long to wait for retries

- Retry period changed from linear to exponential;
- Adds random jitter multiplication to computed backoff.

Final default retry lockout period formula looks like this:
RandRange(JitterMin, JitterMax) * Pow(Base, RetryNumber + 1 + Bias) where RetryNumber is 0, 1, 2, 3...

Defaults:
JitterMin = 0.5
JitterMax = 1.0
Base = 2.0
Bias = 1.0

diff --git a/Engine/Source/Programs/WebTests/Private/TestHttp.cpp b/Engine/Source/Programs/WebTests/Private/TestHttp.cpp
--- a/Engine/Source/Programs/WebTests/Private/TestHttp.cpp
+++ b/Engine/Source/Programs/WebTests/Private/TestHttp.cpp
@@ -1603,6 +1603,42 @@ TEST_CASE_METHOD(FWaitUntilCompleteHttpFixture, "Optionally retry limit can be s
 	HttpRequest->ProcessRequest();
 }
 
+TEST_CASE_METHOD(FWaitUntilCompleteHttpFixture, "Retry fallback with exponential lock out if there is no Retry-After header", HTTP_TAG)
+{
+	if (!bRetryEnabled)
+	{
+		return;
+	}
+
+	DisableWarningsInThisTest();
+
+	TSharedRef<IHttpRequest> HttpRequest = HttpRetryManager->CreateRequest(
+		2/*InRetryLimitCountOverride*/,
+		FHttpRetrySystem::FRetryTimeoutRelativeSecondsSetting()/*InRetryTimeoutRelativeSecondsOverride unused*/,
+		{EHttpResponseCodes::TooManyRequests}/*InRetryResponseCodes*/
+	);
+
+	HttpRequest->SetURL(UrlMockStatus(EHttpResponseCodes::TooManyRequests));
+	HttpRequest->SetVerb(TEXT("GET"));
+
+	ExpectingExtraCallbacks = 2;
+
+	HttpRequest->OnRequestWillRetry().BindLambda([this](FHttpRequestPtr Request, FHttpResponsePtr /*Response*/, float LockoutPeriod) {
+		--ExpectingExtraCallbacks;
+		// Default value in FExponentialBackoffCurve Compute(1) is 4 with default value in FBackoffJitterCoefficient applied
+		CHECK(LockoutPeriod >= (4 * 0.5f));
+		CHECK(LockoutPeriod <= (4 * 1.0f));
+		Request->OnRequestWillRetry().BindLambda([this](FHttpRequestPtr /*Request*/, FHttpResponsePtr /*Response*/, float LockoutPeriod) {
+			--ExpectingExtraCallbacks;
+			// Default value in FExponentialBackoffCurve Compute(2) is 8 with default value in FBackoffJitterCoefficient applied
+			CHECK(LockoutPeriod >= (8 * 0.5f));
+			CHECK(LockoutPeriod <= (8 * 1.0f));
+		});
+	});
+
+	HttpRequest->ProcessRequest();
+}
+
 class FThreadedBatchRequestsFixture : public FWaitThreadedHttpFixture
 {
 public:
diff --git a/Engine/Source/Runtime/Online/HTTP/Private/HttpRetrySystem.cpp b/Engine/Source/Runtime/Online/HTTP/Private/HttpRetrySystem.cpp
--- a/Engine/Source/Runtime/Online/HTTP/Private/HttpRetrySystem.cpp
+++ b/Engine/Source/Runtime/Online/HTTP/Private/HttpRetrySystem.cpp
@@ -70,27 +70,45 @@ TOptional<double> ReadThrottledTimeFromResponseInSeconds(FHttpResponsePtr Respon
 	return LockoutPeriod;
 }
 
+bool FHttpRetrySystem::FExponentialBackoffCurve::IsValid() const
+{
+	return Base > 1.0f
+		&& ExponentBias >= 0.0f
+		&& MinCoefficient <= MaxCoefficient
+		&& MaxCoefficient > 0.001f
+		&& MinCoefficient >= 0.0f;
+}
+
+float FHttpRetrySystem::FExponentialBackoffCurve::Compute(uint32 RetryNumber) const
+{
+	float BackOff = FMath::Pow(Base, static_cast<float>(RetryNumber) + ExponentBias);
+	const float Coefficient = IsValid() ? FMath::RandRange(MinCoefficient, MaxCoefficient) : 1.0f;
+	return FMath::Min(BackOff * Coefficient, MaxBackoffSeconds);
+}
+
 }
 
 FHttpRetrySystem::FRequest::FRequest(
 	TSharedRef<FManager> InManager,
-	const TSharedRef<IHttpRequest, ESPMode::ThreadSafe>& HttpRequest, 
+	const TSharedRef<IHttpRequest, ESPMode::ThreadSafe>& HttpRequest,
 	const FHttpRetrySystem::FRetryLimitCountSetting& InRetryLimitCountOverride,
 	const FHttpRetrySystem::FRetryTimeoutRelativeSecondsSetting& InRetryTimeoutRelativeSecondsOverride,
 	const FHttpRetrySystem::FRetryResponseCodes& InRetryResponseCodes,
 	const FHttpRetrySystem::FRetryVerbs& InRetryVerbs,
 	const FHttpRetrySystem::FRetryDomainsPtr& InRetryDomains,
-	const FRetryLimitCountSetting& InRetryLimitCountForConnectionErrorOverride
-	)
-    : FHttpRequestAdapterBase(HttpRequest)
-    , RetryStatus(FHttpRetrySystem::FRequest::EStatus::NotStarted)
-    , RetryLimitCountOverride(InRetryLimitCountOverride)
-    , RetryLimitCountForConnectionErrorOverride(InRetryLimitCountForConnectionErrorOverride)
-    , RetryTimeoutRelativeSecondsOverride(InRetryTimeoutRelativeSecondsOverride)
+	const FRetryLimitCountSetting& InRetryLimitCountForConnectionErrorOverride,
+	const FExponentialBackoffCurve& InExponentialBackoffCurve
+)
+	: FHttpRequestAdapterBase(HttpRequest)
+	, RetryStatus(FHttpRetrySystem::FRequest::EStatus::NotStarted)
+	, RetryLimitCountOverride(InRetryLimitCountOverride)
+	, RetryLimitCountForConnectionErrorOverride(InRetryLimitCountForConnectionErrorOverride)
+	, RetryTimeoutRelativeSecondsOverride(InRetryTimeoutRelativeSecondsOverride)
 	, RetryResponseCodes(InRetryResponseCodes)
 	, RetryVerbs(InRetryVerbs)
 	, RetryDomains(InRetryDomains)
 	, RetryManager(InManager)
+	, RetryExponentialBackoffCurve(InExponentialBackoffCurve)
 {
     // if the InRetryTimeoutRelativeSecondsOverride override is being used the value cannot be negative
     check(!(InRetryTimeoutRelativeSecondsOverride.IsSet()) || (InRetryTimeoutRelativeSecondsOverride.GetValue() >= 0.0));
@@ -362,26 +380,19 @@ bool FHttpRetrySystem::FManager::ShouldRetry(const FHttpRetryRequestEntry& HttpR
 				// Be default, we will also allow retry for GET and HEAD requests even if they may duplicate on the server
 				static const TSet<FName> DefaultRetryVerbs(TArray<FName>({ FName(TEXT("GET")), FName(TEXT("HEAD")) }));
 
-				const bool bIsRetryVerbsEmpty = HttpRetryRequestEntry.Request->RetryVerbs.Num() == 0;
-				if (bIsRetryVerbsEmpty && DefaultRetryVerbs.Contains(Verb))
-				{
-					bResult = true;
-				}
-				// If retry verbs are specified, only allow retrying the specified list of verbs
-				else if (HttpRetryRequestEntry.Request->RetryVerbs.Contains(Verb))
-				{
-					bResult = true;
-				}
+				const TSet<FName>* RetryVerbsContainer = (
+					HttpRetryRequestEntry.Request->RetryVerbs.Num() == 0
+					? &DefaultRetryVerbs // Use the default list of retry verbs if the request doesn't have a specific set
+					: &HttpRetryRequestEntry.Request->RetryVerbs // Otherwise use the specific set on the request
+				);
+				bResult = RetryVerbsContainer->Contains(Verb);
 			}
 		}
 	}
 	else
 	{
 		// this may be a successful response with one of the explicitly listed response codes we want to retry on
-		if (HttpRetryRequestEntry.Request->RetryResponseCodes.Contains(Response->GetResponseCode()))
-		{
-			bResult = true;
-		}
+		bResult = HttpRetryRequestEntry.Request->RetryResponseCodes.Contains(Response->GetResponseCode());
 	}
 
     return bResult;
@@ -533,23 +544,17 @@ float FHttpRetrySystem::FManager::GetLockoutPeriodSeconds(const FHttpRetryReques
 		LockoutPeriod = static_cast<float>(ResponseLockoutPeriod.GetValue());
 	}
 
-	if (HttpRetryRequestEntry.CurrentRetryCount >= 1)
+	if (LockoutPeriod <= 0.0f)
 	{
-		if (LockoutPeriod <= 0.0f)
+		const bool bFailedToConnect = (HttpRetryRequestEntry.Request->GetStatus() == EHttpRequestStatus::Failed && HttpRetryRequestEntry.Request->GetFailureReason() == EHttpFailureReason::ConnectionError);
+		const bool bHasRetryDomains = HttpRetryRequestEntry.Request->RetryDomains.IsValid();
+		// Skip the lockout period if we failed to connect to a domain and we have other domains to try
+		if (bFailedToConnect && bHasRetryDomains)
 		{
-			const bool bFailedToConnect = (HttpRetryRequestEntry.Request->GetStatus() == EHttpRequestStatus::Failed && HttpRetryRequestEntry.Request->GetFailureReason() == EHttpFailureReason::ConnectionError);
-			const bool bHasRetryDomains = HttpRetryRequestEntry.Request->RetryDomains.IsValid();
-			// Skip the lockout period if we failed to connect to a domain and we have other domains to try
-			const bool bSkipLockoutPeriod = (bFailedToConnect && bHasRetryDomains);
-			if (!bSkipLockoutPeriod)
-			{
-				constexpr const float LockoutPeriodMinimumSeconds = 5.0f;
-				constexpr const float LockoutPeriodEscalationSeconds = 2.5f;
-				constexpr const float LockoutPeriodMaxSeconds = 30.0f;
-				LockoutPeriod = LockoutPeriodMinimumSeconds + LockoutPeriodEscalationSeconds * (HttpRetryRequestEntry.CurrentRetryCount - 1);
-				LockoutPeriod = FMath::Min(LockoutPeriod, LockoutPeriodMaxSeconds);
-			}
+			return 0.0f;
 		}
+		// The first time through this function, the CurrentRetryCount is 0, the second time it's 1, etc. We automatically add 1 to make the input into the backoff function line up with expectations.
+		LockoutPeriod = HttpRetryRequestEntry.Request->RetryExponentialBackoffCurve.Compute(HttpRetryRequestEntry.CurrentRetryCount+1);
 	}
 
 	return LockoutPeriod;
@@ -802,9 +807,9 @@ bool FHttpRetrySystem::FManager::Update(uint32* FileCount, uint32* FailingCount,
 }
 
 FHttpRetrySystem::FManager::FHttpRetryRequestEntry::FHttpRetryRequestEntry(TSharedRef<FHttpRetrySystem::FRequest, ESPMode::ThreadSafe>& InRequest)
-    : bShouldCancel(false)
-    , CurrentRetryCount(0)
-    , CurrentRetryCountForConnectionError(0)
+	: bShouldCancel(false)
+	, CurrentRetryCount(0)
+	, CurrentRetryCountForConnectionError(0)
 	, RequestStartTimeAbsoluteSeconds(FPlatformTime::Seconds())
 	, Request(InRequest)
 {}
diff --git a/Engine/Source/Runtime/Online/HTTP/Public/HttpRetrySystem.h b/Engine/Source/Runtime/Online/HTTP/Public/HttpRetrySystem.h
--- a/Engine/Source/Runtime/Online/HTTP/Public/HttpRetrySystem.h
+++ b/Engine/Source/Runtime/Online/HTTP/Public/HttpRetrySystem.h
@@ -60,6 +60,30 @@ namespace FHttpRetrySystem
 	};
 	typedef TSharedPtr<FRetryDomains, ESPMode::ThreadSafe> FRetryDomainsPtr;
 
+	/** 
+	 * Model for computing exponential backoff using the formula: Base**(CurrentRetryAttempt + Bias) 
+	 * Then applying jitter to the backoff. Jitter application is performed by selecting a random value in the [Min, Max] range and multiplying it against the computed backoff.
+	 * Half jitter can be implemented using { 0.5, 1.0 }, which becomes Backoff' = Backoff * Rand(0.5, 1.0) = Backoff/2 + Rand(0, Backoff/2)
+	 * Full jitter can be implemented using { 0.0, 1.0 }, which becomes Backoff' = Backoff * Rand(0.0, 1.0) = Rand(0.0, Backoff)
+	 * No jitter can be implemented using { 0.0, 0.0 } or any pair that where Min > Max. Backoff' = Backoff
+	*/
+	struct FExponentialBackoffCurve
+	{
+		/** Exponential backoff base */
+		float Base = 2.0f;
+		/** Exponential backoff bias added to the current retry number */
+		float ExponentBias = 1.0f;
+		/** Exponential backoff jitter coefficient minimum value. Defaults to half jitter */
+		float MinCoefficient = 0.5f;
+		/** Exponential backoff jitter coefficient maximum value. Defaults to half jitter */
+		float MaxCoefficient = 1.0f;
+		/** Max back off seconds */
+		float MaxBackoffSeconds = 60.0f;
+
+		bool IsValid() const;
+		float Compute(uint32 RetryNumber) const;
+	};
+
 	/**
 	 * Read the number of seconds a HTTP request is throttled for from the response
 	 * @param Response the HTTP response to read the value from
@@ -111,7 +135,8 @@ namespace FHttpRetrySystem
 			const FRetryResponseCodes& InRetryResponseCodes = FRetryResponseCodes(),
 			const FRetryVerbs& InRetryVerbs = FRetryVerbs(),
 			const FRetryDomainsPtr& InRetryDomains = FRetryDomainsPtr(),
-			const FRetryLimitCountSetting& InRetryLimitCountForConnectionErrorOverride = FRetryLimitCountSetting()
+			const FRetryLimitCountSetting& InRetryLimitCountForConnectionErrorOverride = FRetryLimitCountSetting(),
+			const FExponentialBackoffCurve& InExponentialBackoffCurve = FExponentialBackoffCurve()
 		);
 
 		void HttpOnRequestProgress(FHttpRequestPtr InHttpRequest, uint64 BytesSent, uint64 BytesRcv);
@@ -141,6 +166,9 @@ namespace FHttpRetrySystem
 		/** Save the last response before the retry */
 		FHttpResponsePtr LastResponse;
 		bool bLastSucceeded = false;
+
+		/** Exponential backoff curve */
+		FExponentialBackoffCurve RetryExponentialBackoffCurve;
 	};
 }
 
