Make it not output warning when connection timeout equals activity timeout.

diff --git a/Engine/Source/Runtime/Online/HTTP/Private/Apple/AppleHttp.cpp b/Engine/Source/Runtime/Online/HTTP/Private/Apple/AppleHttp.cpp
--- a/Engine/Source/Runtime/Online/HTTP/Private/Apple/AppleHttp.cpp
+++ b/Engine/Source/Runtime/Online/HTTP/Private/Apple/AppleHttp.cpp
@@ -463,7 +463,16 @@ FAppleHttpRequest::FAppleHttpRequest(NSURLSession* InSession)
 	check(HttpConnectionTimeout > 0.0f);
 	Request.timeoutInterval = HttpConnectionTimeout;
 	
-	UE_CLOG(HttpConnectionTimeout <= FHttpModule::Get().GetHttpActivityTimeout(), LogHttp, Warning, TEXT("HttpConnectionTimeout should be greater than HttpActivityTimeout. Otherwise requests may complete unexpectedly with ConnectionError after HttpConnectionTimeout seconds without activity"));
+	UE_CLOG(
+		HttpConnectionTimeout < FHttpModule::Get().GetHttpActivityTimeout(), 
+		LogHttp, 
+		Warning, 
+		TEXT(
+			"HttpConnectionTimeout can't be less than HttpActivityTimeout, otherwise requests may complete "
+			"unexpectedly with ConnectionError after %.2f(HttpConnectionTimeout) seconds without activity, "
+			"instead of intended %.2f(HttpActivityTimeout) seconds"
+		), 
+		HttpConnectionTimeout, FHttpModule::Get().GetHttpActivityTimeout());
 
 	// Disable cache to mimic WinInet behavior
 	Request.cachePolicy = NSURLRequestReloadIgnoringLocalCacheData;
