Adjust the max concurrent http requests, make it hotfixable.

diff --git a/Engine/Config/BaseEditor.ini b/Engine/Config/BaseEditor.ini
--- a/Engine/Config/BaseEditor.ini
+++ b/Engine/Config/BaseEditor.ini
@@ -681,6 +681,3 @@ bUseWidgetTemplateSelector=True
 +WidgetClassesToHide=/Script/UMG.SpinBox
 +WidgetClassesToHide=/Script/UMG.Throbber
 +WidgetClassesToHide=/Script/UMG.TreeView
-
-[HTTP.HttpThread]
-RunningThreadedRequestLimitEditor=20
diff --git a/Engine/Config/BaseEngine.ini b/Engine/Config/BaseEngine.ini
--- a/Engine/Config/BaseEngine.ini
+++ b/Engine/Config/BaseEngine.ini
@@ -43,9 +43,6 @@ FlushHardTimeLimitShutdown=4.0
 FlushSoftTimeLimitFullFlush=-1.0
 FlushHardTimeLimitFullFlush=-1.0
 
-[HTTP.HttpThread]
-RunningThreadedRequestLimit=11
-
 [HTTP.Curl]
 bAllowSeekFunction=false
 
diff --git a/Engine/Config/ConfigRedirects.ini b/Engine/Config/ConfigRedirects.ini
--- a/Engine/Config/ConfigRedirects.ini
+++ b/Engine/Config/ConfigRedirects.ini
@@ -67,3 +67,9 @@ HttpSendTimeout=HttpActivityTimeout
 FlushSoftTimeLimitBackground=Removed
 ; FlushHardTimeLimitBackground has been removed and will no longer be used
 FlushHardTimeLimitBackground=Removed
+
+[HTTP.HttpThread]
+;RunningThreadedRequestLimitEditor has been deprecated, change it through DefaultMaxConcurrentRequests in build script instead
+RunningThreadedRequestLimitEditor=Removed
+;RunningThreadedRequestLimit has been deprecated, change it through DefaultMaxConcurrentRequests in build script instead
+RunningThreadedRequestLimit=Removed
\ No newline at end of file
diff --git a/Engine/Source/Runtime/Online/HTTP/Private/HttpThread.cpp b/Engine/Source/Runtime/Online/HTTP/Private/HttpThread.cpp
--- a/Engine/Source/Runtime/Online/HTTP/Private/HttpThread.cpp
+++ b/Engine/Source/Runtime/Online/HTTP/Private/HttpThread.cpp
@@ -5,6 +5,7 @@
 #include "HAL/PlatformTime.h"
 #include "HAL/PlatformProcess.h"
 #include "HAL/RunnableThread.h"
+#include "HAL/IConsoleManager.h"
 #include "HttpManager.h"
 #include "Misc/ConfigCacheIni.h"
 #include "Misc/CommandLine.h"
@@ -15,6 +16,13 @@
 #include "PlatformHttp.h"
 #include "Stats/Stats.h"
 
+TAutoConsoleVariable<int32> CVarHttpMaxConcurrentRequests(
+	TEXT("http.MaxConcurrentRequests"),
+	UE_HTTP_DEFAULT_MAX_CONCURRENT_REQUESTS,
+	TEXT("The max number of http requests to run in parallel"),
+	ECVF_SaveForNextBoot
+);
+
 DECLARE_STATS_GROUP(TEXT("HTTP Thread"), STATGROUP_HTTPThread, STATCAT_Advanced);
 DECLARE_CYCLE_STAT(TEXT("Process"), STAT_HTTPThread_Process, STATGROUP_HTTPThread);
 DECLARE_CYCLE_STAT(TEXT("TickThreadedRequest"), STAT_HTTPThread_TickThreadedRequest, STATGROUP_HTTPThread);
@@ -48,7 +56,6 @@ FHttpThreadBase::FHttpThreadBase()
 	: Thread(nullptr)
 	, bIsSingleThread(false)
 	, bIsStopped(true)
-	, RunningThreadedRequestLimit(UE_HTTP_DEFAULT_MAX_CONCURRENT_REQUESTS)
 {
 }
 
@@ -141,25 +148,6 @@ bool FHttpThreadBase::NeedsSingleThreadTick() const
 
 void FHttpThreadBase::UpdateConfigs()
 {
-	int32 LocalRunningThreadedRequestLimit = -1;
-	const bool bFoundLocalRunningThreadedRequestLimit =
-	(
-#if WITH_EDITOR
-		GConfig->GetInt(TEXT("HTTP.HttpThread"), TEXT("RunningThreadedRequestLimitEditor"), LocalRunningThreadedRequestLimit, GEditorIni) ||
-#endif
-		GConfig->GetInt(TEXT("HTTP.HttpThread"), TEXT("RunningThreadedRequestLimit"), LocalRunningThreadedRequestLimit, GEngineIni)
-	);
-	if (bFoundLocalRunningThreadedRequestLimit)
-	{
-		if (LocalRunningThreadedRequestLimit < 1)
-		{
-			UE_LOG(LogHttp, Warning, TEXT("RunningThreadedRequestLimit must be configured as a number greater than 0. The configured value is %d. Ignored. The current value is still %d"), LocalRunningThreadedRequestLimit, RunningThreadedRequestLimit.load());
-		}
-		else
-		{
-			RunningThreadedRequestLimit = LocalRunningThreadedRequestLimit;
-		}
-	}
 }
 
 void FHttpThreadBase::HttpThreadTick(float DeltaSeconds)
@@ -176,11 +164,6 @@ void FHttpThreadBase::CompleteThreadedRequest(IHttpThreadedRequest* Request)
 	// empty
 }
 
-int32 FHttpThreadBase::GetRunningThreadedRequestLimit() const
-{
-	return RunningThreadedRequestLimit.load();
-}
-
 void FHttpThreadBase::Stop()
 {
 	// empty
@@ -235,10 +218,10 @@ void FHttpThreadBase::StartRequestsWaitingInQueue(TArray<IHttpThreadedRequest*>&
 	// Tick new requests separately from existing RunningThreadedRequests so they get a chance 
 	// to send unaffected by possibly large ElapsedTime above
 	int32 RunningThreadedRequestsCounter = RunningThreadedRequests.Num();
-	const int32 LocalRunningThreadedRequestLimit = GetRunningThreadedRequestLimit();
-	if (RunningThreadedRequestsCounter < LocalRunningThreadedRequestLimit)
+	const int32 MaxConcurrentRequests = CVarHttpMaxConcurrentRequests.GetValueOnAnyThread();
+	if (RunningThreadedRequestsCounter < MaxConcurrentRequests)
 	{
-		while(RunningThreadedRequestsCounter < LocalRunningThreadedRequestLimit && !RateLimitedThreadedRequests.IsEmpty())
+		while(RunningThreadedRequestsCounter < MaxConcurrentRequests && !RateLimitedThreadedRequests.IsEmpty())
 		{
 			SCOPE_CYCLE_COUNTER(STAT_HTTPThread_StartThreadedRequest);
 
diff --git a/Engine/Source/Runtime/Online/HTTP/Private/HttpThread.h b/Engine/Source/Runtime/Online/HTTP/Private/HttpThread.h
--- a/Engine/Source/Runtime/Online/HTTP/Private/HttpThread.h
+++ b/Engine/Source/Runtime/Online/HTTP/Private/HttpThread.h
@@ -120,7 +120,6 @@ protected:
 	virtual void CompleteThreadedRequest(IHttpThreadedRequest* Request);
 
 protected:
-	int32 GetRunningThreadedRequestLimit() const;
 
 	// Threading functions
 
@@ -156,9 +155,6 @@ private:
 	/** Tells if the runnable thread is running or stopped */
 	bool bIsStopped;
 
-	/** Limit for threaded http requests running at the same time. If not specified through configuration values, there will be no limit */
-	std::atomic<int32> RunningThreadedRequestLimit;
-
 	/** Last time the thread has been processed. Used in the non-game thread. */
 	double LastTime;
 
