Make sure don't pass received data to stream interface after cancelling the http request or after shutdowning http module;
Add test cases to test that after canceling/shut down request, the response body receive delegate won't ever been called.

diff --git a/Engine/Source/Programs/WebTests/Private/TestHttp.cpp b/Engine/Source/Programs/WebTests/Private/TestHttp.cpp
--- a/Engine/Source/Programs/WebTests/Private/TestHttp.cpp
+++ b/Engine/Source/Programs/WebTests/Private/TestHttp.cpp
@@ -136,7 +136,7 @@ public:
 	FHttpTestLogLevelInitializer HttpTestLogLevelInitializer;
 };
 
-TEST_CASE_METHOD(FHttpModuleTestFixture, "Shutdown http module without issue when there are ongoing http requests.", HTTP_TAG)
+TEST_CASE_METHOD(FHttpModuleTestFixture, "Shutdown http module without issue when there are ongoing upload http requests.", HTTP_TAG)
 {
 	DisableWarningsInThisTest();
 
@@ -544,6 +544,101 @@ TEST_CASE_METHOD(FWaitUntilCompleteHttpFixture, "Streaming http download", HTTP_
 	HttpRequest->ProcessRequest();
 }
 
+// This user streaming class is supposed to be used to receive streaming data through function OnReceivedData 
+// and it's not supposed to be called once destroyed
+class FUserStreamingClass
+{
+public:
+	FUserStreamingClass()
+		: TotalBytesReceived(new int64(0))
+	{
+	}
+
+	~FUserStreamingClass()
+	{
+		delete TotalBytesReceived;
+		TotalBytesReceived = nullptr;
+	}
+
+	bool OnReceivedData(void* Ptr, int64 Length)
+	{
+		*TotalBytesReceived += Length;
+		return true;
+	}
+
+	int64* TotalBytesReceived;
+};
+
+TEST_CASE_METHOD(FWaitUntilCompleteHttpFixture, "In streaming downloading http request won't trigger response body receive delegate after canceling", HTTP_TAG)
+{
+	TSharedRef<IHttpRequest> HttpRequest = CreateRequest();
+	HttpRequest->SetURL(UrlStreamDownload(30, 1024*1024));
+
+	TSharedPtr<FUserStreamingClass> UserInstance = MakeShared<FUserStreamingClass>();
+
+	FHttpRequestStreamDelegate Delegate;
+	Delegate.BindThreadSafeSP(UserInstance.ToSharedRef(), &FUserStreamingClass::OnReceivedData);
+	CHECK(HttpRequest->SetResponseBodyReceiveStreamDelegate(Delegate));
+
+	HttpRequest->OnProcessRequestComplete().BindLambda([](FHttpRequestPtr HttpRequest, FHttpResponsePtr HttpResponse, bool bSucceeded) {
+		CHECK(!bSucceeded);
+		CHECK(HttpResponse != nullptr);
+	});
+	HttpRequest->ProcessRequest();
+
+	while (*UserInstance->TotalBytesReceived == 0) // Make sure it started receiving data
+	{
+		FPlatformProcess::Sleep(0.001f);
+	}
+	CHECK(*UserInstance->TotalBytesReceived < 60 * 1024 * 1024);
+	CHECK(UserInstance.GetSharedReferenceCount() == 1);
+	HttpRequest->CancelRequest();
+	UserInstance.Reset();
+}
+
+class FInvalidateDelegateShutdownFixture : public FHttpModuleTestFixture
+{
+public:
+	FInvalidateDelegateShutdownFixture()
+	{
+		UserStreamingInstance = new FUserStreamingClass;
+	}
+
+	~FInvalidateDelegateShutdownFixture()
+	{
+		delete UserStreamingInstance;
+		UserStreamingInstance = nullptr;
+	}
+
+	FUserStreamingClass* UserStreamingInstance;
+};
+
+TEST_CASE_METHOD(FInvalidateDelegateShutdownFixture, "Shutdown http module without issue when there are ongoing download http requests", HTTP_TAG)
+{
+	DisableWarningsInThisTest();
+
+	for (int32 i = 0; i < 10; ++i)
+	{
+		TSharedRef<IHttpRequest> HttpRequest = HttpModule->CreateRequest();
+		HttpRequest->SetURL(UrlStreamDownload(10, 1024*1024));
+		FHttpRequestStreamDelegate Delegate;
+		Delegate.BindRaw(UserStreamingInstance, &FUserStreamingClass::OnReceivedData);
+		CHECK(HttpRequest->SetResponseBodyReceiveStreamDelegate(Delegate));
+
+		HttpRequest->OnProcessRequestComplete().BindLambda([](FHttpRequestPtr HttpRequest, FHttpResponsePtr HttpResponse, bool bSucceeded) {
+			CHECK(bSucceeded);
+		});
+		HttpRequest->ProcessRequest();
+	}
+
+	while (*UserStreamingInstance->TotalBytesReceived == 0) // Make sure it started receiving data
+	{
+		FPlatformProcess::Sleep(0.001f);
+	}
+
+	HttpModule->GetHttpManager().Tick(0.1f);
+}
+
 TEST_CASE_METHOD(FWaitUntilCompleteHttpFixture, "Can run parallel stream download requests", HTTP_TAG)
 {
 	uint32 Chunks = 5;
diff --git a/Engine/Source/Runtime/Online/HTTP/Private/Apple/AppleHttp.cpp b/Engine/Source/Runtime/Online/HTTP/Private/Apple/AppleHttp.cpp
--- a/Engine/Source/Runtime/Online/HTTP/Private/Apple/AppleHttp.cpp
+++ b/Engine/Source/Runtime/Online/HTTP/Private/Apple/AppleHttp.cpp
@@ -18,21 +18,14 @@
 {
 	/** Holds the payload as we receive it. */
 	@public TArray<uint8> Payload;
-	
-	/** This stream is shared between the request and the delegate. This field is accessed through a thread out of out control 
-	 * from the delegates methods and through mehthods in FAppleHttpRequest. Once we cancel or have an error this can 
-	 * is nullified as soon as possible to avoid receiving more data after completion delegates were triggered */
-	TSharedPtr<FArchive> ResponseBodyReceiveStream;
 
-	/** critical section to properly clear stream */
-	FCriticalSection ResponseStreamLock;
-	
-	/** flag meant to reduce locking on ResponseStreamLock*/
-	BOOL bInitializedWithValidStream;
-	
+	// Flag to indicate the request was initialized with stream. In that case even if stream was set to 
+	// null later on internally, the request itself won't cache received data anymore
+	@public BOOL bInitializedWithValidStream;
+
 	/** Have we received any data? */
 	BOOL bAnyHttpActivity;
-	
+
 	/** Delegate invoked after processing URLSession:dataTask:didReceiveData or URLSession:task:didCompleteWithError:*/
 	@public FNewAppleHttpEventDelegate NewAppleHttpEventDelegate;
 }
@@ -83,20 +76,14 @@
 	FailureReason = EHttpFailureReason::None;
 	bAnyHttpActivity = false;
 	SourceRequest = StaticCastWeakPtr<FAppleHttpRequest>(TWeakPtr<IHttpRequest>(Request.AsShared()));
-	ResponseBodyReceiveStream = Request.GetResponseBodyReceiveStream();
-	bInitializedWithValidStream = (ResponseBodyReceiveStream != nullptr);
-
+	bInitializedWithValidStream = Request.IsInitializedWithValidStream();
+	
 	return self;
 }
 
 - (void)CleanSharedObjects
 {
 	self.SourceRequest = {};
-	if (bInitializedWithValidStream)
-	{
-	    FScopeLock Lock(&ResponseStreamLock);
-		ResponseBodyReceiveStream = nullptr;
-	}
 }
 
 - (void)dealloc
@@ -113,6 +100,15 @@
 	}
 }
 
+- (bool)HandleBodyDataReceived:(void*)Ptr Size:(int64)InSize
+{
+	if (TSharedPtr<FAppleHttpRequest> Request = SourceRequest.Pin())
+	{
+		return Request->PassReceivedDataToStream(Ptr, InSize);
+	}
+	return false;
+}
+
 - (void) SaveEffectiveURL:(const FString&) InEffectiveURL
 {
 	if (TSharedPtr<FAppleHttpRequest> Request = SourceRequest.Pin())
@@ -189,21 +185,16 @@
 	__block int64 NewBytesReceived = 0;
 	if (bInitializedWithValidStream)
 	{
-		if (FScopeLock Lock(&ResponseStreamLock); ResponseBodyReceiveStream)
+		__block bool bSerializeSucceed = false;
+		[data enumerateByteRangesUsingBlock:^(const void *bytes, NSRange byteRange, BOOL *stop) {
+			NewBytesReceived += byteRange.length;
+			bSerializeSucceed = [self HandleBodyDataReceived : const_cast<void*>(bytes) Size : byteRange.length];
+			*stop = bSerializeSucceed? NO : YES;
+		}];
+		
+		if (!bSerializeSucceed && !bCanceled)
 		{
-			__block bool bHadError = false;
-			[data enumerateByteRangesUsingBlock:^(const void *bytes, NSRange byteRange, BOOL *stop) {
-				NewBytesReceived += byteRange.length;
-				ResponseBodyReceiveStream->Serialize(const_cast<void*>(bytes), byteRange.length);
-				bHadError = ResponseBodyReceiveStream->GetError();
-				*stop = bHadError? YES : NO;
-			}];
-			
-			if (bHadError)
-			{
-				[dataTask cancel];
-				ResponseBodyReceiveStream = nullptr;
-			}
+			[dataTask cancel];
 		}
 	}
 	else
@@ -493,11 +484,6 @@ FAppleHttpRequest::~FAppleHttpRequest()
     [Session release];
 }
 
-const TSharedPtr<FArchive> FAppleHttpRequest::GetResponseBodyReceiveStream() const
-{
-	return ResponseBodyReceiveStream;
-}
-
 FString FAppleHttpRequest::GetURL() const
 {
 	SCOPED_AUTORELEASE_POOL;
@@ -712,17 +698,6 @@ bool FAppleHttpRequest::SetContentFromStream(TSharedRef<FArchive, ESPMode::Threa
 	return true;
 }
 
-bool FAppleHttpRequest::SetResponseBodyReceiveStream(TSharedRef<FArchive> Stream)
-{
-	if (CompletionStatus == EHttpRequestStatus::Processing)
-	{
-		UE_LOG(LogHttp, Warning, TEXT("FCurlHttpRequest::SetContentFromStream() - attempted to set content on a request that is inflight"));
-		return false;
-	}
-	ResponseBodyReceiveStream = Stream;
-	return true;
-}
-
 FString FAppleHttpRequest::GetVerb() const
 {
 	FString ConvertedVerb(Request.HTTPMethod);
@@ -888,6 +863,11 @@ void FAppleHttpRequest::Tick(float DeltaSeconds)
 	}
 }
 
+bool FAppleHttpRequest::IsInitializedWithValidStream() const
+{ 
+	return bInitializedWithValidStream;
+}
+
 void FAppleHttpRequest::CheckProgressDelegate()
 {
 	TSharedPtr<FAppleHttpResponse> Response = StaticCastSharedPtr<FAppleHttpResponse>(ResponseCommon);
diff --git a/Engine/Source/Runtime/Online/HTTP/Private/Apple/AppleHttp.h b/Engine/Source/Runtime/Online/HTTP/Private/Apple/AppleHttp.h
--- a/Engine/Source/Runtime/Online/HTTP/Private/Apple/AppleHttp.h
+++ b/Engine/Source/Runtime/Online/HTTP/Private/Apple/AppleHttp.h
@@ -23,7 +23,6 @@ public:
 	// implementation friends
 	friend class FAppleHttpResponse;
 
-
 	//~ Begin IHttpBase Interface
 	virtual FString GetURL() const override;
 	virtual FString GetHeader(const FString& HeaderName) const override;
@@ -42,7 +41,6 @@ public:
 	virtual void SetContentAsString(const FString& ContentString) override;
     virtual bool SetContentAsStreamedFile(const FString& Filename) override;
 	virtual bool SetContentFromStream(TSharedRef<FArchive, ESPMode::ThreadSafe> Stream) override;
-	virtual bool SetResponseBodyReceiveStream(TSharedRef<FArchive> Stream) override;
 	virtual void SetHeader(const FString& HeaderName, const FString& HeaderValue) override;
 	virtual void AppendToHeader(const FString& HeaderName, const FString& AdditionalHeaderValue) override;
 	virtual bool ProcessRequest() override;
@@ -74,11 +72,12 @@ PACKAGE_SCOPE:
 	using FHttpRequestCommon::ResetActivityTimeoutTimer;
 	using FHttpRequestCommon::StopActivityTimeoutTimer;
 
-	const TSharedPtr<FArchive> GetResponseBodyReceiveStream() const;
-
 PACKAGE_SCOPE:
 	using FHttpRequestCommon::TriggerStatusCodeReceivedDelegate;
 	using FHttpRequestCommon::SetEffectiveURL;
+	using FHttpRequestCommon::PassReceivedDataToStream;
+
+	bool IsInitializedWithValidStream() const;
 
 private:
 	/**
@@ -104,8 +103,8 @@ private:
 	/** This is the NSMutableURLRequest, all our Apple functionality will deal with this. */
 	NSMutableURLRequest* Request;
 
-    /** This is the session our request belongs to */
-    NSURLSession* Session;
+	/** This is the session our request belongs to */
+	NSURLSession* Session;
 	
 	/** This is the Task associated to the sessionin charge of our request */
 	NSURLSessionTask* Task;
@@ -119,9 +118,6 @@ private:
 	/** Array used to retrieve back content set on the ObjC request when calling GetContent*/
 	mutable TArray<uint8> StorageForGetContent;
 
-	/** The stream to receive response body */
-	TSharedPtr<FArchive> ResponseBodyReceiveStream;
-
 	/** Time taken to complete/cancel the request. */
 	float ElapsedTime;
 
diff --git a/Engine/Source/Runtime/Online/HTTP/Private/Curl/CurlHttp.cpp b/Engine/Source/Runtime/Online/HTTP/Private/Curl/CurlHttp.cpp
--- a/Engine/Source/Runtime/Online/HTTP/Private/Curl/CurlHttp.cpp
+++ b/Engine/Source/Runtime/Online/HTTP/Private/Curl/CurlHttp.cpp
@@ -361,12 +361,6 @@ bool FCurlHttpRequest::SetContentFromStream(TSharedRef<FArchive, ESPMode::Thread
 	return true;
 }
 
-bool FCurlHttpRequest::SetResponseBodyReceiveStream(TSharedRef<FArchive> Stream)
-{
-	ResponseBodyReceiveStream = Stream;
-	return true;
-}
-
 void FCurlHttpRequest::SetHeader(const FString& HeaderName, const FString& HeaderValue)
 {
 	if (CompletionStatus == EHttpRequestStatus::Processing)
@@ -599,12 +593,16 @@ size_t FCurlHttpRequest::ReceiveResponseBodyCallback(void* Ptr, size_t SizeInBlo
 		return NumberOfBytesProcessed;
 	}
 
-	if (ResponseBodyReceiveStream)
+	if (bInitializedWithValidStream)
 	{
-		ResponseBodyReceiveStream->Serialize(Ptr, SizeToDownload);
-
-		if (!ResponseBodyReceiveStream->GetError())
+		if (PassReceivedDataToStream(Ptr, SizeToDownload))
+		{
+			NumberOfBytesProcessed = SizeToDownload;
+		}
+		else if (bCanceled)
 		{
+			// If it's because of cancellation, set processed size as well so curl don't raise a warning caused by CURLE_WRITE_ERROR 
+			// The transfer will be stopped by cancel flow anyway
 			NumberOfBytesProcessed = SizeToDownload;
 		}
 	}
diff --git a/Engine/Source/Runtime/Online/HTTP/Private/Curl/CurlHttp.h b/Engine/Source/Runtime/Online/HTTP/Private/Curl/CurlHttp.h
--- a/Engine/Source/Runtime/Online/HTTP/Private/Curl/CurlHttp.h
+++ b/Engine/Source/Runtime/Online/HTTP/Private/Curl/CurlHttp.h
@@ -161,7 +161,6 @@ public:
 	virtual void SetContentAsString(const FString& ContentString) override;
 	virtual bool SetContentAsStreamedFile(const FString& Filename) override;
 	virtual bool SetContentFromStream(TSharedRef<FArchive, ESPMode::ThreadSafe> Stream) override;
-	virtual bool SetResponseBodyReceiveStream(TSharedRef<FArchive> Stream) override;
 	virtual void SetHeader(const FString& HeaderName, const FString& HeaderValue) override;
 	virtual void AppendToHeader(const FString& HeaderName, const FString& AdditionalHeaderValue) override;
 	virtual bool ProcessRequest() override;
@@ -378,8 +377,6 @@ private:
 	TUniquePtr<FRequestPayload> RequestPayload;
 	/** Is the request payload seekable? */
 	bool bIsRequestPayloadSeekable = false;
-	/** The stream to receive response body */
-	TSharedPtr<FArchive> ResponseBodyReceiveStream;
 	/** Mapping of header section to values. */
 	TMap<FString, FString> Headers;
 	/** Total elapsed time in seconds since the start of the request */
diff --git a/Engine/Source/Runtime/Online/HTTP/Private/GenericPlatform/HttpRequestCommon.cpp b/Engine/Source/Runtime/Online/HTTP/Private/GenericPlatform/HttpRequestCommon.cpp
--- a/Engine/Source/Runtime/Online/HTTP/Private/GenericPlatform/HttpRequestCommon.cpp
+++ b/Engine/Source/Runtime/Online/HTTP/Private/GenericPlatform/HttpRequestCommon.cpp
@@ -219,6 +219,8 @@ void FHttpRequestCommon::CancelRequest()
 
 	StopActivityTimeoutTimer();
 
+	StopPassingReceivedData();
+
 	bCanceled = true;
 	UE_LOG(LogHttp, Verbose, TEXT("HTTP request canceled. URL=%s"), *GetURL());
 
@@ -402,6 +404,7 @@ void FHttpRequestCommon::Shutdown()
 {
 	FHttpRequestImpl::Shutdown();
 
+	StopPassingReceivedData();
 	StopActivityTimeoutTimer();
 	StopTotalTimeoutTimer();
 }
@@ -445,3 +448,33 @@ void FHttpRequestCommon::SetEffectiveURL(const FString& InEffectiveURL)
 		ResponseCommon->SetEffectiveURL(EffectiveURL);
 	}
 }
+
+bool FHttpRequestCommon::SetResponseBodyReceiveStream(TSharedRef<FArchive> Stream)
+{
+	const FScopeLock StreamLock(&ResponseBodyReceiveStreamCriticalSection);
+
+	ResponseBodyReceiveStream = Stream;
+	bInitializedWithValidStream = true;
+	return true;
+}
+
+bool FHttpRequestCommon::PassReceivedDataToStream(void* Ptr, int64 Length)
+{
+	const FScopeLock StreamLock(&ResponseBodyReceiveStreamCriticalSection);
+
+	if (!ResponseBodyReceiveStream)
+	{
+		return false;
+	}
+
+	ResponseBodyReceiveStream->Serialize(Ptr, Length);
+
+	return !ResponseBodyReceiveStream->GetError();
+}
+
+void FHttpRequestCommon::StopPassingReceivedData()
+{
+	const FScopeLock StreamLock(&ResponseBodyReceiveStreamCriticalSection);
+
+	ResponseBodyReceiveStream = nullptr;
+}
diff --git a/Engine/Source/Runtime/Online/HTTP/Public/GenericPlatform/HttpRequestCommon.h b/Engine/Source/Runtime/Online/HTTP/Public/GenericPlatform/HttpRequestCommon.h
--- a/Engine/Source/Runtime/Online/HTTP/Public/GenericPlatform/HttpRequestCommon.h
+++ b/Engine/Source/Runtime/Online/HTTP/Public/GenericPlatform/HttpRequestCommon.h
@@ -41,6 +41,8 @@ public:
 
 	HTTP_API virtual void ProcessRequestUntilComplete() override;
 
+	HTTP_API virtual bool SetResponseBodyReceiveStream(TSharedRef<FArchive> Stream) override;
+
 protected:
 	/**
 	 * Check if this request is valid or allowed, before actually process the request
@@ -78,6 +80,9 @@ protected:
 
 	HTTP_API void SetEffectiveURL(const FString& InEffectiveURL);
 
+	HTTP_API bool PassReceivedDataToStream(void* Ptr, int64 Length);
+	HTTP_API void StopPassingReceivedData();
+
 protected:
 	/** Current status of request being processed */
 	EHttpRequestStatus::Type CompletionStatus = EHttpRequestStatus::NotStarted;
@@ -92,11 +97,11 @@ protected:
 	TOptional<float> TimeoutSecs;
 
 	/** Indicate the request is timed out, it should quit and fail with EHttpFailureReason::TimedOut */
-	bool bTimedOut = false;
+	std::atomic<bool> bTimedOut = false;
 	/** Indicate the request is activity timed out, it should quit and fail with EHttpFailureReason::ConnectionError */
-	bool bActivityTimedOut = false;
+	std::atomic<bool> bActivityTimedOut = false;
 	/** Indicate the request is cancelled, it should quit and fail with EHttpFailureReason::Cancelled */
-	bool bCanceled = false;
+	std::atomic<bool> bCanceled = false;
 
 	/** TODO: Move this feature into CurlHttp */
 	bool bUsePlatformActivityTimeout = true;
@@ -125,4 +130,14 @@ protected:
 
 	/** The response object which we will use to pair with this request */
 	TSharedPtr<FHttpResponseCommon> ResponseCommon;
+
+	/** The stream to receive response body */
+	TSharedPtr<FArchive> ResponseBodyReceiveStream;
+
+	/** Critical section for accessing ResponseBodyReceiveStream */
+	FCriticalSection ResponseBodyReceiveStreamCriticalSection;
+
+	// Flag to indicate the request was initialized with stream. In that case even if stream was set to 
+	// null later on internally, the request itself won't cache received data anymore
+	std::atomic<bool> bInitializedWithValidStream = false;
 };
diff --git a/Engine/Source/Runtime/Online/HTTP/Public/Interfaces/IHttpRequest.h b/Engine/Source/Runtime/Online/HTTP/Public/Interfaces/IHttpRequest.h
--- a/Engine/Source/Runtime/Online/HTTP/Public/Interfaces/IHttpRequest.h
+++ b/Engine/Source/Runtime/Online/HTTP/Public/Interfaces/IHttpRequest.h
@@ -169,14 +169,14 @@ public:
 	 * @param ContentString - payload to set.
 	 */
 	virtual void SetContentAsString(const FString& ContentString) = 0;
-    
-    /**
-     * Sets the content of the request to stream from a file.
-     *
-     * @param FileName - filename from which to stream the body.
+
+	/**
+	 * Sets the content of the request to stream from a file.
+	 *
+	 * @param FileName - filename from which to stream the body.
 	 * @return True if the file is valid and will be used to stream the request. False otherwise.
-     */
-    virtual bool SetContentAsStreamedFile(const FString& Filename) = 0;
+	 */
+	virtual bool SetContentAsStreamedFile(const FString& Filename) = 0;
 
 	/**
 	 * Sets the content of the request to stream directly from an archive.
@@ -189,8 +189,13 @@ public:
 
 	/**
 	 * Sets the content of the request to stream directly from an delegate.
-	 * NOTE: The delegate will be called from another thread other than the game thread
-	 *
+	 * NOTE: 
+	 *   - The delegate will be called from another thread other than the game thread, make sure 
+	 *     it's thread-safe in there
+	 *   - Make sure the delegate is safe to be called until receiving the process complete callback
+	 *     or after canceling the request
+	 *     For example: don't destroy the instance even if using BindThreadSafeSP, because internally 
+	 *     it's calling Execute to handle error by returned value instead of calling ExecuteIfBound
 	 * @param StreamDelegate - delegate from which the payload should be streamed.
 	 * @return True if the delegate can be used to stream the request. False otherwise.
 	 */
