Fix the dead lock by reducing the scope of request lock in http manager to make sure when trigger http request complete callback in Tick, it doesn't hold the requests lock;
Deperecate the FGenericPlatformHttp::UsesThreadedHttp interface and remove corresponding impls, now all platforms are using threaded http requests;
Re-enable retry system off-game thread support.

Reason:
When accessing requests lock and trigger complete callback in http manager Tick from main thread, the complete callback could have been bound to retry manager which will try to access retry manager lock, while retry manager lock can be holding by FHttpRetrySystem::FManager::RetryHttpRequestWithDelay, which triggered by timer callback from http thread and is trying to access http manager's requests lock in order to add the retry request to http manager.

diff --git a/Engine/Source/Programs/WebTests/Private/TestHttp.cpp b/Engine/Source/Programs/WebTests/Private/TestHttp.cpp
--- a/Engine/Source/Programs/WebTests/Private/TestHttp.cpp
+++ b/Engine/Source/Programs/WebTests/Private/TestHttp.cpp
@@ -1785,6 +1785,29 @@ TEST_CASE_METHOD(FWaitUntilCompleteHttpFixture, "Retry fallback with exponential
 	HttpRequest->ProcessRequest();
 }
 
+TEST_CASE_METHOD(FWaitUntilCompleteHttpFixture, "Dead lock test by retrying requests while completing requests", HTTP_TAG)
+{
+	if (!bRetryEnabled)
+	{
+		return;
+	}
+
+	DisableWarningsInThisTest();
+
+	for (uint32 i = 0; i < 50; ++i)
+	{
+		TSharedRef<IHttpRequest> HttpRequest = HttpRetryManager->CreateRequest(
+			5/*InRetryLimitCountOverride*/,
+			FHttpRetrySystem::FRetryTimeoutRelativeSecondsSetting()/*InRetryTimeoutRelativeSecondsOverride unused*/,
+			{ EHttpResponseCodes::TooManyRequests }/*InRetryResponseCodes*/
+		);
+
+		HttpRequest->SetURL(UrlMockStatus(EHttpResponseCodes::TooManyRequests));
+		HttpRequest->SetHeader(TEXT("Retry-After"), FString::Format(TEXT("{0}"), { 0.1 }));
+		HttpRequest->ProcessRequest();
+	}
+}
+
 class FThreadedBatchRequestsFixture : public FWaitThreadedHttpFixture
 {
 public:
diff --git a/Engine/Source/Runtime/Online/HTTP/Private/Apple/ApplePlatformHttp.cpp b/Engine/Source/Runtime/Online/HTTP/Private/Apple/ApplePlatformHttp.cpp
--- a/Engine/Source/Runtime/Online/HTTP/Private/Apple/ApplePlatformHttp.cpp
+++ b/Engine/Source/Runtime/Online/HTTP/Private/Apple/ApplePlatformHttp.cpp
@@ -275,7 +275,3 @@ IHttpRequest* FApplePlatformHttp::ConstructRequest()
 	return new FAppleHttpRequest(Session);
 }
 
-bool FApplePlatformHttp::UsesThreadedHttp()
-{
-	return true;
-}
diff --git a/Engine/Source/Runtime/Online/HTTP/Private/GenericPlatform/GenericPlatformHttp.cpp b/Engine/Source/Runtime/Online/HTTP/Private/GenericPlatform/GenericPlatformHttp.cpp
--- a/Engine/Source/Runtime/Online/HTTP/Private/GenericPlatform/GenericPlatformHttp.cpp
+++ b/Engine/Source/Runtime/Online/HTTP/Private/GenericPlatform/GenericPlatformHttp.cpp
@@ -309,12 +309,12 @@ IHttpRequest* FGenericPlatformHttp::ConstructRequest()
 	return new FGenericPlatformHttpRequest();
 }
 
+PRAGMA_DISABLE_DEPRECATION_WARNINGS
 bool FGenericPlatformHttp::UsesThreadedHttp()
 {
-	// Many platforms use libcurl.  Our libcurl implementation uses threading.
-	// Platforms that don't use libcurl but have threading will need to override UsesThreadedHttp to return true for their platform.
-	return WITH_CURL;
+	return true;
 }
+PRAGMA_ENABLE_DEPRECATION_WARNINGS
 
 static bool IsAllowedChar(UTF8CHAR LookupChar)
 {
diff --git a/Engine/Source/Runtime/Online/HTTP/Private/HttpManager.cpp b/Engine/Source/Runtime/Online/HTTP/Private/HttpManager.cpp
--- a/Engine/Source/Runtime/Online/HTTP/Private/HttpManager.cpp
+++ b/Engine/Source/Runtime/Online/HTTP/Private/HttpManager.cpp
@@ -66,7 +66,7 @@ PRAGMA_ENABLE_DEPRECATION_WARNINGS
 
 void FHttpManager::Initialize()
 {
-	if (FPlatformHttp::UsesThreadedHttp())
+	if (!Thread)
 	{
 		Thread = CreateHttpThread();
 		Thread->StartThread();
@@ -77,19 +77,21 @@ void FHttpManager::Initialize()
 
 void FHttpManager::Shutdown()
 {
-	FScopeLock ScopeLock(&RequestLock);
-
-	// Don't emit these tracking logs in commandlet runs. Build system traps warnings during cook, and these are not truly fatal, but useful for tracking down shutdown issues.
-	UE_CLOG(ShouldOutputHttpWarnings() && Requests.Num(), LogHttp, Warning, TEXT("[FHttpManager::Shutdown] Unbinding delegates for %d outstanding Http Requests:"), Requests.Num());
-
-	// Clear delegates since they may point to deleted instances
-	for (TArray<FHttpRequestRef>::TIterator It(Requests); It; ++It)
 	{
-		TSharedPtr<IHttpRequest> Request = *It;
-		StaticCastSharedPtr<FHttpRequestImpl>(Request)->Shutdown();
+		FScopeLock ScopeLock(&RequestLock);
 
 		// Don't emit these tracking logs in commandlet runs. Build system traps warnings during cook, and these are not truly fatal, but useful for tracking down shutdown issues.
-		UE_CLOG(ShouldOutputHttpWarnings(), LogHttp, Warning, TEXT("	verb=[%s] url=[%s] refs=[%d] status=%s"), *Request->GetVerb(), *Request->GetURL(), Request.GetSharedReferenceCount(), EHttpRequestStatus::ToString(Request->GetStatus()));
+		UE_CLOG(ShouldOutputHttpWarnings() && Requests.Num(), LogHttp, Warning, TEXT("[FHttpManager::Shutdown] Unbinding delegates for %d outstanding Http Requests:"), Requests.Num());
+
+		// Clear delegates since they may point to deleted instances
+		for (TArray<FHttpRequestRef>::TIterator It(Requests); It; ++It)
+		{
+			TSharedPtr<IHttpRequest> Request = *It;
+			StaticCastSharedPtr<FHttpRequestImpl>(Request)->Shutdown();
+
+			// Don't emit these tracking logs in commandlet runs. Build system traps warnings during cook, and these are not truly fatal, but useful for tracking down shutdown issues.
+			UE_CLOG(ShouldOutputHttpWarnings(), LogHttp, Warning, TEXT("	verb=[%s] url=[%s] refs=[%d] status=%s"), *Request->GetVerb(), *Request->GetURL(), Request.GetSharedReferenceCount(), EHttpRequestStatus::ToString(Request->GetStatus()));
+		}
 	}
 
 	// Clear general delegates since they may point to deleted instances
@@ -371,7 +373,7 @@ void FHttpManager::Flush(EHttpFlushReason FlushReason)
 		{
 			if (Thread)
 			{
-				if( Thread->NeedsSingleThreadTick() )
+				if (Thread->NeedsSingleThreadTick())
 				{
 					if (AppTime >= StallWarnTime)
 					{
@@ -388,10 +390,6 @@ void FHttpManager::Flush(EHttpFlushReason FlushReason)
 					FPlatformProcess::Sleep(SecondsToSleepForOutstandingThreadedRequests);
 				}
 			}
-			else
-			{
-				check(!FPlatformHttp::UsesThreadedHttp());
-			}
 		}
 
 		AppTime = FPlatformTime::Seconds();
@@ -436,14 +434,15 @@ bool FHttpManager::Tick(float DeltaSeconds)
 		}
 	}
 
-	FScopeLock ScopeLock(&RequestLock);
-
 	if (Thread)
 	{
-		// Tick each active request
-		for (const FHttpRequestRef& Request: Requests)
 		{
-			Request->Tick(DeltaSeconds);
+			// Tick each active request
+			FScopeLock ScopeLock(&RequestLock);
+			for (const FHttpRequestRef& Request : Requests)
+			{
+				Request->Tick(DeltaSeconds);
+			}
 		}
 
 		TArray<IHttpThreadedRequest*> CompletedThreadedRequests;
@@ -453,7 +452,12 @@ bool FHttpManager::Tick(float DeltaSeconds)
 		for (IHttpThreadedRequest* CompletedRequest : CompletedThreadedRequests)
 		{
 			FHttpRequestRef CompletedRequestRef = CompletedRequest->AsShared();
-			Requests.Remove(CompletedRequestRef);
+
+			{
+				FScopeLock ScopeLock(&RequestLock);
+				Requests.Remove(CompletedRequestRef);
+			}
+
 			if (CompletedRequest->GetDelegateThreadPolicy() == EHttpRequestDelegateThreadPolicy::CompleteOnGameThread)
 			{
 				CompletedRequest->FinishRequest();
@@ -461,26 +465,7 @@ bool FHttpManager::Tick(float DeltaSeconds)
 			}
 		}
 	}
-	else
-	{
-		TArray<FHttpRequestRef> CompletedRequests;
 
-		// Tick each active request
-		for (const FHttpRequestRef& Request: Requests)
-		{
-			Request->Tick(DeltaSeconds);
-			if (EHttpRequestStatus::IsFinished(Request->GetStatus()))
-			{
-				CompletedRequests.Add(Request);
-			}
-		}
-
-		for (const FHttpRequestRef& CompletedRequest: CompletedRequests)
-		{
-			Requests.Remove(CompletedRequest);
-			BroadcastHttpRequestCompleted(CompletedRequest);
-		}
-	}
 	// keep ticking
 	return true;
 }
diff --git a/Engine/Source/Runtime/Online/HTTP/Private/HttpRetrySystem.cpp b/Engine/Source/Runtime/Online/HTTP/Private/HttpRetrySystem.cpp
--- a/Engine/Source/Runtime/Online/HTTP/Private/HttpRetrySystem.cpp
+++ b/Engine/Source/Runtime/Online/HTTP/Private/HttpRetrySystem.cpp
@@ -17,7 +17,7 @@
 
 TAutoConsoleVariable<bool> CVarHttpRetrySystemNonGameThreadSupportEnabled(
 	TEXT("Http.RetrySystemNonGameThreadSupportEnabled"),
-	false,
+	true,
 	TEXT("Enable retry system non-game thread support")
 );
 
diff --git a/Engine/Source/Runtime/Online/HTTP/Public/Apple/ApplePlatformHttp.h b/Engine/Source/Runtime/Online/HTTP/Public/Apple/ApplePlatformHttp.h
--- a/Engine/Source/Runtime/Online/HTTP/Public/Apple/ApplePlatformHttp.h
+++ b/Engine/Source/Runtime/Online/HTTP/Public/Apple/ApplePlatformHttp.h
@@ -34,13 +34,6 @@ public:
 	 */
 	static IHttpRequest* ConstructRequest();
 
-	/**
-	 * Check if a platform uses the HTTP thread
-	 *
-	 * @return true if the platform uses threaded HTTP, false if not
-	 */
-	static bool UsesThreadedHttp();
-
 private:
     /** Session used to create Apple based requests */
     static inline NSURLSession* Session = nil;
diff --git a/Engine/Source/Runtime/Online/HTTP/Public/GenericPlatform/GenericPlatformHttp.h b/Engine/Source/Runtime/Online/HTTP/Public/GenericPlatform/GenericPlatformHttp.h
--- a/Engine/Source/Runtime/Online/HTTP/Public/GenericPlatform/GenericPlatformHttp.h
+++ b/Engine/Source/Runtime/Online/HTTP/Public/GenericPlatform/GenericPlatformHttp.h
@@ -90,6 +90,7 @@ public:
 	 *
 	 * @return true if the platform uses threaded HTTP, false if not
 	 */
+	UE_DEPRECATED(5.5, "UsesThreadedHttp is deprecated and will be removed")
 	static HTTP_API bool UsesThreadedHttp();
 
 	/**
