Fix the compile error in AppleHttp caused by bCanceled.

diff --git a/Engine/Source/Runtime/Online/HTTP/Private/Apple/AppleHttp.cpp b/Engine/Source/Runtime/Online/HTTP/Private/Apple/AppleHttp.cpp
--- a/Engine/Source/Runtime/Online/HTTP/Private/Apple/AppleHttp.cpp
+++ b/Engine/Source/Runtime/Online/HTTP/Private/Apple/AppleHttp.cpp
@@ -192,7 +192,7 @@
 			*stop = bSerializeSucceed? NO : YES;
 		}];
 		
-		if (!bSerializeSucceed && !bCanceled)
+		if (!bSerializeSucceed)
 		{
 			[dataTask cancel];
 		}
