In Apple http impl change it to read response headers once received status code;
Use cvar for the new flow just in case the headers needs to be read when complete for some urls

diff --git a/Engine/Source/Runtime/Online/HTTP/Private/Apple/AppleHttp.cpp b/Engine/Source/Runtime/Online/HTTP/Private/Apple/AppleHttp.cpp
--- a/Engine/Source/Runtime/Online/HTTP/Private/Apple/AppleHttp.cpp
+++ b/Engine/Source/Runtime/Online/HTTP/Private/Apple/AppleHttp.cpp
@@ -2,13 +2,55 @@
 
 
 #include "AppleHttp.h"
+#include "HAL/IConsoleManager.h"
+#include "HAL/PlatformTime.h"
+#include "Http.h"
 #include "HttpManager.h"
+#include "HttpModule.h"
 #include "Misc/EngineVersion.h"
 #include "Misc/App.h"
 #include "Misc/Base64.h"
-#include "HAL/PlatformTime.h"
-#include "Http.h"
-#include "HttpModule.h"
+
+// It should be safe to read headers early, add the CVar here just in case
+TAutoConsoleVariable<FString> CVarHttpUrlsToReadHeadersWhenComplete(
+	TEXT("http.UrlsToReadHeadersWhenComplete"),
+	TEXT(""),
+	TEXT("List of urls to only read headers when complete the http request\"www.epicgames.com,www.unrealengine.com,...\"")
+);
+
+namespace AppleHTTPRequestInternal
+{
+
+static bool bUpdatedCVarHttpUrlsToReadHeadersWhenComplete = true;
+static TArray<FString> UrlsToReadHeadersWhenComplete;
+
+static void UpdateConfigFromCVar()
+{
+	UE_CALL_ONCE([] {
+		CVarHttpUrlsToReadHeadersWhenComplete.AsVariable()->OnChangedDelegate().AddLambda([](IConsoleVariable* CVar) {
+			bUpdatedCVarHttpUrlsToReadHeadersWhenComplete = true;
+		});
+		bUpdatedCVarHttpUrlsToReadHeadersWhenComplete = true;
+	});
+	if (bUpdatedCVarHttpUrlsToReadHeadersWhenComplete)
+	{
+		CVarHttpUrlsToReadHeadersWhenComplete.GetValueOnAnyThread().ParseIntoArray(UrlsToReadHeadersWhenComplete, TEXT(","));
+		bUpdatedCVarHttpUrlsToReadHeadersWhenComplete = false;
+	}
+}
+
+static bool ShouldReadHeadersWhenComplete(const FString& Url)
+{
+	for (const FString& UrlToReadHeadersWhenComplete : UrlsToReadHeadersWhenComplete)
+	{
+		if (Url.StartsWith(UrlToReadHeadersWhenComplete))
+		{
+			return true;
+		}
+	}
+	return false;
+}
+}
 
 /**
  * Class to hold data from delegate implementation notifications.
@@ -117,6 +159,28 @@
 	}
 }
 
+-(void) BroadcastResponseHeadersReceived
+{
+	if (TSharedPtr<FAppleHttpRequest> Request = SourceRequest.Pin())
+	{
+		AppleHTTPRequestInternal::UpdateConfigFromCVar();
+		if (!AppleHTTPRequestInternal::ShouldReadHeadersWhenComplete(Request->GetURL()))
+		{
+			if (Request->GetDelegateThreadPolicy() == EHttpRequestDelegateThreadPolicy::CompleteOnHttpThread)
+			{
+				Request->BroadcastResponseHeadersReceived();
+			}
+			else if (Request->OnHeaderReceived().IsBound())
+			{
+				FHttpModule::Get().GetHttpManager().AddGameThreadTask([Request]()
+				{
+					Request->BroadcastResponseHeadersReceived();
+				});
+			}
+		}
+	}
+}
+
 - (BOOL) DidValidActivityOcurred:(FStringView) Reason
 {
 	TSharedPtr<FAppleHttpRequest> Request = SourceRequest.Pin();
@@ -166,6 +230,8 @@
 	FString EffectiveURL([Url absoluteString]);
 	[self SaveEffectiveURL: EffectiveURL];
 
+	[self BroadcastResponseHeadersReceived];
+
 	uint64 ExpectedResponseLength = response.expectedContentLength;
 	if(!bInitializedWithValidStream && ExpectedResponseLength != NSURLResponseUnknownLength)
 	{
@@ -819,8 +885,10 @@ void FAppleHttpRequest::FinishRequest()
 	}
 	else
 	{
-		// TODO: Try to broadcast OnHeaderReceived when we receive headers instead of here at the end
-		BroadcastResponseHeadersReceived();
+		if (AppleHTTPRequestInternal::ShouldReadHeadersWhenComplete(GetURL()))
+		{
+			BroadcastResponseHeadersReceived();
+		}
 	}
 
 	OnProcessRequestComplete().ExecuteIfBound(SharedThis(this), Response, bSucceeded);
@@ -944,10 +1012,18 @@ void FAppleHttpResponse::CleanSharedObjects()
 
 FString FAppleHttpResponse::GetHeader(const FString& HeaderName) const
 {
-	SCOPED_AUTORELEASE_POOL;
-	UE_LOG(LogHttp, Verbose, TEXT("FAppleHttpResponse::GetHeader()"));
-	NSString* ConvertedHeaderName = HeaderName.GetNSString();
-	return FString([ResponseDelegate.Response.allHeaderFields objectForKey:ConvertedHeaderName]);
+	if (AppleHTTPRequestInternal::ShouldReadHeadersWhenComplete(GetURL()) && !IsReady())
+	{
+		UE_LOG(LogHttp, Warning, TEXT("Can't get header [%s]. Response still processing for %s."), *HeaderName, *GetURL());
+		return FString();
+	}
+	else
+	{
+		SCOPED_AUTORELEASE_POOL;
+		UE_LOG(LogHttp, Verbose, TEXT("FAppleHttpResponse::GetHeader()"));
+		NSString* ConvertedHeaderName = HeaderName.GetNSString();
+		return FString([ResponseDelegate.Response.allHeaderFields objectForKey:ConvertedHeaderName]);
+	}
 }
 
 TArray<FString> FAppleHttpResponse::GetAllHeaders() const
diff --git a/Engine/Source/Runtime/Online/HTTP/Private/Apple/AppleHttp.h b/Engine/Source/Runtime/Online/HTTP/Private/Apple/AppleHttp.h
--- a/Engine/Source/Runtime/Online/HTTP/Private/Apple/AppleHttp.h
+++ b/Engine/Source/Runtime/Online/HTTP/Private/Apple/AppleHttp.h
@@ -68,11 +68,10 @@ public:
 	virtual ~FAppleHttpRequest();
 
 PACKAGE_SCOPE:
+	using FHttpRequestCommon::BroadcastResponseHeadersReceived;
 	using FHttpRequestCommon::StartActivityTimeoutTimer;
 	using FHttpRequestCommon::ResetActivityTimeoutTimer;
 	using FHttpRequestCommon::StopActivityTimeoutTimer;
-
-PACKAGE_SCOPE:
 	using FHttpRequestCommon::TriggerStatusCodeReceivedDelegate;
 	using FHttpRequestCommon::SetEffectiveURL;
 	using FHttpRequestCommon::PassReceivedDataToStream;
