Fixed issues in AppleHttp so WebTests "Retry immediately without lock out if connect failed and there are alt domains" and "Make sure connection time out can work well for 2nd same http request" are passed
Also fixed a wrong unconditional cast in case of non http response

#changelist validated
#virtualized

diff --git a/Engine/Source/Runtime/Online/HTTP/Private/Apple/AppleHttp.cpp b/Engine/Source/Runtime/Online/HTTP/Private/Apple/AppleHttp.cpp
--- a/Engine/Source/Runtime/Online/HTTP/Private/Apple/AppleHttp.cpp
+++ b/Engine/Source/Runtime/Online/HTTP/Private/Apple/AppleHttp.cpp
@@ -73,7 +73,7 @@ static bool ShouldReadHeadersWhenComplete(const FString& Url)
 }
 
 /** A handle for the response */
-@property(retain) NSHTTPURLResponse* Response;
+@property(retain) NSURLResponse* Response;
 /** The total number of bytes written out during the request/response */
 @property uint64 BytesWritten;
 /** The total number of bytes received out during the request/response */
@@ -85,7 +85,6 @@ static bool ShouldReadHeadersWhenComplete(const FString& Url)
 /** Associated request. Cleared when canceled */
 @property TWeakPtr<FAppleHttpRequest> SourceRequest;
 
-
 /** NSURLSessionDataDelegate delegate methods. Those are called from a thread controlled by the NSURLSession */
 
 /** Sent periodically to notify the delegate of upload progress. */
@@ -108,6 +107,38 @@ static bool ShouldReadHeadersWhenComplete(const FString& Url)
 @synthesize BytesReceived;
 @synthesize SourceRequest;
 
+-(int32)GetStatusCode;
+{
+	if (self.Response == nil)
+	{
+		return 0;
+	}
+	else if ([self.Response isKindOfClass: [NSHTTPURLResponse class]])
+	{
+		return ((NSHTTPURLResponse*)self.Response).statusCode;
+	}
+	else
+	{
+		return 200;
+	}
+}
+
+-(NSDictionary*)GetResponseHeaders;
+{
+	if (self.Response == nil)
+	{
+		return nil;
+	}
+	else if ([self.Response isKindOfClass: [NSHTTPURLResponse class]])
+	{
+		return ((NSHTTPURLResponse*)self.Response).allHeaderFields;
+	}
+	else
+	{
+		return nil;
+	}
+}
+
 - (FAppleHttpResponseDelegate*)initWithRequest:(FAppleHttpRequest&) Request
 {
 	self = [super init];
@@ -221,9 +252,10 @@ static bool ShouldReadHeadersWhenComplete(const FString& Url)
 
 	UE_LOG(LogHttp, Verbose, TEXT("URLSession:dataTask:didReceiveResponse:completionHandler"));
 	
-	self.Response = (NSHTTPURLResponse*)response;
+	self.Response = response;
 
-	int32 StatusCode = [self.Response statusCode];
+	int32 StatusCode = [self GetStatusCode];
+	
 	[self HandleStatusCodeReceived: StatusCode];
 
 	NSURL* Url = [self.Response URL];
@@ -516,20 +548,6 @@ FAppleHttpRequest::FAppleHttpRequest(NSURLSession* InSession)
 	
 	UE_LOG(LogHttp, Verbose, TEXT("FAppleHttpRequest::FAppleHttpRequest()"));
 	Request = [[NSMutableURLRequest alloc] init];
-	float HttpConnectionTimeout = FHttpModule::Get().GetHttpConnectionTimeout();
-	check(HttpConnectionTimeout > 0.0f);
-	Request.timeoutInterval = HttpConnectionTimeout;
-	
-	UE_CLOG(
-		HttpConnectionTimeout < FHttpModule::Get().GetHttpActivityTimeout(), 
-		LogHttp, 
-		Warning, 
-		TEXT(
-			"HttpConnectionTimeout can't be less than HttpActivityTimeout, otherwise requests may complete "
-			"unexpectedly with ConnectionError after %.2f(HttpConnectionTimeout) seconds without activity, "
-			"instead of intended %.2f(HttpActivityTimeout) seconds"
-		), 
-		HttpConnectionTimeout, FHttpModule::Get().GetHttpActivityTimeout());
 
 	// Disable cache to mimic WinInet behavior
 	Request.cachePolicy = NSURLRequestReloadIgnoringLocalCacheData;
@@ -816,6 +834,21 @@ bool FAppleHttpRequest::SetupRequest()
 	LastReportedBytesRead = 0;
 	ElapsedTime = 0.0f;
 
+	float HttpConnectionTimeout = FHttpModule::Get().GetHttpConnectionTimeout();
+	check(HttpConnectionTimeout > 0.0f);
+	Request.timeoutInterval = HttpConnectionTimeout;
+	
+	UE_CLOG(
+		HttpConnectionTimeout < FHttpModule::Get().GetHttpActivityTimeout(), 
+		LogHttp,
+		Warning, 
+		TEXT(
+			"HttpConnectionTimeout can't be less than HttpActivityTimeout, otherwise requests may complete "
+			"unexpectedly with ConnectionError after %.2f(HttpConnectionTimeout) seconds without activity, "
+			"instead of intended %.2f(HttpActivityTimeout) seconds"
+		), 
+		HttpConnectionTimeout, FHttpModule::Get().GetHttpActivityTimeout());
+
 	Task = [Session dataTaskWithRequest: Request];
 	
 	if (Task != nil)
@@ -880,7 +913,7 @@ void FAppleHttpRequest::FinishRequest()
 
 		if (GetFailureReason() == EHttpFailureReason::ConnectionError)
 		{
-			Response = nullptr;
+			ResponseCommon = nullptr;
 		}
 	}
 	else
@@ -891,7 +924,7 @@ void FAppleHttpRequest::FinishRequest()
 		}
 	}
 
-	OnProcessRequestComplete().ExecuteIfBound(SharedThis(this), Response, bSucceeded);
+	OnProcessRequestComplete().ExecuteIfBound(SharedThis(this), ResponseCommon, bSucceeded);
 }
 
 void FAppleHttpRequest::CleanupRequest()
@@ -1020,9 +1053,17 @@ FString FAppleHttpResponse::GetHeader(const FString& HeaderName) const
 	else
 	{
 		SCOPED_AUTORELEASE_POOL;
-		UE_LOG(LogHttp, Verbose, TEXT("FAppleHttpResponse::GetHeader()"));
-		NSString* ConvertedHeaderName = HeaderName.GetNSString();
-		return FString([ResponseDelegate.Response.allHeaderFields objectForKey:ConvertedHeaderName]);
+		if(NSDictionary* Headers = [ResponseDelegate GetResponseHeaders])
+		{
+			UE_LOG(LogHttp, Verbose, TEXT("FAppleHttpResponse::GetHeader()"));
+			NSHTTPURLResponse* Response = (NSHTTPURLResponse*)ResponseDelegate.Response;
+			NSString* ConvertedHeaderName = HeaderName.GetNSString();
+			return FString([Response.allHeaderFields objectForKey:ConvertedHeaderName]);
+		}
+		else
+		{
+			return FString();
+		}
 	}
 }
 
@@ -1030,14 +1071,17 @@ TArray<FString> FAppleHttpResponse::GetAllHeaders() const
 {
 	UE_LOG(LogHttp, Verbose, TEXT("FAppleHttpResponse::GetAllHeaders()"));
 
-	NSDictionary* Headers = ResponseDelegate.Response.allHeaderFields;
 	TArray<FString> Result;
-	Result.Reserve([Headers count]);
-	for (NSString* Key in [Headers allKeys])
+	SCOPED_AUTORELEASE_POOL;
+	if (NSDictionary* Headers = [ResponseDelegate GetResponseHeaders])
 	{
-		FString ConvertedValue([Headers objectForKey:Key]);
-		FString ConvertedKey(Key);
-		Result.Add( FString::Printf( TEXT("%s: %s"), *ConvertedKey, *ConvertedValue ) );
+		Result.Reserve([Headers count]);
+		for (NSString* Key in [Headers allKeys])
+		{
+			FString ConvertedValue([Headers objectForKey:Key]);
+			FString ConvertedKey(Key);
+			Result.Add( FString::Printf( TEXT("%s: %s"), *ConvertedKey, *ConvertedValue ) );
+		}
 	}
 	return Result;
 }
@@ -1087,7 +1131,7 @@ int32 FAppleHttpResponse::GetResponseCode() const
 {
 	UE_LOG(LogHttp, Verbose, TEXT("FAppleHttpResponse::GetResponseCode()"));
 
-	return ResponseDelegate.Response.statusCode;
+	return [ResponseDelegate GetStatusCode];
 }
 
 bool FAppleHttpResponse::IsReady() const
