Fix the http issue that FinishRequest could potentially be called twice in CurlHttp impl.

Reason: When FHttpRequestCommon::CancelRequest get called, it will set bCanceled as true and abort request which will in the end call FinishRequest. But in http thread it will also check IsThreadedRequestComplete to see if the request is complete and call FinishRequest, in that function we shouldn't check bCanceled or bTimedOut, because request has already been aborted when those flags changed.

#tests Reproduced and fixed with a test case

diff --git a/Engine/Source/Runtime/Online/HTTP/Private/Curl/CurlHttp.cpp b/Engine/Source/Runtime/Online/HTTP/Private/Curl/CurlHttp.cpp
--- a/Engine/Source/Runtime/Online/HTTP/Private/Curl/CurlHttp.cpp
+++ b/Engine/Source/Runtime/Online/HTTP/Private/Curl/CurlHttp.cpp
@@ -1091,17 +1091,7 @@ bool FCurlHttpRequest::StartThreadedRequest()
 
 bool FCurlHttpRequest::IsThreadedRequestComplete()
 {
-	if (bCanceled)
-	{
-		return true;
-	}
-
-	if (bTimedOut || bActivityTimedOut)
-	{
-		return true;
-	}
-	
-	if (bCurlRequestCompleted  && ElapsedTime >= FHttpModule::Get().GetHttpDelayTime())
+	if (bCurlRequestCompleted && ElapsedTime >= FHttpModule::Get().GetHttpDelayTime())
 	{
 		return true;
 	}
diff --git a/Engine/Source/Runtime/Online/HTTP/Private/GenericPlatform/HttpRequestCommon.cpp b/Engine/Source/Runtime/Online/HTTP/Private/GenericPlatform/HttpRequestCommon.cpp
--- a/Engine/Source/Runtime/Online/HTTP/Private/GenericPlatform/HttpRequestCommon.cpp
+++ b/Engine/Source/Runtime/Online/HTTP/Private/GenericPlatform/HttpRequestCommon.cpp
@@ -175,7 +175,7 @@ void FHttpRequestCommon::SetStatus(EHttpRequestStatus::Type InCompletionStatus)
 
 void FHttpRequestCommon::SetFailureReason(EHttpFailureReason InFailureReason)
 {
-	check(FailureReason == EHttpFailureReason::None);
+	UE_CLOG(FailureReason != EHttpFailureReason::None, LogHttp, Warning, TEXT("FailureReason had been set to %s, now setting to %s"), LexToString(FailureReason), LexToString(InFailureReason));
 	FailureReason = InFailureReason;
 
 	if (ResponseCommon)
@@ -212,7 +212,8 @@ const FHttpResponsePtr FHttpRequestCommon::GetResponse() const
 
 void FHttpRequestCommon::CancelRequest()
 {
-	if (bCanceled)
+	bool bWasCanceled = bCanceled.exchange(true);
+	if (bWasCanceled)
 	{
 		return;
 	}
@@ -221,7 +222,6 @@ void FHttpRequestCommon::CancelRequest()
 
 	StopPassingReceivedData();
 
-	bCanceled = true;
 	UE_LOG(LogHttp, Verbose, TEXT("HTTP request canceled. URL=%s"), *GetURL());
 
 	FHttpModule::Get().GetHttpManager().AddHttpThreadTask([StrongThis = StaticCastSharedRef<FHttpRequestCommon>(AsShared())]()
diff --git a/Engine/Source/Runtime/Online/HTTP/Public/Interfaces/IHttpBase.h b/Engine/Source/Runtime/Online/HTTP/Public/Interfaces/IHttpBase.h
--- a/Engine/Source/Runtime/Online/HTTP/Public/Interfaces/IHttpBase.h
+++ b/Engine/Source/Runtime/Online/HTTP/Public/Interfaces/IHttpBase.h
@@ -77,6 +77,20 @@ enum class EHttpFailureReason : uint8
 	Other
 };
 
+/** @return the stringified version of the enum passed in */
+inline const TCHAR* LexToString(EHttpFailureReason HttpFailureReason)
+{
+	switch (HttpFailureReason)
+	{
+	case EHttpFailureReason::None: return TEXT("None");
+	case EHttpFailureReason::ConnectionError: return TEXT("ConnectionError");
+	case EHttpFailureReason::Cancelled: return TEXT("Cancelled");
+	case EHttpFailureReason::TimedOut: return TEXT("TimedOut");
+	case EHttpFailureReason::Other: return TEXT("Other");
+	default: checkNoEntry(); return TEXT("Invalid");
+	}
+}
+
 /**
  * Base interface for Http Requests and Responses.
  */
