Fix the random failure in test case "Retry manager and http manager is thread safe for flushing".

Reason:
When FHttpManager::Tick get called from different threads, it might get same request in CompletedThreadedRequests through the http thread, because Thread->GetCompletedRequests is not thread safe. In the end it will crash due to access destroyed CompletedRequest.

diff --git a/Engine/Source/Runtime/Online/HTTP/Private/HttpManager.cpp b/Engine/Source/Runtime/Online/HTTP/Private/HttpManager.cpp
--- a/Engine/Source/Runtime/Online/HTTP/Private/HttpManager.cpp
+++ b/Engine/Source/Runtime/Online/HTTP/Private/HttpManager.cpp
@@ -20,6 +20,7 @@
 // FHttpManager
 
 FCriticalSection FHttpManager::RequestLock;
+FCriticalSection FHttpManager::CompletedRequestLock;
 
 const TCHAR* LexToString(const EHttpFlushReason& FlushReason)
 {
@@ -446,7 +447,12 @@ bool FHttpManager::Tick(float DeltaSeconds)
 		}
 
 		TArray<IHttpThreadedRequest*> CompletedThreadedRequests;
-		Thread->GetCompletedRequests(CompletedThreadedRequests);
+
+		{
+			// Thread->GetCompletedRequests doesn't support multi-thread access
+			FScopeLock ScopeLock(&CompletedRequestLock);
+			Thread->GetCompletedRequests(CompletedThreadedRequests);
+		}
 
 		// Finish and remove any completed requests
 		for (IHttpThreadedRequest* CompletedRequest : CompletedThreadedRequests)
diff --git a/Engine/Source/Runtime/Online/HTTP/Public/HttpManager.h b/Engine/Source/Runtime/Online/HTTP/Public/HttpManager.h
--- a/Engine/Source/Runtime/Online/HTTP/Public/HttpManager.h
+++ b/Engine/Source/Runtime/Online/HTTP/Public/HttpManager.h
@@ -342,6 +342,10 @@ PACKAGE_SCOPE:
 
 	/** Used to lock access to add/remove/find requests */
 	static FCriticalSection RequestLock;
+
+	/** Used to lock access to get completed requests */
+	static FCriticalSection CompletedRequestLock;
+
 	/**
 	 * Broadcast that a non-threaded HTTP request is complete.
 	 * Called automatically internally for threaded requests.
