Add more comments to http complete callback;
Format the indents in IHttpRequest.h.

diff --git a/Engine/Source/Runtime/Online/HTTP/Public/Interfaces/IHttpRequest.h b/Engine/Source/Runtime/Online/HTTP/Public/Interfaces/IHttpRequest.h
--- a/Engine/Source/Runtime/Online/HTTP/Public/Interfaces/IHttpRequest.h
+++ b/Engine/Source/Runtime/Online/HTTP/Public/Interfaces/IHttpRequest.h
@@ -30,10 +30,16 @@ typedef TSharedRef<IHttpResponse, ESPMode::ThreadSafe> FHttpResponseRef;
  * Delegate called when an Http request completes
  *
  * @param Request original Http request that started things
- * @param Response response received from the server if a successful connection was established
- * @param bConnectedSuccessfully - indicates whether or not the request was able to connect successfully
+ * @param Response response received from the server if a successful connection was established. Note that
+ *        even if response is not null, it doesn't mean the request was processed successfully. Use
+ *        Response->GetStatus(), Response->GetFailureReason() and Response->GetResponseCode() to decide
+ *        how to handle it
+ * @param bProcessedSuccessfully this flag indicates whether or not the request was able to be processed successfully.
+ *        It could be connect error, timeout error or platform http error. Note that when this flag is true
+ *        doesn't mean the http request get 2xx success status code, it just means this request got response
+ *        from http server and completed the processing
  */
-using FHttpRequestCompleteDelegate = TTSDelegate<void(FHttpRequestPtr /*Request*/, FHttpResponsePtr /*Response*/, bool /*bConnectedSuccessfully*/)>;
+using FHttpRequestCompleteDelegate = TTSDelegate<void(FHttpRequestPtr /*Request*/, FHttpResponsePtr /*Response*/, bool /*bProcessedSuccessfully*/)>;
 
 /**
  * Delegate called when an Http request receives status code
@@ -114,14 +120,14 @@ private:
 /**
  * Interface for Http requests (created using FHttpFactory)
  */
-class IHttpRequest : 
+class IHttpRequest :
 	public IHttpBase, public TSharedFromThis<IHttpRequest, ESPMode::ThreadSafe>
 {
 public:
 
 	/**
 	 * Gets the verb (GET, PUT, POST) used by the request.
-	 * 
+	 *
 	 * @return the verb string
 	 */
 	virtual FString GetVerb() const = 0;
@@ -137,7 +143,7 @@ public:
 	virtual void SetVerb(const FString& Verb) = 0;
 
 	/**
-	 * Sets the URL for the request 
+	 * Sets the URL for the request
 	 * Eg. (http://my.domain.com/something.ext?key=value&key2=value).
 	 * Must be set before calling ProcessRequest.
 	 *
@@ -202,12 +208,12 @@ public:
 	bool SetContentFromStreamDelegate(FHttpRequestStreamDelegate StreamDelegate) { return SetContentFromStream(MakeShared<FArchiveWithDelegate>(StreamDelegate)); }
 
 	/**
-	 * Sets the stream to receive the response body. Make sure to handle the cleanup of stream when 
-	 * Serialize generated error(Stream->GetError returns true after Stream->Serialize call), this 
+	 * Sets the stream to receive the response body. Make sure to handle the cleanup of stream when
+	 * Serialize generated error(Stream->GetError returns true after Stream->Serialize call), this
 	 * http request will fail and quit.
 	 *
-	 * NOTE: Once set, the data will no longer be cached in response, IHttpResponse::GetContent() and 
-	 * IHttpResponse::GetContentAsString() will return empty result. The Stream->Serialize will be called 
+	 * NOTE: Once set, the data will no longer be cached in response, IHttpResponse::GetContent() and
+	 * IHttpResponse::GetContentAsString() will return empty result. The Stream->Serialize will be called
 	 * from another thread other than the game thread
 	 *
 	 * @param Stream - will be used to receive the response body
@@ -216,10 +222,10 @@ public:
 	virtual bool SetResponseBodyReceiveStream(TSharedRef<FArchive> Stream) = 0;
 
 	/**
-	 * Sets the delegate to receive the response body. Make sure to handle the cleanup of received data when 
+	 * Sets the delegate to receive the response body. Make sure to handle the cleanup of received data when
 	 * failed to process the data(StreamDelegate return false), this http request will fail and quit.
 	 *
-	 * NOTE: Once set, the data will no longer be cached in response, IHttpResponse::GetContent() and 
+	 * NOTE: Once set, the data will no longer be cached in response, IHttpResponse::GetContent() and
 	 * IHttpResponse::GetContentAsString() will return empty result. The delegate will be called from
 	 * another thread other than the game thread
 	 *
@@ -242,7 +248,7 @@ public:
 	virtual void SetHeader(const FString& HeaderName, const FString& HeaderValue) = 0;
 
 	/**
-	* Appends to the value already set in the header. 
+	* Appends to the value already set in the header.
 	* If there is already content in that header, a comma delimiter is used.
 	* If the header is as of yet unset, the result is the same as calling SetHeader
 	* Content-Length is the only header set for you.
@@ -300,18 +306,18 @@ public:
 	 * Delegate called to update the request/response progress. See FHttpRequestProgressDelegate64
 	 */
 	virtual FHttpRequestProgressDelegate64& OnRequestProgress64() = 0;
-	
+
 	/**
 	* Delegate called when the request will be retried
 	*/
 	virtual FHttpRequestWillRetryDelegate& OnRequestWillRetry() = 0;
 
-	/** 
+	/**
 	 * Delegate called to signal the receipt of a header.  See FHttpRequestHeaderReceivedDelegate
 	 */
 	virtual FHttpRequestHeaderReceivedDelegate& OnHeaderReceived() = 0;
 
-	/** 
+	/**
 	 * Delegate called to signal the receipt of a header.  See FHttpRequestStatusCodeReceivedDelegate
 	 */
 	virtual FHttpRequestStatusCodeReceivedDelegate& OnStatusCodeReceived() = 0;
@@ -337,26 +343,26 @@ public:
 
 	/**
 	 * Gets the time that it took for the server to fully respond to the request.
-	 * 
+	 *
 	 * @return elapsed time in seconds.
 	 */
 	virtual float GetElapsedTime() const = 0;
 
 	/**
-	 * Set thread policy about which thread to trigger the delegates, set by FHttpManager::SetRequestCompletedDelegate, 
+	 * Set thread policy about which thread to trigger the delegates, set by FHttpManager::SetRequestCompletedDelegate,
 	 * IHttpRequest::OnStatusCodeReceived, IHttpRequest::OnHeaderReceived, IHttpRequest::OnRequestProgress64 and IHttpRequest::OnProcessRequestComplete.
 	 *
-	 * Note that when set it as CompleteOnHttpThread, the thread to trigger delegates could be any thread 
-	 * depends on the implementation. User code should make the delegate thread-safe and shouldn't assume 
+	 * Note that when set it as CompleteOnHttpThread, the thread to trigger delegates could be any thread
+	 * depends on the implementation. User code should make the delegate thread-safe and shouldn't assume
 	 * it's triggered by the thread where this request get created.
-	 * 
+	 *
 	 * @param InThreadPolicy - The thread policy to indicate which thread to trigger the delegates
 	 */
 	virtual void SetDelegateThreadPolicy(EHttpRequestDelegateThreadPolicy InThreadPolicy) = 0;
 
 	/**
 	 * Get thread policy about which thread to complete this request
-	 * 
+	 *
 	 * @return The thread policy
 	 */
 	virtual EHttpRequestDelegateThreadPolicy GetDelegateThreadPolicy() const = 0;
@@ -364,7 +370,7 @@ public:
 	/**
 	 * Blocking call to wait the request until it's completed
 	 *
-	 * WARNINGS: 
+	 * WARNINGS:
 	 *
 	 * - This is a blocking call, DON'T use this in a time-sensitive context
 	 * - Complete delegate will be used in this function so customized complete delegate is not supported
@@ -374,8 +380,8 @@ public:
 	 */
 	virtual void ProcessRequestUntilComplete() = 0;
 
-	/** 
-	 * Destructor for overrides 
+	/**
+	 * Destructor for overrides
 	 */
 	virtual ~IHttpRequest() = default;
 };
