Enable event loop by default on most of the platforms.

#tests Passed WebTests on Windows

diff --git a/Engine/Source/Runtime/Online/HTTP/HTTP.Build.cs b/Engine/Source/Runtime/Online/HTTP/HTTP.Build.cs
--- a/Engine/Source/Runtime/Online/HTTP/HTTP.Build.cs
+++ b/Engine/Source/Runtime/Online/HTTP/HTTP.Build.cs
@@ -4,6 +4,8 @@ using UnrealBuildTool;
 
 public class HTTP : ModuleRules
 {
+	protected virtual bool bPlatformEventLoopEnabledByDefault { get { return true; } }
+
 	protected virtual bool bPlatformSupportsWinHttp
 	{
 		get
@@ -93,6 +95,7 @@ public class HTTP : ModuleRules
 			}
 		}
 
+		PrivateDefinitions.Add("UE_HTTP_EVENT_LOOP_ENABLE_CHANCE_BY_DEFAULT=" + (bPlatformEventLoopEnabledByDefault ? "100" : "0"));
 		PrivateDefinitions.Add("WITH_CURL_LIBCURL =" + (bPlatformSupportsLibCurl ? "1" : "0"));
 		PrivateDefinitions.Add("WITH_CURL_XCURL=" + (bPlatformSupportsXCurl ? "1" : "0"));
 		PrivateDefinitions.Add("WITH_CURL_MULTIPOLL=" + (bPlatformSupportsCurlMultiPoll ? "1" : "0"));
diff --git a/Engine/Source/Runtime/Online/HTTP/Private/Curl/CurlHttpManager.cpp b/Engine/Source/Runtime/Online/HTTP/Private/Curl/CurlHttpManager.cpp
--- a/Engine/Source/Runtime/Online/HTTP/Private/Curl/CurlHttpManager.cpp
+++ b/Engine/Source/Runtime/Online/HTTP/Private/Curl/CurlHttpManager.cpp
@@ -37,12 +37,7 @@
 #define DISABLE_UNVERIFIED_CERTIFICATE_LOADING 0
 #endif
 
-TAutoConsoleVariable<int32> CVarCurlEventLoopEnableChance(
-	TEXT("http.CurlEventLoopEnableChance"),
-	0,
-	TEXT("Enable chance of curl event loop, from 0 to 100"),
-	ECVF_SaveForNextBoot
-);
+extern TAutoConsoleVariable<int32> CVarHttpEventLoopEnableChance;
 
 CURLM* FCurlHttpManager::GMultiHandle = nullptr;
 #if !WITH_CURL_XCURL
@@ -452,10 +447,10 @@ void FCurlHttpManager::UpdateConfigs()
 
 FHttpThreadBase* FCurlHttpManager::CreateHttpThread()
 {
-	bool bUseEventLoop = (FMath::RandRange(0, 99) < CVarCurlEventLoopEnableChance.GetValueOnGameThread());
+	bool bUseEventLoop = (FMath::RandRange(0, 99) < CVarHttpEventLoopEnableChance.GetValueOnGameThread());
 
 	// Also support to change it through runtime args.
-	// Can't set cvar CVarCurlEventLoopEnableChance through runtime args or .ini files because http module initialized too early
+	// Can't set cvar CVarHttpEventLoopEnableChance through runtime args or .ini files because http module initialized too early
 	FParse::Bool(FCommandLine::Get(), TEXT("useeventloop="), bUseEventLoop);
 
 	if (bUseEventLoop)
diff --git a/Engine/Source/Runtime/Online/HTTP/Private/HttpModule.cpp b/Engine/Source/Runtime/Online/HTTP/Private/HttpModule.cpp
--- a/Engine/Source/Runtime/Online/HTTP/Private/HttpModule.cpp
+++ b/Engine/Source/Runtime/Online/HTTP/Private/HttpModule.cpp
@@ -22,6 +22,13 @@ TAutoConsoleVariable<bool> CVarHttpInsecureProtocolEnabled(
 );
 #endif
 
+TAutoConsoleVariable<int32> CVarHttpEventLoopEnableChance(
+	TEXT("http.EventLoopEnableChance"),
+	UE_HTTP_EVENT_LOOP_ENABLE_CHANCE_BY_DEFAULT,
+	TEXT("Enable chance of event loop, from 0 to 100"),
+	ECVF_SaveForNextBoot
+);
+
 DEFINE_LOG_CATEGORY(LogHttp);
 
 // FHttpModule
