Allow '.' in FHttpPath validation in HttpServer's implementation.

diff --git a/Engine/Source/Runtime/Online/HTTPServer/Private/HttpPath.cpp b/Engine/Source/Runtime/Online/HTTPServer/Private/HttpPath.cpp
--- a/Engine/Source/Runtime/Online/HTTPServer/Private/HttpPath.cpp
+++ b/Engine/Source/Runtime/Online/HTTPServer/Private/HttpPath.cpp
@@ -42,11 +42,10 @@ bool FHttpPath::IsValidPath() const
 	{ 
 		return 
 			(C <= 32) || (C >= 127) ||
-			(C == TEXT(' ')) ||	(C == TEXT('.')) ||	(C == TEXT(',')) ||
-			(C == TEXT('<')) || (C == TEXT('>')) ||	(C == TEXT(']')) ||
-			(C == TEXT('[')) || (C == TEXT('}')) ||	(C == TEXT('{')) ||
-			(C == TEXT('#')) || (C == TEXT('|')) ||	(C == TEXT('^')) ||
-			(C == TEXT('\\'));
+			(C == TEXT(' ')) || (C == TEXT(',')) || (C == TEXT('\\')) ||
+			(C == TEXT('<')) || (C == TEXT('>')) || (C == TEXT(']')) ||
+			(C == TEXT('[')) || (C == TEXT('}')) || (C == TEXT('{')) ||
+			(C == TEXT('#')) || (C == TEXT('|')) || (C == TEXT('^'));
 	}; 
 
 	return (INDEX_NONE == Path.FindLastCharByPredicate(IsInvalidUriChar));
