Use the same cvar CVarHttpEventLoopEnableChance to enable event loop by default on apple http impl;
Moved the code to HttpManager to get rid of the duplicated code.

diff --git a/Engine/Source/Runtime/Online/HTTP/Private/Apple/AppleHttpManager.cpp b/Engine/Source/Runtime/Online/HTTP/Private/Apple/AppleHttpManager.cpp
--- a/Engine/Source/Runtime/Online/HTTP/Private/Apple/AppleHttpManager.cpp
+++ b/Engine/Source/Runtime/Online/HTTP/Private/Apple/AppleHttpManager.cpp
@@ -9,21 +9,8 @@
 #include "HttpModule.h"
 #include "Http.h"
 
-TAutoConsoleVariable<int32> CVarAppleEventLoopEnableChance(
-	TEXT("http.AppleEventLoopEnableChance"),
-	0,
-	TEXT("Enable chance of Apple event loop, from 0 to 100"),
-	ECVF_SaveForNextBoot
-);
-
 FHttpThreadBase* FAppleHttpManager::CreateHttpThread()
 {
-	bool bUseEventLoop = !(FMath::RandRange(0, 99) < CVarAppleEventLoopEnableChance.GetValueOnGameThread());
-
-	// Also support to change it through runtime args.
-	// Can't set cvar CVarCurlEventLoopEnableChance through runtime args or .ini files because http module initialized too early
-	FParse::Bool(FCommandLine::Get(), TEXT("useeventloop="), bUseEventLoop);
-
 	if (bUseEventLoop)
 	{
 		UE_LOG(LogHttp, Log, TEXT("CreateHttpThread using FAppleEventLoopHttpThread"));
diff --git a/Engine/Source/Runtime/Online/HTTP/Private/Curl/CurlHttpManager.cpp b/Engine/Source/Runtime/Online/HTTP/Private/Curl/CurlHttpManager.cpp
--- a/Engine/Source/Runtime/Online/HTTP/Private/Curl/CurlHttpManager.cpp
+++ b/Engine/Source/Runtime/Online/HTTP/Private/Curl/CurlHttpManager.cpp
@@ -37,8 +37,6 @@
 #define DISABLE_UNVERIFIED_CERTIFICATE_LOADING 0
 #endif
 
-extern TAutoConsoleVariable<int32> CVarHttpEventLoopEnableChance;
-
 CURLM* FCurlHttpManager::GMultiHandle = nullptr;
 #if !WITH_CURL_XCURL
 CURLSH* FCurlHttpManager::GShareHandle = nullptr;
@@ -447,12 +445,6 @@ void FCurlHttpManager::UpdateConfigs()
 
 FHttpThreadBase* FCurlHttpManager::CreateHttpThread()
 {
-	bool bUseEventLoop = (FMath::RandRange(0, 99) < CVarHttpEventLoopEnableChance.GetValueOnGameThread());
-
-	// Also support to change it through runtime args.
-	// Can't set cvar CVarHttpEventLoopEnableChance through runtime args or .ini files because http module initialized too early
-	FParse::Bool(FCommandLine::Get(), TEXT("useeventloop="), bUseEventLoop);
-
 	if (bUseEventLoop)
 	{
 #if WITH_CURL_MULTIPOLL
diff --git a/Engine/Source/Runtime/Online/HTTP/Private/HttpManager.cpp b/Engine/Source/Runtime/Online/HTTP/Private/HttpManager.cpp
--- a/Engine/Source/Runtime/Online/HTTP/Private/HttpManager.cpp
+++ b/Engine/Source/Runtime/Online/HTTP/Private/HttpManager.cpp
@@ -2,6 +2,7 @@
 
 #include "HttpManager.h"
 #include "HttpModule.h"
+#include "HAL/IConsoleManager.h"
 #include "HAL/PlatformTime.h"
 #include "HAL/PlatformProcess.h"
 #include "Misc/ScopeLock.h"
@@ -17,6 +18,13 @@
 #include "Stats/Stats.h"
 #include "Containers/BackgroundableTicker.h"
 
+TAutoConsoleVariable<int32> CVarHttpEventLoopEnableChance(
+	TEXT("http.EventLoopEnableChance"),
+	UE_HTTP_EVENT_LOOP_ENABLE_CHANCE_BY_DEFAULT,
+	TEXT("Enable chance of event loop, from 0 to 100"),
+	ECVF_SaveForNextBoot
+);
+
 // FHttpManager
 
 FCriticalSection FHttpManager::RequestLock;
@@ -69,6 +77,12 @@ void FHttpManager::Initialize()
 {
 	if (!Thread)
 	{
+		bUseEventLoop = (FMath::RandRange(0, 99) < CVarHttpEventLoopEnableChance.GetValueOnGameThread());
+
+		// Also support to change it through runtime args.
+		// Can't set cvar CVarHttpEventLoopEnableChance through runtime args or .ini files because http module initialized too early
+		FParse::Bool(FCommandLine::Get(), TEXT("useeventloop="), bUseEventLoop);
+
 		Thread = CreateHttpThread();
 		Thread->StartThread();
 	}
diff --git a/Engine/Source/Runtime/Online/HTTP/Private/HttpModule.cpp b/Engine/Source/Runtime/Online/HTTP/Private/HttpModule.cpp
--- a/Engine/Source/Runtime/Online/HTTP/Private/HttpModule.cpp
+++ b/Engine/Source/Runtime/Online/HTTP/Private/HttpModule.cpp
@@ -22,13 +22,6 @@ TAutoConsoleVariable<bool> CVarHttpInsecureProtocolEnabled(
 );
 #endif
 
-TAutoConsoleVariable<int32> CVarHttpEventLoopEnableChance(
-	TEXT("http.EventLoopEnableChance"),
-	UE_HTTP_EVENT_LOOP_ENABLE_CHANCE_BY_DEFAULT,
-	TEXT("Enable chance of event loop, from 0 to 100"),
-	ECVF_SaveForNextBoot
-);
-
 DEFINE_LOG_CATEGORY(LogHttp);
 
 // FHttpModule
diff --git a/Engine/Source/Runtime/Online/HTTP/Public/HttpManager.h b/Engine/Source/Runtime/Online/HTTP/Public/HttpManager.h
--- a/Engine/Source/Runtime/Online/HTTP/Public/HttpManager.h
+++ b/Engine/Source/Runtime/Online/HTTP/Public/HttpManager.h
@@ -338,6 +338,8 @@ protected:
 
 	FHttpStats HttpStats;
 
+	bool bUseEventLoop = true;
+
 PACKAGE_SCOPE:
 
 	/** Used to lock access to add/remove/find requests */
