Log warning instead of check when http activity timer has been started.

diff --git a/Engine/Source/Runtime/Online/HTTP/Private/GenericPlatform/HttpRequestCommon.cpp b/Engine/Source/Runtime/Online/HTTP/Private/GenericPlatform/HttpRequestCommon.cpp
--- a/Engine/Source/Runtime/Online/HTTP/Private/GenericPlatform/HttpRequestCommon.cpp
+++ b/Engine/Source/Runtime/Online/HTTP/Private/GenericPlatform/HttpRequestCommon.cpp
@@ -267,7 +267,11 @@ void FHttpRequestCommon::StartActivityTimeoutTimer()
 
 void FHttpRequestCommon::StartActivityTimeoutTimerBy(double DelayToTrigger)
 {
-	check(ActivityTimeoutHttpTaskTimerHandle == nullptr);
+	if (ActivityTimeoutHttpTaskTimerHandle != nullptr)
+	{
+		UE_LOG(LogHttp, Warning, TEXT("Request %p already started activity timeout timer"), this);
+		return;
+	}
 
 	TWeakPtr<IHttpRequest> RequestWeakPtr(AsShared());
 	ActivityTimeoutHttpTaskTimerHandle = FHttpModule::Get().GetHttpManager().AddHttpThreadTask([RequestWeakPtr]() {
