Keep the old config name for event loop enable chance, so existing hotfix for it can still work.

#tests Tried in local WindowsClient

diff --git a/Engine/Source/Runtime/Online/HTTP/Private/HttpManager.cpp b/Engine/Source/Runtime/Online/HTTP/Private/HttpManager.cpp
--- a/Engine/Source/Runtime/Online/HTTP/Private/HttpManager.cpp
+++ b/Engine/Source/Runtime/Online/HTTP/Private/HttpManager.cpp
@@ -19,7 +19,7 @@
 #include "Containers/BackgroundableTicker.h"
 
 TAutoConsoleVariable<int32> CVarHttpEventLoopEnableChance(
-	TEXT("http.EventLoopEnableChance"),
+	TEXT("http.CurlEventLoopEnableChance"),
 	UE_HTTP_EVENT_LOOP_ENABLE_CHANCE_BY_DEFAULT,
 	TEXT("Enable chance of event loop, from 0 to 100"),
 	ECVF_SaveForNextBoot
