Make sure to re-open the file when retry http request with streaming upload request.

diff --git a/Engine/Source/Programs/WebTests/Private/TestHttp.cpp b/Engine/Source/Programs/WebTests/Private/TestHttp.cpp
--- a/Engine/Source/Programs/WebTests/Private/TestHttp.cpp
+++ b/Engine/Source/Programs/WebTests/Private/TestHttp.cpp
@@ -1116,6 +1116,43 @@ TEST_CASE_METHOD(FWaitUntilCompleteHttpFixture, "Streaming http upload from file
 	HttpRequest->ProcessRequest();
 }
 
+TEST_CASE_METHOD(FWaitUntilCompleteHttpFixture, "Streaming uploading http request will re-open file when retry", HTTP_TAG)
+{
+	if (!bRetryEnabled)
+	{
+		return;
+	}
+	DisableWarningsInThisTest();
+
+	FString Filename = FString(FPlatformProcess::UserSettingsDir()) / TEXT("TestStreamUploadRetry.dat");
+	UE::TestHttp::WriteTestFile(Filename, 1*1024*1024/*1MB*/);
+
+	TSharedRef<IHttpRequest> HttpRequest = HttpRetryManager->CreateRequest(
+		1/*InRetryLimitCountOverride*/,
+		FHttpRetrySystem::FRetryTimeoutRelativeSecondsSetting()/*InRetryTimeoutRelativeSecondsOverride unused*/,
+		{EHttpResponseCodes::TooManyRequests}/*InRetryResponseCodes*/
+	);
+
+	HttpRequest->SetURL(UrlMockStatus(EHttpResponseCodes::TooManyRequests));
+	HttpRequest->SetHeader(TEXT("Retry-After"), TEXT("1")); // Will be forwarded back in response
+	HttpRequest->SetVerb(TEXT("PUT"));
+	HttpRequest->SetHeader(TEXT("Content-Disposition"), TEXT("attachment;filename=TestStreamUploadRetry.dat"));
+	HttpRequest->SetContentAsStreamedFile(Filename);
+
+	++ExpectingExtraCallbacks;
+	HttpRequest->OnRequestWillRetry().BindLambda([this](FHttpRequestPtr Request, FHttpResponsePtr /*Response*/, float LockoutPeriod) {
+		--ExpectingExtraCallbacks;
+		Request->SetURL(UrlStreamUpload());
+	});
+
+	HttpRequest->OnProcessRequestComplete().BindLambda([Filename](FHttpRequestPtr HttpRequest, FHttpResponsePtr HttpResponse, bool bSucceeded) {
+		CHECK(bSucceeded);
+		CHECK(HttpResponse->GetResponseCode() == 200);
+		IFileManager::Get().Delete(*Filename);
+	});
+	HttpRequest->ProcessRequest();
+}
+
 TEST_CASE_METHOD(FWaitUntilCompleteHttpFixture, "Redirect enabled by default and can work well", HTTP_TAG)
 {
 	TSharedRef<IHttpRequest> HttpRequest = CreateRequest();
diff --git a/Engine/Source/Runtime/Online/HTTP/Private/Curl/CurlHttp.cpp b/Engine/Source/Runtime/Online/HTTP/Private/Curl/CurlHttp.cpp
--- a/Engine/Source/Runtime/Online/HTTP/Private/Curl/CurlHttp.cpp
+++ b/Engine/Source/Runtime/Online/HTTP/Private/Curl/CurlHttp.cpp
@@ -11,7 +11,6 @@
 #include "Misc/Paths.h"
 #include "Curl/CurlHttpManager.h"
 #include "Misc/ScopeLock.h"
-#include "HAL/FileManager.h"
 #include "HAL/IConsoleManager.h"
 #include "Internationalization/Regex.h"
 
@@ -332,18 +331,9 @@ bool FCurlHttpRequest::SetContentAsStreamedFile(const FString& Filename)
 		return false;
 	}
 
-	FArchive* File = IFileManager::Get().CreateFileReader(*Filename);
-	if (File)
-	{
-		RequestPayload = MakeUnique<FRequestPayloadInFileStream>(MakeShareable(File), true/*bInCloseWhenComplete*/);
-	}
-	else
-	{
-		UE_LOG(LogHttp, Warning, TEXT("FCurlHttpRequest::SetContentAsStreamedFile Failed to open %s for reading"), *Filename);
-		RequestPayload.Reset();
-	}
+	RequestPayload = MakeUnique<FRequestPayloadInFileStream>(*Filename);
 	bIsRequestPayloadSeekable = false;
-	return RequestPayload.IsValid();
+	return true;
 }
 
 bool FCurlHttpRequest::SetContentFromStream(TSharedRef<FArchive, ESPMode::ThreadSafe> Stream)
@@ -818,13 +808,6 @@ bool FCurlHttpRequest::SetupRequest()
 	QUICK_SCOPE_CYCLE_COUNTER(STAT_FCurlHttpRequest_SetupRequest);
 	check(EasyHandle);
 
-	if ((GetVerb().IsEmpty() || GetVerb().Equals(TEXT("GET"), ESearchCase::IgnoreCase))
-		&& (RequestPayload.IsValid() && RequestPayload->GetContentLength() > 0))
-	{
-		UE_LOG(LogHttp, Warning, TEXT("An HTTP Get request cannot contain a payload."));
-		return false;
-	}
-
 	// set up request
 
 	if (!RequestPayload.IsValid())
@@ -833,6 +816,18 @@ bool FCurlHttpRequest::SetupRequest()
 		bIsRequestPayloadSeekable = true;
 	}
 
+	if (!RequestPayload->Open())
+	{
+		UE_LOG(LogHttp, Warning, TEXT("Failed to open request payload."));
+		return false;
+	}
+
+	if ((GetVerb().IsEmpty() || GetVerb().Equals(TEXT("GET"), ESearchCase::IgnoreCase)) && RequestPayload->GetContentLength() > 0)
+	{
+		UE_LOG(LogHttp, Warning, TEXT("An HTTP Get request cannot contain a payload."));
+		return false;
+	}
+
 	bCurlRequestCompleted = false;
 	CurlAddToMultiResult = CURLM_OK;
 	LastReportedBytesSent = 0;
diff --git a/Engine/Source/Runtime/Online/HTTP/Private/GenericPlatform/HttpRequestPayload.cpp b/Engine/Source/Runtime/Online/HTTP/Private/GenericPlatform/HttpRequestPayload.cpp
--- a/Engine/Source/Runtime/Online/HTTP/Private/GenericPlatform/HttpRequestPayload.cpp
+++ b/Engine/Source/Runtime/Online/HTTP/Private/GenericPlatform/HttpRequestPayload.cpp
@@ -4,6 +4,8 @@
 #include "GenericPlatform/GenericPlatformFile.h"
 #include "GenericPlatform/GenericPlatformHttp.h"
 #include "HAL/PlatformFileManager.h"
+#include "HAL/FileManager.h"
+#include "Http.h"
 
 bool FGenericPlatformHttp::IsURLEncoded(const TArray<uint8>& Payload)
 {
@@ -33,7 +35,12 @@ bool FGenericPlatformHttp::IsURLEncoded(const TArray<uint8>& Payload)
 	return true;
 }
 
-FRequestPayloadInFileStream::FRequestPayloadInFileStream(TSharedRef<FArchive, ESPMode::ThreadSafe> InFile, bool bInCloseWhenComplete) 
+FRequestPayloadInFileStream::FRequestPayloadInFileStream(const FString& InFilename)
+	: Filename(InFilename)
+{
+}
+
+FRequestPayloadInFileStream::FRequestPayloadInFileStream(TSharedRef<FArchive> InFile, bool bInCloseWhenComplete) 
 	: File(InFile)
 	, bCloseWhenComplete(bInCloseWhenComplete)
 {
@@ -45,7 +52,7 @@ FRequestPayloadInFileStream::~FRequestPayloadInFileStream()
 
 uint64 FRequestPayloadInFileStream::GetContentLength() const
 {
-	return File->TotalSize();
+	return File ? File->TotalSize() : 0;
 }
 
 const TArray<uint8>& FRequestPayloadInFileStream::GetContent() const
@@ -69,6 +76,11 @@ size_t FRequestPayloadInFileStream::FillOutputBuffer(void* OutputBuffer, size_t
 
 size_t FRequestPayloadInFileStream::FillOutputBuffer(TArrayView<uint8> OutputBuffer, size_t SizeAlreadySent)
 {
+	if (!File)
+	{
+		return 0;
+	}
+
 	const size_t ContentLength = GetContentLength();
 	check(SizeAlreadySent <= ContentLength);
 	const size_t SizeToSend = ContentLength - SizeAlreadySent;
@@ -84,11 +96,30 @@ size_t FRequestPayloadInFileStream::FillOutputBuffer(TArrayView<uint8> OutputBuf
 	return SizeToSendThisTime;
 }
 
+bool FRequestPayloadInFileStream::Open()
+{
+	if (!File && !Filename.IsEmpty())
+	{
+		FArchive* RawFile = IFileManager::Get().CreateFileReader(*Filename);
+		if (!RawFile)
+		{
+			UE_LOG(LogHttp, Warning, TEXT("FRequestPayloadInFileStream::Open Failed to open %s for reading"), *Filename);
+			return false;
+		}
+
+		File = MakeShareable(RawFile);
+		bCloseWhenComplete = true;
+	}
+
+	return File != nullptr;
+}
+
 void FRequestPayloadInFileStream::Close()
 {
-	if (bCloseWhenComplete)
+	if (bCloseWhenComplete && File)
 	{
 		File->Close();
+		File.Reset();
 	}
 }
 
@@ -137,6 +168,11 @@ size_t FRequestPayloadInMemory::FillOutputBuffer(TArrayView<uint8> OutputBuffer,
 	return SizeToSendThisTime;
 }
 
+bool FRequestPayloadInMemory::Open()
+{
+	return true;
+}
+
 void FRequestPayloadInMemory::Close()
 {
 }
diff --git a/Engine/Source/Runtime/Online/HTTP/Public/GenericPlatform/HttpRequestPayload.h b/Engine/Source/Runtime/Online/HTTP/Public/GenericPlatform/HttpRequestPayload.h
--- a/Engine/Source/Runtime/Online/HTTP/Public/GenericPlatform/HttpRequestPayload.h
+++ b/Engine/Source/Runtime/Online/HTTP/Public/GenericPlatform/HttpRequestPayload.h
@@ -41,6 +41,9 @@ public:
 	 */
 	virtual size_t FillOutputBuffer(TArrayView<uint8> OutputBuffer, size_t SizeAlreadySent) = 0;
 
+	/** Open the upload payload before start */
+	virtual bool Open() = 0;
+
 	/** Close the upload payload when the http request don't need to use it anymore */
 	virtual void Close() = 0;
 };
@@ -48,16 +51,20 @@ public:
 class FRequestPayloadInFileStream : public FRequestPayload
 {
 public:
-	FRequestPayloadInFileStream(TSharedRef<FArchive, ESPMode::ThreadSafe> InFile, bool bInCloseWhenComplete = false);
+	FRequestPayloadInFileStream(const FString& InFilename);
+	FRequestPayloadInFileStream(TSharedRef<FArchive> InFile, bool bInCloseWhenComplete = false);
 	virtual ~FRequestPayloadInFileStream();
 	virtual uint64 GetContentLength() const override;
 	virtual const TArray<uint8>& GetContent() const override;
 	virtual bool IsURLEncoded() const override;
 	virtual size_t FillOutputBuffer(void* OutputBuffer, size_t MaxOutputBufferSize, size_t SizeAlreadySent) override;
 	virtual size_t FillOutputBuffer(TArrayView<uint8> OutputBuffer, size_t SizeAlreadySent) override;
+	virtual bool Open() override;
 	virtual void Close() override;
+
 private:
-	TSharedRef<FArchive, ESPMode::ThreadSafe> File;
+	FString Filename;
+	TSharedPtr<FArchive> File;
 	bool bCloseWhenComplete = false;
 };
 
@@ -72,6 +79,7 @@ public:
 	virtual bool IsURLEncoded() const override;
 	virtual size_t FillOutputBuffer(void* OutputBuffer, size_t MaxOutputBufferSize, size_t SizeAlreadySent) override;
 	virtual size_t FillOutputBuffer(TArrayView<uint8> OutputBuffer, size_t SizeAlreadySent) override;
+	virtual bool Open() override;
 	virtual void Close() override;
 private:
 	TArray<uint8> Buffer;
