Disable event loop by default for pc until the random crash when shutdown http manager get fixed.

diff --git a/Engine/Source/Runtime/Online/HTTP/HTTP.Build.cs b/Engine/Source/Runtime/Online/HTTP/HTTP.Build.cs
--- a/Engine/Source/Runtime/Online/HTTP/HTTP.Build.cs
+++ b/Engine/Source/Runtime/Online/HTTP/HTTP.Build.cs
@@ -4,7 +4,8 @@ using UnrealBuildTool;
 
 public class HTTP : ModuleRules
 {
-	protected virtual bool bPlatformEventLoopEnabledByDefault { get { return true; } }
+	// Currently there is a random event loop crash when shutdown HTTP manager on PC
+	protected virtual bool bPlatformEventLoopEnabledByDefault { get { return !Target.Platform.IsInGroup(UnrealPlatformGroup.Windows); } }
 
 	protected virtual bool bPlatformSupportsWinHttp
 	{
