Integrate the missing piece code to set activity timeout per request when resolve conflicts in apple http impl.

diff --git a/Engine/Source/Runtime/Online/HTTP/Private/Apple/AppleHttp.cpp b/Engine/Source/Runtime/Online/HTTP/Private/Apple/AppleHttp.cpp
--- a/Engine/Source/Runtime/Online/HTTP/Private/Apple/AppleHttp.cpp
+++ b/Engine/Source/Runtime/Online/HTTP/Private/Apple/AppleHttp.cpp
@@ -839,7 +839,7 @@ bool FAppleHttpRequest::SetupRequest()
 	Request.timeoutInterval = HttpConnectionTimeout;
 	
 	UE_CLOG(
-		HttpConnectionTimeout < FHttpModule::Get().GetHttpActivityTimeout(), 
+		HttpConnectionTimeout < GetActivityTimeoutOrDefault(),
 		LogHttp,
 		Warning, 
 		TEXT(
@@ -847,7 +847,7 @@ bool FAppleHttpRequest::SetupRequest()
 			"unexpectedly with ConnectionError after %.2f(HttpConnectionTimeout) seconds without activity, "
 			"instead of intended %.2f(HttpActivityTimeout) seconds"
 		), 
-		HttpConnectionTimeout, FHttpModule::Get().GetHttpActivityTimeout());
+		HttpConnectionTimeout, GetActivityTimeoutOrDefault());
 
 	Task = [Session dataTaskWithRequest: Request];
 	
