Extract duplicated code for setting content stream file into HttpRequestCommon.

diff --git a/Engine/Source/Runtime/Online/HTTP/Private/Curl/CurlHttp.cpp b/Engine/Source/Runtime/Online/HTTP/Private/Curl/CurlHttp.cpp
--- a/Engine/Source/Runtime/Online/HTTP/Private/Curl/CurlHttp.cpp
+++ b/Engine/Source/Runtime/Online/HTTP/Private/Curl/CurlHttp.cpp
@@ -323,15 +323,11 @@ void FCurlHttpRequest::SetContentAsString(const FString& ContentString)
 
 bool FCurlHttpRequest::SetContentAsStreamedFile(const FString& Filename)
 {
-	UE_LOG(LogHttp, Verbose, TEXT("FCurlHttpRequest::SetContentAsStreamedFile() - %s"), *Filename);
-
-	if (CompletionStatus == EHttpRequestStatus::Processing)
+	if (!SetContentAsStreamedFileDefaultImpl(Filename))
 	{
-		UE_LOG(LogHttp, Warning, TEXT("FCurlHttpRequest::SetContentAsStreamedFile() - attempted to set content on a request that is inflight"));
 		return false;
 	}
 
-	RequestPayload = MakeUnique<FRequestPayloadInFileStream>(*Filename);
 	bIsRequestPayloadSeekable = false;
 	return true;
 }
@@ -808,23 +804,14 @@ bool FCurlHttpRequest::SetupRequest()
 	QUICK_SCOPE_CYCLE_COUNTER(STAT_FCurlHttpRequest_SetupRequest);
 	check(EasyHandle);
 
-	// set up request
-
 	if (!RequestPayload.IsValid())
 	{
 		RequestPayload = MakeUnique<FRequestPayloadInMemory>(TArray<uint8>());
 		bIsRequestPayloadSeekable = true;
 	}
 
-	if (!RequestPayload->Open())
+	if (!OpenRequestPayloadDefaultImpl())
 	{
-		UE_LOG(LogHttp, Warning, TEXT("Failed to open request payload."));
-		return false;
-	}
-
-	if ((GetVerb().IsEmpty() || GetVerb().Equals(TEXT("GET"), ESearchCase::IgnoreCase)) && RequestPayload->GetContentLength() > 0)
-	{
-		UE_LOG(LogHttp, Warning, TEXT("An HTTP Get request cannot contain a payload."));
 		return false;
 	}
 
@@ -1347,12 +1334,9 @@ float FCurlHttpRequest::GetElapsedTime() const
 
 void FCurlHttpRequest::CleanupRequest()
 {
-	if (RequestPayload.IsValid())
-	{
-		RequestPayload->Close();
-	}
-
 	curl_easy_setopt(EasyHandle, CURLOPT_SHARE, nullptr);
+
+	CloseRequestPayloadDefaultImpl();
 }
 
 // FCurlHttpRequest
diff --git a/Engine/Source/Runtime/Online/HTTP/Private/Curl/CurlHttp.h b/Engine/Source/Runtime/Online/HTTP/Private/Curl/CurlHttp.h
--- a/Engine/Source/Runtime/Online/HTTP/Private/Curl/CurlHttp.h
+++ b/Engine/Source/Runtime/Online/HTTP/Private/Curl/CurlHttp.h
@@ -373,8 +373,6 @@ private:
 	CURLMcode		CurlAddToMultiResult;
 	/** Operation result code as returned by libcurl */
 	CURLcode		CurlCompletionResult;
-	/** Payload to use with the request. Typically for POST, PUT, or PATCH */
-	TUniquePtr<FRequestPayload> RequestPayload;
 	/** Is the request payload seekable? */
 	bool bIsRequestPayloadSeekable = false;
 	/** Mapping of header section to values. */
diff --git a/Engine/Source/Runtime/Online/HTTP/Private/GenericPlatform/HttpRequestCommon.cpp b/Engine/Source/Runtime/Online/HTTP/Private/GenericPlatform/HttpRequestCommon.cpp
--- a/Engine/Source/Runtime/Online/HTTP/Private/GenericPlatform/HttpRequestCommon.cpp
+++ b/Engine/Source/Runtime/Online/HTTP/Private/GenericPlatform/HttpRequestCommon.cpp
@@ -1,6 +1,7 @@
 // Copyright Epic Games, Inc. All Rights Reserved.
 
 #include "GenericPlatform/HttpRequestCommon.h"
+#include "GenericPlatform/HttpRequestPayload.h"
 #include "GenericPlatform/HttpResponseCommon.h"
 #include "HAL/Event.h"
 #include "Http.h"
@@ -500,5 +501,47 @@ void FHttpRequestCommon::StopPassingReceivedData()
 float FHttpRequestCommon::GetActivityTimeoutOrDefault() const
 {
 	return ActivityTimeoutSecs.Get(FHttpModule::Get().GetHttpActivityTimeout());
+}
+
+bool FHttpRequestCommon::SetContentAsStreamedFileDefaultImpl(const FString& Filename)
+{
+	UE_LOG(LogHttp, Verbose, TEXT("FHttpRequestCommon::SetContentAsStreamedFileDefaultImpl() - %s"), *Filename);
+
+	if (CompletionStatus == EHttpRequestStatus::Processing)
+	{
+		UE_LOG(LogHttp, Warning, TEXT("FHttpRequestCommon::SetContentAsStreamedFileDefaultImpl() - attempted to set content on a request that is inflight"));
+		return false;
+	}
+
+	RequestPayload = MakeUnique<FRequestPayloadInFileStream>(*Filename);
+	return true;
+}
+
+bool FHttpRequestCommon::OpenRequestPayloadDefaultImpl()
+{
+	if (!RequestPayload)
+	{
+		return true;
+	}
 
-}
\ No newline at end of file
+	if (!RequestPayload->Open())
+	{
+		return false;
+	}
+
+	if ((GetVerb().IsEmpty() || GetVerb().Equals(TEXT("GET"), ESearchCase::IgnoreCase)) && RequestPayload->GetContentLength() > 0)
+	{
+		UE_LOG(LogHttp, Warning, TEXT("An HTTP Get request cannot contain a payload."));
+		return false;
+	}
+
+	return true;
+}
+
+void FHttpRequestCommon::CloseRequestPayloadDefaultImpl()
+{
+	if (RequestPayload.IsValid())
+	{
+		RequestPayload->Close();
+	}
+}
diff --git a/Engine/Source/Runtime/Online/HTTP/Public/GenericPlatform/HttpRequestCommon.h b/Engine/Source/Runtime/Online/HTTP/Public/GenericPlatform/HttpRequestCommon.h
--- a/Engine/Source/Runtime/Online/HTTP/Public/GenericPlatform/HttpRequestCommon.h
+++ b/Engine/Source/Runtime/Online/HTTP/Public/GenericPlatform/HttpRequestCommon.h
@@ -6,6 +6,7 @@
 
 class FHttpResponseCommon;
 class IHttpTaskTimerHandle;
+class FRequestPayload;
 
 /**
  * Contains implementation of some common functions that don't vary between implementations of different platforms
@@ -87,6 +88,10 @@ protected:
 
 	HTTP_API float GetActivityTimeoutOrDefault() const;
 
+	HTTP_API bool SetContentAsStreamedFileDefaultImpl(const FString& Filename);
+	HTTP_API bool OpenRequestPayloadDefaultImpl();
+	HTTP_API void CloseRequestPayloadDefaultImpl();
+
 protected:
 	/** Current status of request being processed */
 	EHttpRequestStatus::Type CompletionStatus = EHttpRequestStatus::NotStarted;
@@ -147,4 +152,7 @@ protected:
 	// Flag to indicate the request was initialized with stream. In that case even if stream was set to 
 	// null later on internally, the request itself won't cache received data anymore
 	std::atomic<bool> bInitializedWithValidStream = false;
+
+	/** Payload to use with the request. Typically for POST, PUT, or PATCH */
+	TUniquePtr<FRequestPayload> RequestPayload;
 };
