Fix the non-unity compile error on android and apple, caused by the request payload refactor.

diff --git a/Engine/Source/Runtime/Online/HTTP/Private/Curl/CurlHttp.h b/Engine/Source/Runtime/Online/HTTP/Private/Curl/CurlHttp.h
--- a/Engine/Source/Runtime/Online/HTTP/Private/Curl/CurlHttp.h
+++ b/Engine/Source/Runtime/Online/HTTP/Private/Curl/CurlHttp.h
@@ -7,7 +7,6 @@
 #include "HAL/ThreadSafeCounter.h"
 #include "IHttpThreadedRequest.h"
 #include "Containers/SpscQueue.h"
-#include "GenericPlatform/HttpRequestPayload.h"
 #include "GenericPlatform/HttpResponseCommon.h"
 #include "HAL/ThreadSafeBool.h"
 
diff --git a/Engine/Source/Runtime/Online/HTTP/Private/GenericPlatform/HttpRequestCommon.cpp b/Engine/Source/Runtime/Online/HTTP/Private/GenericPlatform/HttpRequestCommon.cpp
--- a/Engine/Source/Runtime/Online/HTTP/Private/GenericPlatform/HttpRequestCommon.cpp
+++ b/Engine/Source/Runtime/Online/HTTP/Private/GenericPlatform/HttpRequestCommon.cpp
@@ -1,7 +1,6 @@
 // Copyright Epic Games, Inc. All Rights Reserved.
 
 #include "GenericPlatform/HttpRequestCommon.h"
-#include "GenericPlatform/HttpRequestPayload.h"
 #include "GenericPlatform/HttpResponseCommon.h"
 #include "HAL/Event.h"
 #include "Http.h"
diff --git a/Engine/Source/Runtime/Online/HTTP/Public/GenericPlatform/HttpRequestCommon.h b/Engine/Source/Runtime/Online/HTTP/Public/GenericPlatform/HttpRequestCommon.h
--- a/Engine/Source/Runtime/Online/HTTP/Public/GenericPlatform/HttpRequestCommon.h
+++ b/Engine/Source/Runtime/Online/HTTP/Public/GenericPlatform/HttpRequestCommon.h
@@ -3,10 +3,10 @@
 #pragma once
 
 #include "GenericPlatform/HttpRequestImpl.h"
+#include "GenericPlatform/HttpRequestPayload.h"
 
 class FHttpResponseCommon;
 class IHttpTaskTimerHandle;
-class FRequestPayload;
 
 /**
  * Contains implementation of some common functions that don't vary between implementations of different platforms
