[Http]
* Fixed thread sanitizer errors by changing TQueue<IHttpThreadedRequest*, EQueueMode::Mpsc> to TMpscQueue<IHttpThreadedRequest*>

diff --git a/Engine/Source/Runtime/Online/HTTP/Private/HttpThread.h b/Engine/Source/Runtime/Online/HTTP/Private/HttpThread.h
--- a/Engine/Source/Runtime/Online/HTTP/Private/HttpThread.h
+++ b/Engine/Source/Runtime/Online/HTTP/Private/HttpThread.h
@@ -163,13 +163,13 @@ protected:
 	 * Threaded requests that are waiting to be processed on the http thread.
 	 * Added to on (any) non-HTTP thread, processed then cleared on HTTP thread.
 	 */
-	TQueue<IHttpThreadedRequest*, EQueueMode::Mpsc> NewThreadedRequests;
+	TMpscQueue<IHttpThreadedRequest*> NewThreadedRequests;
 
 	/**
 	 * Threaded requests that are waiting to be cancelled on the http thread.
 	 * Added to on (any) non-HTTP thread, processed then cleared on HTTP thread.
 	 */
-	TQueue<IHttpThreadedRequest*, EQueueMode::Mpsc> CancelledThreadedRequests;
+	TMpscQueue<IHttpThreadedRequest*> CancelledThreadedRequests;
 
 	/**
 	 * Threaded requests that are ready to run, but waiting due to the running request limit (not in any of the other lists, except potentially CancelledThreadedRequests).
