Make sure to create a new NSInputStream instance when retry http request with streaming upload request.

diff --git a/Engine/Source/Runtime/Online/HTTP/Private/Apple/AppleHttp.cpp b/Engine/Source/Runtime/Online/HTTP/Private/Apple/AppleHttp.cpp
--- a/Engine/Source/Runtime/Online/HTTP/Private/Apple/AppleHttp.cpp
+++ b/Engine/Source/Runtime/Online/HTTP/Private/Apple/AppleHttp.cpp
@@ -531,6 +531,7 @@ static bool ShouldReadHeadersWhenComplete(const FString& Url)
 	return YES;
 }
 @end
+
 /****************************************************************************
  * FAppleHttpRequest implementation
  ***************************************************************************/
@@ -538,7 +539,6 @@ static bool ShouldReadHeadersWhenComplete(const FString& Url)
 FAppleHttpRequest::FAppleHttpRequest(NSURLSession* InSession)
 :   Session([InSession retain])
 ,   Task(nil)
-,	bIsPayloadFile(false)
 ,	ContentBytesLength(0)
 ,	ElapsedTime(0.0f)
 ,	LastReportedBytesWritten(0)
@@ -648,16 +648,16 @@ const TArray<uint8>& FAppleHttpRequest::GetContent() const
 {
 	UE_LOG(LogHttp, Verbose, TEXT("FAppleHttpRequest::GetContent()"));
 	StorageForGetContent.Empty();
-	if (bIsPayloadFile)
-	{
-		UE_LOG(LogHttp, Warning, TEXT("FAppleHttpRequest::GetContent() called on a request that is set up for streaming a file. Return value is an empty buffer"));
-	}
-	else
+	if (StreamedContentSource.IsType<FNoStreamSource>())
 	{
 		SCOPED_AUTORELEASE_POOL;
 		NSData* Body = Request.HTTPBody; // accessing HTTPBody will call retain autorelease on the value, increasing its retain count
 		StorageForGetContent.Append((const uint8*)Body.bytes, Body.length);
 	}
+	else
+	{
+		UE_LOG(LogHttp, Warning, TEXT("FAppleHttpRequest::GetContent() called on a request that is set up for streaming a file. Return value is an empty buffer"));
+	}
 	return StorageForGetContent;
 }
 
@@ -668,12 +668,11 @@ void FAppleHttpRequest::SetContent(const TArray<uint8>& ContentPayload)
 		UE_LOG(LogHttp, Warning, TEXT("FAppleHttpRequest::SetContent() - attempted to set content on a request that is inflight"));
 		return;
 	}
-
+	
 	UE_LOG(LogHttp, Verbose, TEXT("FAppleHttpRequest::SetContent()"));
-	Request.HTTPBodyStream = nil;
+	StreamedContentSource.Emplace<FNoStreamSource>();
 	Request.HTTPBody = [NSData dataWithBytes:ContentPayload.GetData() length:ContentPayload.Num()];
 	ContentBytesLength = ContentPayload.Num();
-	bIsPayloadFile = false;
 }
 
 void FAppleHttpRequest::SetContent(TArray<uint8>&& ContentPayload)
@@ -686,12 +685,11 @@ void FAppleHttpRequest::SetContent(TArray<uint8>&& ContentPayload)
 
 	UE_LOG(LogHttp, Verbose, TEXT("FAppleHttpRequest::SetContent()"));
 
-	Request.HTTPBodyStream = nil;
+	StreamedContentSource.Emplace<FNoStreamSource>();
 	// We cannot use NSData dataWithBytesNoCopy:length:freeWhenDone: and keep the data in this instance because we don't have control
 	// over the lifetime of the request copy that NSURLSessionTask keeps
 	Request.HTTPBody = [NSData dataWithBytes:ContentPayload.GetData() length:ContentPayload.Num()];
 	ContentBytesLength = ContentPayload.Num();
-	bIsPayloadFile = false;
 
 	// Clear argument content since client code probably expects that
 	ContentPayload.Empty();
@@ -721,11 +719,10 @@ void FAppleHttpRequest::SetContentAsString(const FString& ContentString)
 	UE_LOG(LogHttp, Verbose, TEXT("FAppleHttpRequest::SetContentAsString() - %s"), *ContentString);
 	FTCHARToUTF8 Converter(*ContentString);
 
-	Request.HTTPBodyStream = nil;
+	StreamedContentSource.Emplace<FNoStreamSource>();
 	// The extra length computation here is unfortunate, but it's technically not safe to assume the length is the same.
 	Request.HTTPBody = [NSData dataWithBytes:(ANSICHAR*)Converter.Get() length:Converter.Length()];
 	ContentBytesLength = Converter.Length();
-	bIsPayloadFile = false;
 }
 
 bool FAppleHttpRequest::SetContentAsStreamedFile(const FString& Filename)
@@ -747,24 +744,20 @@ bool FAppleHttpRequest::SetContentAsStreamedFile(const FString& Filename)
 	if (stat(PlatformFilename.fileSystemRepresentation, &FileAttrs) == 0)
 	{
 		UE_LOG(LogHttp, VeryVerbose, TEXT("FAppleHttpRequest::SetContentAsStreamedFile succeeded in getting the file size - %lld"), FileAttrs.st_size);
-		// Under the hood, the Foundation framework unsets HTTPBody, and takes over as the stream delegate.
-		// The stream itself should be unopened when passed to setHTTPBodyStream.
-		Request.HTTPBodyStream = [NSInputStream inputStreamWithFileAtPath: PlatformFilename];
+		StreamedContentSource.Emplace<FString>(Filename);
 		ContentBytesLength = FileAttrs.st_size;
-		bIsPayloadFile = true;
+		return true;
 	}
 	else
 	{
 		UE_LOG(LogHttp, Warning, TEXT("FAppleHttpRequest::SetContentAsStreamedFile failed to get file size"));
-		Request.HTTPBodyStream = nil;
+		StreamedContentSource.Emplace<FNoStreamSource>();
 		ContentBytesLength = 0;
-		bIsPayloadFile = false;
+		return false;
 	}
-
-	return bIsPayloadFile;
 }
 
-bool FAppleHttpRequest::SetContentFromStream(TSharedRef<FArchive, ESPMode::ThreadSafe> Stream)
+bool FAppleHttpRequest::SetContentFromStream(TSharedRef<FArchive> Stream)
 {
 	SCOPED_AUTORELEASE_POOL;
 	if (CompletionStatus == EHttpRequestStatus::Processing)
@@ -774,10 +767,8 @@ bool FAppleHttpRequest::SetContentFromStream(TSharedRef<FArchive, ESPMode::Threa
 	}
 
 	Request.HTTPBody = nil;
-
-	Request.HTTPBodyStream = [FNSInputStreamFromArchive initWithArchive: Stream];
 	ContentBytesLength = Stream->TotalSize();
-	bIsPayloadFile = true;
+	StreamedContentSource.Emplace<TSharedRef<FArchive>>(MoveTemp(Stream));
 
 	return true;
 }
@@ -809,11 +800,28 @@ bool FAppleHttpRequest::ProcessRequest()
 	return true;
 }
 
+struct FAppleHttpRequest::FAppleHttpStreamFactory
+{
+	NSInputStream *operator()(FNoStreamSource)
+	{
+		return nil;
+	}
+	
+	NSInputStream *operator()(const FString& Filename)
+	{
+		return [NSInputStream inputStreamWithFileAtPath: Filename.GetNSString()];
+	}
+	
+	NSInputStream *operator()(const TSharedRef<FArchive>& Archive)
+	{
+		return [FNSInputStreamFromArchive initWithArchive: Archive];
+	}
+};
+
 bool FAppleHttpRequest::SetupRequest()
 {
 	SCOPED_AUTORELEASE_POOL;
 	UE_LOG(LogHttp, Verbose, TEXT("FAppleHttpRequest::SetupRequest()"));
-	bool bStarted = false;
 
 	// set the content-length and user-agent (it is possible that the OS ignores this value)
 	if(GetContentLength() > 0)
@@ -849,6 +857,18 @@ bool FAppleHttpRequest::SetupRequest()
 		), 
 		HttpConnectionTimeout, GetActivityTimeoutOrDefault());
 
+	if (NSInputStream *HttpBodyStream = Visit(FAppleHttpStreamFactory{}, StreamedContentSource))
+	{
+		Request.HTTPBodyStream = HttpBodyStream;
+	}
+	else
+	{
+		UE_CLOG(!StreamedContentSource.IsType<FNoStreamSource>(), LogHttp, Warning, TEXT("Could not create native stream from stream source"));
+		SetStatus(EHttpRequestStatus::Failed);
+		SetFailureReason(EHttpFailureReason::Other);
+		return false;
+	}
+	
 	Task = [Session dataTaskWithRequest: Request];
 	
 	if (Task != nil)
@@ -869,15 +889,15 @@ bool FAppleHttpRequest::SetupRequest()
 
 		[[Task retain] resume];
 		UE_LOG(LogHttp, Verbose, TEXT("[NSURLSessionTask resume]"));
+		return true;
 	}
 	else
 	{
-		UE_LOG(LogHttp, Warning, TEXT("ProcessRequest failed. Could not initialize Internet connection."));
+		UE_LOG(LogHttp, Warning, TEXT("SetupRequest failed. Could not initialize NSURLSessionTask."));
 		SetStatus(EHttpRequestStatus::Failed);
 		SetFailureReason(EHttpFailureReason::ConnectionError);
+		return false;
 	}
-
-	return bStarted;
 }
 
 void FAppleHttpRequest::FinishRequest()
diff --git a/Engine/Source/Runtime/Online/HTTP/Private/Apple/AppleHttp.h b/Engine/Source/Runtime/Online/HTTP/Private/Apple/AppleHttp.h
--- a/Engine/Source/Runtime/Online/HTTP/Private/Apple/AppleHttp.h
+++ b/Engine/Source/Runtime/Online/HTTP/Private/Apple/AppleHttp.h
@@ -7,6 +7,7 @@
 #include "IHttpThreadedRequest.h"
 #include "PlatformHttp.h"
 #include "HttpPackage.h"
+#include "Misc/TVariant.h"
 
 /**
  * Delegate invoked when in progress Task completes. It is invoked in an out of our control thread
@@ -40,7 +41,7 @@ public:
 	virtual void SetContent(TArray<uint8>&& ContentPayload) override;
 	virtual void SetContentAsString(const FString& ContentString) override;
     virtual bool SetContentAsStreamedFile(const FString& Filename) override;
-	virtual bool SetContentFromStream(TSharedRef<FArchive, ESPMode::ThreadSafe> Stream) override;
+	virtual bool SetContentFromStream(TSharedRef<FArchive> Stream) override;
 	virtual void SetHeader(const FString& HeaderName, const FString& HeaderValue) override;
 	virtual void AppendToHeader(const FString& HeaderName, const FString& AdditionalHeaderValue) override;
 	virtual bool ProcessRequest() override;
@@ -108,8 +109,14 @@ private:
 	/** This is the Task associated to the sessionin charge of our request */
 	NSURLSessionTask* Task;
     
-	/** Flag whether the request payload source is a file */
-	bool bIsPayloadFile;
+	struct FAppleHttpStreamFactory;
+	using FNoStreamSource = FEmptyVariantState;
+	/** Source to create stream from
+		FNoStreamedSource: No streamed data
+		FString: Filename set from SetContentAsStreamedFile
+		TSharedRef<FArchive>: Stream set from SetContentFromStream
+	 */
+	TVariant<FNoStreamSource, FString, TSharedRef<FArchive>> StreamedContentSource;
 
 	/** The request payload length in bytes. This must be tracked separately for a file stream */
 	uint64 ContentBytesLength;
