Removed assignment from removed variable

diff --git a/Engine/Source/Runtime/Online/HTTP/Private/Apple/AppleHttp.cpp b/Engine/Source/Runtime/Online/HTTP/Private/Apple/AppleHttp.cpp
--- a/Engine/Source/Runtime/Online/HTTP/Private/Apple/AppleHttp.cpp
+++ b/Engine/Source/Runtime/Online/HTTP/Private/Apple/AppleHttp.cpp
@@ -873,8 +873,6 @@ bool FAppleHttpRequest::SetupRequest()
 	
 	if (Task != nil)
 	{
-		bStarted = true;
-
 		SetStatus(EHttpRequestStatus::Processing);
 		SetFailureReason(EHttpFailureReason::None);
 
