Fix all non stream related HTTP requests being marked as failed on apple platforms

diff --git a/Engine/Source/Runtime/Online/HTTP/Private/Apple/AppleHttp.cpp b/Engine/Source/Runtime/Online/HTTP/Private/Apple/AppleHttp.cpp
--- a/Engine/Source/Runtime/Online/HTTP/Private/Apple/AppleHttp.cpp
+++ b/Engine/Source/Runtime/Online/HTTP/Private/Apple/AppleHttp.cpp
@@ -861,9 +861,9 @@ bool FAppleHttpRequest::SetupRequest()
 	{
 		Request.HTTPBodyStream = HttpBodyStream;
 	}
-	else
+	else if (!StreamedContentSource.IsType<FNoStreamSource>())
 	{
-		UE_CLOG(!StreamedContentSource.IsType<FNoStreamSource>(), LogHttp, Warning, TEXT("Could not create native stream from stream source"));
+		UE_LOG(LogHttp, Warning, TEXT("Could not create native stream from stream source"));
 		SetStatus(EHttpRequestStatus::Failed);
 		SetFailureReason(EHttpFailureReason::Other);
 		return false;
