Make sure when cancel http request from retry system, it can still trigger complete delegate when non-game thread support enabled.

diff --git a/Engine/Source/Programs/WebTests/Private/TestHttp.cpp b/Engine/Source/Programs/WebTests/Private/TestHttp.cpp
--- a/Engine/Source/Programs/WebTests/Private/TestHttp.cpp
+++ b/Engine/Source/Programs/WebTests/Private/TestHttp.cpp
@@ -1695,14 +1695,22 @@ TEST_CASE_METHOD(FWaitUntilCompleteHttpFixture, "Cancel http request connect bef
 	HttpRequest->SetVerb(TEXT("GET"));
 	HttpRequest->SetTimeout(7);
 	const double StartTime = FPlatformTime::Seconds();
-	HttpRequest->OnProcessRequestComplete().BindLambda([StartTime](FHttpRequestPtr HttpRequest, FHttpResponsePtr HttpResponse, bool bSucceeded) {
+	++ExpectingExtraCallbacks;
+	HttpRequest->OnProcessRequestComplete().BindLambda([this, StartTime](FHttpRequestPtr HttpRequest, FHttpResponsePtr HttpResponse, bool bSucceeded) {
+		--ExpectingExtraCallbacks;
 		CHECK(!bSucceeded);
 		const double DurationInSeconds = FPlatformTime::Seconds() - StartTime;
 		CHECK(DurationInSeconds < 2.0);
 		CHECK(HttpRequest->GetFailureReason() == EHttpFailureReason::Cancelled);
 	});
-	HttpRequest->ProcessRequest();
-	FPlatformProcess::Sleep(0.5);
+	SECTION("ProcessRequest called")
+	{
+		HttpRequest->ProcessRequest();
+		FPlatformProcess::Sleep(0.5);
+	}
+	SECTION("ProcessRequest not called")
+	{
+	}
 	HttpRequest->CancelRequest();
 	HttpRequest->CancelRequest(); // Duplicated calls to CancelRequest should be fine
 }
diff --git a/Engine/Source/Runtime/Online/HTTP/Private/HttpRetrySystem.cpp b/Engine/Source/Runtime/Online/HTTP/Private/HttpRetrySystem.cpp
--- a/Engine/Source/Runtime/Online/HTTP/Private/HttpRetrySystem.cpp
+++ b/Engine/Source/Runtime/Online/HTTP/Private/HttpRetrySystem.cpp
@@ -129,6 +129,20 @@ FHttpRetrySystem::FRequest::FRequest(
 	}
 }
 
+void FHttpRetrySystem::FRequest::BindAdaptorDelegates()
+{
+	if (!bBoundAdaptorDelegates)
+	{
+		bBoundAdaptorDelegates = true;
+
+		// Can't BindRaw&Unbind from ctor&dtor because with thread-safe delegate, it can cause issue when delete this request during complete callback, then unbind the callback
+		HttpRequest->OnProcessRequestComplete().BindThreadSafeSP(this, &FHttpRetrySystem::FRequest::HttpOnProcessRequestComplete);
+		HttpRequest->OnRequestProgress64().BindThreadSafeSP(this, &FHttpRetrySystem::FRequest::HttpOnRequestProgress);
+		HttpRequest->OnStatusCodeReceived().BindThreadSafeSP(this, &FHttpRetrySystem::FRequest::HttpOnStatusCodeReceived);
+		HttpRequest->OnHeaderReceived().BindThreadSafeSP(this, &FHttpRetrySystem::FRequest::HttpOnHeaderReceived);
+	}
+}
+
 bool FHttpRetrySystem::FRequest::ProcessRequest()
 {
 	TSharedRef<FRequest> RetryRequest = StaticCastSharedRef<FRequest>(AsShared());
@@ -139,10 +153,7 @@ bool FHttpRetrySystem::FRequest::ProcessRequest()
 		SetUrlFromRetryDomains();
 	}
 
-	HttpRequest->OnRequestProgress64().BindThreadSafeSP(RetryRequest, &FHttpRetrySystem::FRequest::HttpOnRequestProgress);
-	HttpRequest->OnProcessRequestComplete().BindThreadSafeSP(RetryRequest, &FHttpRetrySystem::FRequest::HttpOnProcessRequestComplete);
-	HttpRequest->OnStatusCodeReceived().BindThreadSafeSP(RetryRequest, &FHttpRetrySystem::FRequest::HttpOnStatusCodeReceived);
-	HttpRequest->OnHeaderReceived().BindThreadSafeSP(RetryRequest, &FHttpRetrySystem::FRequest::HttpOnHeaderReceived);
+	BindAdaptorDelegates();
 
 	TSharedPtr<FManager> RetryManagerPtr = RetryManager.Pin();
 
@@ -182,6 +193,8 @@ void FHttpRetrySystem::FRequest::CancelRequest()
 { 
 	TSharedRef<FRequest, ESPMode::ThreadSafe> RetryRequest = StaticCastSharedRef<FRequest>(AsShared());
 
+	BindAdaptorDelegates();
+
 	if (TSharedPtr<FManager> RetryManagerPtr = RetryManager.Pin())
 	{
 		RetryManagerPtr->CancelRequest(RetryRequest);
diff --git a/Engine/Source/Runtime/Online/HTTP/Public/HttpRetrySystem.h b/Engine/Source/Runtime/Online/HTTP/Public/HttpRetrySystem.h
--- a/Engine/Source/Runtime/Online/HTTP/Public/HttpRetrySystem.h
+++ b/Engine/Source/Runtime/Online/HTTP/Public/HttpRetrySystem.h
@@ -149,6 +149,8 @@ namespace FHttpRetrySystem
 		/** Move to the next retry domain from our RetryDomains */
 		void MoveToNextRetryDomain();
 
+		void BindAdaptorDelegates();
+
 		EStatus::Type RetryStatus;
 
 		FRetryLimitCountSetting RetryLimitCountOverride;
@@ -169,6 +171,8 @@ namespace FHttpRetrySystem
 
 		/** Exponential backoff curve */
 		FExponentialBackoffCurve RetryExponentialBackoffCurve;
+
+		bool bBoundAdaptorDelegates = false;
 	};
 }
 
