Refactor function FHttpRetrySystem::FManager::ShouldRetry to use early returns.

#tests Covered by test case in WebTests.

diff --git a/Engine/Source/Programs/WebTests/Private/TestHttp.cpp b/Engine/Source/Programs/WebTests/Private/TestHttp.cpp
--- a/Engine/Source/Programs/WebTests/Private/TestHttp.cpp
+++ b/Engine/Source/Programs/WebTests/Private/TestHttp.cpp
@@ -1857,6 +1857,40 @@ TEST_CASE_METHOD(FWaitUntilCompleteHttpFixture, "Request can time out during ret
 	HttpRequest->ProcessRequest();
 }
 
+TEST_CASE_METHOD(FWaitUntilCompleteHttpFixture, "Request will not retry", HTTP_TAG)
+{
+	if (!bRetryEnabled)
+	{
+		return;
+	}
+
+	DisableWarningsInThisTest();
+
+	TSharedRef<IHttpRequest> HttpRequest = HttpRetryManager->CreateRequest(1/*InRetryLimitCountOverride*/);
+	SECTION("When response code is not listed for retry")
+	{
+		HttpRequest->SetURL(UrlMockStatus(EHttpResponseCodes::TooManyRequests));
+		// Will be forwarded back in response
+		HttpRequest->SetHeader(TEXT("Retry-After"), FString::Format(TEXT("{0}"), { 2 }));
+	}
+	SECTION("When there is any response and timed out during streaming download")
+	{
+		HttpRequest->SetURL(UrlStreamDownload(3/*Chunks*/, HTTP_TEST_TIMEOUT_CHUNK_SIZE, 2/*ChunkLatency*/));
+		HttpRequest->SetTimeout(3.0f);
+
+		HttpRequest->OnProcessRequestComplete().BindLambda([](FHttpRequestPtr HttpRequest, FHttpResponsePtr HttpResponse, bool bSucceeded) {
+			CHECK(HttpRequest->GetStatus() == EHttpRequestStatus::Failed);
+			CHECK(HttpRequest->GetFailureReason() == EHttpFailureReason::TimedOut);
+		});
+	}
+
+	HttpRequest->OnRequestWillRetry().BindLambda([this](FHttpRequestPtr Request, FHttpResponsePtr Response, float LockoutPeriod) {
+		CHECK(false);
+	});
+
+	HttpRequest->ProcessRequest();
+}
+
 TEST_CASE_METHOD(FWaitUntilCompleteHttpFixture, "Retry immediately without lock out if connect failed and there are alt domains", HTTP_TAG)
 {
 	if (!bRetryEnabled)
diff --git a/Engine/Source/Runtime/Online/HTTP/Private/HttpRetrySystem.cpp b/Engine/Source/Runtime/Online/HTTP/Private/HttpRetrySystem.cpp
--- a/Engine/Source/Runtime/Online/HTTP/Private/HttpRetrySystem.cpp
+++ b/Engine/Source/Runtime/Online/HTTP/Private/HttpRetrySystem.cpp
@@ -375,42 +375,40 @@ TSharedRef<FHttpRetrySystem::FRequest, ESPMode::ThreadSafe> FHttpRetrySystem::FM
 
 bool FHttpRetrySystem::FManager::ShouldRetry(const FHttpRetryRequestEntry& HttpRetryRequestEntry)
 {
-    bool bResult = false;
-
 	FHttpResponsePtr Response = HttpRetryRequestEntry.Request->GetResponse();
-	// invalid response means connection or network error but we need to know which one
-	if (!Response.IsValid())
+	if (Response)
 	{
-		// ONLY retry bad responses if they are connection errors (NOT protocol errors or unknown) otherwise request may be sent (and processed!) twice
-		if (HttpRetryRequestEntry.Request->GetStatus() == EHttpRequestStatus::Failed)
-		{
-			if (HttpRetryRequestEntry.Request->GetFailureReason() == EHttpFailureReason::ConnectionError)
-			{
-				bResult = true;
-			}
-			else
-			{
-				const FName Verb = FName(*HttpRetryRequestEntry.Request->GetVerb());
+		return HttpRetryRequestEntry.Request->RetryResponseCodes.Contains(Response->GetResponseCode());
+	}
 
-				// Be default, we will also allow retry for GET and HEAD requests even if they may duplicate on the server
-				static const TSet<FName> DefaultRetryVerbs(TArray<FName>({ FName(TEXT("GET")), FName(TEXT("HEAD")) }));
+	// ONLY continue to check retry if no response. If there is any response, it means at least the http 
+	// connection was established, we shouldn't attempt to retry. Otherwise request may be sent (and 
+	// processed) twice
 
-				const TSet<FName>* RetryVerbsContainer = (
-					HttpRetryRequestEntry.Request->RetryVerbs.Num() == 0
-					? &DefaultRetryVerbs // Use the default list of retry verbs if the request doesn't have a specific set
-					: &HttpRetryRequestEntry.Request->RetryVerbs // Otherwise use the specific set on the request
-				);
-				bResult = RetryVerbsContainer->Contains(Verb);
-			}
-		}
+	// Safe check
+	if (HttpRetryRequestEntry.Request->GetStatus() != EHttpRequestStatus::Failed)
+	{
+		// This shouldn't happen when response is null, but just in case
+		return false;
 	}
-	else
+
+	// Should retry if couldn't connect at all
+	if (HttpRetryRequestEntry.Request->GetFailureReason() == EHttpFailureReason::ConnectionError)
+	{
+		return true;
+	}
+
+	// Should retry for idempotent verbs if there is network error
+	const FName Verb = FName(*HttpRetryRequestEntry.Request->GetVerb());
+
+	if (!HttpRetryRequestEntry.Request->RetryVerbs.IsEmpty())
 	{
-		// this may be a successful response with one of the explicitly listed response codes we want to retry on
-		bResult = HttpRetryRequestEntry.Request->RetryResponseCodes.Contains(Response->GetResponseCode());
+		return HttpRetryRequestEntry.Request->RetryVerbs.Contains(Verb);
 	}
 
-    return bResult;
+	// Be default, we will also allow retry for GET and HEAD requests even if they may duplicate on the server
+	static const TSet<FName> DefaultRetryVerbs(TArray<FName>({ FName(TEXT("GET")), FName(TEXT("HEAD")) }));
+	return DefaultRetryVerbs.Contains(Verb);
 }
 
 bool FHttpRetrySystem::FManager::RetryLimitForConnectionErrorIsSet(const FHttpRetryRequestEntry& HttpRetryRequestEntry)
