Trigger delegates in FAppleHttp::AbortRequest in case the request was not yet added to the http manager

diff --git a/Engine/Source/Runtime/Online/HTTP/Private/Apple/AppleHttp.cpp b/Engine/Source/Runtime/Online/HTTP/Private/Apple/AppleHttp.cpp
--- a/Engine/Source/Runtime/Online/HTTP/Private/Apple/AppleHttp.cpp
+++ b/Engine/Source/Runtime/Online/HTTP/Private/Apple/AppleHttp.cpp
@@ -927,6 +927,10 @@ void FAppleHttpRequest::FinishRequest()
 				}
 			}
 		}
+		else if (bCanceled)
+		{
+			Reason = EHttpFailureReason::Cancelled;
+		}
 		SetFailureReason(Reason);
 
 		if (GetFailureReason() == EHttpFailureReason::ConnectionError)
@@ -972,6 +976,11 @@ void FAppleHttpRequest::AbortRequest()
 	{
 		[Task cancel];
 	}
+	else
+	{
+		// No Task means SetupRequest was not called, so we were not added to the HttpManager yet
+		FinishRequestNotInHttpManager();
+	}
 }
 
 void FAppleHttpRequest::Tick(float DeltaSeconds)
