Add support to log response for specified url patterns, through CVarHttpUrlPatternsToLogResponse. By default it only logs json response but can enable for all content type by disable CVarHttpLogJsonResponseOnly;
Move response headers logging from CurlHttp into common http impl.

#tests Tried on pc when the command line -ini:Engine:[ConsoleVariables]:http.UrlPatternsToLogResponse=epicgames.net

diff --git a/Engine/Source/Runtime/Online/HTTP/Private/Apple/AppleHttp.cpp b/Engine/Source/Runtime/Online/HTTP/Private/Apple/AppleHttp.cpp
--- a/Engine/Source/Runtime/Online/HTTP/Private/Apple/AppleHttp.cpp
+++ b/Engine/Source/Runtime/Online/HTTP/Private/Apple/AppleHttp.cpp
@@ -540,7 +540,6 @@ FAppleHttpRequest::FAppleHttpRequest(NSURLSession* InSession)
 :   Session([InSession retain])
 ,   Task(nil)
 ,	ContentBytesLength(0)
-,	ElapsedTime(0.0f)
 ,	LastReportedBytesWritten(0)
 ,	LastReportedBytesRead(0)
 {
@@ -1015,11 +1014,6 @@ PRAGMA_ENABLE_DEPRECATION_WARNINGS
 	}
 }
 
-float FAppleHttpRequest::GetElapsedTime() const
-{
-	return ElapsedTime;
-}
-
 bool FAppleHttpRequest::StartThreadedRequest()
 {
 	return true;
diff --git a/Engine/Source/Runtime/Online/HTTP/Private/Apple/AppleHttp.h b/Engine/Source/Runtime/Online/HTTP/Private/Apple/AppleHttp.h
--- a/Engine/Source/Runtime/Online/HTTP/Private/Apple/AppleHttp.h
+++ b/Engine/Source/Runtime/Online/HTTP/Private/Apple/AppleHttp.h
@@ -46,7 +46,6 @@ public:
 	virtual void AppendToHeader(const FString& HeaderName, const FString& AdditionalHeaderValue) override;
 	virtual bool ProcessRequest() override;
 	virtual void Tick(float DeltaSeconds) override;
-	virtual float GetElapsedTime() const override;
 	//~ End IHttpRequest Interface
 
 	//~ Begin IHttpRequestThreaded Interface
@@ -124,9 +123,6 @@ private:
 	/** Array used to retrieve back content set on the ObjC request when calling GetContent*/
 	mutable TArray<uint8> StorageForGetContent;
 
-	/** Time taken to complete/cancel the request. */
-	float ElapsedTime;
-
 	/** Last reported bytes written */
 	int32 LastReportedBytesWritten;
 
diff --git a/Engine/Source/Runtime/Online/HTTP/Private/Curl/CurlHttp.cpp b/Engine/Source/Runtime/Online/HTTP/Private/Curl/CurlHttp.cpp
--- a/Engine/Source/Runtime/Online/HTTP/Private/Curl/CurlHttp.cpp
+++ b/Engine/Source/Runtime/Online/HTTP/Private/Curl/CurlHttp.cpp
@@ -19,12 +19,6 @@
 #include <openssl/ssl.h>
 #endif
 
-TAutoConsoleVariable<bool> CVarCurlDebugServerResponseEnabled(
-	TEXT("http.CurlDebugServerResponseEnabled"),
-	false,
-	TEXT("Enable debugging of server response")
-);
-
 #if WITH_SSL
 static int SslCertVerify(int PreverifyOk, X509_STORE_CTX* Context)
 {
@@ -76,7 +70,6 @@ FCurlHttpRequest::FCurlHttpRequest()
 	, bRedirected(false)
 	, CurlAddToMultiResult(CURLM_OK)
 	, CurlCompletionResult(CURLE_OK)
-	, ElapsedTime(0.0f)
 	, bAnyHttpActivity(false)
 	, BytesSent(0)
 	, TotalBytesSent(0)
@@ -1220,40 +1213,6 @@ void FCurlHttpRequest::FinishRequest()
 
 	if (Response.IsValid() && Response->bSucceeded)
 	{
-		bool bDebugServerResponse = CVarCurlDebugServerResponseEnabled.GetValueOnAnyThread() && (Response->GetResponseCode() >= 500 && Response->GetResponseCode() <= 503);
-
-		// log info about error responses to identify failed downloads
-		if (UE_LOG_ACTIVE(LogHttp, Verbose) || bDebugServerResponse)
-		{
-			if (bDebugServerResponse)
-			{
-				UE_LOG(LogHttp, Warning, TEXT("%p: request has been successfully processed. URL: %s, HTTP code: %d, content length: %llu, actual payload size: %llu, elapsed: %.2fs"),
-					this, *GetURL(), Response->HttpCode, Response->ContentLength, TotalBytesRead.load(), ElapsedTime);
-			}
-			else
-			{
-				UE_LOG(LogHttp, Log, TEXT("%p: request has been successfully processed. URL: %s, HTTP code: %d, content length: %llu, actual payload size: %llu, elapsed: %.2fs"),
-					this, *GetURL(), Response->HttpCode, Response->ContentLength, TotalBytesRead.load(), ElapsedTime);
-			}
-
-			TArray<FString> AllHeaders = Response->GetAllHeaders();
-			for (TArray<FString>::TConstIterator It(AllHeaders); It; ++It)
-			{
-				const FString& HeaderStr = *It;
-				if (!HeaderStr.StartsWith(TEXT("Authorization")) && !HeaderStr.StartsWith(TEXT("Set-Cookie")))
-				{
-					if (bDebugServerResponse)
-					{
-						UE_LOG(LogHttp, Warning, TEXT("%p Response Header %s"), this, *HeaderStr);
-					}
-					else
-					{
-						UE_LOG(LogHttp, Verbose, TEXT("%p Response Header %s"), this, *HeaderStr);
-					}
-				}
-			}
-		}
-
 		HandleRequestSucceed(Response);
 	}
 	else
@@ -1327,11 +1286,6 @@ void FCurlHttpRequest::FinishRequest()
 	}
 }
 
-float FCurlHttpRequest::GetElapsedTime() const
-{
-	return ElapsedTime;
-}
-
 void FCurlHttpRequest::CleanupRequest()
 {
 	curl_easy_setopt(EasyHandle, CURLOPT_SHARE, nullptr);
diff --git a/Engine/Source/Runtime/Online/HTTP/Private/Curl/CurlHttp.h b/Engine/Source/Runtime/Online/HTTP/Private/Curl/CurlHttp.h
--- a/Engine/Source/Runtime/Online/HTTP/Private/Curl/CurlHttp.h
+++ b/Engine/Source/Runtime/Online/HTTP/Private/Curl/CurlHttp.h
@@ -164,7 +164,6 @@ public:
 	virtual void AppendToHeader(const FString& HeaderName, const FString& AdditionalHeaderValue) override;
 	virtual bool ProcessRequest() override;
 	virtual void Tick(float DeltaSeconds) override;
-	virtual float GetElapsedTime() const override;
 	//~ End IHttpRequest Interface
 
 	//~ Begin IHttpRequestThreaded Interface
@@ -376,8 +375,6 @@ private:
 	bool bIsRequestPayloadSeekable = false;
 	/** Mapping of header section to values. */
 	TMap<FString, FString> Headers;
-	/** Total elapsed time in seconds since the start of the request */
-	float ElapsedTime;
 	/** Have we had any HTTP activity with the host? Sending headers, SSL handshake, etc */
 	bool bAnyHttpActivity;
 	/** Newly received headers we need to inform listeners about */
diff --git a/Engine/Source/Runtime/Online/HTTP/Private/GenericPlatform/HttpRequestCommon.cpp b/Engine/Source/Runtime/Online/HTTP/Private/GenericPlatform/HttpRequestCommon.cpp
--- a/Engine/Source/Runtime/Online/HTTP/Private/GenericPlatform/HttpRequestCommon.cpp
+++ b/Engine/Source/Runtime/Online/HTTP/Private/GenericPlatform/HttpRequestCommon.cpp
@@ -3,11 +3,61 @@
 #include "GenericPlatform/HttpRequestCommon.h"
 #include "GenericPlatform/HttpResponseCommon.h"
 #include "HAL/Event.h"
+#include "HAL/IConsoleManager.h"
 #include "Http.h"
 #include "HttpManager.h"
+#include "Logging/StructuredLog.h"
 #include "Misc/CommandLine.h"
 #include "Stats/Stats.h"
 
+namespace UE::HttpRequestCommon::Private
+{
+
+TAutoConsoleVariable<FString> CVarHttpUrlPatternsToLogResponse(
+	TEXT("http.UrlPatternsToLogResponse"),
+	TEXT(""),
+	TEXT("List of url patterns to log headers and json content: \"epicgames.com,unrealengine.com,...\""),
+	ECVF_SaveForNextBoot
+);
+
+TAutoConsoleVariable<bool> CVarHttpLogJsonResponseOnly(
+	TEXT("http.LogJsonResponseOnly"),
+	true,
+	TEXT("When log response payload, log json content only"),
+	ECVF_SaveForNextBoot
+);
+
+bool ShouldLogResponse(FStringView Url)
+{
+	static std::atomic<bool> bUpdatedCVarHttpUrlPatternsToLogResponse = true;
+	UE_CALL_ONCE([] {
+		CVarHttpUrlPatternsToLogResponse.AsVariable()->OnChangedDelegate().AddLambda([](IConsoleVariable* CVar) {
+			bUpdatedCVarHttpUrlPatternsToLogResponse = true;
+		});
+	});
+
+	static TArray<FString> UrlPatternsToLogResponse;
+	static FCriticalSection UrlPatternsToLogResponseCriticalSection;
+	const FScopeLock CacheLock(&UrlPatternsToLogResponseCriticalSection);
+	if (bUpdatedCVarHttpUrlPatternsToLogResponse)
+	{
+		CVarHttpUrlPatternsToLogResponse.GetValueOnAnyThread().ParseIntoArray(UrlPatternsToLogResponse, TEXT(","));
+		bUpdatedCVarHttpUrlPatternsToLogResponse = false;
+	}
+
+	for (const FString& UrlPatternToLogResponse : UrlPatternsToLogResponse)
+	{
+		if (Url.Contains(UrlPatternToLogResponse))
+		{
+			return true;
+		}
+	}
+
+	return false;
+}
+
+}
+
 FHttpRequestCommon::FHttpRequestCommon()
 	: RequestStartTimeAbsoluteSeconds(FPlatformTime::Seconds())
 	, ActivityTimeoutAt(0.0)
@@ -159,6 +209,9 @@ EHttpRequestDelegateThreadPolicy FHttpRequestCommon::GetDelegateThreadPolicy() c
 void FHttpRequestCommon::HandleRequestSucceed(TSharedPtr<IHttpResponse> InResponse)
 {
 	SetStatus(EHttpRequestStatus::Succeeded);
+
+	LogResponse(InResponse);
+
 	OnProcessRequestComplete().ExecuteIfBound(SharedThis(this), InResponse, true);
 	FHttpModule::Get().GetHttpManager().RecordStatTimeToConnect(ConnectTime);
 }
@@ -476,6 +529,11 @@ bool FHttpRequestCommon::SetResponseBodyReceiveStream(TSharedRef<FArchive> Strea
 	return true;
 }
 
+float FHttpRequestCommon::GetElapsedTime() const
+{
+	return ElapsedTime;
+}
+
 bool FHttpRequestCommon::PassReceivedDataToStream(void* Ptr, int64 Length)
 {
 	const FScopeLock StreamLock(&ResponseBodyReceiveStreamCriticalSection);
@@ -497,6 +555,7 @@ void FHttpRequestCommon::StopPassingReceivedData()
 	ResponseBodyReceiveStream = nullptr;
 }
 
+
 float FHttpRequestCommon::GetActivityTimeoutOrDefault() const
 {
 	return ActivityTimeoutSecs.Get(FHttpModule::Get().GetHttpActivityTimeout());
@@ -544,3 +603,63 @@ void FHttpRequestCommon::CloseRequestPayloadDefaultImpl()
 		RequestPayload->Close();
 	}
 }
+
+#define UE_HTTP_LOG_RESPONSE_PRIVATE(Condition, Format, ...) \
+	if (Condition) \
+	{ \
+		UE_LOG(LogHttp, Warning, Format, ##__VA_ARGS__); \
+	} \
+	else \
+	{ \
+		UE_LOG(LogHttp, Verbose, Format, ##__VA_ARGS__); \
+	}
+
+void FHttpRequestCommon::LogResponse(const TSharedPtr<IHttpResponse>& InResponse)
+{
+	bool bShouldLogResponse = UE::HttpRequestCommon::Private::ShouldLogResponse(GetURL());
+	UE_HTTP_LOG_RESPONSE_PRIVATE(bShouldLogResponse, TEXT("%p %s %s completed with code %d after %.2fs. Content length: %ld"), this, *GetVerb(), *GetURL(), InResponse->GetResponseCode(), ElapsedTime, InResponse->GetContentLength());
+
+	TArray<FString> AllHeaders = InResponse->GetAllHeaders();
+	for (const FString& HeaderStr : AllHeaders)
+	{
+		if (!HeaderStr.StartsWith(TEXT("Authorization")) && !HeaderStr.StartsWith(TEXT("Set-Cookie")))
+		{
+			UE_HTTP_LOG_RESPONSE_PRIVATE(bShouldLogResponse, TEXT("%p Response Header %s"), this, *HeaderStr);
+		}
+	}
+
+	if (!bShouldLogResponse || InResponse->GetContentLength() == 0)
+	{
+		return;
+	}
+
+	if (UE::HttpRequestCommon::Private::CVarHttpLogJsonResponseOnly.GetValueOnAnyThread())
+	{
+		bool bIsContentTypeJson = !InResponse->GetHeader(TEXT("Content-Type")).Compare(TEXT("application/json"), ESearchCase::IgnoreCase);
+		if (!bIsContentTypeJson)
+			return;
+	}
+
+	const TArray<uint8>& Content = InResponse->GetContent();
+	FUtf8StringView ResponseStringView(reinterpret_cast<const UTF8CHAR*>(Content.GetData()), Content.Num());
+	int32 StartPos = 0;
+	int32 EndPos = 0;
+	// The response payload could exceed the maximum length supported by UE_LOG/UE_LOGFMT, so log it line by line if there are multiple lines
+	while (StartPos < ResponseStringView.Len())
+	{
+		EndPos = ResponseStringView.Find("\n", StartPos);
+		if (EndPos != INDEX_NONE)
+		{
+			FUtf8StringView Line(&ResponseStringView[StartPos], EndPos - StartPos);
+			UE_LOGFMT(LogHttp, Warning, "{Line}", Line);
+		}
+		else
+		{
+			FUtf8StringView Remain(&ResponseStringView[StartPos], ResponseStringView.Len() - StartPos);
+			UE_LOGFMT(LogHttp, Warning, "{Remain}", Remain);
+			break;
+		}
+
+		StartPos = EndPos + 1;
+	}
+}
\ No newline at end of file
diff --git a/Engine/Source/Runtime/Online/HTTP/Public/GenericPlatform/HttpRequestCommon.h b/Engine/Source/Runtime/Online/HTTP/Public/GenericPlatform/HttpRequestCommon.h
--- a/Engine/Source/Runtime/Online/HTTP/Public/GenericPlatform/HttpRequestCommon.h
+++ b/Engine/Source/Runtime/Online/HTTP/Public/GenericPlatform/HttpRequestCommon.h
@@ -46,6 +46,8 @@ public:
 
 	HTTP_API virtual bool SetResponseBodyReceiveStream(TSharedRef<FArchive> Stream) override;
 
+	HTTP_API virtual float GetElapsedTime() const override;
+
 protected:
 	/**
 	 * Check if this request is valid or allowed, before actually process the request
@@ -92,6 +94,8 @@ protected:
 	HTTP_API bool OpenRequestPayloadDefaultImpl();
 	HTTP_API void CloseRequestPayloadDefaultImpl();
 
+	HTTP_API void LogResponse(const TSharedPtr<IHttpResponse>& InResponse);
+
 protected:
 	/** Current status of request being processed */
 	EHttpRequestStatus::Type CompletionStatus = EHttpRequestStatus::NotStarted;
@@ -155,4 +159,7 @@ protected:
 
 	/** Payload to use with the request. Typically for POST, PUT, or PATCH */
 	TUniquePtr<FRequestPayload> RequestPayload;
+
+	/** Total elapsed time in seconds since the start of the request */
+	float ElapsedTime = 0.0f;
 };
