Reduce log noise in FAppleHttp for Verbose and VeryVerbose

diff --git a/Engine/Source/Runtime/Online/HTTP/Private/Apple/AppleHttp.cpp b/Engine/Source/Runtime/Online/HTTP/Private/Apple/AppleHttp.cpp
--- a/Engine/Source/Runtime/Online/HTTP/Private/Apple/AppleHttp.cpp
+++ b/Engine/Source/Runtime/Online/HTTP/Private/Apple/AppleHttp.cpp
@@ -224,6 +224,8 @@ static bool ShouldReadHeadersWhenComplete(const FString& Url)
 	{
 		bAnyHttpActivity = true;
 
+		Request->ConnectTime = FPlatformTime::Seconds() - Request->StartProcessTime;
+
 		Request->StartActivityTimeoutTimer();
 	}
 
@@ -237,7 +239,7 @@ static bool ShouldReadHeadersWhenComplete(const FString& Url)
 	{
 		return;
 	}
-	UE_LOG(LogHttp, Verbose, TEXT("URLSession:task:didSendBodyData:totalBytesSent:totalBytesExpectedToSend: totalBytesSent = %lld, totalBytesSent = %lld: %p"), totalBytesSent, totalBytesExpectedToSend, self);
+	UE_LOG(LogHttp, VeryVerbose, TEXT("URLSession:task:didSendBodyData:totalBytesSent:totalBytesExpectedToSend: totalBytesSent = %lld, totalBytesSent = %lld: %p"), totalBytesSent, totalBytesExpectedToSend, self);
 	self.BytesWritten = totalBytesSent;
 }
 
@@ -249,8 +251,6 @@ static bool ShouldReadHeadersWhenComplete(const FString& Url)
 		return;
 	}
 
-	UE_LOG(LogHttp, Verbose, TEXT("URLSession:dataTask:didReceiveResponse:completionHandler"));
-	
 	self.Response = response;
 
 	int32 StatusCode = [self GetStatusCode];
@@ -268,7 +268,7 @@ static bool ShouldReadHeadersWhenComplete(const FString& Url)
 	{
 		Payload.Empty(ExpectedResponseLength);
 	}
-	UE_LOG(LogHttp, Verbose, TEXT("URLSession:dataTask:didReceiveResponse:completionHandler: expectedContentLength = %lld. Length = %llu: %p"), ExpectedResponseLength, Payload.Max(), self);
+	UE_LOG(LogHttp, VeryVerbose, TEXT("URLSession:dataTask:didReceiveResponse:completionHandler: expectedContentLength = %lld. Length = %llu: %p"), ExpectedResponseLength, Payload.Max(), self);
 	completionHandler(NSURLSessionResponseAllow);
 }
 
@@ -303,7 +303,7 @@ static bool ShouldReadHeadersWhenComplete(const FString& Url)
 	}
 	// Keep BytesReceived as a separated value to avoid concurrent accesses to Payload
 	self.BytesReceived += NewBytesReceived;
-	UE_LOG(LogHttp, Verbose, TEXT("URLSession:dataTask:didReceiveData with %llu bytes. After Append, Payload Length = %llu: %p"), NewBytesReceived, self.BytesReceived, self);
+	UE_LOG(LogHttp, VeryVerbose, TEXT("URLSession:dataTask:didReceiveData with %llu bytes. After Append, Payload Length = %llu: %p"), NewBytesReceived, self.BytesReceived, self);
 	
 	NewAppleHttpEventDelegate.ExecuteIfBound();
 }
@@ -320,7 +320,7 @@ static bool ShouldReadHeadersWhenComplete(const FString& Url)
 	self.RequestStatus = EHttpRequestStatus::Failed;
 	if (error == nil)
 	{
-		UE_LOG(LogHttp, Verbose, TEXT("URLSession:task:didCompleteWithError. Http request succeeded: %p"), self);
+		UE_LOG(LogHttp, VeryVerbose, TEXT("URLSession:task:didCompleteWithError. Http request succeeded: %p"), self);
 		self.RequestStatus = EHttpRequestStatus::Succeeded;
 	}
 	else
@@ -544,7 +544,6 @@ FAppleHttpRequest::FAppleHttpRequest(NSURLSession* InSession)
 {
 	bUsePlatformActivityTimeout = false;
 	
-	UE_LOG(LogHttp, Verbose, TEXT("FAppleHttpRequest::FAppleHttpRequest()"));
 	Request = [[NSMutableURLRequest alloc] init];
 
 	// Disable cache to mimic WinInet behavior
@@ -560,7 +559,6 @@ FAppleHttpRequest::FAppleHttpRequest(NSURLSession* InSession)
 
 FAppleHttpRequest::~FAppleHttpRequest()
 {
-	UE_LOG(LogHttp, Verbose, TEXT("FAppleHttpRequest::~FAppleHttpRequest()"));
 	PostProcess();
 	[Request release];
     [Session release];
@@ -573,12 +571,10 @@ FString FAppleHttpRequest::GetURL() const
 	if (URL != nullptr)
 	{
 		FString ConvertedURL(URL.absoluteString);
-		UE_LOG(LogHttp, Verbose, TEXT("FAppleHttpRequest::GetURL() - %s"), *ConvertedURL);
 		return ConvertedURL;
 	}
 	else
 	{
-		UE_LOG(LogHttp, Verbose, TEXT("FAppleHttpRequest::GetURL() - NULL"));
 		return FString();
 	}
 }
@@ -594,7 +590,6 @@ FString FAppleHttpRequest::GetHeader(const FString& HeaderName) const
 {
 	SCOPED_AUTORELEASE_POOL;
 	FString Header([Request valueForHTTPHeaderField:HeaderName.GetNSString()]);
-	UE_LOG(LogHttp, Verbose, TEXT("FAppleHttpRequest::GetHeader() - %s"), *Header);
 	return Header;
 }
 
@@ -627,7 +622,6 @@ void FAppleHttpRequest::AppendToHeader(const FString& HeaderName, const FString&
 TArray<FString> FAppleHttpRequest::GetAllHeaders() const
 {
 	SCOPED_AUTORELEASE_POOL;
-	UE_LOG(LogHttp, Verbose, TEXT("FAppleHttpRequest::GetAllHeaders()"));
 	NSDictionary* Headers = Request.allHTTPHeaderFields;
 	TArray<FString> Result;
 	Result.Reserve(Headers.count);
@@ -635,8 +629,6 @@ TArray<FString> FAppleHttpRequest::GetAllHeaders() const
 	{
 		FString ConvertedValue(Headers[Key]);
 		FString ConvertedKey(Key);
-		UE_LOG(LogHttp, Verbose, TEXT("Header= %s, Key= %s"), *ConvertedValue, *ConvertedKey);
-
 		Result.Add( FString::Printf( TEXT("%s: %s"), *ConvertedKey, *ConvertedValue ) );
 	}
 	return Result;
@@ -644,7 +636,6 @@ TArray<FString> FAppleHttpRequest::GetAllHeaders() const
 
 const TArray<uint8>& FAppleHttpRequest::GetContent() const
 {
-	UE_LOG(LogHttp, Verbose, TEXT("FAppleHttpRequest::GetContent()"));
 	StorageForGetContent.Empty();
 	if (StreamedContentSource.IsType<FNoStreamSource>())
 	{
@@ -667,7 +658,6 @@ void FAppleHttpRequest::SetContent(const TArray<uint8>& ContentPayload)
 		return;
 	}
 	
-	UE_LOG(LogHttp, Verbose, TEXT("FAppleHttpRequest::SetContent()"));
 	StreamedContentSource.Emplace<FNoStreamSource>();
 	Request.HTTPBody = [NSData dataWithBytes:ContentPayload.GetData() length:ContentPayload.Num()];
 	ContentBytesLength = ContentPayload.Num();
@@ -681,7 +671,7 @@ void FAppleHttpRequest::SetContent(TArray<uint8>&& ContentPayload)
 		return;
 	}
 
-	UE_LOG(LogHttp, Verbose, TEXT("FAppleHttpRequest::SetContent()"));
+	UE_LOG(LogHttp, Verbose, TEXT("FAppleHttpRequest::SetContent(). Payload size %d"), ContentPayload.Num());
 
 	StreamedContentSource.Emplace<FNoStreamSource>();
 	// We cannot use NSData dataWithBytesNoCopy:length:freeWhenDone: and keep the data in this instance because we don't have control
@@ -696,13 +686,11 @@ void FAppleHttpRequest::SetContent(TArray<uint8>&& ContentPayload)
 FString FAppleHttpRequest::GetContentType() const
 {
 	FString ContentType = GetHeader(TEXT("Content-Type"));
-	UE_LOG(LogHttp, Verbose, TEXT("FAppleHttpRequest::GetContentType() - %s"), *ContentType);
 	return ContentType;
 }
 
 uint64 FAppleHttpRequest::GetContentLength() const
 {
-	UE_LOG(LogHttp, Verbose, TEXT("FAppleHttpRequest::GetContentLength() - %i"), ContentBytesLength);
 	return ContentBytesLength;
 }
 
@@ -758,6 +746,8 @@ bool FAppleHttpRequest::SetContentAsStreamedFile(const FString& Filename)
 bool FAppleHttpRequest::SetContentFromStream(TSharedRef<FArchive> Stream)
 {
 	SCOPED_AUTORELEASE_POOL;
+	UE_LOG(LogHttp, Verbose, TEXT("FAppleHttpRequest::SetContentFromStream() - %p"), &Stream.Get());
+
 	if (CompletionStatus == EHttpRequestStatus::Processing)
 	{
 		UE_LOG(LogHttp, Warning, TEXT("FAppleHttpRequest::SetContentFromStream() - attempted to set content on a request that is inflight"));
@@ -774,7 +764,6 @@ bool FAppleHttpRequest::SetContentFromStream(TSharedRef<FArchive> Stream)
 FString FAppleHttpRequest::GetVerb() const
 {
 	FString ConvertedVerb(Request.HTTPMethod);
-	UE_LOG(LogHttp, Verbose, TEXT("FAppleHttpRequest::GetVerb() - %s"), *ConvertedVerb);
 	return ConvertedVerb;
 }
 
@@ -788,13 +777,14 @@ void FAppleHttpRequest::SetVerb(const FString& Verb)
 bool FAppleHttpRequest::ProcessRequest()
 {
 	SCOPED_AUTORELEASE_POOL;
-	UE_LOG(LogHttp, Verbose, TEXT("FAppleHttpRequest::ProcessRequest()"));
 
 	if (!PreProcess())
 	{
 		return false;
 	}
 	
+	StartProcessTime = FPlatformTime::Seconds();
+
 	AppleHTTPRequestInternal::UpdateConfigFromCVar();
 
 	return true;
@@ -821,19 +811,21 @@ struct FAppleHttpRequest::FAppleHttpStreamFactory
 bool FAppleHttpRequest::SetupRequest()
 {
 	SCOPED_AUTORELEASE_POOL;
-	UE_LOG(LogHttp, Verbose, TEXT("FAppleHttpRequest::SetupRequest()"));
 
 	// set the content-length and user-agent (it is possible that the OS ignores this value)
 	if(GetContentLength() > 0)
 	{
+		UE_LOG(LogHttp, VeryVerbose, TEXT("Setting content length: %d"), GetContentLength());
 		[Request setValue:[NSString stringWithFormat:@"%llu", GetContentLength()] forHTTPHeaderField:@"Content-Length"];
 	}
 
-	const FString UserAgent = GetHeader("User-Agent");
-	if(UserAgent.IsEmpty())
+	const FString CurrentUserAgent = GetHeader("User-Agent");
+	if(CurrentUserAgent.IsEmpty())
 	{
-		NSString* Tag = FPlatformHttp::GetDefaultUserAgent().GetNSString();
-		[Request setValue:Tag forHTTPHeaderField:@"User-Agent"];
+		FString DefaultUserAgent = FPlatformHttp::GetDefaultUserAgent();
+		UE_LOG(LogHttp, Verbose, TEXT("Setting default User-Agent %s"), *DefaultUserAgent);
+		NSString* UserAgent = DefaultUserAgent.GetNSString();
+		[Request setValue:UserAgent forHTTPHeaderField:@"User-Agent"];
 	}
 
 	PostProcess();
@@ -886,7 +878,7 @@ bool FAppleHttpRequest::SetupRequest()
 		FHttpModule::Get().GetHttpManager().AddThreadedRequest(SharedThis(this));
 
 		[[Task retain] resume];
-		UE_LOG(LogHttp, Verbose, TEXT("[NSURLSessionTask resume]"));
+		UE_LOG(LogHttp, Verbose, TEXT("Task started %p"), this);
 		return true;
 	}
 	else
@@ -900,16 +892,21 @@ bool FAppleHttpRequest::SetupRequest()
 
 void FAppleHttpRequest::FinishRequest()
 {
-	UE_LOG(LogHttp, Verbose, TEXT("FAppleHttpRequest::FinishRequest()"));
-
 	PostProcess();
 
 	TSharedPtr<FAppleHttpResponse> Response = StaticCastSharedPtr<FAppleHttpResponse>(ResponseCommon);
 	bool bSucceeded = (Response && Response->GetStatusFromDelegate() == EHttpRequestStatus::Succeeded);
-	UE_LOG(LogHttp, Verbose, TEXT("Request %s"), bSucceeded ? TEXT("succeeded") : TEXT("failed"));
 	SetStatus(bSucceeded ? EHttpRequestStatus::Succeeded : EHttpRequestStatus::Failed);
 
-	if (!bSucceeded)
+	if (bSucceeded)
+	{
+		if (AppleHTTPRequestInternal::ShouldReadHeadersWhenComplete(GetURL()))
+		{
+			BroadcastResponseHeadersReceived();
+		}
+		HandleRequestSucceed(Response);
+	}
+	else
 	{
 		EHttpFailureReason Reason = EHttpFailureReason::Other;
 		if (Response)
@@ -933,26 +930,17 @@ void FAppleHttpRequest::FinishRequest()
 		}
 		SetFailureReason(Reason);
 
+		UE_LOG(LogHttp, Verbose, TEXT("Request failed: %p Reason %s"), this, LexToString(Reason));
 		if (GetFailureReason() == EHttpFailureReason::ConnectionError)
 		{
 			ResponseCommon = nullptr;
 		}
+		OnProcessRequestComplete().ExecuteIfBound(SharedThis(this), ResponseCommon, bSucceeded);
 	}
-	else
-	{
-		if (AppleHTTPRequestInternal::ShouldReadHeadersWhenComplete(GetURL()))
-		{
-			BroadcastResponseHeadersReceived();
-		}
-	}
-
-	OnProcessRequestComplete().ExecuteIfBound(SharedThis(this), ResponseCommon, bSucceeded);
 }
 
 void FAppleHttpRequest::CleanupRequest()
 {
-	UE_LOG(LogHttp, Verbose, TEXT("FAppleHttpRequest::CleanupRequest()"));
-	
 	TSharedPtr<FAppleHttpResponse> Response = StaticCastSharedPtr<FAppleHttpResponse>(ResponseCommon);
 	if (Response != nullptr)
 	{
@@ -1043,14 +1031,12 @@ void FAppleHttpRequest::TickThreadedRequest(float DeltaSeconds)
 FAppleHttpResponse::FAppleHttpResponse(FAppleHttpRequest& InRequest)
 	: FHttpResponseCommon(InRequest)
 {
-	UE_LOG(LogHttp, Verbose, TEXT("FAppleHttpResponse::FAppleHttpResponse()"));
 	ResponseDelegate = [[FAppleHttpResponseDelegate alloc] initWithRequest: InRequest];
+	UE_LOG(LogHttp, VeryVerbose, TEXT("FAppleHttpResponse::FAppleHttpResponse(). Request: %p ResponseDelegate: %p"), &InRequest, ResponseDelegate);
 }
 
 FAppleHttpResponse::~FAppleHttpResponse()
 {
-	UE_LOG(LogHttp, Verbose, TEXT("FAppleHttpResponse::~FAppleHttpResponse()"));
-	
 	[ResponseDelegate release];
 	ResponseDelegate = nil;
 }
@@ -1077,7 +1063,6 @@ FString FAppleHttpResponse::GetHeader(const FString& HeaderName) const
 		SCOPED_AUTORELEASE_POOL;
 		if(NSDictionary* Headers = [ResponseDelegate GetResponseHeaders])
 		{
-			UE_LOG(LogHttp, Verbose, TEXT("FAppleHttpResponse::GetHeader()"));
 			NSHTTPURLResponse* Response = (NSHTTPURLResponse*)ResponseDelegate.Response;
 			NSString* ConvertedHeaderName = HeaderName.GetNSString();
 			return FString([Response.allHeaderFields objectForKey:ConvertedHeaderName]);
@@ -1091,8 +1076,6 @@ FString FAppleHttpResponse::GetHeader(const FString& HeaderName) const
 
 TArray<FString> FAppleHttpResponse::GetAllHeaders() const
 {
-	UE_LOG(LogHttp, Verbose, TEXT("FAppleHttpResponse::GetAllHeaders()"));
-
 	TArray<FString> Result;
 	SCOPED_AUTORELEASE_POOL;
 	if (NSDictionary* Headers = [ResponseDelegate GetResponseHeaders])
@@ -1110,15 +1093,11 @@ TArray<FString> FAppleHttpResponse::GetAllHeaders() const
 
 FString FAppleHttpResponse::GetContentType() const
 {
-	UE_LOG(LogHttp, Verbose, TEXT("FAppleHttpResponse::GetContentType()"));
-
 	return GetHeader( TEXT( "Content-Type" ) );
 }
 
 uint64 FAppleHttpResponse::GetContentLength() const
 {
-	UE_LOG(LogHttp, Verbose, TEXT("FAppleHttpResponse::GetContentLength()"));
-	
 	return ResponseDelegate.Response.expectedContentLength;
 }
 
@@ -1128,17 +1107,11 @@ const TArray<uint8>& FAppleHttpResponse::GetContent() const
 	{
 		UE_LOG(LogHttp, Warning, TEXT("Payload is incomplete. Response still processing. %s"), *GetURL());
 	}
-	else
-	{
-		UE_LOG(LogHttp, Verbose, TEXT("FAppleHttpResponse::GetContent() - Num: %i"), ResponseDelegate->Payload.Num());
-	}
 	return ResponseDelegate->Payload;
 }
 
 FString FAppleHttpResponse::GetContentAsString() const
 {
-	UE_LOG(LogHttp, Verbose, TEXT("FAppleHttpResponse::GetContentAsString()"));
-
 	// Fill in our data.
 	const TArray<uint8>& Payload = GetContent();
 
@@ -1151,8 +1124,6 @@ FString FAppleHttpResponse::GetContentAsString() const
 
 int32 FAppleHttpResponse::GetResponseCode() const
 {
-	UE_LOG(LogHttp, Verbose, TEXT("FAppleHttpResponse::GetResponseCode()"));
-
 	return [ResponseDelegate GetStatusCode];
 }
 
diff --git a/Engine/Source/Runtime/Online/HTTP/Private/Apple/AppleHttp.h b/Engine/Source/Runtime/Online/HTTP/Private/Apple/AppleHttp.h
--- a/Engine/Source/Runtime/Online/HTTP/Private/Apple/AppleHttp.h
+++ b/Engine/Source/Runtime/Online/HTTP/Private/Apple/AppleHttp.h
@@ -68,6 +68,8 @@ public:
 	virtual ~FAppleHttpRequest();
 
 PACKAGE_SCOPE:
+	using FHttpRequestCommon::StartProcessTime;
+	using FHttpRequestCommon::ConnectTime;
 	using FHttpRequestCommon::BroadcastResponseHeadersReceived;
 	using FHttpRequestCommon::StartActivityTimeoutTimer;
 	using FHttpRequestCommon::ResetActivityTimeoutTimer;
