Clean the cvar to trigger header received callback before http request complete;
Move the url patterns for logging response into http manager.

diff --git a/Engine/Source/Runtime/Online/HTTP/Private/Apple/AppleHttp.cpp b/Engine/Source/Runtime/Online/HTTP/Private/Apple/AppleHttp.cpp
--- a/Engine/Source/Runtime/Online/HTTP/Private/Apple/AppleHttp.cpp
+++ b/Engine/Source/Runtime/Online/HTTP/Private/Apple/AppleHttp.cpp
@@ -11,47 +11,6 @@
 #include "Misc/App.h"
 #include "Misc/Base64.h"
 
-// It should be safe to read headers early, add the CVar here just in case
-TAutoConsoleVariable<FString> CVarHttpUrlsToReadHeadersWhenComplete(
-	TEXT("http.UrlsToReadHeadersWhenComplete"),
-	TEXT(""),
-	TEXT("List of urls to only read headers when complete the http request\"www.epicgames.com,www.unrealengine.com,...\"")
-);
-
-namespace AppleHTTPRequestInternal
-{
-
-static bool bUpdatedCVarHttpUrlsToReadHeadersWhenComplete = true;
-static TArray<FString> UrlsToReadHeadersWhenComplete;
-
-static void UpdateConfigFromCVar()
-{
-	UE_CALL_ONCE([] {
-		CVarHttpUrlsToReadHeadersWhenComplete.AsVariable()->OnChangedDelegate().AddLambda([](IConsoleVariable* CVar) {
-			bUpdatedCVarHttpUrlsToReadHeadersWhenComplete = true;
-		});
-		bUpdatedCVarHttpUrlsToReadHeadersWhenComplete = true;
-	});
-	if (bUpdatedCVarHttpUrlsToReadHeadersWhenComplete)
-	{
-		CVarHttpUrlsToReadHeadersWhenComplete.GetValueOnAnyThread().ParseIntoArray(UrlsToReadHeadersWhenComplete, TEXT(","));
-		bUpdatedCVarHttpUrlsToReadHeadersWhenComplete = false;
-	}
-}
-
-static bool ShouldReadHeadersWhenComplete(const FString& Url)
-{
-	for (const FString& UrlToReadHeadersWhenComplete : UrlsToReadHeadersWhenComplete)
-	{
-		if (Url.StartsWith(UrlToReadHeadersWhenComplete))
-		{
-			return true;
-		}
-	}
-	return false;
-}
-}
-
 /**
  * Class to hold data from delegate implementation notifications.
  */
@@ -194,19 +153,16 @@ static bool ShouldReadHeadersWhenComplete(const FString& Url)
 {
 	if (TSharedPtr<FAppleHttpRequest> Request = SourceRequest.Pin())
 	{
-		if (!AppleHTTPRequestInternal::ShouldReadHeadersWhenComplete(Request->GetURL()))
+		if (Request->GetDelegateThreadPolicy() == EHttpRequestDelegateThreadPolicy::CompleteOnHttpThread)
+		{
+			Request->BroadcastResponseHeadersReceived();
+		}
+		else if (Request->OnHeaderReceived().IsBound())
 		{
-			if (Request->GetDelegateThreadPolicy() == EHttpRequestDelegateThreadPolicy::CompleteOnHttpThread)
+			FHttpModule::Get().GetHttpManager().AddGameThreadTask([Request]()
 			{
 				Request->BroadcastResponseHeadersReceived();
-			}
-			else if (Request->OnHeaderReceived().IsBound())
-			{
-				FHttpModule::Get().GetHttpManager().AddGameThreadTask([Request]()
-				{
-					Request->BroadcastResponseHeadersReceived();
-				});
-			}
+			});
 		}
 	}
 }
@@ -1053,24 +1009,16 @@ void FAppleHttpResponse::CleanSharedObjects()
 
 FString FAppleHttpResponse::GetHeader(const FString& HeaderName) const
 {
-	if (AppleHTTPRequestInternal::ShouldReadHeadersWhenComplete(GetURL()) && !IsReady())
+	SCOPED_AUTORELEASE_POOL;
+	if(NSDictionary* Headers = [ResponseDelegate GetResponseHeaders])
 	{
-		UE_LOG(LogHttp, Warning, TEXT("Can't get header [%s]. Response still processing for %s."), *HeaderName, *GetURL());
-		return FString();
+		NSHTTPURLResponse* Response = (NSHTTPURLResponse*)ResponseDelegate.Response;
+		NSString* ConvertedHeaderName = HeaderName.GetNSString();
+		return FString([Response.allHeaderFields objectForKey:ConvertedHeaderName]);
 	}
 	else
 	{
-		SCOPED_AUTORELEASE_POOL;
-		if(NSDictionary* Headers = [ResponseDelegate GetResponseHeaders])
-		{
-			NSHTTPURLResponse* Response = (NSHTTPURLResponse*)ResponseDelegate.Response;
-			NSString* ConvertedHeaderName = HeaderName.GetNSString();
-			return FString([Response.allHeaderFields objectForKey:ConvertedHeaderName]);
-		}
-		else
-		{
-			return FString();
-		}
+		return FString();
 	}
 }
 
diff --git a/Engine/Source/Runtime/Online/HTTP/Private/GenericPlatform/HttpRequestCommon.cpp b/Engine/Source/Runtime/Online/HTTP/Private/GenericPlatform/HttpRequestCommon.cpp
--- a/Engine/Source/Runtime/Online/HTTP/Private/GenericPlatform/HttpRequestCommon.cpp
+++ b/Engine/Source/Runtime/Online/HTTP/Private/GenericPlatform/HttpRequestCommon.cpp
@@ -13,13 +13,6 @@
 namespace UE::HttpRequestCommon::Private
 {
 
-TAutoConsoleVariable<FString> CVarHttpUrlPatternsToLogResponse(
-	TEXT("http.UrlPatternsToLogResponse"),
-	TEXT(""),
-	TEXT("List of url patterns to log headers and json content: \"epicgames.com,unrealengine.com,...\""),
-	ECVF_SaveForNextBoot
-);
-
 TAutoConsoleVariable<bool> CVarHttpLogJsonResponseOnly(
 	TEXT("http.LogJsonResponseOnly"),
 	true,
@@ -27,35 +20,6 @@ TAutoConsoleVariable<bool> CVarHttpLogJsonResponseOnly(
 	ECVF_SaveForNextBoot
 );
 
-bool ShouldLogResponse(FStringView Url)
-{
-	static std::atomic<bool> bUpdatedCVarHttpUrlPatternsToLogResponse = true;
-	UE_CALL_ONCE([] {
-		CVarHttpUrlPatternsToLogResponse.AsVariable()->OnChangedDelegate().AddLambda([](IConsoleVariable* CVar) {
-			bUpdatedCVarHttpUrlPatternsToLogResponse = true;
-		});
-	});
-
-	static TArray<FString> UrlPatternsToLogResponse;
-	static FCriticalSection UrlPatternsToLogResponseCriticalSection;
-	const FScopeLock CacheLock(&UrlPatternsToLogResponseCriticalSection);
-	if (bUpdatedCVarHttpUrlPatternsToLogResponse)
-	{
-		CVarHttpUrlPatternsToLogResponse.GetValueOnAnyThread().ParseIntoArray(UrlPatternsToLogResponse, TEXT(","));
-		bUpdatedCVarHttpUrlPatternsToLogResponse = false;
-	}
-
-	for (const FString& UrlPatternToLogResponse : UrlPatternsToLogResponse)
-	{
-		if (Url.Contains(UrlPatternToLogResponse))
-		{
-			return true;
-		}
-	}
-
-	return false;
-}
-
 }
 
 FHttpRequestCommon::FHttpRequestCommon()
@@ -616,7 +580,7 @@ void FHttpRequestCommon::CloseRequestPayloadDefaultImpl()
 
 void FHttpRequestCommon::LogResponse(const TSharedPtr<IHttpResponse>& InResponse)
 {
-	bool bShouldLogResponse = UE::HttpRequestCommon::Private::ShouldLogResponse(GetURL());
+	bool bShouldLogResponse = FHttpModule::Get().GetHttpManager().ShouldLogResponse(GetURL());
 	UE_HTTP_LOG_RESPONSE_PRIVATE(bShouldLogResponse, TEXT("%p %s %s completed with code %d after %.2fs. Content length: %ld"), this, *GetVerb(), *GetURL(), InResponse->GetResponseCode(), ElapsedTime, InResponse->GetContentLength());
 
 	TArray<FString> AllHeaders = InResponse->GetAllHeaders();
diff --git a/Engine/Source/Runtime/Online/HTTP/Private/HttpManager.cpp b/Engine/Source/Runtime/Online/HTTP/Private/HttpManager.cpp
--- a/Engine/Source/Runtime/Online/HTTP/Private/HttpManager.cpp
+++ b/Engine/Source/Runtime/Online/HTTP/Private/HttpManager.cpp
@@ -25,10 +25,14 @@ TAutoConsoleVariable<int32> CVarHttpEventLoopEnableChance(
 	ECVF_SaveForNextBoot
 );
 
-// FHttpManager
+TAutoConsoleVariable<FString> CVarHttpUrlPatternsToLogResponse(
+	TEXT("http.UrlPatternsToLogResponse"),
+	TEXT(""),
+	TEXT("List of url patterns to log headers and json content: \"epicgames.com,unrealengine.com,...\""),
+	ECVF_SaveForNextBoot
+);
 
-FCriticalSection FHttpManager::RequestLock;
-FCriticalSection FHttpManager::CompletedRequestLock;
+// FHttpManager
 
 const TCHAR* LexToString(const EHttpFlushReason& FlushReason)
 {
@@ -88,10 +92,15 @@ void FHttpManager::Initialize()
 	}
 
 	UpdateConfigs();
+
+	UpdateUrlPatternsToLogResponse(CVarHttpUrlPatternsToLogResponse.AsVariable());
+	CVarHttpUrlPatternsToLogResponse.AsVariable()->OnChangedDelegate().AddRaw(this, &FHttpManager::UpdateUrlPatternsToLogResponse);
 }
 
 void FHttpManager::Shutdown()
 {
+	CVarHttpUrlPatternsToLogResponse.AsVariable()->OnChangedDelegate().Clear();
+
 	{
 		FScopeLock ScopeLock(&RequestLock);
 
@@ -599,3 +608,25 @@ void FHttpManager::RecordMaxTimeToWaitInQueue(float Duration)
 {
 	HttpStats.MaxTimeToWaitInQueue = FGenericPlatformMath::Max(Duration, HttpStats.MaxTimeToWaitInQueue);
 }
+
+void FHttpManager::UpdateUrlPatternsToLogResponse(IConsoleVariable* CVar)
+{
+	const FScopeLock CacheLock(&UrlPatternsToLogResponseCriticalSection);
+	const FString UrlPatternsToLogResponseStr = CVar->AsVariable()->GetString();
+	UrlPatternsToLogResponseStr.ParseIntoArray(UrlPatternsToLogResponse, TEXT(","));
+}
+
+bool FHttpManager::ShouldLogResponse(FStringView Url)
+{
+	const FScopeLock CacheLock(&UrlPatternsToLogResponseCriticalSection);
+	for (const FString& UrlPatternToLogResponse : UrlPatternsToLogResponse)
+	{
+		if (Url.Contains(UrlPatternToLogResponse))
+		{
+			return true;
+		}
+	}
+
+	return false;
+}
+
diff --git a/Engine/Source/Runtime/Online/HTTP/Public/HttpManager.h b/Engine/Source/Runtime/Online/HTTP/Public/HttpManager.h
--- a/Engine/Source/Runtime/Online/HTTP/Public/HttpManager.h
+++ b/Engine/Source/Runtime/Online/HTTP/Public/HttpManager.h
@@ -285,6 +285,8 @@ protected:
 
 	HTTP_API bool HasAnyBoundDelegate() const;
 
+	HTTP_API void UpdateUrlPatternsToLogResponse(IConsoleVariable* CVar);
+
 protected:
 	/** List of Http requests that are actively being processed */
 	TArray<FHttpRequestRef> Requests;
@@ -340,13 +342,18 @@ protected:
 
 	bool bUseEventLoop = true;
 
+	TArray<FString> UrlPatternsToLogResponse;
+
 PACKAGE_SCOPE:
 
 	/** Used to lock access to add/remove/find requests */
-	static FCriticalSection RequestLock;
+	mutable FCriticalSection RequestLock;
 
 	/** Used to lock access to get completed requests */
-	static FCriticalSection CompletedRequestLock;
+	FCriticalSection CompletedRequestLock;
+
+	/** Used to lock access to url patterns for logging */
+	FCriticalSection UrlPatternsToLogResponseCriticalSection;
 
 	/**
 	 * Broadcast that a non-threaded HTTP request is complete.
@@ -368,4 +375,6 @@ PACKAGE_SCOPE:
 
 	/** Record the time to wait in queue, to have a general idea how long the client usually wait before actually starting, to adjust the requests */
 	HTTP_API void RecordMaxTimeToWaitInQueue(float Duration);
+
+	HTTP_API bool ShouldLogResponse(FStringView Url);
 };
