[Backout] - CL33670189
Original CL Desc
-----------------------------------------------------------------
Clean the cvar to trigger header received callback before http request complete;
Move the url patterns for logging response into http manager.

diff --git a/Engine/Source/Runtime/Online/HTTP/Private/Apple/AppleHttp.cpp b/Engine/Source/Runtime/Online/HTTP/Private/Apple/AppleHttp.cpp
--- a/Engine/Source/Runtime/Online/HTTP/Private/Apple/AppleHttp.cpp
+++ b/Engine/Source/Runtime/Online/HTTP/Private/Apple/AppleHttp.cpp
@@ -11,6 +11,47 @@
 #include "Misc/App.h"
 #include "Misc/Base64.h"
 
+// It should be safe to read headers early, add the CVar here just in case
+TAutoConsoleVariable<FString> CVarHttpUrlsToReadHeadersWhenComplete(
+	TEXT("http.UrlsToReadHeadersWhenComplete"),
+	TEXT(""),
+	TEXT("List of urls to only read headers when complete the http request\"www.epicgames.com,www.unrealengine.com,...\"")
+);
+
+namespace AppleHTTPRequestInternal
+{
+
+static bool bUpdatedCVarHttpUrlsToReadHeadersWhenComplete = true;
+static TArray<FString> UrlsToReadHeadersWhenComplete;
+
+static void UpdateConfigFromCVar()
+{
+	UE_CALL_ONCE([] {
+		CVarHttpUrlsToReadHeadersWhenComplete.AsVariable()->OnChangedDelegate().AddLambda([](IConsoleVariable* CVar) {
+			bUpdatedCVarHttpUrlsToReadHeadersWhenComplete = true;
+		});
+		bUpdatedCVarHttpUrlsToReadHeadersWhenComplete = true;
+	});
+	if (bUpdatedCVarHttpUrlsToReadHeadersWhenComplete)
+	{
+		CVarHttpUrlsToReadHeadersWhenComplete.GetValueOnAnyThread().ParseIntoArray(UrlsToReadHeadersWhenComplete, TEXT(","));
+		bUpdatedCVarHttpUrlsToReadHeadersWhenComplete = false;
+	}
+}
+
+static bool ShouldReadHeadersWhenComplete(const FString& Url)
+{
+	for (const FString& UrlToReadHeadersWhenComplete : UrlsToReadHeadersWhenComplete)
+	{
+		if (Url.StartsWith(UrlToReadHeadersWhenComplete))
+		{
+			return true;
+		}
+	}
+	return false;
+}
+}
+
 /**
  * Class to hold data from delegate implementation notifications.
  */
@@ -153,16 +194,19 @@
 {
 	if (TSharedPtr<FAppleHttpRequest> Request = SourceRequest.Pin())
 	{
-		if (Request->GetDelegateThreadPolicy() == EHttpRequestDelegateThreadPolicy::CompleteOnHttpThread)
-		{
-			Request->BroadcastResponseHeadersReceived();
-		}
-		else if (Request->OnHeaderReceived().IsBound())
+		if (!AppleHTTPRequestInternal::ShouldReadHeadersWhenComplete(Request->GetURL()))
 		{
-			FHttpModule::Get().GetHttpManager().AddGameThreadTask([Request]()
+			if (Request->GetDelegateThreadPolicy() == EHttpRequestDelegateThreadPolicy::CompleteOnHttpThread)
 			{
 				Request->BroadcastResponseHeadersReceived();
-			});
+			}
+			else if (Request->OnHeaderReceived().IsBound())
+			{
+				FHttpModule::Get().GetHttpManager().AddGameThreadTask([Request]()
+				{
+					Request->BroadcastResponseHeadersReceived();
+				});
+			}
 		}
 	}
 }
@@ -1009,16 +1053,24 @@ void FAppleHttpResponse::CleanSharedObjects()
 
 FString FAppleHttpResponse::GetHeader(const FString& HeaderName) const
 {
-	SCOPED_AUTORELEASE_POOL;
-	if(NSDictionary* Headers = [ResponseDelegate GetResponseHeaders])
+	if (AppleHTTPRequestInternal::ShouldReadHeadersWhenComplete(GetURL()) && !IsReady())
 	{
-		NSHTTPURLResponse* Response = (NSHTTPURLResponse*)ResponseDelegate.Response;
-		NSString* ConvertedHeaderName = HeaderName.GetNSString();
-		return FString([Response.allHeaderFields objectForKey:ConvertedHeaderName]);
+		UE_LOG(LogHttp, Warning, TEXT("Can't get header [%s]. Response still processing for %s."), *HeaderName, *GetURL());
+		return FString();
 	}
 	else
 	{
-		return FString();
+		SCOPED_AUTORELEASE_POOL;
+		if(NSDictionary* Headers = [ResponseDelegate GetResponseHeaders])
+		{
+			NSHTTPURLResponse* Response = (NSHTTPURLResponse*)ResponseDelegate.Response;
+			NSString* ConvertedHeaderName = HeaderName.GetNSString();
+			return FString([Response.allHeaderFields objectForKey:ConvertedHeaderName]);
+		}
+		else
+		{
+			return FString();
+		}
 	}
 }
 
diff --git a/Engine/Source/Runtime/Online/HTTP/Private/GenericPlatform/HttpRequestCommon.cpp b/Engine/Source/Runtime/Online/HTTP/Private/GenericPlatform/HttpRequestCommon.cpp
--- a/Engine/Source/Runtime/Online/HTTP/Private/GenericPlatform/HttpRequestCommon.cpp
+++ b/Engine/Source/Runtime/Online/HTTP/Private/GenericPlatform/HttpRequestCommon.cpp
@@ -13,6 +13,13 @@
 namespace UE::HttpRequestCommon::Private
 {
 
+TAutoConsoleVariable<FString> CVarHttpUrlPatternsToLogResponse(
+	TEXT("http.UrlPatternsToLogResponse"),
+	TEXT(""),
+	TEXT("List of url patterns to log headers and json content: \"epicgames.com,unrealengine.com,...\""),
+	ECVF_SaveForNextBoot
+);
+
 TAutoConsoleVariable<bool> CVarHttpLogJsonResponseOnly(
 	TEXT("http.LogJsonResponseOnly"),
 	true,
@@ -20,6 +27,35 @@ TAutoConsoleVariable<bool> CVarHttpLogJsonResponseOnly(
 	ECVF_SaveForNextBoot
 );
 
+bool ShouldLogResponse(FStringView Url)
+{
+	static std::atomic<bool> bUpdatedCVarHttpUrlPatternsToLogResponse = true;
+	UE_CALL_ONCE([] {
+		CVarHttpUrlPatternsToLogResponse.AsVariable()->OnChangedDelegate().AddLambda([](IConsoleVariable* CVar) {
+			bUpdatedCVarHttpUrlPatternsToLogResponse = true;
+		});
+	});
+
+	static TArray<FString> UrlPatternsToLogResponse;
+	static FCriticalSection UrlPatternsToLogResponseCriticalSection;
+	const FScopeLock CacheLock(&UrlPatternsToLogResponseCriticalSection);
+	if (bUpdatedCVarHttpUrlPatternsToLogResponse)
+	{
+		CVarHttpUrlPatternsToLogResponse.GetValueOnAnyThread().ParseIntoArray(UrlPatternsToLogResponse, TEXT(","));
+		bUpdatedCVarHttpUrlPatternsToLogResponse = false;
+	}
+
+	for (const FString& UrlPatternToLogResponse : UrlPatternsToLogResponse)
+	{
+		if (Url.Contains(UrlPatternToLogResponse))
+		{
+			return true;
+		}
+	}
+
+	return false;
+}
+
 }
 
 FHttpRequestCommon::FHttpRequestCommon()
@@ -580,7 +616,7 @@ void FHttpRequestCommon::CloseRequestPayloadDefaultImpl()
 
 void FHttpRequestCommon::LogResponse(const TSharedPtr<IHttpResponse>& InResponse)
 {
-	bool bShouldLogResponse = FHttpModule::Get().GetHttpManager().ShouldLogResponse(GetURL());
+	bool bShouldLogResponse = UE::HttpRequestCommon::Private::ShouldLogResponse(GetURL());
 	UE_HTTP_LOG_RESPONSE_PRIVATE(bShouldLogResponse, TEXT("%p %s %s completed with code %d after %.2fs. Content length: %ld"), this, *GetVerb(), *GetURL(), InResponse->GetResponseCode(), ElapsedTime, InResponse->GetContentLength());
 
 	TArray<FString> AllHeaders = InResponse->GetAllHeaders();
diff --git a/Engine/Source/Runtime/Online/HTTP/Private/HttpManager.cpp b/Engine/Source/Runtime/Online/HTTP/Private/HttpManager.cpp
--- a/Engine/Source/Runtime/Online/HTTP/Private/HttpManager.cpp
+++ b/Engine/Source/Runtime/Online/HTTP/Private/HttpManager.cpp
@@ -25,15 +25,11 @@ TAutoConsoleVariable<int32> CVarHttpEventLoopEnableChance(
 	ECVF_SaveForNextBoot
 );
 
-TAutoConsoleVariable<FString> CVarHttpUrlPatternsToLogResponse(
-	TEXT("http.UrlPatternsToLogResponse"),
-	TEXT(""),
-	TEXT("List of url patterns to log headers and json content: \"epicgames.com,unrealengine.com,...\""),
-	ECVF_SaveForNextBoot
-);
-
 // FHttpManager
 
+FCriticalSection FHttpManager::RequestLock;
+FCriticalSection FHttpManager::CompletedRequestLock;
+
 const TCHAR* LexToString(const EHttpFlushReason& FlushReason)
 {
 	switch (FlushReason)
@@ -92,15 +88,10 @@ void FHttpManager::Initialize()
 	}
 
 	UpdateConfigs();
-
-	UpdateUrlPatternsToLogResponse(CVarHttpUrlPatternsToLogResponse.AsVariable());
-	CVarHttpUrlPatternsToLogResponse.AsVariable()->OnChangedDelegate().AddRaw(this, &FHttpManager::UpdateUrlPatternsToLogResponse);
 }
 
 void FHttpManager::Shutdown()
 {
-	CVarHttpUrlPatternsToLogResponse.AsVariable()->OnChangedDelegate().Clear();
-
 	{
 		FScopeLock ScopeLock(&RequestLock);
 
@@ -608,25 +599,3 @@ void FHttpManager::RecordMaxTimeToWaitInQueue(float Duration)
 {
 	HttpStats.MaxTimeToWaitInQueue = FGenericPlatformMath::Max(Duration, HttpStats.MaxTimeToWaitInQueue);
 }
-
-void FHttpManager::UpdateUrlPatternsToLogResponse(IConsoleVariable* CVar)
-{
-	const FScopeLock CacheLock(&UrlPatternsToLogResponseCriticalSection);
-	const FString UrlPatternsToLogResponseStr = CVar->AsVariable()->GetString();
-	UrlPatternsToLogResponseStr.ParseIntoArray(UrlPatternsToLogResponse, TEXT(","));
-}
-
-bool FHttpManager::ShouldLogResponse(FStringView Url)
-{
-	const FScopeLock CacheLock(&UrlPatternsToLogResponseCriticalSection);
-	for (const FString& UrlPatternToLogResponse : UrlPatternsToLogResponse)
-	{
-		if (Url.Contains(UrlPatternToLogResponse))
-		{
-			return true;
-		}
-	}
-
-	return false;
-}
-
diff --git a/Engine/Source/Runtime/Online/HTTP/Public/HttpManager.h b/Engine/Source/Runtime/Online/HTTP/Public/HttpManager.h
--- a/Engine/Source/Runtime/Online/HTTP/Public/HttpManager.h
+++ b/Engine/Source/Runtime/Online/HTTP/Public/HttpManager.h
@@ -285,8 +285,6 @@ protected:
 
 	HTTP_API bool HasAnyBoundDelegate() const;
 
-	HTTP_API void UpdateUrlPatternsToLogResponse(IConsoleVariable* CVar);
-
 protected:
 	/** List of Http requests that are actively being processed */
 	TArray<FHttpRequestRef> Requests;
@@ -342,18 +340,13 @@ protected:
 
 	bool bUseEventLoop = true;
 
-	TArray<FString> UrlPatternsToLogResponse;
-
 PACKAGE_SCOPE:
 
 	/** Used to lock access to add/remove/find requests */
-	mutable FCriticalSection RequestLock;
+	static FCriticalSection RequestLock;
 
 	/** Used to lock access to get completed requests */
-	FCriticalSection CompletedRequestLock;
-
-	/** Used to lock access to url patterns for logging */
-	FCriticalSection UrlPatternsToLogResponseCriticalSection;
+	static FCriticalSection CompletedRequestLock;
 
 	/**
 	 * Broadcast that a non-threaded HTTP request is complete.
@@ -375,6 +368,4 @@ PACKAGE_SCOPE:
 
 	/** Record the time to wait in queue, to have a general idea how long the client usually wait before actually starting, to adjust the requests */
 	HTTP_API void RecordMaxTimeToWaitInQueue(float Duration);
-
-	HTTP_API bool ShouldLogResponse(FStringView Url);
 };
