Make it possible to set the priority of the http thread via the http.threadpriority cvar. This can be changed at runtime (e.g for hotfix purposes)

diff --git a/Engine/Source/Runtime/Online/HTTP/Private/HttpThread.cpp b/Engine/Source/Runtime/Online/HTTP/Private/HttpThread.cpp
--- a/Engine/Source/Runtime/Online/HTTP/Private/HttpThread.cpp
+++ b/Engine/Source/Runtime/Online/HTTP/Private/HttpThread.cpp
@@ -15,6 +15,7 @@
 #include "Http.h"
 #include "PlatformHttp.h"
 #include "Stats/Stats.h"
+#include "Templates/UnrealTemplate.h"
 
 TAutoConsoleVariable<int32> CVarHttpMaxConcurrentRequests(
 	TEXT("http.MaxConcurrentRequests"),
@@ -23,6 +24,35 @@ TAutoConsoleVariable<int32> CVarHttpMaxConcurrentRequests(
 	ECVF_SaveForNextBoot
 );
 
+// Thread priority cvar (settable at runtime)
+// We declare these explicitly rather than just casting the cvar in case the enum changes in future
+const int32 GHttpThreadPriorities[] =
+{
+	EThreadPriority::TPri_Lowest,
+	EThreadPriority::TPri_BelowNormal,
+	EThreadPriority::TPri_SlightlyBelowNormal,
+	EThreadPriority::TPri_Normal,
+	EThreadPriority::TPri_AboveNormal
+};
+
+const TCHAR* GHttpThreadPriortyNames[] =
+{
+	TEXT("TPri_Lowest"),
+	TEXT("TPri_BelowNormal"),
+	TEXT("TPri_SlightlyBelowNormal"),
+	TEXT("TPri_Normal"),
+	TEXT("TPri_AboveNormal")
+};
+
+// Warning: Due to a bug with http module console variables, this cvar is not settable via the console (or via -execcmds). It needs to be set via ini (or via -dpcvars on startup). Hotfixing is supported
+static int32 GHttpThreadPriorityIndex = 3; // EThreadPriority::TPri_Normal
+FAutoConsoleVariableRef CVarHttpThreadPriority(
+	TEXT("http.ThreadPriority"), 
+	GHttpThreadPriorityIndex, 
+	TEXT("Thread priority of the Http Manager thread: 0=Lowest, 1=BelowNormal, 2=SlightlyBelowNormal, 3=Normal, 4=AboveNormal\n")
+	TEXT("Note that this is switchable at runtime"),
+	ECVF_Default);
+
 DECLARE_STATS_GROUP(TEXT("HTTP Thread"), STATGROUP_HTTPThread, STATCAT_Advanced);
 DECLARE_CYCLE_STAT(TEXT("Process"), STAT_HTTPThread_Process, STATGROUP_HTTPThread);
 DECLARE_CYCLE_STAT(TEXT("TickThreadedRequest"), STAT_HTTPThread_TickThreadedRequest, STATGROUP_HTTPThread);
@@ -56,6 +86,7 @@ FHttpThreadBase::FHttpThreadBase()
 	: Thread(nullptr)
 	, bIsSingleThread(false)
 	, bIsStopped(true)
+	, CurrentThreadPriority(EThreadPriority::TPri_Num)
 {
 }
 
@@ -70,10 +101,13 @@ void FHttpThreadBase::StartThread()
 
 	const bool bDisableForkedHTTPThread = FParse::Param(FCommandLine::Get(), TEXT("DisableForkedHTTPThread"));
 
+	// Get the requested thread priority from the cvar
+	CurrentThreadPriority = (EThreadPriority)GHttpThreadPriorities[FMath::Clamp(GHttpThreadPriorityIndex, 0, UE_ARRAY_COUNT(GHttpThreadPriorities)-1)];
+
 	if (FForkProcessHelper::IsForkedMultithreadInstance() && bDisableForkedHTTPThread == false)
 	{
 		// We only create forkable threads on the forked instance since the HTTPManager cannot safely transition from fake to real seamlessly
-		Thread = FForkProcessHelper::CreateForkableThread(this, TEXT("HttpManagerThread"), 128 * 1024, TPri_Normal);
+		Thread = FForkProcessHelper::CreateForkableThread(this, TEXT("HttpManagerThread"), 128 * 1024, CurrentThreadPriority);
 	}
 	else
 	{
@@ -83,12 +117,27 @@ void FHttpThreadBase::StartThread()
 			bIsSingleThread = true;
 		}
 
-		Thread = FRunnableThread::Create(this, TEXT("HttpManagerThread"), 128 * 1024, TPri_Normal);
+		Thread = FRunnableThread::Create(this, TEXT("HttpManagerThread"), 128 * 1024, CurrentThreadPriority);
 	}
 
 	bIsStopped = false;
 }
 
+void FHttpThreadBase::UpdateThreadPriorityIfNeeded()
+{
+	if ( !bIsSingleThread && ensure(!IsInGameThread()))
+	{
+		int32 ThreadPriorityIndex = FMath::Clamp(GHttpThreadPriorityIndex, 0, UE_ARRAY_COUNT(GHttpThreadPriorities) - 1);
+		EThreadPriority DesiredThreadPriority = (EThreadPriority)GHttpThreadPriorities[ThreadPriorityIndex];
+		if (DesiredThreadPriority != CurrentThreadPriority)
+		{
+			UE_LOG(LogHttp, Display, TEXT("Updating HTTP thread priority to %s"), GHttpThreadPriortyNames[ThreadPriorityIndex]);
+			FPlatformProcess::SetThreadPriority(DesiredThreadPriority);
+			CurrentThreadPriority = DesiredThreadPriority;
+		}
+	}
+}
+
 void FHttpThreadBase::StopThread()
 {
 	if (Thread != nullptr)
@@ -323,6 +372,8 @@ void FHttpThreadBase::Process(TArray<IHttpThreadedRequest*>& RequestsToCancel, T
 {
 	SCOPE_CYCLE_COUNTER(STAT_HTTPThread_Process);
 
+	UpdateThreadPriorityIfNeeded();
+
 	ConsumeCanceledRequestsAndNewRequests(RequestsToCancel, RequestsToComplete);
 
 	MoveCompletingRequestsToCompletedRequests(RequestsToComplete);
diff --git a/Engine/Source/Runtime/Online/HTTP/Private/HttpThread.h b/Engine/Source/Runtime/Online/HTTP/Private/HttpThread.h
--- a/Engine/Source/Runtime/Online/HTTP/Private/HttpThread.h
+++ b/Engine/Source/Runtime/Online/HTTP/Private/HttpThread.h
@@ -143,6 +143,7 @@ private:
 	void MoveCompletingRequestsToCompletedRequests(TArray<IHttpThreadedRequest*>& RequestsToComplete);
 	void StartRequestsWaitingInQueue(TArray<IHttpThreadedRequest*>& RequestsToComplete);
 	void FinishRequestsFromHttpThreadWithCallbacks(TArray<IHttpThreadedRequest*>& RequestsToComplete);
+	void UpdateThreadPriorityIfNeeded();
 
 protected:
 	/** Pointer to Runnable Thread */
@@ -158,6 +159,9 @@ private:
 	/** Last time the thread has been processed. Used in the non-game thread. */
 	double LastTime;
 
+	/** Current thread priority of the thread. Used to detect when a priority change is requested */
+	EThreadPriority CurrentThreadPriority;
+
 protected:
 	/** 
 	 * Threaded requests that are waiting to be processed on the http thread.
