Change event loop managed storage to use TQueue::Dequeue instead of TQueue::Peek and TQueue::Pop.

diff --git a/Engine/Source/Runtime/Online/Experimental/EventLoop/Public/EventLoop/EventLoop.h b/Engine/Source/Runtime/Online/Experimental/EventLoop/Public/EventLoop/EventLoop.h
--- a/Engine/Source/Runtime/Online/Experimental/EventLoop/Public/EventLoop/EventLoop.h
+++ b/Engine/Source/Runtime/Online/Experimental/EventLoop/Public/EventLoop/EventLoop.h
@@ -293,10 +293,10 @@ public:
 private:
 	void RunAsyncTasks()
 	{
-		while (FAsyncTask* AsyncTask = AsyncTasks.Peek())
+		FAsyncTask AsyncTask;
+		while (AsyncTasks.Dequeue(AsyncTask))
 		{
-			(*AsyncTask)();
-			AsyncTasks.Pop();
+			AsyncTask();
 		}
 	}
 
diff --git a/Engine/Source/Runtime/Online/Experimental/EventLoop/Public/EventLoop/EventLoopManagedStorage.h b/Engine/Source/Runtime/Online/Experimental/EventLoop/Public/EventLoop/EventLoopManagedStorage.h
--- a/Engine/Source/Runtime/Online/Experimental/EventLoop/Public/EventLoop/EventLoopManagedStorage.h
+++ b/Engine/Source/Runtime/Online/Experimental/EventLoop/Public/EventLoop/EventLoopManagedStorage.h
@@ -253,19 +253,19 @@ public:
 		// Add new queued items.
 		if (OutAddedHandles)
 		{
-			while (FAddRequest* AddRequest = AddRequests.Peek())
+			FAddRequest AddRequest;
+			while (AddRequests.Dequeue(AddRequest))
 			{
-				OutAddedHandles->Add(AddImpl(AddRequest->Handle, MoveTemp(AddRequest->Data)));
-				AddRequests.Pop();
+				OutAddedHandles->Add(AddImpl(AddRequest.Handle, MoveTemp(AddRequest.Data)));
 				++NumChanges;
 			}
 		}
 		else
 		{
-			while (FAddRequest* AddRequest = AddRequests.Peek())
+			FAddRequest AddRequest;
+			while (AddRequests.Dequeue(AddRequest))
 			{
-				AddImpl(AddRequest->Handle, MoveTemp(AddRequest->Data));
-				AddRequests.Pop();
+				AddImpl(AddRequest.Handle, MoveTemp(AddRequest.Data));
 				++NumChanges;
 			}
 		}
@@ -273,19 +273,19 @@ public:
 		// Remove queued items.
 		if (OutRemovedHandles)
 		{
-			while (FRemoveRequest* RemoveRequest = RemoveRequests.Peek())
+			FRemoveRequest RemoveRequest;
+			while (RemoveRequests.Dequeue(RemoveRequest))
 			{
-				OutRemovedHandles->Add(RemoveImpl(RemoveRequest->Handle, RemoveRequest->OnComplete));
-				RemoveRequests.Pop();
+				OutRemovedHandles->Add(RemoveImpl(RemoveRequest.Handle, RemoveRequest.OnComplete));
 				++NumChanges;
 			}
 		}
 		else
 		{
-			while (FRemoveRequest* RemoveRequest = RemoveRequests.Peek())
+			FRemoveRequest RemoveRequest;
+			while (RemoveRequests.Dequeue(RemoveRequest))
 			{
-				RemoveImpl(RemoveRequest->Handle, RemoveRequest->OnComplete);
-				RemoveRequests.Pop();
+				RemoveImpl(RemoveRequest.Handle, RemoveRequest.OnComplete);
 				++NumChanges;
 			}
 		}
