Log all error codes for issue diagnosis.
#rnx

diff --git a/Engine/Source/Runtime/Online/BuildPatchServices/Private/Installer/OptimisedDelta.cpp b/Engine/Source/Runtime/Online/BuildPatchServices/Private/Installer/OptimisedDelta.cpp
--- a/Engine/Source/Runtime/Online/BuildPatchServices/Private/Installer/OptimisedDelta.cpp
+++ b/Engine/Source/Runtime/Online/BuildPatchServices/Private/Installer/OptimisedDelta.cpp
@@ -211,18 +211,20 @@ namespace BuildPatchServices
 				}
 				else
 				{
+					UE_LOG(LogOptimisedDelta, Log, TEXT("Failed to load %s, ShouldRetry = false"), *RelativeDeltaFilePath);
 					SetFailedDownload();
 				}
 			}
 			else if (ShouldRetry(Download))
 			{
 				++RetryCount;
-				UE_LOG(LogOptimisedDelta, Log, TEXT("Failed to download %s, retrying %i/%i"), *RelativeDeltaFilePath, RetryCount, DeltaRetries);
+				UE_LOG(LogOptimisedDelta, Log, TEXT("Failed to download (response code=%d) %s, retrying %i/%i"), Download->GetResponseCode(), *RelativeDeltaFilePath, RetryCount, DeltaRetries);
 				CloudDirIdx = (CloudDirIdx + RetryCount) % Configuration.CloudDirectories.Num();
 				DownloadService->RequestFile(Configuration.CloudDirectories[CloudDirIdx] / RelativeDeltaFilePath, ChunkDeltaComplete, ChunkDeltaProgress);
 			}
 			else
 			{
+				UE_LOG(LogOptimisedDelta, Log, TEXT("Failed to download (response code=%d) %s, ShouldRetry = false"), Download->GetResponseCode(), *RelativeDeltaFilePath);
 				ErrorCode = DownloadErrorCodes::MissingDeltaFile;
 				SetFailedDownload();
 			}
