Enable HTTP retries for 202 (Accepted) status results.
Produce regular logs if retries are for non-errors status results.

diff --git a/Engine/Source/Runtime/Online/HTTP/Private/HttpRetrySystem.cpp b/Engine/Source/Runtime/Online/HTTP/Private/HttpRetrySystem.cpp
--- a/Engine/Source/Runtime/Online/HTTP/Private/HttpRetrySystem.cpp
+++ b/Engine/Source/Runtime/Online/HTTP/Private/HttpRetrySystem.cpp
@@ -479,7 +479,15 @@ void FHttpRetrySystem::FManager::RetryHttpRequest(FHttpRetryRequestEntry& Reques
 		++RequestEntry.CurrentRetryCountForConnectionError;
 	}
 	RequestEntry.Request->RetryStatus = FRequest::EStatus::Processing;
-	UE_LOG(LogHttp, Warning, TEXT("Retry %d on %s"), RequestEntry.CurrentRetryCount, *(RequestEntry.Request->GetURL()));
+	if (RequestEntry.Request->GetResponse()->GetResponseCode() >= 400)
+	{
+		UE_LOG(LogHttp, Warning, TEXT("Retry %d on %s"), RequestEntry.CurrentRetryCount, *(RequestEntry.Request->GetURL()));
+	}
+	else
+	{
+		UE_LOG(LogHttp, Log, TEXT("Retry %d on %s"), RequestEntry.CurrentRetryCount, *(RequestEntry.Request->GetURL()));
+	}
+	
 	RequestEntry.Request->HttpRequest->ProcessRequest();
 }
 
