Enable HTTP retries for 202 (Accepted) status results.
Produce regular logs if retries are for non-errors status results.

diff --git a/Engine/Source/Runtime/Online/HTTP/Private/HttpRetrySystem.cpp b/Engine/Source/Runtime/Online/HTTP/Private/HttpRetrySystem.cpp
--- a/Engine/Source/Runtime/Online/HTTP/Private/HttpRetrySystem.cpp
+++ b/Engine/Source/Runtime/Online/HTTP/Private/HttpRetrySystem.cpp
@@ -479,7 +479,32 @@ void FHttpRetrySystem::FManager::RetryHttpRequest(FHttpRetryRequestEntry& Reques
 		++RequestEntry.CurrentRetryCountForConnectionError;
 	}
 	RequestEntry.Request->RetryStatus = FRequest::EStatus::Processing;
-	UE_LOG(LogHttp, Warning, TEXT("Retry %d on %s"), RequestEntry.CurrentRetryCount, *(RequestEntry.Request->GetURL()));
+
+	if (const FHttpResponsePtr Response = RequestEntry.Request->GetResponse())
+	{
+		if (const int32 ResponseCode = Response->GetResponseCode(); ResponseCode < 400)
+		{
+			// 1XX, 2XX, and 3XX are non error responses, regular log level
+			UE_LOG(LogHttp, Log, TEXT("Retry %d on %s with response %d"),
+				RequestEntry.CurrentRetryCount,
+				*(RequestEntry.Request->GetURL()),
+				ResponseCode);
+		}
+		else
+		{
+			// 4XX, 5XX are error responses, warning log level
+			UE_LOG(LogHttp, Warning, TEXT("Retry %d on %s with response %d"),
+				RequestEntry.CurrentRetryCount,
+				*(RequestEntry.Request->GetURL()),
+				ResponseCode);	
+		}
+	}
+	else
+	{
+		// We don't know the response code, default to warning log level
+		UE_LOG(LogHttp, Warning, TEXT("Retry %d on %s"), RequestEntry.CurrentRetryCount, *(RequestEntry.Request->GetURL()));
+	}
+	
 	RequestEntry.Request->HttpRequest->ProcessRequest();
 }
 
