Merge AppleHTTP fixes from EOSSDK

diff --git a/Engine/Source/Runtime/Online/HTTP/Private/Apple/AppleHttp.cpp b/Engine/Source/Runtime/Online/HTTP/Private/Apple/AppleHttp.cpp
--- a/Engine/Source/Runtime/Online/HTTP/Private/Apple/AppleHttp.cpp
+++ b/Engine/Source/Runtime/Online/HTTP/Private/Apple/AppleHttp.cpp
@@ -362,10 +362,10 @@
 
 @implementation FNSInputStreamFromArchive
 
-+(FNSInputStreamFromArchive*)initWithArchive:(TSharedRef<FArchive>) Archive
++(FNSInputStreamFromArchive*)inputStreamWithArchive:(TSharedRef<FArchive>) Archive
 {
 	FNSInputStreamFromArchive* Ret = [[[FNSInputStreamFromArchive alloc] init] autorelease];
-	Ret->Archive = MoveTemp(Archive);
+	Ret->Archive = Archive;
 	return Ret;
 }
 
@@ -374,8 +374,8 @@
 	self = [super init];
 	if (self)
 	{
+		AlreadySentContent = 0;
 		StreamStatus = NSStreamStatusNotOpen;
-
 		// Docs say it is good practice that streams are it's own delegates by default
 		Delegate = self;
 	}
@@ -417,14 +417,7 @@
 
 - (void)setDelegate:(id<NSStreamDelegate>)InDelegate
 {
-	if (InDelegate == nil)
-	{
-		InDelegate = self;
-	}
-	else
-	{
-		Delegate = InDelegate;
-	}
+	Delegate = InDelegate ?: self;
 }
 
 - (id)propertyForKey:(NSString *)key
@@ -758,7 +751,7 @@ struct FAppleHttpRequest::FAppleHttpStreamFactory
 	
 	NSInputStream *operator()(const TSharedRef<FArchive>& Archive)
 	{
-		return [FNSInputStreamFromArchive initWithArchive: Archive];
+		return [FNSInputStreamFromArchive inputStreamWithArchive: Archive];
 	}
 };
 
