Refactor: Merge IHttpThreadedRequest into FHttpRequestCommon, deprecate FHttpRequestCommon and FHttpResponseCommon from public.

diff --git a/Engine/Source/Runtime/Online/HTTP/Private/Apple/AppleEventLoopHttpThread.cpp b/Engine/Source/Runtime/Online/HTTP/Private/Apple/AppleEventLoopHttpThread.cpp
--- a/Engine/Source/Runtime/Online/HTTP/Private/Apple/AppleEventLoopHttpThread.cpp
+++ b/Engine/Source/Runtime/Online/HTTP/Private/Apple/AppleEventLoopHttpThread.cpp
@@ -4,7 +4,8 @@
 
 #include "Apple/AppleHttp.h" 
 
-bool FAppleEventLoopHttpThread::StartThreadedRequest(IHttpThreadedRequest* Request)
+PRAGMA_DISABLE_DEPRECATION_WARNINGS
+bool FAppleEventLoopHttpThread::StartThreadedRequest(FHttpRequestCommon* Request)
 {
 	FHttpResponsePtr Response = Request->GetResponse();
 	auto AppleResponse = StaticCastSharedPtr<FAppleHttpResponse>(Response);
@@ -15,7 +16,7 @@ bool FAppleEventLoopHttpThread::StartThreadedRequest(IHttpThreadedRequest* Reque
 	return FEventLoopHttpThread::StartThreadedRequest(Request);
 }
 
-void FAppleEventLoopHttpThread::CompleteThreadedRequest(IHttpThreadedRequest* Request)
+void FAppleEventLoopHttpThread::CompleteThreadedRequest(FHttpRequestCommon* Request)
 {
 }
 
@@ -24,12 +25,13 @@ void FAppleEventLoopHttpThread::CreateEventLoop()
 	UE::EventLoop::TEventLoop<FAppleHTTPIOManager>::FParams EventLoopParams;
 	EventLoopParams.IOManagerParams.ProcessRequests = [this]()
 	{ 
-		TArray<IHttpThreadedRequest*> RequestsToCancel;
-		TArray<IHttpThreadedRequest*> RequestsToComplete;
+		TArray<FHttpRequestCommon*> RequestsToCancel;
+		TArray<FHttpRequestCommon*> RequestsToComplete;
 		Process(RequestsToCancel, RequestsToComplete);
 	};
 	EventLoop.Emplace(MoveTemp(EventLoopParams));
 }
+PRAGMA_ENABLE_DEPRECATION_WARNINGS
 
 void FAppleEventLoopHttpThread::DestroyEventLoop()
 {
diff --git a/Engine/Source/Runtime/Online/HTTP/Private/Apple/AppleEventLoopHttpThread.h b/Engine/Source/Runtime/Online/HTTP/Private/Apple/AppleEventLoopHttpThread.h
--- a/Engine/Source/Runtime/Online/HTTP/Private/Apple/AppleEventLoopHttpThread.h
+++ b/Engine/Source/Runtime/Online/HTTP/Private/Apple/AppleEventLoopHttpThread.h
@@ -87,8 +87,10 @@ class FAppleEventLoopHttpThread
 {
 protected:
 	//~ Begin FHttpThread Interface
-	virtual bool StartThreadedRequest(IHttpThreadedRequest* Request) override;
-	virtual void CompleteThreadedRequest(IHttpThreadedRequest* Request) override;
+PRAGMA_DISABLE_DEPRECATION_WARNINGS
+	virtual bool StartThreadedRequest(FHttpRequestCommon* Request) override;
+	virtual void CompleteThreadedRequest(FHttpRequestCommon* Request) override;
+PRAGMA_ENABLE_DEPRECATION_WARNINGS
 	//~ End FHttpThread Interface
 protected:
 	virtual void CreateEventLoop() override;
diff --git a/Engine/Source/Runtime/Online/HTTP/Private/Apple/AppleHttp.h b/Engine/Source/Runtime/Online/HTTP/Private/Apple/AppleHttp.h
--- a/Engine/Source/Runtime/Online/HTTP/Private/Apple/AppleHttp.h
+++ b/Engine/Source/Runtime/Online/HTTP/Private/Apple/AppleHttp.h
@@ -4,7 +4,7 @@
 
 #include "CoreMinimal.h"
 #include "GenericPlatform/HttpResponseCommon.h"
-#include "IHttpThreadedRequest.h"
+#include "GenericPlatform/HttpRequestCommon.h"
 #include "PlatformHttp.h"
 #include "HttpPackage.h"
 #include "Misc/TVariant.h"
@@ -18,7 +18,9 @@ DECLARE_DELEGATE(FNewAppleHttpEventDelegate);
 /**
  * Apple implementation of an Http request
  */
-class FAppleHttpRequest : public IHttpThreadedRequest
+PRAGMA_DISABLE_DEPRECATION_WARNINGS
+class FAppleHttpRequest : public FHttpRequestCommon
+PRAGMA_ENABLE_DEPRECATION_WARNINGS
 {
 public:
 	// implementation friends
@@ -68,6 +70,7 @@ public:
 	virtual ~FAppleHttpRequest();
 
 PACKAGE_SCOPE:
+PRAGMA_DISABLE_DEPRECATION_WARNINGS
 	using FHttpRequestCommon::StartProcessTime;
 	using FHttpRequestCommon::ConnectTime;
 	using FHttpRequestCommon::BroadcastResponseHeadersReceived;
@@ -77,6 +80,7 @@ PACKAGE_SCOPE:
 	using FHttpRequestCommon::TriggerStatusCodeReceivedDelegate;
 	using FHttpRequestCommon::SetEffectiveURL;
 	using FHttpRequestCommon::PassReceivedDataToStream;
+PRAGMA_ENABLE_DEPRECATION_WARNINGS
 
 	bool IsInitializedWithValidStream() const;
 
@@ -137,7 +141,9 @@ private:
 /**
  * Apple implementation of an Http response
  */
+PRAGMA_DISABLE_DEPRECATION_WARNINGS
 class FAppleHttpResponse : public FHttpResponseCommon
+PRAGMA_ENABLE_DEPRECATION_WARNINGS
 {
 public:
 	//~ Begin IHttpBase Interface
diff --git a/Engine/Source/Runtime/Online/HTTP/Private/Curl/CurlHttp.cpp b/Engine/Source/Runtime/Online/HTTP/Private/Curl/CurlHttp.cpp
--- a/Engine/Source/Runtime/Online/HTTP/Private/Curl/CurlHttp.cpp
+++ b/Engine/Source/Runtime/Online/HTTP/Private/Curl/CurlHttp.cpp
@@ -62,6 +62,7 @@ static CURLcode sslctx_function(CURL * curl, void * sslctx, void * parm)
 }
 #endif //#if WITH_SSL
 
+PRAGMA_DISABLE_DEPRECATION_WARNINGS
 FCurlHttpRequest::FCurlHttpRequest()
 	: EasyHandle(nullptr)
 	, HeaderList(nullptr)
@@ -1037,7 +1038,7 @@ bool FCurlHttpRequest::ProcessRequest()
 
 void FCurlHttpRequest::ClearInCaseOfRetry()
 {
-	IHttpThreadedRequest::ClearInCaseOfRetry();
+	FHttpRequestCommon::ClearInCaseOfRetry();
 
 	// Clear out response. If this is a re-used request, Response could point to a stale response until SetupRequestHttpThread is called
 	LastReportedBytesRead = 0;
@@ -1125,10 +1126,8 @@ void FCurlHttpRequest::CheckProgressDelegate()
 	{
 		LastReportedBytesSent = CurrentBytesSent;
 		LastReportedBytesRead = CurrentBytesRead;
-PRAGMA_DISABLE_DEPRECATION_WARNINGS
 		// Update response progress
 		OnRequestProgress().ExecuteIfBound(SharedThis(this), LastReportedBytesSent, LastReportedBytesRead);
-PRAGMA_ENABLE_DEPRECATION_WARNINGS
 		OnRequestProgress64().ExecuteIfBound(SharedThis(this), LastReportedBytesSent, LastReportedBytesRead);
 	}
 }
@@ -1370,5 +1369,6 @@ FString FCurlHttpResponse::GetContentAsString() const
 	FUTF8ToTCHAR TCHARData(reinterpret_cast<const ANSICHAR*>(Payload.GetData()), Payload.Num());
 	return FString(TCHARData.Length(), TCHARData.Get());
 }
+PRAGMA_ENABLE_DEPRECATION_WARNINGS
 
 #endif //WITH_CURL
diff --git a/Engine/Source/Runtime/Online/HTTP/Private/Curl/CurlHttp.h b/Engine/Source/Runtime/Online/HTTP/Private/Curl/CurlHttp.h
--- a/Engine/Source/Runtime/Online/HTTP/Private/Curl/CurlHttp.h
+++ b/Engine/Source/Runtime/Online/HTTP/Private/Curl/CurlHttp.h
@@ -3,12 +3,12 @@
 #pragma once
 
 #include "CoreMinimal.h"
+#include "Containers/SpscQueue.h"
 #include "HAL/LowLevelMemTracker.h"
+#include "HAL/ThreadSafeBool.h"
 #include "HAL/ThreadSafeCounter.h"
-#include "IHttpThreadedRequest.h"
-#include "Containers/SpscQueue.h"
+#include "GenericPlatform/HttpRequestCommon.h"
 #include "GenericPlatform/HttpResponseCommon.h"
-#include "HAL/ThreadSafeBool.h"
 
 class FCurlHttpResponse;
 
@@ -135,7 +135,9 @@ namespace
 /**
  * Curl implementation of an HTTP request
  */
-class FCurlHttpRequest : public IHttpThreadedRequest
+PRAGMA_DISABLE_DEPRECATION_WARNINGS
+class FCurlHttpRequest : public FHttpRequestCommon
+PRAGMA_ENABLE_DEPRECATION_WARNINGS
 {
 public:
 
@@ -402,7 +404,9 @@ private:
 /**
  * Curl implementation of an HTTP response
  */
+PRAGMA_DISABLE_DEPRECATION_WARNINGS
 class FCurlHttpResponse : public FHttpResponseCommon
+PRAGMA_ENABLE_DEPRECATION_WARNINGS
 {
 public:
 	// implementation friends
diff --git a/Engine/Source/Runtime/Online/HTTP/Private/Curl/CurlHttpThread.cpp b/Engine/Source/Runtime/Online/HTTP/Private/Curl/CurlHttpThread.cpp
--- a/Engine/Source/Runtime/Online/HTTP/Private/Curl/CurlHttpThread.cpp
+++ b/Engine/Source/Runtime/Online/HTTP/Private/Curl/CurlHttpThread.cpp
@@ -46,7 +46,9 @@ void FCurlHttpThread::HttpThreadTick(float DeltaSeconds)
 					CURL* CompletedHandle = Message->easy_handle;
 					curl_multi_remove_handle(FCurlHttpManager::GMultiHandle, CompletedHandle);
 
-					IHttpThreadedRequest** Request = HandlesToRequests.Find(CompletedHandle);
+PRAGMA_DISABLE_DEPRECATION_WARNINGS
+					FHttpRequestCommon** Request = HandlesToRequests.Find(CompletedHandle);
+PRAGMA_ENABLE_DEPRECATION_WARNINGS
 					if (Request)
 					{
 						FCurlHttpRequest* CurlRequest = static_cast<FCurlHttpRequest*>(*Request);
@@ -68,7 +70,8 @@ void FCurlHttpThread::HttpThreadTick(float DeltaSeconds)
 	FLegacyHttpThread::HttpThreadTick(DeltaSeconds);
 }
 
-bool FCurlHttpThread::StartThreadedRequest(IHttpThreadedRequest* Request)
+PRAGMA_DISABLE_DEPRECATION_WARNINGS
+bool FCurlHttpThread::StartThreadedRequest(FHttpRequestCommon* Request)
 {
 	FCurlHttpRequest* CurlRequest = static_cast<FCurlHttpRequest*>(Request);
 	CURL* EasyHandle = CurlRequest->GetEasyHandle();
@@ -94,7 +97,7 @@ bool FCurlHttpThread::StartThreadedRequest(IHttpThreadedRequest* Request)
 	return FLegacyHttpThread::StartThreadedRequest(Request);
 }
 
-void FCurlHttpThread::CompleteThreadedRequest(IHttpThreadedRequest* Request)
+void FCurlHttpThread::CompleteThreadedRequest(FHttpRequestCommon* Request)
 {
 	FCurlHttpRequest* CurlRequest = static_cast<FCurlHttpRequest*>(Request);
 	CURL* EasyHandle = CurlRequest->GetEasyHandle();
@@ -105,5 +108,6 @@ void FCurlHttpThread::CompleteThreadedRequest(IHttpThreadedRequest* Request)
 		HandlesToRequests.Remove(EasyHandle);
 	}
 }
+PRAGMA_ENABLE_DEPRECATION_WARNINGS
 
 #endif
diff --git a/Engine/Source/Runtime/Online/HTTP/Private/Curl/CurlHttpThread.h b/Engine/Source/Runtime/Online/HTTP/Private/Curl/CurlHttpThread.h
--- a/Engine/Source/Runtime/Online/HTTP/Private/Curl/CurlHttpThread.h
+++ b/Engine/Source/Runtime/Online/HTTP/Private/Curl/CurlHttpThread.h
@@ -22,7 +22,7 @@
 
 #endif //WITH_CURL
 
-class IHttpThreadedRequest;
+class FHttpRequestCommon;
 
 #if WITH_CURL
 
@@ -36,13 +36,15 @@ public:
 protected:
 	//~ Begin FHttpThread Interface
 	virtual void HttpThreadTick(float DeltaSeconds) override;
-	virtual bool StartThreadedRequest(IHttpThreadedRequest* Request) override;
-	virtual void CompleteThreadedRequest(IHttpThreadedRequest* Request) override;
+PRAGMA_DISABLE_DEPRECATION_WARNINGS
+	virtual bool StartThreadedRequest(FHttpRequestCommon* Request) override;
+	virtual void CompleteThreadedRequest(FHttpRequestCommon* Request) override;
 	//~ End FHttpThread Interface
 protected:
 
 	/** Mapping of libcurl easy handles to HTTP requests */
-	TMap<CURL*, IHttpThreadedRequest*> HandlesToRequests;
+	TMap<CURL*, FHttpRequestCommon*> HandlesToRequests;
+PRAGMA_ENABLE_DEPRECATION_WARNINGS
 };
 
 
diff --git a/Engine/Source/Runtime/Online/HTTP/Private/Curl/CurlMultiPollEventLoopHttpThread.cpp b/Engine/Source/Runtime/Online/HTTP/Private/Curl/CurlMultiPollEventLoopHttpThread.cpp
--- a/Engine/Source/Runtime/Online/HTTP/Private/Curl/CurlMultiPollEventLoopHttpThread.cpp
+++ b/Engine/Source/Runtime/Online/HTTP/Private/Curl/CurlMultiPollEventLoopHttpThread.cpp
@@ -68,7 +68,8 @@ FCurlMultiPollEventLoopHttpThread::FCurlMultiPollEventLoopHttpThread()
 {
 }
 
-bool FCurlMultiPollEventLoopHttpThread::StartThreadedRequest(IHttpThreadedRequest* Request)
+PRAGMA_DISABLE_DEPRECATION_WARNINGS
+bool FCurlMultiPollEventLoopHttpThread::StartThreadedRequest(FHttpRequestCommon* Request)
 {
 	FCurlHttpRequest* CurlRequest = static_cast<FCurlHttpRequest*>(Request);
 	CURL* EasyHandle = CurlRequest->GetEasyHandle();
@@ -94,7 +95,7 @@ bool FCurlMultiPollEventLoopHttpThread::StartThreadedRequest(IHttpThreadedReques
 	return FEventLoopHttpThread::StartThreadedRequest(Request);
 }
 
-void FCurlMultiPollEventLoopHttpThread::CompleteThreadedRequest(IHttpThreadedRequest* Request)
+void FCurlMultiPollEventLoopHttpThread::CompleteThreadedRequest(FHttpRequestCommon* Request)
 {
 	FCurlHttpRequest* CurlRequest = static_cast<FCurlHttpRequest*>(Request);
 	CURL* EasyHandle = CurlRequest->GetEasyHandle();
@@ -105,6 +106,7 @@ void FCurlMultiPollEventLoopHttpThread::CompleteThreadedRequest(IHttpThreadedReq
 		HandlesToRequests.Remove(EasyHandle);
 	}
 }
+PRAGMA_ENABLE_DEPRECATION_WARNINGS
 
 void FCurlMultiPollEventLoopHttpThread::CreateEventLoop()
 {
@@ -154,7 +156,9 @@ void FCurlMultiPollEventLoopHttpThread::ProcessCurlRequests()
 			CURL* CompletedHandle = Message->easy_handle;
 			curl_multi_remove_handle(FCurlHttpManager::GMultiHandle, CompletedHandle);
 
-			IHttpThreadedRequest** Request = HandlesToRequests.Find(CompletedHandle);
+PRAGMA_DISABLE_DEPRECATION_WARNINGS
+			FHttpRequestCommon** Request = HandlesToRequests.Find(CompletedHandle);
+PRAGMA_ENABLE_DEPRECATION_WARNINGS
 			if (Request)
 			{
 				FCurlHttpRequest* CurlRequest = static_cast<FCurlHttpRequest*>(*Request);
@@ -175,8 +179,10 @@ void FCurlMultiPollEventLoopHttpThread::ProcessCurlRequests()
 	// If any requests completed, immediately process requests to handle completion event.
 	if (CompletedRequest)
 	{
-		TArray<IHttpThreadedRequest*> RequestsToCancel;
-		TArray<IHttpThreadedRequest*> RequestsToComplete;
+PRAGMA_DISABLE_DEPRECATION_WARNINGS
+		TArray<FHttpRequestCommon*> RequestsToCancel;
+		TArray<FHttpRequestCommon*> RequestsToComplete;
+PRAGMA_ENABLE_DEPRECATION_WARNINGS
 		Process(RequestsToCancel, RequestsToComplete);
 	}
 }
diff --git a/Engine/Source/Runtime/Online/HTTP/Private/Curl/CurlMultiPollEventLoopHttpThread.h b/Engine/Source/Runtime/Online/HTTP/Private/Curl/CurlMultiPollEventLoopHttpThread.h
--- a/Engine/Source/Runtime/Online/HTTP/Private/Curl/CurlMultiPollEventLoopHttpThread.h
+++ b/Engine/Source/Runtime/Online/HTTP/Private/Curl/CurlMultiPollEventLoopHttpThread.h
@@ -23,7 +23,7 @@
 #include "Microsoft/HideMicrosoftPlatformTypes.h"
 #endif
 
-class IHttpThreadedRequest;
+class FHttpRequestCommon;
 class FCurlMultiPollIOManager;
 
 class FCurlMultiPollIOManagerIOAccess final : public FNoncopyable
@@ -77,8 +77,10 @@ public:
 
 protected:
 	//~ Begin FHttpThread Interface
-	virtual bool StartThreadedRequest(IHttpThreadedRequest* Request) override;
-	virtual void CompleteThreadedRequest(IHttpThreadedRequest* Request) override;
+PRAGMA_DISABLE_DEPRECATION_WARNINGS
+	virtual bool StartThreadedRequest(FHttpRequestCommon* Request) override;
+	virtual void CompleteThreadedRequest(FHttpRequestCommon* Request) override;
+PRAGMA_ENABLE_DEPRECATION_WARNINGS
 	//~ End FHttpThread Interface
 protected:
 	virtual void CreateEventLoop() override;
@@ -91,8 +93,10 @@ protected:
 
 	TOptional<UE::EventLoop::TEventLoop<FCurlMultiPollIOManager>> EventLoop;
 
+PRAGMA_DISABLE_DEPRECATION_WARNINGS
 	/** Mapping of libcurl easy handles to HTTP requests */
-	TMap<CURL*, IHttpThreadedRequest*> HandlesToRequests;
+	TMap<CURL*, FHttpRequestCommon*> HandlesToRequests;
+PRAGMA_ENABLE_DEPRECATION_WARNINGS
 };
 
 #endif // WITH_CURL_MULTIPOLL
diff --git a/Engine/Source/Runtime/Online/HTTP/Private/Curl/CurlMultiWaitEventLoopHttpThread.cpp b/Engine/Source/Runtime/Online/HTTP/Private/Curl/CurlMultiWaitEventLoopHttpThread.cpp
--- a/Engine/Source/Runtime/Online/HTTP/Private/Curl/CurlMultiWaitEventLoopHttpThread.cpp
+++ b/Engine/Source/Runtime/Online/HTTP/Private/Curl/CurlMultiWaitEventLoopHttpThread.cpp
@@ -87,7 +87,8 @@ FCurlMultiWaitEventLoopHttpThread::FCurlMultiWaitEventLoopHttpThread()
 {
 }
 
-bool FCurlMultiWaitEventLoopHttpThread::StartThreadedRequest(IHttpThreadedRequest* Request)
+PRAGMA_DISABLE_DEPRECATION_WARNINGS
+bool FCurlMultiWaitEventLoopHttpThread::StartThreadedRequest(FHttpRequestCommon* Request)
 {
 	FCurlHttpRequest* CurlRequest = static_cast<FCurlHttpRequest*>(Request);
 	CURL* EasyHandle = CurlRequest->GetEasyHandle();
@@ -113,7 +114,7 @@ bool FCurlMultiWaitEventLoopHttpThread::StartThreadedRequest(IHttpThreadedReques
 	return FEventLoopHttpThread::StartThreadedRequest(Request);
 }
 
-void FCurlMultiWaitEventLoopHttpThread::CompleteThreadedRequest(IHttpThreadedRequest* Request)
+void FCurlMultiWaitEventLoopHttpThread::CompleteThreadedRequest(FHttpRequestCommon* Request)
 {
 	FCurlHttpRequest* CurlRequest = static_cast<FCurlHttpRequest*>(Request);
 	CURL* EasyHandle = CurlRequest->GetEasyHandle();
@@ -123,6 +124,7 @@ void FCurlMultiWaitEventLoopHttpThread::CompleteThreadedRequest(IHttpThreadedReq
 		curl_multi_remove_handle(FCurlHttpManager::GMultiHandle, EasyHandle);
 	}
 }
+PRAGMA_ENABLE_DEPRECATION_WARNINGS
 
 void FCurlMultiWaitEventLoopHttpThread::CreateEventLoop()
 {
@@ -172,7 +174,9 @@ void FCurlMultiWaitEventLoopHttpThread::ProcessCurlRequests()
 			CURL* CompletedHandle = Message->easy_handle;
 			curl_multi_remove_handle(FCurlHttpManager::GMultiHandle, CompletedHandle);
 
-			IHttpThreadedRequest** Request = HandlesToRequests.Find(CompletedHandle);
+PRAGMA_DISABLE_DEPRECATION_WARNINGS
+			FHttpRequestCommon** Request = HandlesToRequests.Find(CompletedHandle);
+PRAGMA_ENABLE_DEPRECATION_WARNINGS
 			if (Request)
 			{
 				FCurlHttpRequest* CurlRequest = static_cast<FCurlHttpRequest*>(*Request);
@@ -193,8 +197,10 @@ void FCurlMultiWaitEventLoopHttpThread::ProcessCurlRequests()
 	// If any requests completed, immediately process requests to handle completion event.
 	if (CompletedRequest)
 	{
-		TArray<IHttpThreadedRequest*> RequestsToCancel;
-		TArray<IHttpThreadedRequest*> RequestsToComplete;
+PRAGMA_DISABLE_DEPRECATION_WARNINGS
+		TArray<FHttpRequestCommon*> RequestsToCancel;
+		TArray<FHttpRequestCommon*> RequestsToComplete;
+PRAGMA_ENABLE_DEPRECATION_WARNINGS
 		Process(RequestsToCancel, RequestsToComplete);
 	}
 }
diff --git a/Engine/Source/Runtime/Online/HTTP/Private/Curl/CurlMultiWaitEventLoopHttpThread.h b/Engine/Source/Runtime/Online/HTTP/Private/Curl/CurlMultiWaitEventLoopHttpThread.h
--- a/Engine/Source/Runtime/Online/HTTP/Private/Curl/CurlMultiWaitEventLoopHttpThread.h
+++ b/Engine/Source/Runtime/Online/HTTP/Private/Curl/CurlMultiWaitEventLoopHttpThread.h
@@ -23,7 +23,7 @@
 #include "Microsoft/HideMicrosoftPlatformTypes.h"
 #endif
 
-class IHttpThreadedRequest;
+class FHttpRequestCommon;
 class FCurlMultiWaitIOManager;
 
 class FCurlMultiWaitIOManagerIOAccess final : public FNoncopyable
@@ -78,8 +78,10 @@ public:
 
 protected:
 	//~ Begin FHttpThread Interface
-	virtual bool StartThreadedRequest(IHttpThreadedRequest* Request) override;
-	virtual void CompleteThreadedRequest(IHttpThreadedRequest* Request) override;
+PRAGMA_DISABLE_DEPRECATION_WARNINGS
+	virtual bool StartThreadedRequest(FHttpRequestCommon* Request) override;
+	virtual void CompleteThreadedRequest(FHttpRequestCommon* Request) override;
+PRAGMA_ENABLE_DEPRECATION_WARNINGS
 	//~ End FHttpThread Interface
 protected:
 	virtual void CreateEventLoop() override;
@@ -93,7 +95,9 @@ protected:
 	TOptional<UE::EventLoop::TEventLoop<FCurlMultiWaitIOManager>> EventLoop;
 
 	/** Mapping of libcurl easy handles to HTTP requests */
-	TMap<CURL*, IHttpThreadedRequest*> HandlesToRequests;
+PRAGMA_DISABLE_DEPRECATION_WARNINGS
+	TMap<CURL*, FHttpRequestCommon*> HandlesToRequests;
+PRAGMA_ENABLE_DEPRECATION_WARNINGS
 };
 
 #endif // WITH_CURL_MULTIWAIT
diff --git a/Engine/Source/Runtime/Online/HTTP/Private/Curl/CurlSocketEventLoopHttpThread.cpp b/Engine/Source/Runtime/Online/HTTP/Private/Curl/CurlSocketEventLoopHttpThread.cpp
--- a/Engine/Source/Runtime/Online/HTTP/Private/Curl/CurlSocketEventLoopHttpThread.cpp
+++ b/Engine/Source/Runtime/Online/HTTP/Private/Curl/CurlSocketEventLoopHttpThread.cpp
@@ -59,7 +59,8 @@ void FCurlSocketEventLoopHttpThread::HttpThreadTick(float DeltaSeconds)
 	FEventLoopHttpThread::HttpThreadTick(DeltaSeconds);
 }
 
-bool FCurlSocketEventLoopHttpThread::StartThreadedRequest(IHttpThreadedRequest* Request)
+PRAGMA_DISABLE_DEPRECATION_WARNINGS
+bool FCurlSocketEventLoopHttpThread::StartThreadedRequest(FHttpRequestCommon* Request)
 {
 	FCurlHttpRequest* CurlRequest = static_cast<FCurlHttpRequest*>(Request);
 	CURL* EasyHandle = CurlRequest->GetEasyHandle();
@@ -85,7 +86,7 @@ bool FCurlSocketEventLoopHttpThread::StartThreadedRequest(IHttpThreadedRequest*
 	return FEventLoopHttpThread::StartThreadedRequest(Request);
 }
 
-void FCurlSocketEventLoopHttpThread::CompleteThreadedRequest(IHttpThreadedRequest* Request)
+void FCurlSocketEventLoopHttpThread::CompleteThreadedRequest(FHttpRequestCommon* Request)
 {
 	FCurlHttpRequest* CurlRequest = static_cast<FCurlHttpRequest*>(Request);
 	CURL* EasyHandle = CurlRequest->GetEasyHandle();
@@ -96,6 +97,7 @@ void FCurlSocketEventLoopHttpThread::CompleteThreadedRequest(IHttpThreadedReques
 		HandlesToRequests.Remove(EasyHandle);
 	}
 }
+PRAGMA_ENABLE_DEPRECATION_WARNINGS
 
 void FCurlSocketEventLoopHttpThread::CreateEventLoop()
 {
@@ -244,6 +246,7 @@ void FCurlSocketEventLoopHttpThread::ProcessCurlSocketEvent(curl_socket_t Socket
 	}
 }
 
+PRAGMA_DISABLE_DEPRECATION_WARNINGS
 void FCurlSocketEventLoopHttpThread::ProcessCurlRequests()
 {
 	QUICK_SCOPE_CYCLE_COUNTER(STAT_FCurlSocketEventLoopHttpThread_ProcessCurlRequests);
@@ -265,7 +268,7 @@ void FCurlSocketEventLoopHttpThread::ProcessCurlRequests()
 			CURL* CompletedHandle = Message->easy_handle;
 			curl_multi_remove_handle(FCurlHttpManager::GMultiHandle, CompletedHandle);
 
-			IHttpThreadedRequest** Request = HandlesToRequests.Find(CompletedHandle);
+			FHttpRequestCommon** Request = HandlesToRequests.Find(CompletedHandle);
 			if (Request)
 			{
 				FCurlHttpRequest* CurlRequest = static_cast<FCurlHttpRequest*>(*Request);
@@ -286,11 +289,12 @@ void FCurlSocketEventLoopHttpThread::ProcessCurlRequests()
 	// If any requests completed, immediately process requests to handle completion event.
 	if (CompletedRequest)
 	{
-		TArray<IHttpThreadedRequest*> RequestsToCancel;
-		TArray<IHttpThreadedRequest*> RequestsToComplete;
+		TArray<FHttpRequestCommon*> RequestsToCancel;
+		TArray<FHttpRequestCommon*> RequestsToComplete;
 		Process(RequestsToCancel, RequestsToComplete);
 	}
 }
+PRAGMA_ENABLE_DEPRECATION_WARNINGS
 
 #endif // WITH_CURL_MULTISOCKET
 #endif // WITH_CURL
diff --git a/Engine/Source/Runtime/Online/HTTP/Private/Curl/CurlSocketEventLoopHttpThread.h b/Engine/Source/Runtime/Online/HTTP/Private/Curl/CurlSocketEventLoopHttpThread.h
--- a/Engine/Source/Runtime/Online/HTTP/Private/Curl/CurlSocketEventLoopHttpThread.h
+++ b/Engine/Source/Runtime/Online/HTTP/Private/Curl/CurlSocketEventLoopHttpThread.h
@@ -24,7 +24,7 @@
 #include "Microsoft/HideMicrosoftPlatformTypes.h"
 #endif
 
-class IHttpThreadedRequest;
+class FHttpRequestCommon;
 
 class FCurlSocketEventLoopHttpThread
 	: public FEventLoopHttpThread
@@ -36,8 +36,10 @@ public:
 protected:
 	//~ Begin FHttpThread Interface
 	virtual void HttpThreadTick(float DeltaSeconds) override;
-	virtual bool StartThreadedRequest(IHttpThreadedRequest* Request) override;
-	virtual void CompleteThreadedRequest(IHttpThreadedRequest* Request) override;
+PRAGMA_DISABLE_DEPRECATION_WARNINGS
+	virtual bool StartThreadedRequest(FHttpRequestCommon* Request) override;
+	virtual void CompleteThreadedRequest(FHttpRequestCommon* Request) override;
+PRAGMA_ENABLE_DEPRECATION_WARNINGS
 	//~ End FHttpThread Interface
 protected:
 	virtual void CreateEventLoop() override;
@@ -63,7 +65,9 @@ protected:
 	};
 
 	/** Mapping of libcurl easy handles to HTTP requests */
-	TMap<CURL*, IHttpThreadedRequest*> HandlesToRequests;
+PRAGMA_DISABLE_DEPRECATION_WARNINGS
+	TMap<CURL*, FHttpRequestCommon*> HandlesToRequests;
+PRAGMA_ENABLE_DEPRECATION_WARNINGS
 
 	UE::EventLoop::FTimerHandle RequestTimeoutTimer;
 
diff --git a/Engine/Source/Runtime/Online/HTTP/Private/EventLoopHttpThread.cpp b/Engine/Source/Runtime/Online/HTTP/Private/EventLoopHttpThread.cpp
--- a/Engine/Source/Runtime/Online/HTTP/Private/EventLoopHttpThread.cpp
+++ b/Engine/Source/Runtime/Online/HTTP/Private/EventLoopHttpThread.cpp
@@ -3,7 +3,7 @@
 #include "EventLoopHttpThread.h"
 #include "EventLoop/EventLoop.h"
 #include "EventLoop/EventLoopIOManagerNull.h"
-#include "IHttpThreadedRequest.h"
+#include "GenericPlatform/HttpRequestCommon.h"
 #include "Http.h"
 #include "PlatformHttp.h"
 #include "Stats/Stats.h"
@@ -99,7 +99,8 @@ void FEventLoopHttpThread::UpdateConfigs()
 	FHttpThreadBase::UpdateConfigs();
 }
 
-void FEventLoopHttpThread::AddRequest(IHttpThreadedRequest* Request)
+PRAGMA_DISABLE_DEPRECATION_WARNINGS
+void FEventLoopHttpThread::AddRequest(FHttpRequestCommon* Request)
 {
 	FHttpThreadBase::AddRequest(Request);
 
@@ -108,14 +109,14 @@ void FEventLoopHttpThread::AddRequest(IHttpThreadedRequest* Request)
 		// Force a wakeup to process new tasks.
 		EventLoop->PostAsyncTask([this]()
 		{
-			TArray<IHttpThreadedRequest*> RequestsToCancel;
-			TArray<IHttpThreadedRequest*> RequestsToComplete;
+			TArray<FHttpRequestCommon*> RequestsToCancel;
+			TArray<FHttpRequestCommon*> RequestsToComplete;
 			Process(RequestsToCancel, RequestsToComplete);
 		});
 	}
 }
 
-void FEventLoopHttpThread::CancelRequest(IHttpThreadedRequest* Request)
+void FEventLoopHttpThread::CancelRequest(FHttpRequestCommon* Request)
 {
 	FHttpThreadBase::CancelRequest(Request);
 
@@ -124,17 +125,18 @@ void FEventLoopHttpThread::CancelRequest(IHttpThreadedRequest* Request)
 		// Force a wakeup to process new tasks.
 		EventLoop->PostAsyncTask([this]()
 		{
-			TArray<IHttpThreadedRequest*> RequestsToCancel;
-			TArray<IHttpThreadedRequest*> RequestsToComplete;
+			TArray<FHttpRequestCommon*> RequestsToCancel;
+			TArray<FHttpRequestCommon*> RequestsToComplete;
 			Process(RequestsToCancel, RequestsToComplete);
 		});
 	}
 }
 
-void FEventLoopHttpThread::GetCompletedRequests(TArray<IHttpThreadedRequest*>& OutCompletedRequests)
+void FEventLoopHttpThread::GetCompletedRequests(TArray<FHttpRequestCommon*>& OutCompletedRequests)
 {
 	FHttpThreadBase::GetCompletedRequests(OutCompletedRequests);
 }
+PRAGMA_ENABLE_DEPRECATION_WARNINGS
 
 void FEventLoopHttpThread::Tick()
 {
@@ -171,9 +173,11 @@ void FEventLoopHttpThread::ResetTickTimer()
 
 	RequestTickTimer = EventLoop.SetTimer([this]()
 	{
-		TArray<IHttpThreadedRequest*> RequestsToCancel;
-		TArray<IHttpThreadedRequest*> RequestsToComplete;
+PRAGMA_DISABLE_DEPRECATION_WARNINGS
+		TArray<FHttpRequestCommon*> RequestsToCancel;
+		TArray<FHttpRequestCommon*> RequestsToComplete;
 		Process(RequestsToCancel, RequestsToComplete);
+PRAGMA_ENABLE_DEPRECATION_WARNINGS
 	},
 	FTimespan::FromSeconds(FHttpModule::Get().GetHttpEventLoopThreadTickIntervalInSeconds()),
 	true /* repeat */);
diff --git a/Engine/Source/Runtime/Online/HTTP/Private/EventLoopHttpThread.h b/Engine/Source/Runtime/Online/HTTP/Private/EventLoopHttpThread.h
--- a/Engine/Source/Runtime/Online/HTTP/Private/EventLoopHttpThread.h
+++ b/Engine/Source/Runtime/Online/HTTP/Private/EventLoopHttpThread.h
@@ -16,9 +16,11 @@ public:
 	virtual void StartThread() override final;
 	virtual void StopThread() override final;
 	virtual void UpdateConfigs() override final;
-	virtual void AddRequest(IHttpThreadedRequest* Request) override final;
-	virtual void CancelRequest(IHttpThreadedRequest* Request) override final;
-	virtual void GetCompletedRequests(TArray<IHttpThreadedRequest*>& OutCompletedRequests) override final;
+PRAGMA_DISABLE_DEPRECATION_WARNINGS
+	virtual void AddRequest(FHttpRequestCommon* Request) override final;
+	virtual void CancelRequest(FHttpRequestCommon* Request) override final;
+	virtual void GetCompletedRequests(TArray<FHttpRequestCommon*>& OutCompletedRequests) override final;
+PRAGMA_ENABLE_DEPRECATION_WARNINGS
 
 	//~ Begin FSingleThreadRunnable Interface
 	// Cannot be overridden to ensure identical behavior with the threaded tick
diff --git a/Engine/Source/Runtime/Online/HTTP/Private/GenericPlatform/HttpRequestCommon.cpp b/Engine/Source/Runtime/Online/HTTP/Private/GenericPlatform/HttpRequestCommon.cpp
--- a/Engine/Source/Runtime/Online/HTTP/Private/GenericPlatform/HttpRequestCommon.cpp
+++ b/Engine/Source/Runtime/Online/HTTP/Private/GenericPlatform/HttpRequestCommon.cpp
@@ -10,6 +10,7 @@
 #include "Misc/CommandLine.h"
 #include "Stats/Stats.h"
 
+PRAGMA_DISABLE_DEPRECATION_WARNINGS
 namespace UE::HttpRequestCommon::Private
 {
 
@@ -498,6 +499,17 @@ float FHttpRequestCommon::GetElapsedTime() const
 	return ElapsedTime;
 }
 
+void FHttpRequestCommon::StartWaitingInQueue()
+{
+	TimeStartedWaitingInQueue = FPlatformTime::Seconds();
+}
+
+float FHttpRequestCommon::GetTimeStartedWaitingInQueue() const
+{
+	check(TimeStartedWaitingInQueue != 0);
+	return TimeStartedWaitingInQueue;
+}
+
 bool FHttpRequestCommon::PassReceivedDataToStream(void* Ptr, int64 Length)
 {
 	const FScopeLock StreamLock(&ResponseBodyReceiveStreamCriticalSection);
@@ -626,4 +638,5 @@ void FHttpRequestCommon::LogResponse(const TSharedPtr<IHttpResponse>& InResponse
 
 		StartPos = EndPos + 1;
 	}
-}
\ No newline at end of file
+}
+PRAGMA_ENABLE_DEPRECATION_WARNINGS
\ No newline at end of file
diff --git a/Engine/Source/Runtime/Online/HTTP/Private/GenericPlatform/HttpResponseCommon.cpp b/Engine/Source/Runtime/Online/HTTP/Private/GenericPlatform/HttpResponseCommon.cpp
--- a/Engine/Source/Runtime/Online/HTTP/Private/GenericPlatform/HttpResponseCommon.cpp
+++ b/Engine/Source/Runtime/Online/HTTP/Private/GenericPlatform/HttpResponseCommon.cpp
@@ -4,6 +4,7 @@
 #include "GenericPlatform/HttpRequestCommon.h"
 #include "GenericPlatform/GenericPlatformHttp.h"
 
+PRAGMA_DISABLE_DEPRECATION_WARNINGS
 FHttpResponseCommon::FHttpResponseCommon(const FHttpRequestCommon& HttpRequest)
 	: URL(HttpRequest.GetURL())
 	, EffectiveURL(HttpRequest.GetEffectiveURL())
@@ -56,3 +57,4 @@ void FHttpResponseCommon::SetEffectiveURL(const FString& InEffectiveURL)
 {
 	EffectiveURL = InEffectiveURL;
 }
+PRAGMA_ENABLE_DEPRECATION_WARNINGS
diff --git a/Engine/Source/Runtime/Online/HTTP/Private/HttpManager.cpp b/Engine/Source/Runtime/Online/HTTP/Private/HttpManager.cpp
--- a/Engine/Source/Runtime/Online/HTTP/Private/HttpManager.cpp
+++ b/Engine/Source/Runtime/Online/HTTP/Private/HttpManager.cpp
@@ -11,7 +11,7 @@
 #include "Misc/Guid.h"
 #include "Misc/Fork.h"
 #include "HttpThread.h"
-#include "IHttpThreadedRequest.h"
+#include "GenericPlatform/HttpRequestCommon.h"
 #include "Misc/ConfigCacheIni.h"
 #include "Misc/CommandLine.h"
 
@@ -471,7 +471,8 @@ bool FHttpManager::Tick(float DeltaSeconds)
 			}
 		}
 
-		TArray<IHttpThreadedRequest*> CompletedThreadedRequests;
+PRAGMA_DISABLE_DEPRECATION_WARNINGS
+		TArray<FHttpRequestCommon*> CompletedThreadedRequests;
 
 		{
 			// Thread->GetCompletedRequests doesn't support multi-thread access
@@ -480,7 +481,7 @@ bool FHttpManager::Tick(float DeltaSeconds)
 		}
 
 		// Finish and remove any completed requests
-		for (IHttpThreadedRequest* CompletedRequest : CompletedThreadedRequests)
+		for (FHttpRequestCommon* CompletedRequest : CompletedThreadedRequests)
 		{
 			FHttpRequestRef CompletedRequestRef = CompletedRequest->AsShared();
 
@@ -495,6 +496,7 @@ bool FHttpManager::Tick(float DeltaSeconds)
 				BroadcastHttpRequestCompleted(CompletedRequestRef);
 			}
 		}
+PRAGMA_ENABLE_DEPRECATION_WARNINGS
 	}
 
 	// keep ticking
@@ -523,22 +525,22 @@ void FHttpManager::RemoveRequest(const FHttpRequestRef& Request)
 	Requests.Remove(Request);
 }
 
-void FHttpManager::AddThreadedRequest(const TSharedRef<IHttpThreadedRequest, ESPMode::ThreadSafe>& Request)
+PRAGMA_DISABLE_DEPRECATION_WARNINGS
+void FHttpManager::AddThreadedRequest(const TSharedRef<FHttpRequestCommon, ESPMode::ThreadSafe>& Request)
 {
 	check(Thread);
 	{
-PRAGMA_DISABLE_DEPRECATION_WARNINGS
 		AddRequest(Request);
-PRAGMA_ENABLE_DEPRECATION_WARNINGS
 	}
 	Thread->AddRequest(&Request.Get());
 }
 
-void FHttpManager::CancelThreadedRequest(const TSharedRef<IHttpThreadedRequest, ESPMode::ThreadSafe>& Request)
+void FHttpManager::CancelThreadedRequest(const TSharedRef<FHttpRequestCommon, ESPMode::ThreadSafe>& Request)
 {
 	check(Thread);
 	Thread->CancelRequest(&Request.Get());
 }
+PRAGMA_ENABLE_DEPRECATION_WARNINGS
 
 bool FHttpManager::IsValidRequest(const IHttpRequest* RequestPtr) const
 {
diff --git a/Engine/Source/Runtime/Online/HTTP/Private/HttpThread.cpp b/Engine/Source/Runtime/Online/HTTP/Private/HttpThread.cpp
--- a/Engine/Source/Runtime/Online/HTTP/Private/HttpThread.cpp
+++ b/Engine/Source/Runtime/Online/HTTP/Private/HttpThread.cpp
@@ -1,7 +1,7 @@
 // Copyright Epic Games, Inc. All Rights Reserved.
 
 #include "HttpThread.h"
-#include "IHttpThreadedRequest.h"
+#include "GenericPlatform/HttpRequestCommon.h"
 #include "HAL/PlatformTime.h"
 #include "HAL/PlatformProcess.h"
 #include "HAL/RunnableThread.h"
@@ -151,24 +151,26 @@ void FHttpThreadBase::StopThread()
 	bIsSingleThread = true;
 }
 
-void FHttpThreadBase::AddRequest(IHttpThreadedRequest* Request)
+PRAGMA_DISABLE_DEPRECATION_WARNINGS
+void FHttpThreadBase::AddRequest(FHttpRequestCommon* Request)
 {
 	NewThreadedRequests.Enqueue(Request);
 }
 
-void FHttpThreadBase::CancelRequest(IHttpThreadedRequest* Request)
+void FHttpThreadBase::CancelRequest(FHttpRequestCommon* Request)
 {
 	CancelledThreadedRequests.Enqueue(Request);
 }
 
-void FHttpThreadBase::GetCompletedRequests(TArray<IHttpThreadedRequest*>& OutCompletedRequests)
+void FHttpThreadBase::GetCompletedRequests(TArray<FHttpRequestCommon*>& OutCompletedRequests)
 {
-	IHttpThreadedRequest* Request = nullptr;
+	FHttpRequestCommon* Request = nullptr;
 	while (CompletedThreadedRequests.Dequeue(Request))
 	{
 		OutCompletedRequests.Add(Request);
 	}
 }
+PRAGMA_ENABLE_DEPRECATION_WARNINGS
 
 bool FHttpThreadBase::Init()
 {
@@ -203,12 +205,13 @@ void FHttpThreadBase::HttpThreadTick(float DeltaSeconds)
 {
 }
 
-bool FHttpThreadBase::StartThreadedRequest(IHttpThreadedRequest* Request)
+PRAGMA_DISABLE_DEPRECATION_WARNINGS
+bool FHttpThreadBase::StartThreadedRequest(FHttpRequestCommon* Request)
 {
 	return Request->StartThreadedRequest();
 }
 
-void FHttpThreadBase::CompleteThreadedRequest(IHttpThreadedRequest* Request)
+void FHttpThreadBase::CompleteThreadedRequest(FHttpRequestCommon* Request)
 {
 	// empty
 }
@@ -223,11 +226,11 @@ void FHttpThreadBase::Exit()
 	// empty
 }
 
-void FHttpThreadBase::ConsumeCanceledRequestsAndNewRequests(TArray<IHttpThreadedRequest*>& RequestsToCancel, TArray<IHttpThreadedRequest*>& RequestsToComplete)
+void FHttpThreadBase::ConsumeCanceledRequestsAndNewRequests(TArray<FHttpRequestCommon*>& RequestsToCancel, TArray<FHttpRequestCommon*>& RequestsToComplete)
 {
 	// cache all cancelled and new requests
 	{
-		IHttpThreadedRequest* Request = nullptr;
+		FHttpRequestCommon* Request = nullptr;
 
 		RequestsToCancel.Reset();
 		while (CancelledThreadedRequests.Dequeue(Request))
@@ -243,7 +246,7 @@ void FHttpThreadBase::ConsumeCanceledRequestsAndNewRequests(TArray<IHttpThreaded
 	}
 
 	// Cancel any pending cancel requests
-	for (IHttpThreadedRequest* Request : RequestsToCancel)
+	for (FHttpRequestCommon* Request : RequestsToCancel)
 	{
 		if (RunningThreadedRequests.Remove(Request) > 0)
 		{
@@ -261,7 +264,7 @@ void FHttpThreadBase::ConsumeCanceledRequestsAndNewRequests(TArray<IHttpThreaded
 	}
 }
 
-void FHttpThreadBase::StartRequestsWaitingInQueue(TArray<IHttpThreadedRequest*>& RequestsToComplete)
+void FHttpThreadBase::StartRequestsWaitingInQueue(TArray<FHttpRequestCommon*>& RequestsToComplete)
 {
 	// We'll start rate limited requests until we hit the limit
 	// Tick new requests separately from existing RunningThreadedRequests so they get a chance 
@@ -274,7 +277,7 @@ void FHttpThreadBase::StartRequestsWaitingInQueue(TArray<IHttpThreadedRequest*>&
 		{
 			SCOPE_CYCLE_COUNTER(STAT_HTTPThread_StartThreadedRequest);
 
-			IHttpThreadedRequest* ReadyThreadedRequest = RateLimitedThreadedRequests[0];
+			FHttpRequestCommon* ReadyThreadedRequest = RateLimitedThreadedRequests[0];
 			RateLimitedThreadedRequests.RemoveAt(0);
 
 			float DurationInQueue = FPlatformTime::Seconds() - ReadyThreadedRequest->GetTimeStartedWaitingInQueue();
@@ -306,7 +309,7 @@ void FHttpThreadBase::StartRequestsWaitingInQueue(TArray<IHttpThreadedRequest*>&
 	}
 }
 
-void FHttpThreadBase::MoveCompletingRequestsToCompletedRequests(TArray<IHttpThreadedRequest*>& RequestsToComplete)
+void FHttpThreadBase::MoveCompletingRequestsToCompletedRequests(TArray<FHttpRequestCommon*>& RequestsToComplete)
 {
 	const double AppTime = FPlatformTime::Seconds();
 	const double ElapsedTime = AppTime - LastTime;
@@ -314,7 +317,7 @@ void FHttpThreadBase::MoveCompletingRequestsToCompletedRequests(TArray<IHttpThre
 
 	// Tick any running requests
 	// as long as they properly finish in HttpThreadTick below they are unaffected by a possibly large ElapsedTime above
-	for (IHttpThreadedRequest* Request : RunningThreadedRequests)
+	for (FHttpRequestCommon* Request : RunningThreadedRequests)
 	{
 		SCOPE_CYCLE_COUNTER(STAT_HTTPThread_TickThreadedRequest);
 
@@ -334,7 +337,7 @@ void FHttpThreadBase::MoveCompletingRequestsToCompletedRequests(TArray<IHttpThre
 	{
 		SCOPE_CYCLE_COUNTER(STAT_HTTPThread_IsThreadedRequestComplete);
 
-		IHttpThreadedRequest* Request = RunningThreadedRequests[Index];
+		FHttpRequestCommon* Request = RunningThreadedRequests[Index];
 
 		if (Request->IsThreadedRequestComplete())
 		{
@@ -346,11 +349,11 @@ void FHttpThreadBase::MoveCompletingRequestsToCompletedRequests(TArray<IHttpThre
 	}
 }
 
-void FHttpThreadBase::FinishRequestsFromHttpThreadWithCallbacks(TArray<IHttpThreadedRequest*>& RequestsToComplete)
+void FHttpThreadBase::FinishRequestsFromHttpThreadWithCallbacks(TArray<FHttpRequestCommon*>& RequestsToComplete)
 {
 	if (RequestsToComplete.Num() > 0)
 	{
-		for (IHttpThreadedRequest* Request : RequestsToComplete)
+		for (FHttpRequestCommon* Request : RequestsToComplete)
 		{
 			SCOPE_CYCLE_COUNTER(STAT_HTTPThread_CompleteThreadedRequest);
 
@@ -368,7 +371,7 @@ void FHttpThreadBase::FinishRequestsFromHttpThreadWithCallbacks(TArray<IHttpThre
 	}
 }
 
-void FHttpThreadBase::Process(TArray<IHttpThreadedRequest*>& RequestsToCancel, TArray<IHttpThreadedRequest*>& RequestsToComplete)
+void FHttpThreadBase::Process(TArray<FHttpRequestCommon*>& RequestsToCancel, TArray<FHttpRequestCommon*>& RequestsToComplete)
 {
 	SCOPE_CYCLE_COUNTER(STAT_HTTPThread_Process);
 
@@ -409,17 +412,17 @@ void FLegacyHttpThread::StopThread()
 	FHttpThreadBase::StopThread();
 }
 
-void FLegacyHttpThread::AddRequest(IHttpThreadedRequest* Request)
+void FLegacyHttpThread::AddRequest(FHttpRequestCommon* Request)
 {
 	FHttpThreadBase::AddRequest(Request);
 }
 
-void FLegacyHttpThread::CancelRequest(IHttpThreadedRequest* Request)
+void FLegacyHttpThread::CancelRequest(FHttpRequestCommon* Request)
 {
 	FHttpThreadBase::CancelRequest(Request);
 }
 
-void FLegacyHttpThread::GetCompletedRequests(TArray<IHttpThreadedRequest*>& OutCompletedRequests)
+void FLegacyHttpThread::GetCompletedRequests(TArray<FHttpRequestCommon*>& OutCompletedRequests)
 {
 	FHttpThreadBase::GetCompletedRequests(OutCompletedRequests);
 }
@@ -430,11 +433,12 @@ void FLegacyHttpThread::Tick()
 
 	if (ensure(NeedsSingleThreadTick()))
 	{
-		TArray<IHttpThreadedRequest*> RequestsToCancel;
-		TArray<IHttpThreadedRequest*> RequestsToComplete;
+		TArray<FHttpRequestCommon*> RequestsToCancel;
+		TArray<FHttpRequestCommon*> RequestsToComplete;
 		Process(RequestsToCancel, RequestsToComplete);
 	}
 }
+PRAGMA_ENABLE_DEPRECATION_WARNINGS
 
 TSharedPtr<IHttpTaskTimerHandle> FLegacyHttpThread::AddHttpThreadTask(TFunction<void()>&& Task, float InDelay)
 {
@@ -470,9 +474,11 @@ bool FLegacyHttpThread::Init()
 UE_DISABLE_OPTIMIZATION_SHIP
 uint32 FLegacyHttpThread::Run()
 {
+PRAGMA_DISABLE_DEPRECATION_WARNINGS
 	// Arrays declared outside of loop to re-use memory
-	TArray<IHttpThreadedRequest*> RequestsToCancel;
-	TArray<IHttpThreadedRequest*> RequestsToComplete;
+	TArray<FHttpRequestCommon*> RequestsToCancel;
+	TArray<FHttpRequestCommon*> RequestsToComplete;
+PRAGMA_ENABLE_DEPRECATION_WARNINGS
 	while (!ExitRequest.GetValue())
 	{
 		if (ensureMsgf(!NeedsSingleThreadTick(), TEXT("HTTP Thread was set to singlethread mode while it was running autonomously!")))
diff --git a/Engine/Source/Runtime/Online/HTTP/Private/HttpThread.h b/Engine/Source/Runtime/Online/HTTP/Private/HttpThread.h
--- a/Engine/Source/Runtime/Online/HTTP/Private/HttpThread.h
+++ b/Engine/Source/Runtime/Online/HTTP/Private/HttpThread.h
@@ -15,7 +15,7 @@
 
 #include <atomic>
 
-class IHttpThreadedRequest;
+class FHttpRequestCommon;
 class FHttpThreadBase;
 
 class IHttpTaskTimerHandle
@@ -53,26 +53,28 @@ public:
 	 */
 	bool IsStopped() const { return bIsStopped; }
 
+PRAGMA_DISABLE_DEPRECATION_WARNINGS
 	/**
 	 * Add a request to begin processing on HTTP thread.
 	 *
 	 * @param Request the request to be processed on the HTTP thread
 	 */
-	virtual void AddRequest(IHttpThreadedRequest* Request);
+	virtual void AddRequest(FHttpRequestCommon* Request);
 
 	/**
 	 * Mark a request as cancelled.    Called on non-HTTP thread.
 	 *
 	 * @param Request the request to be processed on the HTTP thread
 	 */
-	virtual void CancelRequest(IHttpThreadedRequest* Request);
+	virtual void CancelRequest(FHttpRequestCommon* Request);
 
 	/**
 	 * Get completed requests.  Clears internal arrays.  Called on non-HTTP thread.
 	 *
 	 * @param OutCompletedRequests array of requests that have been completed
 	 */
-	virtual void GetCompletedRequests(TArray<IHttpThreadedRequest*>& OutCompletedRequests);
+	virtual void GetCompletedRequests(TArray<FHttpRequestCommon*>& OutCompletedRequests);
+PRAGMA_ENABLE_DEPRECATION_WARNINGS
 
 	//~ Begin FSingleThreadRunnable Interface
 	virtual void Tick() override;
@@ -109,15 +111,16 @@ protected:
 	 */
 	virtual void HttpThreadTick(float DeltaSeconds);
 	
+PRAGMA_DISABLE_DEPRECATION_WARNINGS
 	/** 
 	 * Start processing a request on the http thread
 	 */
-	virtual bool StartThreadedRequest(IHttpThreadedRequest* Request);
+	virtual bool StartThreadedRequest(FHttpRequestCommon* Request);
 
 	/** 
 	 * Complete a request on the http thread
 	 */
-	virtual void CompleteThreadedRequest(IHttpThreadedRequest* Request);
+	virtual void CompleteThreadedRequest(FHttpRequestCommon* Request);
 
 protected:
 
@@ -130,7 +133,7 @@ protected:
 	virtual void Exit() override;
 	//~ End FRunnable Interface
 
-	void Process(TArray<IHttpThreadedRequest*>& RequestsToCancel, TArray<IHttpThreadedRequest*>& RequestsToComplete);
+	void Process(TArray<FHttpRequestCommon*>& RequestsToCancel, TArray<FHttpRequestCommon*>& RequestsToComplete);
 
 	/**
 	*  FSingleThreadRunnable accessor for ticking this FRunnable when multi-threading is disabled.
@@ -139,11 +142,12 @@ protected:
 	virtual class FSingleThreadRunnable* GetSingleThreadInterface() override { return this; }
 
 private:
-	void ConsumeCanceledRequestsAndNewRequests(TArray<IHttpThreadedRequest*>& RequestsToCancel, TArray<IHttpThreadedRequest*>& RequestsToComplete);
-	void MoveCompletingRequestsToCompletedRequests(TArray<IHttpThreadedRequest*>& RequestsToComplete);
-	void StartRequestsWaitingInQueue(TArray<IHttpThreadedRequest*>& RequestsToComplete);
-	void FinishRequestsFromHttpThreadWithCallbacks(TArray<IHttpThreadedRequest*>& RequestsToComplete);
+	void ConsumeCanceledRequestsAndNewRequests(TArray<FHttpRequestCommon*>& RequestsToCancel, TArray<FHttpRequestCommon*>& RequestsToComplete);
+	void MoveCompletingRequestsToCompletedRequests(TArray<FHttpRequestCommon*>& RequestsToComplete);
+	void StartRequestsWaitingInQueue(TArray<FHttpRequestCommon*>& RequestsToComplete);
+	void FinishRequestsFromHttpThreadWithCallbacks(TArray<FHttpRequestCommon*>& RequestsToComplete);
 	void UpdateThreadPriorityIfNeeded();
+PRAGMA_ENABLE_DEPRECATION_WARNINGS
 
 protected:
 	/** Pointer to Runnable Thread */
@@ -163,35 +167,37 @@ private:
 	EThreadPriority CurrentThreadPriority;
 
 protected:
+PRAGMA_DISABLE_DEPRECATION_WARNINGS
 	/** 
 	 * Threaded requests that are waiting to be processed on the http thread.
 	 * Added to on (any) non-HTTP thread, processed then cleared on HTTP thread.
 	 */
-	TMpscQueue<IHttpThreadedRequest*> NewThreadedRequests;
+	TMpscQueue<FHttpRequestCommon*> NewThreadedRequests;
 
 	/**
 	 * Threaded requests that are waiting to be cancelled on the http thread.
 	 * Added to on (any) non-HTTP thread, processed then cleared on HTTP thread.
 	 */
-	TMpscQueue<IHttpThreadedRequest*> CancelledThreadedRequests;
+	TMpscQueue<FHttpRequestCommon*> CancelledThreadedRequests;
 
 	/**
 	 * Threaded requests that are ready to run, but waiting due to the running request limit (not in any of the other lists, except potentially CancelledThreadedRequests).
 	 * Only accessed on the HTTP thread.
 	 */
-	TArray<IHttpThreadedRequest*> RateLimitedThreadedRequests;
+	TArray<FHttpRequestCommon*> RateLimitedThreadedRequests;
 
 	/**
 	 * Currently running threaded requests (not in any of the other lists, except potentially CancelledThreadedRequests).
 	 * Only accessed on the HTTP thread.
 	 */
-	TArray<IHttpThreadedRequest*> RunningThreadedRequests;
+	TArray<FHttpRequestCommon*> RunningThreadedRequests;
 
 	/**
 	 * Threaded requests that have completed and are waiting for the game thread to process.
 	 * Added to on HTTP thread, processed then cleared on game thread (Single producer, single consumer)
 	 */
-	TSpscQueue<IHttpThreadedRequest*> CompletedThreadedRequests;
+	TSpscQueue<FHttpRequestCommon*> CompletedThreadedRequests;
+PRAGMA_ENABLE_DEPRECATION_WARNINGS
 };
 
 class FLegacyHttpThread	: public FHttpThreadBase
@@ -203,9 +209,11 @@ public:
 
 	virtual void StartThread() override final;
 	virtual void StopThread() override final;
-	virtual void AddRequest(IHttpThreadedRequest* Request) override final;
-	virtual void CancelRequest(IHttpThreadedRequest* Request) override final;
-	virtual void GetCompletedRequests(TArray<IHttpThreadedRequest*>& OutCompletedRequests) override final;
+PRAGMA_DISABLE_DEPRECATION_WARNINGS
+	virtual void AddRequest(FHttpRequestCommon* Request) override final;
+	virtual void CancelRequest(FHttpRequestCommon* Request) override final;
+	virtual void GetCompletedRequests(TArray<FHttpRequestCommon*>& OutCompletedRequests) override final;
+PRAGMA_ENABLE_DEPRECATION_WARNINGS
 
 	//~ Begin FSingleThreadRunnable Interface
 	// Cannot be overridden to ensure identical behavior with the threaded tick
diff --git a/Engine/Source/Runtime/Online/HTTP/Private/IHttpThreadedRequest.cpp b/Engine/Source/Runtime/Online/HTTP/Private/IHttpThreadedRequest.cpp
deleted file mode 100644
diff --git a/Engine/Source/Runtime/Online/HTTP/Private/IHttpThreadedRequest.h b/Engine/Source/Runtime/Online/HTTP/Private/IHttpThreadedRequest.h
deleted file mode 100644
diff --git a/Engine/Source/Runtime/Online/HTTP/Public/GenericPlatform/HttpRequestCommon.h b/Engine/Source/Runtime/Online/HTTP/Public/GenericPlatform/HttpRequestCommon.h
--- a/Engine/Source/Runtime/Online/HTTP/Public/GenericPlatform/HttpRequestCommon.h
+++ b/Engine/Source/Runtime/Online/HTTP/Public/GenericPlatform/HttpRequestCommon.h
@@ -11,7 +11,7 @@ class IHttpTaskTimerHandle;
 /**
  * Contains implementation of some common functions that don't vary between implementations of different platforms
  */
-class FHttpRequestCommon : public FHttpRequestImpl
+class UE_DEPRECATED(5.5, "FHttpRequestCommon is deprecated and will be moved to internal") FHttpRequestCommon : public FHttpRequestImpl
 {
 public:
 	FHttpRequestCommon();
@@ -48,6 +48,13 @@ public:
 
 	HTTP_API virtual float GetElapsedTime() const override;
 
+	HTTP_API virtual bool IsThreadedRequestComplete() = 0;
+	HTTP_API virtual bool StartThreadedRequest() = 0;
+	HTTP_API virtual void TickThreadedRequest(float DeltaSeconds) = 0;
+
+	HTTP_API void StartWaitingInQueue();
+	HTTP_API float GetTimeStartedWaitingInQueue() const;
+
 protected:
 	/**
 	 * Check if this request is valid or allowed, before actually process the request
@@ -145,7 +152,9 @@ protected:
 	FString EffectiveURL;
 
 	/** The response object which we will use to pair with this request */
+PRAGMA_DISABLE_DEPRECATION_WARNINGS
 	TSharedPtr<FHttpResponseCommon> ResponseCommon;
+PRAGMA_ENABLE_DEPRECATION_WARNINGS
 
 	/** The stream to receive response body */
 	TSharedPtr<FArchive> ResponseBodyReceiveStream;
@@ -162,4 +171,7 @@ protected:
 
 	/** Total elapsed time in seconds since the start of the request */
 	float ElapsedTime = 0.0f;
+
+	/** Record the time started to wait in the queue */
+	float TimeStartedWaitingInQueue = 0.0f;
 };
diff --git a/Engine/Source/Runtime/Online/HTTP/Public/GenericPlatform/HttpResponseCommon.h b/Engine/Source/Runtime/Online/HTTP/Public/GenericPlatform/HttpResponseCommon.h
--- a/Engine/Source/Runtime/Online/HTTP/Public/GenericPlatform/HttpResponseCommon.h
+++ b/Engine/Source/Runtime/Online/HTTP/Public/GenericPlatform/HttpResponseCommon.h
@@ -9,12 +9,14 @@ class FHttpRequestCommon;
 /**
  * Contains implementation of some common functions that don't vary between implementations of different platforms
  */
-class FHttpResponseCommon : public IHttpResponse
+class UE_DEPRECATED(5.5, "FHttpResponseCommon is deprecated and will be moved to internal") FHttpResponseCommon : public IHttpResponse
 {
+PRAGMA_DISABLE_DEPRECATION_WARNINGS
 	friend FHttpRequestCommon;
 
 public:
 	HTTP_API FHttpResponseCommon(const FHttpRequestCommon& HttpRequest);
+PRAGMA_ENABLE_DEPRECATION_WARNINGS
 
 	// IHttpBase
 	HTTP_API virtual FString GetURLParameter(const FString& ParameterName) const override;
diff --git a/Engine/Source/Runtime/Online/HTTP/Public/HttpManager.h b/Engine/Source/Runtime/Online/HTTP/Public/HttpManager.h
--- a/Engine/Source/Runtime/Online/HTTP/Public/HttpManager.h
+++ b/Engine/Source/Runtime/Online/HTTP/Public/HttpManager.h
@@ -19,9 +19,11 @@
 #include "Templates/SharedPointer.h"
 
 class FHttpThreadBase;
+PRAGMA_DISABLE_DEPRECATION_WARNINGS
+class FHttpRequestCommon;
+PRAGMA_ENABLE_DEPRECATION_WARNINGS
 class FOutputDevice;
 class IHttpTaskTimerHandle;
-class IHttpThreadedRequest;
 
 enum class EHttpFlushReason : uint8
 {
@@ -160,19 +162,21 @@ public:
 	 */
 	HTTP_API virtual void FlushTick(float DeltaSeconds);
 
+PRAGMA_DISABLE_DEPRECATION_WARNINGS
 	/** 
 	 * Add a http request to be executed on the http thread
 	 *
 	 * @param Request - the request object to add
 	 */
-	HTTP_API void AddThreadedRequest(const TSharedRef<IHttpThreadedRequest, ESPMode::ThreadSafe>& Request);
+	HTTP_API void AddThreadedRequest(const TSharedRef<FHttpRequestCommon, ESPMode::ThreadSafe>& Request);
 
 	/**
 	 * Mark a threaded http request as cancelled to be removed from the http thread
 	 *
 	 * @param Request - the request object to cancel
 	 */
-	HTTP_API void CancelThreadedRequest(const TSharedRef<IHttpThreadedRequest, ESPMode::ThreadSafe>& Request);
+	HTTP_API void CancelThreadedRequest(const TSharedRef<FHttpRequestCommon, ESPMode::ThreadSafe>& Request);
+PRAGMA_ENABLE_DEPRECATION_WARNINGS
 
 	/**
 	 * List all of the Http requests currently being processed
