Fix the link error in editor by exporting http api SetResponseBodyReceiveStreamDelegate and SetResponseBodyReceiveStreamDelegateV2 in IHttpRequest.

diff --git a/Engine/Source/Runtime/Online/HTTP/Public/Interfaces/IHttpRequest.h b/Engine/Source/Runtime/Online/HTTP/Public/Interfaces/IHttpRequest.h
--- a/Engine/Source/Runtime/Online/HTTP/Public/Interfaces/IHttpRequest.h
+++ b/Engine/Source/Runtime/Online/HTTP/Public/Interfaces/IHttpRequest.h
@@ -240,7 +240,7 @@ public:
 	 * @return True if the delegate can be used. False otherwise.
 	 */
 	UE_DEPRECATED(5.5, "SetResponseBodyReceiveStreamDelegate has been deprecated, use SetResponseBodyReceiveStreamDelegateV2 instead")
-	bool SetResponseBodyReceiveStreamDelegate(FHttpRequestStreamDelegate StreamDelegate);
+	HTTP_API bool SetResponseBodyReceiveStreamDelegate(FHttpRequestStreamDelegate StreamDelegate);
 
 	/**
 	 * Sets the delegate to receive the response body. Make sure to handle the cleanup of received data when
@@ -253,7 +253,7 @@ public:
 	 * @param StreamDelegate - will be used to receive the response body
 	 * @return True if the delegate can be used. False otherwise.
 	 */
-	bool SetResponseBodyReceiveStreamDelegateV2(FHttpRequestStreamDelegateV2 StreamDelegate);
+	HTTP_API bool SetResponseBodyReceiveStreamDelegateV2(FHttpRequestStreamDelegateV2 StreamDelegate);
 
 	/**
 	 * Sets optional header info.
