[PerfTracking] Add csv stats to track HttpManagerThread perf
#tests Run client with -autoperf and generated perf reports from resulting csvs for verification -autoperf and generat perf reports from resulting csvs for verification

diff --git a/Engine/Source/Runtime/Online/HTTP/Private/HttpManager.cpp b/Engine/Source/Runtime/Online/HTTP/Private/HttpManager.cpp
--- a/Engine/Source/Runtime/Online/HTTP/Private/HttpManager.cpp
+++ b/Engine/Source/Runtime/Online/HTTP/Private/HttpManager.cpp
@@ -14,6 +14,7 @@
 #include "GenericPlatform/HttpRequestCommon.h"
 #include "Misc/ConfigCacheIni.h"
 #include "Misc/CommandLine.h"
+#include "ProfilingDebugging/CsvProfiler.h"
 
 #include "Stats/Stats.h"
 #include "Containers/BackgroundableTicker.h"
@@ -58,6 +59,16 @@ namespace
 	}
 }
 
+CSV_DEFINE_CATEGORY(HttpManager, true);
+CSV_DEFINE_STAT(HttpManager, RequestsInQueue);
+CSV_DEFINE_STAT(HttpManager, MaxRequestsInQueue);
+CSV_DEFINE_STAT(HttpManager, RequestsInFlight);
+CSV_DEFINE_STAT(HttpManager, MaxRequestsInFlight);
+CSV_DEFINE_STAT(HttpManager, MaxTimeToWaitInQueue);
+CSV_DEFINE_STAT(HttpManager, DownloadedMB);
+CSV_DEFINE_STAT(HttpManager, BandwidthMbps);
+CSV_DEFINE_STAT(HttpManager, DurationMsAvg);
+
 PRAGMA_DISABLE_DEPRECATION_WARNINGS
 FHttpManager::FHttpManager()
 	: FTSTickerObjectBase(0.0f, FTSBackgroundableTicker::GetCoreTicker())
@@ -498,6 +509,17 @@ PRAGMA_DISABLE_DEPRECATION_WARNINGS
 		}
 PRAGMA_ENABLE_DEPRECATION_WARNINGS
 	}
+	
+	// Report csv stats.
+	int32 TotalDownloadedMB = int32(HttpStats.TotalDownloadedBytes >> 20);
+	CSV_CUSTOM_STAT_DEFINED(RequestsInQueue, HttpStats.RequestsInQueue, ECsvCustomStatOp::Set);
+	CSV_CUSTOM_STAT_DEFINED(MaxRequestsInQueue, int32(HttpStats.MaxRequestsInQueue), ECsvCustomStatOp::Set);
+	CSV_CUSTOM_STAT_DEFINED(RequestsInFlight, HttpStats.RequestsInFlight, ECsvCustomStatOp::Set);
+	CSV_CUSTOM_STAT_DEFINED(MaxRequestsInFlight, int32(HttpStats.MaxRequestsInFlight), ECsvCustomStatOp::Set);
+	CSV_CUSTOM_STAT_DEFINED(MaxTimeToWaitInQueue, int32(HttpStats.MaxTimeToWaitInQueue), ECsvCustomStatOp::Set);
+	CSV_CUSTOM_STAT_DEFINED(DownloadedMB, TotalDownloadedMB, ECsvCustomStatOp::Set);
+	CSV_CUSTOM_STAT_DEFINED(BandwidthMbps, int32(HttpStats.BandwidthMbps), ECsvCustomStatOp::Set);
+	CSV_CUSTOM_STAT_DEFINED(DurationMsAvg, int32(HttpStats.HttpDurationMsAvg), ECsvCustomStatOp::Set);
 
 	// keep ticking
 	return true;
@@ -588,6 +610,30 @@ bool FHttpManager::SupportsDynamicProxy() const
 
 void FHttpManager::BroadcastHttpRequestCompleted(const FHttpRequestRef& Request)
 {
+	FHttpResponsePtr Response = Request->GetResponse();
+	if (Response.IsValid())
+	{
+		const int64 OldDuration = HttpStatsHistory.DurationMs[HttpStatsHistory.HistoryIndex];
+		const int64 NewDuration = FMath::RoundToInt64(Request->GetElapsedTime() * 1000.0f);
+
+		HttpStatsHistory.TotalDuration -= OldDuration;
+		HttpStatsHistory.TotalDuration += NewDuration;
+		HttpStatsHistory.DurationMs[HttpStatsHistory.HistoryIndex] = NewDuration;
+
+		const int64 SizeBytes = Response->GetContentLength();
+		HttpStats.TotalDownloadedBytes += SizeBytes;
+
+		HttpStatsHistory.TotalDownloadedBytes -= HttpStatsHistory.DownloadedBytes[HttpStatsHistory.HistoryIndex];
+		HttpStatsHistory.TotalDownloadedBytes += SizeBytes;
+		HttpStatsHistory.DownloadedBytes[HttpStatsHistory.HistoryIndex] = SizeBytes;
+
+		HttpStats.BandwidthMbps = ((HttpStatsHistory.TotalDownloadedBytes * 8) / (HttpStatsHistory.TotalDuration + 1) / 1000);
+		HttpStats.HttpDurationMsAvg = HttpStatsHistory.TotalDuration / FHttpStatsHistory::HttpHistoryCount;
+
+		// Increment index
+		HttpStatsHistory.HistoryIndex = (HttpStatsHistory.HistoryIndex + 1) % FHttpStatsHistory::HttpHistoryCount;
+	}
+		
 	RequestCompletedDelegate.ExecuteIfBound(Request);
 }
 
@@ -601,8 +647,15 @@ void FHttpManager::RecordStatTimeToConnect(float Duration)
 	HttpStats.MaxTimeToConnect = FGenericPlatformMath::Max(Duration, HttpStats.MaxTimeToConnect);
 }
 
+void FHttpManager::RecordStatRequestsInFlight(uint32 RequestsInFlight)
+{
+	HttpStats.RequestsInFlight = RequestsInFlight;
+	HttpStats.MaxRequestsInFlight = FGenericPlatformMath::Max(RequestsInFlight, HttpStats.MaxRequestsInFlight);
+}
+
 void FHttpManager::RecordStatRequestsInQueue(uint32 RequestsInQueue)
 {
+	HttpStats.RequestsInQueue = RequestsInQueue;
 	HttpStats.MaxRequestsInQueue = FGenericPlatformMath::Max(RequestsInQueue, HttpStats.MaxRequestsInQueue);
 }
 
diff --git a/Engine/Source/Runtime/Online/HTTP/Private/HttpThread.cpp b/Engine/Source/Runtime/Online/HTTP/Private/HttpThread.cpp
--- a/Engine/Source/Runtime/Online/HTTP/Private/HttpThread.cpp
+++ b/Engine/Source/Runtime/Online/HTTP/Private/HttpThread.cpp
@@ -266,6 +266,8 @@ void FHttpThreadBase::ConsumeCanceledRequestsAndNewRequests(TArray<FHttpRequestC
 
 void FHttpThreadBase::StartRequestsWaitingInQueue(TArray<FHttpRequestCommon*>& RequestsToComplete)
 {
+	FHttpManager& HttpManager = FHttpModule::Get().GetHttpManager();
+
 	// We'll start rate limited requests until we hit the limit
 	// Tick new requests separately from existing RunningThreadedRequests so they get a chance 
 	// to send unaffected by possibly large ElapsedTime above
@@ -285,7 +287,7 @@ void FHttpThreadBase::StartRequestsWaitingInQueue(TArray<FHttpRequestCommon*>& R
 			float StartImmediately = 0.01f;
 			if (DurationInQueue > StartImmediately)
 			{
-				FHttpModule::Get().GetHttpManager().RecordMaxTimeToWaitInQueue(DurationInQueue);
+				HttpManager.RecordMaxTimeToWaitInQueue(DurationInQueue);
 			}
 
 			if (StartThreadedRequest(ReadyThreadedRequest))
@@ -303,9 +305,10 @@ void FHttpThreadBase::StartRequestsWaitingInQueue(TArray<FHttpRequestCommon*>& R
 		}
 	}
 
+	HttpManager.RecordStatRequestsInFlight(RunningThreadedRequestsCounter);
 	if (!RateLimitedThreadedRequests.IsEmpty())
 	{
-		FHttpModule::Get().GetHttpManager().RecordStatRequestsInQueue(RateLimitedThreadedRequests.Num());
+		HttpManager.RecordStatRequestsInQueue(RateLimitedThreadedRequests.Num());
 	}
 }
 
diff --git a/Engine/Source/Runtime/Online/HTTP/Public/HttpManager.h b/Engine/Source/Runtime/Online/HTTP/Public/HttpManager.h
--- a/Engine/Source/Runtime/Online/HTTP/Public/HttpManager.h
+++ b/Engine/Source/Runtime/Online/HTTP/Public/HttpManager.h
@@ -54,16 +54,34 @@ DECLARE_DELEGATE_OneParam(FHttpManagerRequestCompletedDelegate, const FHttpReque
 
 struct FHttpStats
 {
+	/** The number of requests waiting in queue in http manager */
+	int32 RequestsInQueue = 0;
+	/** The number of requests in flight in http manager */
+	int32 RequestsInFlight = 0;
 	/** The max time to successfully connect the backend */
 	float MaxTimeToConnect = -1.0f;
 	/** The max waiting queue in http manager */
 	uint32 MaxRequestsInQueue = 0;
+	/** The max number of requests in flight in http manager */
+	uint32 MaxRequestsInFlight = 0;
 	/** The max waiting time in queue of http manager */
 	float MaxTimeToWaitInQueue = 0.0f;
+	/** The total bytes downloaded so far */
+	int64 TotalDownloadedBytes = 0;
+	/** Approximate download bandwidth used */
+	int64 BandwidthMbps = 0;
+	/** Avg duration (in milliseconds) from request to response */
+	int64 HttpDurationMsAvg = 0;
 
 	bool operator==(const FHttpStats& Other) const
 	{
-		return MaxRequestsInQueue == Other.MaxRequestsInQueue
+		return RequestsInQueue == Other.RequestsInQueue
+			&& RequestsInFlight == Other.RequestsInFlight
+			&& MaxRequestsInQueue == Other.MaxRequestsInQueue
+			&& MaxRequestsInFlight == Other.MaxRequestsInFlight
+			&& TotalDownloadedBytes == Other.TotalDownloadedBytes
+			&& BandwidthMbps == Other.BandwidthMbps
+			&& HttpDurationMsAvg == Other.HttpDurationMsAvg
 			&& FMath::IsNearlyEqual(MaxTimeToConnect, Other.MaxTimeToConnect)
 			&& FMath::IsNearlyEqual(MaxTimeToWaitInQueue, Other.MaxTimeToWaitInQueue);
 	}
@@ -343,6 +361,19 @@ protected:
 	TMap<EHttpFlushReason, FHttpFlushTimeLimit> FlushTimeLimitsMap;
 
 	FHttpStats HttpStats;
+	struct FHttpStatsHistory
+	{
+		static constexpr int32 HttpHistoryCount = 16;
+		int32 HistoryIndex = 0;
+
+		int64 DownloadedBytes[HttpHistoryCount] = { 0 };
+		int64 DurationMs[HttpHistoryCount] = { 0 };
+
+		int64 TotalDownloadedBytes = 0;
+		int64 TotalUploadedBytes = 0;
+		int64 TotalDuration = 0;
+	};
+	FHttpStatsHistory HttpStatsHistory;
 
 	bool bUseEventLoop = true;
 
@@ -374,6 +405,9 @@ PACKAGE_SCOPE:
 	/** Record the time to connect, to have a general idea how long the client usually take to connect for success requests, to adjust the connection timeout */
 	HTTP_API void RecordStatTimeToConnect(float Duration);
 
+	/** Record the requests waiting in flight, to have an idea if there are too many concurrent requests */
+	HTTP_API void RecordStatRequestsInFlight(uint32 RequestsInFlight);
+
 	/** Record the requests waiting in queue, to have an idea if there are too many requests or if request number limit is too small */
 	HTTP_API void RecordStatRequestsInQueue(uint32 RequestsInQueue);
 
