Add platform http stats to monitor memory usage.

diff --git a/Engine/Source/Runtime/Online/HTTP/Private/HttpManager.cpp b/Engine/Source/Runtime/Online/HTTP/Private/HttpManager.cpp
--- a/Engine/Source/Runtime/Online/HTTP/Private/HttpManager.cpp
+++ b/Engine/Source/Runtime/Online/HTTP/Private/HttpManager.cpp
@@ -664,6 +664,11 @@ void FHttpManager::RecordMaxTimeToWaitInQueue(float Duration)
 	HttpStats.MaxTimeToWaitInQueue = FGenericPlatformMath::Max(Duration, HttpStats.MaxTimeToWaitInQueue);
 }
 
+void FHttpManager::RecordPlatformStats(const FHttpStatsPlatform& PlatformStats)
+{
+	HttpStats.PlatformStats = PlatformStats;
+}
+
 void FHttpManager::UpdateUrlPatternsToLogResponse(IConsoleVariable* CVar)
 {
 	const FScopeLock CacheLock(&UrlPatternsToLogResponseCriticalSection);
diff --git a/Engine/Source/Runtime/Online/HTTP/Public/HttpManager.h b/Engine/Source/Runtime/Online/HTTP/Public/HttpManager.h
--- a/Engine/Source/Runtime/Online/HTTP/Public/HttpManager.h
+++ b/Engine/Source/Runtime/Online/HTTP/Public/HttpManager.h
@@ -52,6 +52,20 @@ DECLARE_DELEGATE_OneParam(FHttpManagerRequestAddedDelegate, const FHttpRequestRe
  */
 DECLARE_DELEGATE_OneParam(FHttpManagerRequestCompletedDelegate, const FHttpRequestRef& /*Request*/);
 
+struct FHttpStatsPlatformMemoryPool
+{
+	uint64 PoolSize = 0;
+	uint64 MaxInUseSize = 0;
+	uint64 CurrentInUseSize = 0;
+};
+
+struct FHttpStatsPlatform
+{
+	FHttpStatsPlatformMemoryPool MemoryPoolConnection;
+	FHttpStatsPlatformMemoryPool MemoryPoolSsl;
+	FHttpStatsPlatformMemoryPool MemoryPoolNet;
+};
+
 struct FHttpStats
 {
 	/** The number of requests waiting in queue in http manager */
@@ -72,7 +86,10 @@ struct FHttpStats
 	int64 BandwidthMbps = 0;
 	/** Avg duration (in milliseconds) from request to response */
 	int64 HttpDurationMsAvg = 0;
+	/** The optional http stats on specific platform */
+	TOptional<FHttpStatsPlatform> PlatformStats;
 
+	UE_DEPRECATED(5.5, "operator== for FHttpStats has been deprecated and will be removed.")
 	bool operator==(const FHttpStats& Other) const
 	{
 		return RequestsInQueue == Other.RequestsInQueue
@@ -414,5 +431,8 @@ PACKAGE_SCOPE:
 	/** Record the time to wait in queue, to have a general idea how long the client usually wait before actually starting, to adjust the requests */
 	HTTP_API void RecordMaxTimeToWaitInQueue(float Duration);
 
+	/** Record platform specific stats */
+	HTTP_API void RecordPlatformStats(const FHttpStatsPlatform& PlatformStats);
+
 	HTTP_API bool ShouldLogResponse(FStringView Url);
 };
