Extract http response code from http impls into FHttpRequestCommon;
Change CVarHttpUrlPatternsToLogResponse to use space instead of comma to separate url patterns;
Add CVarHttpUrlPatternsToMockFailure to mock http connect error or response failures through url patterns.

#tests Tried in game client also passed all the tests including new added test cases in WebTests

diff --git a/Engine/Source/Programs/WebTests/Private/TestHttp.cpp b/Engine/Source/Programs/WebTests/Private/TestHttp.cpp
--- a/Engine/Source/Programs/WebTests/Private/TestHttp.cpp
+++ b/Engine/Source/Programs/WebTests/Private/TestHttp.cpp
@@ -39,6 +39,7 @@
 extern TAutoConsoleVariable<bool> CVarHttpInsecureProtocolEnabled;
 extern TAutoConsoleVariable<bool> CVarHttpRetrySystemNonGameThreadSupportEnabled;
 extern TAutoConsoleVariable<int32> CVarHttpMaxConcurrentRequests;
+extern TAutoConsoleVariable<FString> CVarHttpUrlPatternsToMockFailure;
 
 class FMockHttpModule : public FHttpModule
 {
@@ -395,6 +396,41 @@ TEST_CASE_METHOD(FWaitUntilCompleteHttpFixture, "Can process https request", HTT
 	HttpRequest->ProcessRequest();
 }
 
+TEST_CASE_METHOD(FWaitUntilCompleteHttpFixture, "Can mock connection error through CVar", HTTP_TAG)
+{
+	CVarHttpUrlPatternsToMockFailure->Set(TEXT("epicgames.com->0 unrealengine.com->503"));
+
+	float ExpectedTimeoutDuration = 2.0f;
+	HttpModule->HttpConnectionTimeout = ExpectedTimeoutDuration;
+	const double StartTime = FPlatformTime::Seconds();
+
+	TSharedRef<IHttpRequest> HttpRequest = CreateRequest();
+	HttpRequest->SetURL(TEXT("https://www.epicgames.com/"));
+	HttpRequest->OnProcessRequestComplete().BindLambda([StartTime, ExpectedTimeoutDuration](FHttpRequestPtr HttpRequest, FHttpResponsePtr HttpResponse, bool bSucceeded) {
+		CHECK(!bSucceeded);
+		CHECK(!HttpResponse);
+		CHECK(HttpRequest->GetFailureReason() == EHttpFailureReason::ConnectionError);
+		const double DurationInSeconds  = FPlatformTime::Seconds() - StartTime;
+		CHECK(FMath::IsNearlyEqual(DurationInSeconds, ExpectedTimeoutDuration, HTTP_TIME_DIFF_TOLERANCE_OF_REQUEST));
+	});
+	HttpRequest->ProcessRequest();
+}
+
+TEST_CASE_METHOD(FWaitUntilCompleteHttpFixture, "Can mock response failure through CVar", HTTP_TAG)
+{
+	CVarHttpUrlPatternsToMockFailure->Set(TEXT("epicgames.com->0 unrealengine.com->503"));
+
+	TSharedRef<IHttpRequest> HttpRequest = CreateRequest();
+	HttpRequest->SetVerb(TEXT("GET"));
+	HttpRequest->SetURL(TEXT("https://www.unrealengine.com/"));
+	HttpRequest->OnProcessRequestComplete().BindLambda([](FHttpRequestPtr HttpRequest, FHttpResponsePtr HttpResponse, bool bSucceeded) {
+		CHECK(bSucceeded);
+		REQUIRE(HttpResponse != nullptr);
+		CHECK(HttpResponse->GetResponseCode() == 503);
+	});
+	HttpRequest->ProcessRequest();
+}
+
 TEST_CASE_METHOD(FWaitUntilCompleteHttpFixture, "Can complete successfully for different response codes", HTTP_TAG)
 {
 	TSharedRef<IHttpRequest> HttpRequest = CreateRequest();
diff --git a/Engine/Source/Runtime/Online/HTTP/Private/Apple/AppleHttp.cpp b/Engine/Source/Runtime/Online/HTTP/Private/Apple/AppleHttp.cpp
--- a/Engine/Source/Runtime/Online/HTTP/Private/Apple/AppleHttp.cpp
+++ b/Engine/Source/Runtime/Online/HTTP/Private/Apple/AppleHttp.cpp
@@ -126,14 +126,20 @@ extern TAutoConsoleVariable<bool> CVarHttpSetGeneralFailureReasonFromCommonCode;
 	[super dealloc];
 }
 
-- (void) HandleStatusCodeReceived:(int32) StatusCode
+- (void) HandleStatusCodeReceived
 {
 	if (TSharedPtr<FAppleHttpRequest> Request = SourceRequest.Pin())
 	{
-		Request->TriggerStatusCodeReceivedDelegate(StatusCode);
+		int32 StatusCode = [self GetStatusCode];
+		Request->HandleStatusCodeReceived(StatusCode);
 	}
 }
 
+- (void)SetRequestStatus:(EHttpRequestStatus::Type)InRequestStatus
+{
+	self.RequestStatus = InRequestStatus;
+}
+
 - (bool)HandleBodyDataReceived:(void*)Ptr Size:(int64)InSize
 {
 	if (TSharedPtr<FAppleHttpRequest> Request = SourceRequest.Pin())
@@ -211,9 +217,7 @@ extern TAutoConsoleVariable<bool> CVarHttpSetGeneralFailureReasonFromCommonCode;
 
 	self.Response = response;
 
-	int32 StatusCode = [self GetStatusCode];
-	
-	[self HandleStatusCodeReceived: StatusCode];
+	[self HandleStatusCodeReceived];
 
 	NSURL* Url = [self.Response URL];
 	FString EffectiveURL([Url absoluteString]);
@@ -817,10 +821,10 @@ bool FAppleHttpRequest::SetupRequest()
 		SetStatus(EHttpRequestStatus::Processing);
 		SetFailureReason(EHttpFailureReason::None);
 
-		TSharedPtr<FAppleHttpResponse> Response = MakeShared<FAppleHttpResponse>(*this);
-		ResponseCommon = Response;
+		InitResponse();
 
 		// Both Task and Response keep a strong reference to the delegate
+		TSharedPtr<FAppleHttpResponse> Response = StaticCastSharedPtr<FAppleHttpResponse>(ResponseCommon);
 		Task.delegate = Response->ResponseDelegate;
 
 		//Setup delegates before starting the request
@@ -839,6 +843,17 @@ bool FAppleHttpRequest::SetupRequest()
 	}
 }
 
+FHttpResponsePtr FAppleHttpRequest::CreateResponse()
+{
+	return MakeShared<FAppleHttpResponse>(*this);
+}
+
+void FAppleHttpRequest::MockResponseData()
+{
+	TSharedPtr<FAppleHttpResponse> Response = StaticCastSharedPtr<FAppleHttpResponse>(ResponseCommon);
+	[Response->ResponseDelegate SetRequestStatus: EHttpRequestStatus::Succeeded];
+}
+
 void FAppleHttpRequest::FinishRequest()
 {
 	PostProcess();
@@ -852,29 +867,31 @@ void FAppleHttpRequest::FinishRequest()
 	}
 	else
 	{
-		EHttpFailureReason Reason = EHttpFailureReason::Other;
-		if (Response)
+		if (FailureReason == EHttpFailureReason::None) // FailureReason could have been set by FHttpRequestCommon::WillTriggerMockFailure
 		{
-			Reason = Response->GetFailureReasonFromDelegate();
-			if (Reason == EHttpFailureReason::Cancelled)
+			EHttpFailureReason Reason = EHttpFailureReason::Other;
+			if (Response)
 			{
-				if (bTimedOut)
-				{
-					Reason = EHttpFailureReason::TimedOut;
-				}
-				else if (bActivityTimedOut)
+				Reason = Response->GetFailureReasonFromDelegate();
+				if (Reason == EHttpFailureReason::Cancelled)
 				{
-					Reason = EHttpFailureReason::ConnectionError;
+					if (bTimedOut)
+					{
+						Reason = EHttpFailureReason::TimedOut;
+					}
+					else if (bActivityTimedOut)
+					{
+						Reason = EHttpFailureReason::ConnectionError;
+					}
 				}
 			}
+			else if (bCanceled)
+			{
+				Reason = EHttpFailureReason::Cancelled;
+			}
+			SetFailureReason(Reason);
 		}
-		else if (bCanceled)
-		{
-			Reason = EHttpFailureReason::Cancelled;
-		}
-		SetFailureReason(Reason);
 
-		UE_LOG(LogHttp, Verbose, TEXT("Request failed: %p Reason %s"), this, LexToString(Reason));
 		if (GetFailureReason() == EHttpFailureReason::ConnectionError)
 		{
 			ResponseCommon = nullptr;
@@ -1068,11 +1085,6 @@ FString FAppleHttpResponse::GetContentAsString() const
 	return UTF8_TO_TCHAR( ZeroTerminatedPayload.GetData() );
 }
 
-int32 FAppleHttpResponse::GetResponseCode() const
-{
-	return [ResponseDelegate GetStatusCode];
-}
-
 bool FAppleHttpResponse::IsReady() const
 {
 	return EHttpRequestStatus::IsFinished(ResponseDelegate.RequestStatus);
diff --git a/Engine/Source/Runtime/Online/HTTP/Private/Apple/AppleHttp.h b/Engine/Source/Runtime/Online/HTTP/Private/Apple/AppleHttp.h
--- a/Engine/Source/Runtime/Online/HTTP/Private/Apple/AppleHttp.h
+++ b/Engine/Source/Runtime/Online/HTTP/Private/Apple/AppleHttp.h
@@ -77,7 +77,7 @@ PRAGMA_DISABLE_DEPRECATION_WARNINGS
 	using FHttpRequestCommon::StartActivityTimeoutTimer;
 	using FHttpRequestCommon::ResetActivityTimeoutTimer;
 	using FHttpRequestCommon::StopActivityTimeoutTimer;
-	using FHttpRequestCommon::TriggerStatusCodeReceivedDelegate;
+	using FHttpRequestCommon::HandleStatusCodeReceived;
 	using FHttpRequestCommon::SetEffectiveURL;
 	using FHttpRequestCommon::PassReceivedDataToStream;
 PRAGMA_ENABLE_DEPRECATION_WARNINGS
@@ -97,6 +97,9 @@ private:
 	 */
 	virtual bool SetupRequest() override;
 
+	virtual FHttpResponsePtr CreateResponse() override;
+	virtual void MockResponseData() override;
+
 	virtual void AbortRequest() override;
 
 	/**
@@ -155,7 +158,6 @@ public:
 	//~ End IHttpBase Interface
 
 	//~ Begin IHttpResponse Interface
-	virtual int32 GetResponseCode() const override;
 	virtual FString GetContentAsString() const override;
 	//~ End IHttpResponse Interface
 
diff --git a/Engine/Source/Runtime/Online/HTTP/Private/Curl/CurlHttp.cpp b/Engine/Source/Runtime/Online/HTTP/Private/Curl/CurlHttp.cpp
--- a/Engine/Source/Runtime/Online/HTTP/Private/Curl/CurlHttp.cpp
+++ b/Engine/Source/Runtime/Online/HTTP/Private/Curl/CurlHttp.cpp
@@ -21,6 +21,20 @@
 
 extern TAutoConsoleVariable<bool> CVarHttpSetGeneralFailureReasonFromCommonCode;
 
+TAutoConsoleVariable<bool> CVarHttpCurlSetResponseCodeWhenMarkAsCompleted(
+	TEXT("http.CurlSetResponseCodeWhenMarkAsCompleted"),
+	true,
+	TEXT("Set response code in when mark as completed, instead of setting it in FinishRequest."),
+	ECVF_SaveForNextBoot
+);
+
+TAutoConsoleVariable<bool> CVarHttpCurlReadContentLengthWhenFinish(
+	TEXT("http.CurlReadContentLengthWhenFinish"),
+	false,
+	TEXT("Whether to read content length from curl handle when finish request"),
+	ECVF_SaveForNextBoot
+);
+
 #if WITH_SSL
 static int SslCertVerify(int PreverifyOk, X509_STORE_CTX* Context)
 {
@@ -447,11 +461,7 @@ size_t FCurlHttpRequest::ReceiveResponseHeaderCallback(void* Ptr, size_t SizeInB
 {
 	QUICK_SCOPE_CYCLE_COUNTER(STAT_FCurlHttpRequest_ReceiveResponseHeaderCallback);
 
-	if (!ResponseCommon.IsValid())
-	{
-		ResponseCommon = MakeShared<FCurlHttpResponse>(*this);
-		TotalBytesRead = 0;
-	}
+	InitResponse();
 
 	OnAnyActivityOccur(TEXTVIEW("Received header"));
 
@@ -550,10 +560,7 @@ size_t FCurlHttpRequest::ReceiveResponseBodyCallback(void* Ptr, size_t SizeInBlo
 	QUICK_SCOPE_CYCLE_COUNTER(STAT_FCurlHttpRequest_ReceiveResponseBodyCallback);
 	LLM_SCOPE(ELLMTag::Networking);
 
-	if (!ResponseCommon.IsValid())
-	{
-		ResponseCommon = MakeShared<FCurlHttpResponse>(*this);
-	}
+	InitResponse();
 
 	OnAnyActivityOccur(TEXTVIEW("Received body"));
 
@@ -1149,6 +1156,19 @@ void FCurlHttpRequest::MarkAsCompleted(CURLcode InCurlCompletionResult)
 	CurlCompletionResult = InCurlCompletionResult;
 	bCurlRequestCompleted = true;
 
+	if (CVarHttpCurlSetResponseCodeWhenMarkAsCompleted.GetValueOnAnyThread())
+	{
+		if (TSharedPtr<FCurlHttpResponse> Response = StaticCastSharedPtr<FCurlHttpResponse>(ResponseCommon))
+		{
+			// get the information
+			long HttpCode = 0;
+			if (CURLE_OK == curl_easy_getinfo(EasyHandle, CURLINFO_RESPONSE_CODE, &HttpCode))
+			{
+				Response->SetResponseCode(HttpCode);
+			}
+		}
+	}
+
 	StopActivityTimeoutTimer();
 }
 
@@ -1169,18 +1189,23 @@ void FCurlHttpRequest::FinishRequest()
 		{
 			Response->bSucceeded = (CURLE_OK == CurlCompletionResult);
 
-			// get the information
-			long HttpCode = 0;
-			if (CURLE_OK == curl_easy_getinfo(EasyHandle, CURLINFO_RESPONSE_CODE, &HttpCode))
+			if (!CVarHttpCurlSetResponseCodeWhenMarkAsCompleted.GetValueOnAnyThread())
 			{
-				Response->HttpCode = HttpCode;
+				// get the information
+				long HttpCode = 0;
+				if (CURLE_OK == curl_easy_getinfo(EasyHandle, CURLINFO_RESPONSE_CODE, &HttpCode))
+				{
+					Response->SetResponseCode(HttpCode);
+				}
 			}
 
 			// If content length wasn't received through response header 
 			if (Response->ContentLength == 0)
 			{
 				double ContentLengthDownload = 0.0;
-				if (CURLE_OK == curl_easy_getinfo(EasyHandle, CURLINFO_CONTENT_LENGTH_DOWNLOAD, &ContentLengthDownload) && ContentLengthDownload > 0.0)
+				if (CVarHttpCurlReadContentLengthWhenFinish.GetValueOnAnyThread() && 
+					CURLE_OK == curl_easy_getinfo(EasyHandle, CURLINFO_CONTENT_LENGTH_DOWNLOAD, &ContentLengthDownload) && 
+					ContentLengthDownload > 0.0)
 				{
 					Response->ContentLength = static_cast< uint64 >(ContentLengthDownload);
 				}
@@ -1191,10 +1216,10 @@ void FCurlHttpRequest::FinishRequest()
 				}
 			}
 
-			if (Response->HttpCode <= 0 && URL.StartsWith(TEXT("Http"), ESearchCase::IgnoreCase))
+			if (Response->GetResponseCode() <= 0 && URL.StartsWith(TEXT("Http"), ESearchCase::IgnoreCase))
 			{
 				UE_LOG(LogHttp, Warning, TEXT("%p: invalid HTTP response code received. URL: %s, HTTP code: %d, content length: %llu, actual payload size: %llu"),
-					this, *GetURL(), Response->HttpCode, Response->ContentLength, TotalBytesRead.load());
+					this, *GetURL(), Response->GetResponseCode(), Response->ContentLength, TotalBytesRead.load());
 				Response->bSucceeded = false;
 			}
 		}
@@ -1322,11 +1347,24 @@ void FCurlHttpRequest::CleanupRequest()
 	CloseRequestPayloadDefaultImpl();
 }
 
+FHttpResponsePtr FCurlHttpRequest::CreateResponse()
+{
+	return MakeShared<FCurlHttpResponse>(*this);
+}
+
+void FCurlHttpRequest::MockResponseData()
+{
+	CurlCompletionResult = CURLE_OK;
+	bCurlRequestCompleted = true;
+
+	TSharedPtr<FCurlHttpResponse> Response = StaticCastSharedPtr<FCurlHttpResponse>(ResponseCommon);
+	Response->bSucceeded = true;
+}
+
 // FCurlHttpRequest
 
 FCurlHttpResponse::FCurlHttpResponse(const FCurlHttpRequest& InRequest)
 	: FHttpResponseCommon(InRequest)
-	, HttpCode(EHttpResponseCodes::Unknown)
 	, ContentLength(0)
 	, bIsReady(0)
 	, bSucceeded(0)
@@ -1388,17 +1426,13 @@ const TArray<uint8>& FCurlHttpResponse::GetContent() const
 	return Payload;
 }
 
-int32 FCurlHttpResponse::GetResponseCode() const
-{
-	return HttpCode;
-}
-
 FString FCurlHttpResponse::GetContentAsString() const
 {
 	// Content is NOT null-terminated; we need to specify lengths here
 	FUTF8ToTCHAR TCHARData(reinterpret_cast<const ANSICHAR*>(Payload.GetData()), Payload.Num());
 	return FString(TCHARData.Length(), TCHARData.Get());
 }
+
 PRAGMA_ENABLE_DEPRECATION_WARNINGS
 
 #endif //WITH_CURL
diff --git a/Engine/Source/Runtime/Online/HTTP/Private/Curl/CurlHttp.h b/Engine/Source/Runtime/Online/HTTP/Private/Curl/CurlHttp.h
--- a/Engine/Source/Runtime/Online/HTTP/Private/Curl/CurlHttp.h
+++ b/Engine/Source/Runtime/Online/HTTP/Private/Curl/CurlHttp.h
@@ -355,6 +355,9 @@ private:
 
 	virtual void ClearInCaseOfRetry() override;
 
+	virtual FHttpResponsePtr CreateResponse() override;
+	virtual void MockResponseData() override;
+
 private:
 
 	/** Pointer to an easy handle specific to this request */
@@ -421,7 +424,6 @@ public:
 	//~ End IHttpBase Interface
 
 	//~ Begin IHttpResponse Interface
-	virtual int32 GetResponseCode() const override;
 	virtual FString GetContentAsString() const override;
 	//~ End IHttpResponse Interface
 
@@ -439,8 +441,6 @@ private:
 	TSharedPtr<FArchive> ResponseBodyReceiveStream;
 	/** Cached key/value header pairs. Parsed once request completes. Only accessible on the game thread. */
 	TMap<FString, FString> Headers;
-	/** Cached code from completed response */
-	int32 HttpCode;
 	/** Cached content length from completed response */
 	uint64 ContentLength;
 	/** True when the response has finished async processing */
diff --git a/Engine/Source/Runtime/Online/HTTP/Private/GenericPlatform/HttpRequestCommon.cpp b/Engine/Source/Runtime/Online/HTTP/Private/GenericPlatform/HttpRequestCommon.cpp
--- a/Engine/Source/Runtime/Online/HTTP/Private/GenericPlatform/HttpRequestCommon.cpp
+++ b/Engine/Source/Runtime/Online/HTTP/Private/GenericPlatform/HttpRequestCommon.cpp
@@ -105,11 +105,65 @@ bool FHttpRequestCommon::PreCheck() const
 	return true;
 }
 
+bool FHttpRequestCommon::WillTriggerMockFailure()
+{
+	TOptional<int32> MockResponseCode = FHttpModule::Get().GetHttpManager().GetMockFailure(GetURL());
+	if (MockResponseCode.IsSet())
+	{
+		if (MockResponseCode.GetValue() == EHttpResponseCodes::Unknown)
+		{
+			int32 HttpConnectionTimeout = FHttpModule::Get().GetHttpConnectionTimeout();
+			TWeakPtr<FHttpRequestCommon> RequestWeakPtr(SharedThis(this));
+			FHttpModule::Get().GetHttpManager().AddHttpThreadTask([RequestWeakPtr]() {
+				if (TSharedPtr<FHttpRequestCommon> RequestPtr = RequestWeakPtr.Pin())
+				{
+					RequestPtr->SetFailureReason(EHttpFailureReason::ConnectionError);
+					RequestPtr->FinishRequestNotInHttpManager();
+				}
+			}, HttpConnectionTimeout);
+		}
+		else
+		{
+			InitResponse();
+			ResponseCommon->SetResponseCode(MockResponseCode.GetValue());
+			MockResponseData();
+			FinishRequestNotInHttpManager();
+		}
+
+		return true;
+	}
+
+	return false;
+}
+
+void FHttpRequestCommon::InitResponse()
+{
+	if (!ResponseCommon)
+	{
+		FHttpResponsePtr Response = CreateResponse();
+		ResponseCommon = StaticCastSharedPtr<FHttpResponseCommon>(Response);
+	}
+}
+
 bool FHttpRequestCommon::PreProcess()
 {
 	ClearInCaseOfRetry();
 
-	if (!PreCheck() || !SetupRequest())
+	if (!PreCheck())
+	{
+		FinishRequestNotInHttpManager();
+		return false;
+	}
+
+	if (WillTriggerMockFailure())
+	{
+		// Connect timeout mocking will trigger FinishRequest after a delay, still make sure total timeout 
+		// works when mocking connect timeout
+		StartTotalTimeoutTimer();
+		return false;
+	}
+
+	if (!SetupRequest())
 	{
 		FinishRequestNotInHttpManager();
 		return false;
@@ -339,12 +393,11 @@ void FHttpRequestCommon::StartActivityTimeoutTimerBy(double DelayToTrigger)
 		return;
 	}
 
-	TWeakPtr<IHttpRequest> RequestWeakPtr(AsShared());
+	TWeakPtr<FHttpRequestCommon> RequestWeakPtr(SharedThis(this));
 	ActivityTimeoutHttpTaskTimerHandle = FHttpModule::Get().GetHttpManager().AddHttpThreadTask([RequestWeakPtr]() {
-		if (TSharedPtr<IHttpRequest> RequestPtr = RequestWeakPtr.Pin())
+		if (TSharedPtr<FHttpRequestCommon> RequestPtr = RequestWeakPtr.Pin())
 		{
-			TSharedPtr<FHttpRequestCommon> RequestCommonPtr = StaticCastSharedPtr<FHttpRequestCommon>(RequestPtr);
-			RequestCommonPtr->OnActivityTimeoutTimerTaskTrigger();
+			RequestPtr->OnActivityTimeoutTimerTaskTrigger();
 		}
 	}, DelayToTrigger + 0.05);
 }
@@ -503,6 +556,15 @@ void FHttpRequestCommon::ProcessRequestUntilComplete()
 	FPlatformProcess::ReturnSynchEventToPool(Event);
 }
 
+void FHttpRequestCommon::HandleStatusCodeReceived(int32 StatusCode)
+{
+	if (ResponseCommon)
+	{
+		ResponseCommon->SetResponseCode(StatusCode);
+	}
+	TriggerStatusCodeReceivedDelegate(StatusCode);
+}
+
 void FHttpRequestCommon::TriggerStatusCodeReceivedDelegate(int32 StatusCode)
 {
 	if (DelegateThreadPolicy == EHttpRequestDelegateThreadPolicy::CompleteOnHttpThread)
diff --git a/Engine/Source/Runtime/Online/HTTP/Private/GenericPlatform/HttpResponseCommon.cpp b/Engine/Source/Runtime/Online/HTTP/Private/GenericPlatform/HttpResponseCommon.cpp
--- a/Engine/Source/Runtime/Online/HTTP/Private/GenericPlatform/HttpResponseCommon.cpp
+++ b/Engine/Source/Runtime/Online/HTTP/Private/GenericPlatform/HttpResponseCommon.cpp
@@ -57,4 +57,15 @@ void FHttpResponseCommon::SetEffectiveURL(const FString& InEffectiveURL)
 {
 	EffectiveURL = InEffectiveURL;
 }
+
+int32 FHttpResponseCommon::GetResponseCode() const
+{
+	return ResponseCode;
+}
+
+void FHttpResponseCommon::SetResponseCode(int32 InResponseCode)
+{
+	ResponseCode = InResponseCode;
+}
+
 PRAGMA_ENABLE_DEPRECATION_WARNINGS
diff --git a/Engine/Source/Runtime/Online/HTTP/Private/HttpManager.cpp b/Engine/Source/Runtime/Online/HTTP/Private/HttpManager.cpp
--- a/Engine/Source/Runtime/Online/HTTP/Private/HttpManager.cpp
+++ b/Engine/Source/Runtime/Online/HTTP/Private/HttpManager.cpp
@@ -29,7 +29,14 @@ TAutoConsoleVariable<int32> CVarHttpEventLoopEnableChance(
 TAutoConsoleVariable<FString> CVarHttpUrlPatternsToLogResponse(
 	TEXT("http.UrlPatternsToLogResponse"),
 	TEXT(""),
-	TEXT("List of url patterns to log headers and json content: \"epicgames.com,unrealengine.com,...\""),
+	TEXT("List of url patterns to log headers and json content: \"epicgames.com unrealengine.com ...\""),
+	ECVF_SaveForNextBoot
+);
+
+TAutoConsoleVariable<FString> CVarHttpUrlPatternsToMockFailure(
+	TEXT("http.UrlPatternsToMockFailure"),
+	TEXT(""),
+	TEXT("List of url patterns to mock failure with response code, 0 indicates ConnectionError: \"epicgames.com->0 unrealengine.com->503 ...\""),
 	ECVF_SaveForNextBoot
 );
 
@@ -106,11 +113,15 @@ void FHttpManager::Initialize()
 
 	UpdateUrlPatternsToLogResponse(CVarHttpUrlPatternsToLogResponse.AsVariable());
 	CVarHttpUrlPatternsToLogResponse.AsVariable()->OnChangedDelegate().AddRaw(this, &FHttpManager::UpdateUrlPatternsToLogResponse);
+
+	UpdateUrlPatternsToMockFailure(CVarHttpUrlPatternsToMockFailure.AsVariable());
+	CVarHttpUrlPatternsToMockFailure.AsVariable()->OnChangedDelegate().AddRaw(this, &FHttpManager::UpdateUrlPatternsToMockFailure);
 }
 
 void FHttpManager::Shutdown()
 {
 	CVarHttpUrlPatternsToLogResponse.AsVariable()->OnChangedDelegate().Clear();
+	CVarHttpUrlPatternsToMockFailure.AsVariable()->OnChangedDelegate().Clear();
 
 	{
 		FScopeLock ScopeLock(&RequestLock);
@@ -673,7 +684,7 @@ void FHttpManager::UpdateUrlPatternsToLogResponse(IConsoleVariable* CVar)
 {
 	const FScopeLock CacheLock(&UrlPatternsToLogResponseCriticalSection);
 	const FString UrlPatternsToLogResponseStr = CVar->AsVariable()->GetString();
-	UrlPatternsToLogResponseStr.ParseIntoArray(UrlPatternsToLogResponse, TEXT(","));
+	UrlPatternsToLogResponseStr.ParseIntoArray(UrlPatternsToLogResponse, TEXT(" "));
 }
 
 bool FHttpManager::ShouldLogResponse(FStringView Url)
@@ -690,3 +701,39 @@ bool FHttpManager::ShouldLogResponse(FStringView Url)
 	return false;
 }
 
+void FHttpManager::UpdateUrlPatternsToMockFailure(IConsoleVariable* CVar)
+{
+	const FScopeLock CacheLock(&UrlPatternsToMockFailureCriticalSection);
+	const FString UrlPatternsToMockFailureStr = CVar->AsVariable()->GetString();
+
+	TArray<FString> UrlPatternsToMockFailureStrings;
+	UrlPatternsToMockFailureStr.ParseIntoArray(UrlPatternsToMockFailureStrings, TEXT(" "));
+
+	for (const FString& UrlPatternToMockFailureString : UrlPatternsToMockFailureStrings)
+	{
+		TArray<FString> UrlPattern;
+		UrlPatternToMockFailureString.ParseIntoArray(UrlPattern, TEXT("->"));
+		if (UrlPattern.Num() == 2)
+		{
+			int32 ResponseCode = FCString::Atoi(*UrlPattern[1]);
+			UrlPatternsToMockFailure.Emplace(UrlPattern[0], ResponseCode);
+		}
+	}
+}
+
+TOptional<int32> FHttpManager::GetMockFailure(FStringView Url)
+{
+	TOptional<int32> Result;
+	const FScopeLock CacheLock(&UrlPatternsToMockFailureCriticalSection);
+
+	for (const TPair<FString, int32>& UrlPattern : UrlPatternsToMockFailure)
+	{
+		if (Url.Contains(UrlPattern.Key))
+		{
+			Result = UrlPattern.Value;
+		}
+	}
+
+	return Result;
+}
+
diff --git a/Engine/Source/Runtime/Online/HTTP/Public/GenericPlatform/HttpRequestCommon.h b/Engine/Source/Runtime/Online/HTTP/Public/GenericPlatform/HttpRequestCommon.h
--- a/Engine/Source/Runtime/Online/HTTP/Public/GenericPlatform/HttpRequestCommon.h
+++ b/Engine/Source/Runtime/Online/HTTP/Public/GenericPlatform/HttpRequestCommon.h
@@ -90,6 +90,7 @@ protected:
 
 	HTTP_API virtual void CleanupRequest() = 0;
 
+	HTTP_API void HandleStatusCodeReceived(int32 StatusCode);
 	HTTP_API void TriggerStatusCodeReceivedDelegate(int32 StatusCode);
 
 	HTTP_API void SetEffectiveURL(const FString& InEffectiveURL);
@@ -104,6 +105,13 @@ protected:
 	HTTP_API void CloseRequestPayloadDefaultImpl();
 
 	HTTP_API void LogResponse(const TSharedPtr<IHttpResponse>& InResponse);
+	HTTP_API bool WillTriggerMockFailure();
+
+	HTTP_API void InitResponse();
+
+	HTTP_API virtual FHttpResponsePtr CreateResponse() = 0;
+	// The function for each impl to make sure when mocking response code, the internal state is valid for finishing request successfully
+	HTTP_API virtual void MockResponseData() = 0;
 
 protected:
 	/** Current status of request being processed */
diff --git a/Engine/Source/Runtime/Online/HTTP/Public/GenericPlatform/HttpResponseCommon.h b/Engine/Source/Runtime/Online/HTTP/Public/GenericPlatform/HttpResponseCommon.h
--- a/Engine/Source/Runtime/Online/HTTP/Public/GenericPlatform/HttpResponseCommon.h
+++ b/Engine/Source/Runtime/Online/HTTP/Public/GenericPlatform/HttpResponseCommon.h
@@ -24,14 +24,17 @@ PRAGMA_ENABLE_DEPRECATION_WARNINGS
 	HTTP_API virtual const FString& GetEffectiveURL() const override;
 	HTTP_API virtual EHttpRequestStatus::Type GetStatus() const override;
 	HTTP_API virtual EHttpFailureReason GetFailureReason() const override;
+	HTTP_API virtual int32 GetResponseCode() const override;
 
 protected:
 	HTTP_API void SetRequestStatus(EHttpRequestStatus::Type InCompletionStatus);
 	HTTP_API void SetRequestFailureReason(EHttpFailureReason InFailureReason);
 	HTTP_API void SetEffectiveURL(const FString& InEffectiveURL);
+	HTTP_API void SetResponseCode(int32 InResponseCode);
 
 	FString URL;
 	FString EffectiveURL;
 	EHttpRequestStatus::Type CompletionStatus;
 	EHttpFailureReason FailureReason;
+	int32 ResponseCode = EHttpResponseCodes::Unknown;
 };
diff --git a/Engine/Source/Runtime/Online/HTTP/Public/HttpManager.h b/Engine/Source/Runtime/Online/HTTP/Public/HttpManager.h
--- a/Engine/Source/Runtime/Online/HTTP/Public/HttpManager.h
+++ b/Engine/Source/Runtime/Online/HTTP/Public/HttpManager.h
@@ -325,6 +325,7 @@ protected:
 	HTTP_API bool HasAnyBoundDelegate() const;
 
 	HTTP_API void UpdateUrlPatternsToLogResponse(IConsoleVariable* CVar);
+	HTTP_API void UpdateUrlPatternsToMockFailure(IConsoleVariable* CVar);
 
 protected:
 	/** List of Http requests that are actively being processed */
@@ -395,6 +396,10 @@ protected:
 	bool bUseEventLoop = true;
 
 	TArray<FString> UrlPatternsToLogResponse;
+	FCriticalSection UrlPatternsToLogResponseCriticalSection;
+
+	TMap<FString, int32> UrlPatternsToMockFailure;
+	FCriticalSection UrlPatternsToMockFailureCriticalSection;
 
 PACKAGE_SCOPE:
 
@@ -404,9 +409,6 @@ PACKAGE_SCOPE:
 	/** Used to lock access to get completed requests */
 	FCriticalSection CompletedRequestLock;
 
-	/** Used to lock access to url patterns for logging */
-	FCriticalSection UrlPatternsToLogResponseCriticalSection;
-
 	/**
 	 * Broadcast that a non-threaded HTTP request is complete.
 	 * Called automatically internally for threaded requests.
@@ -435,4 +437,6 @@ PACKAGE_SCOPE:
 	HTTP_API void RecordPlatformStats(const FHttpStatsPlatform& PlatformStats);
 
 	HTTP_API bool ShouldLogResponse(FStringView Url);
+
+	HTTP_API TOptional<int32> GetMockFailure(FStringView Url);
 };
