Do not restrict handling of Retry-After during HTTP retries.

diff --git a/Engine/Source/Programs/WebTests/Private/TestHttp.cpp b/Engine/Source/Programs/WebTests/Private/TestHttp.cpp
--- a/Engine/Source/Programs/WebTests/Private/TestHttp.cpp
+++ b/Engine/Source/Programs/WebTests/Private/TestHttp.cpp
@@ -1868,10 +1868,6 @@ TEST_CASE_METHOD(FWaitUntilCompleteHttpFixture, "Retry respect Retry-After heade
 	{
 		HttpRequest->SetURL(UrlMockStatus(EHttpResponseCodes::TooManyRequests));
 	}
-	SECTION("ServiceUnavail")
-	{
-		HttpRequest->SetURL(UrlMockStatus(EHttpResponseCodes::ServiceUnavail));
-	}
 
 	uint32 RetryAfter = 4;
 
diff --git a/Engine/Source/Runtime/Online/HTTP/Private/HttpRetrySystem.cpp b/Engine/Source/Runtime/Online/HTTP/Private/HttpRetrySystem.cpp
--- a/Engine/Source/Runtime/Online/HTTP/Private/HttpRetrySystem.cpp
+++ b/Engine/Source/Runtime/Online/HTTP/Private/HttpRetrySystem.cpp
@@ -32,40 +32,36 @@ TOptional<double> ReadThrottledTimeFromResponseInSeconds(FHttpResponsePtr Respon
 	// Check if there was a Retry-After header
 	if (Response.IsValid())
 	{
-		int32 ResponseCode = Response->GetResponseCode();
-		if (ResponseCode == EHttpResponseCodes::TooManyRequests || ResponseCode == EHttpResponseCodes::ServiceUnavail)
+		FString RetryAfter = Response->GetHeader(TEXT("Retry-After"));
+		if (!RetryAfter.IsEmpty())
 		{
-			FString RetryAfter = Response->GetHeader(TEXT("Retry-After"));
-			if (!RetryAfter.IsEmpty())
+			if (RetryAfter.IsNumeric())
 			{
-				if (RetryAfter.IsNumeric())
-				{
-					// seconds
-					LockoutPeriod.Emplace(FCString::Atof(*RetryAfter));
-				}
-				else
-				{
-					// http date
-					FDateTime UTCServerTime;
-					if (FDateTime::ParseHttpDate(RetryAfter, UTCServerTime))
-					{
-						const FDateTime UTCNow = FDateTime::UtcNow();
-						LockoutPeriod.Emplace((UTCServerTime - UTCNow).GetTotalSeconds());
-					}
-				}
+				// seconds
+				LockoutPeriod.Emplace(FCString::Atof(*RetryAfter));
 			}
 			else
 			{
-				FString RateLimitReset = Response->GetHeader(TEXT("X-Rate-Limit-Reset"));
-				if (!RateLimitReset.IsEmpty())
+				// http date
+				FDateTime UTCServerTime;
+				if (FDateTime::ParseHttpDate(RetryAfter, UTCServerTime))
 				{
-					// UTC seconds
-					const FDateTime UTCServerTime = FDateTime::FromUnixTimestamp(FCString::Atoi64(*RateLimitReset));
 					const FDateTime UTCNow = FDateTime::UtcNow();
 					LockoutPeriod.Emplace((UTCServerTime - UTCNow).GetTotalSeconds());
 				}
 			}
 		}
+		else
+		{
+			FString RateLimitReset = Response->GetHeader(TEXT("X-Rate-Limit-Reset"));
+			if (!RateLimitReset.IsEmpty())
+			{
+				// UTC seconds
+				const FDateTime UTCServerTime = FDateTime::FromUnixTimestamp(FCString::Atoi64(*RateLimitReset));
+				const FDateTime UTCNow = FDateTime::UtcNow();
+				LockoutPeriod.Emplace((UTCServerTime - UTCNow).GetTotalSeconds());
+			}
+		}
 	}
 	return LockoutPeriod;
 }
