PR #11944: BuildPatchServices: Fix DownloadService issues

diff --git a/Engine/Source/Runtime/Online/BuildPatchServices/Private/Installer/DownloadService.cpp b/Engine/Source/Runtime/Online/BuildPatchServices/Private/Installer/DownloadService.cpp
--- a/Engine/Source/Runtime/Online/BuildPatchServices/Private/Installer/DownloadService.cpp
+++ b/Engine/Source/Runtime/Online/BuildPatchServices/Private/Installer/DownloadService.cpp
@@ -265,13 +265,12 @@ namespace BuildPatchServices
 		{
 			// Load file from drive/network.
 			TUniquePtr<FFileRequest> FileRequest = MakeUnique<FFileRequest>();
+			FileRequest->Task = MakeFileLoadTask(RequestId, FileUri, FileRequest.Get());
+
 			{
 				FScopeLock ScopeLock(&ActiveRequestsCS);
 				ActiveFileRequests.Add(RequestId, MoveTemp(FileRequest));
 			}
-
-			FFileRequest* FileRequestPtr = FileRequest.Get();
-			FileRequest->Task = MakeFileLoadTask(RequestId, FileUri, FileRequestPtr);
 		}
 		DownloadServiceStat->OnDownloadStarted(RequestId, FileUri);
 
@@ -331,7 +330,7 @@ namespace BuildPatchServices
 						bool bIsCanceled = false;
 						while (BytesRead < FileSize && !FileRequest->ShouldCancel.IsCanceled())
 						{
-							bIsCanceled = !FileRequest->ShouldCancel.IsCanceled();
+							bIsCanceled = FileRequest->ShouldCancel.IsCanceled();
 							if (bIsCanceled)
 							{
 								bSuccess = false;
