Fixed some 'deprecated' FString usage.

diff --git a/Engine/Source/Runtime/Online/WebSockets/Private/WinHttp/Support/WinHttpConnectionWebSocket.cpp b/Engine/Source/Runtime/Online/WebSockets/Private/WinHttp/Support/WinHttpConnectionWebSocket.cpp
--- a/Engine/Source/Runtime/Online/WebSockets/Private/WinHttp/Support/WinHttpConnectionWebSocket.cpp
+++ b/Engine/Source/Runtime/Online/WebSockets/Private/WinHttp/Support/WinHttpConnectionWebSocket.cpp
@@ -710,7 +710,7 @@ void FWinHttpConnectionWebSocket::HandleWebSocketCloseComplete()
 	}
 
 	FUTF8ToTCHAR Converter(CloseReasonBuffer, CloseReasonActualLength);
-	FString CloseReason(Converter.Length(), Converter.Get());
+	FString CloseReason = FString::ConstructFromPtrSize(Converter.Get(), Converter.Length());
 
 	UE_LOG(LogWinHttp, Verbose, TEXT("WinHttp WebSocket[%p]: Received Close message from the server. Code=[%d] Reason=[%s]"), this, static_cast<int32>(CloseCode), *CloseReason);
 
diff --git a/Engine/Source/Runtime/Online/WebSockets/Private/WinHttp/WinHttpWebSocket.cpp b/Engine/Source/Runtime/Online/WebSockets/Private/WinHttp/WinHttpWebSocket.cpp
--- a/Engine/Source/Runtime/Online/WebSockets/Private/WinHttp/WinHttpWebSocket.cpp
+++ b/Engine/Source/Runtime/Online/WebSockets/Private/WinHttp/WinHttpWebSocket.cpp
@@ -383,7 +383,7 @@ void FWinHttpWebSocket::HandleWebSocketMessage(EWebSocketMessageType MessageType
 	if (MessageType == EWebSocketMessageType::Utf8 && OnMessage().IsBound())
 	{
 		const FUTF8ToTCHAR TCHARConverter(reinterpret_cast<const ANSICHAR*>(MessagePayload.GetData()), MessagePayload.Num());
-		const FString Message(TCHARConverter.Length(), TCHARConverter.Get());
+		const FString Message = FString::ConstructFromPtrSize(TCHARConverter.Get(), TCHARConverter.Length());
 
 		OnMessage().Broadcast(Message);
 	}
