If we have both hashes, only check Sha for chunk serialization.

#rnx

diff --git a/Engine/Source/Runtime/Online/BuildPatchServices/Private/Data/ChunkData.cpp b/Engine/Source/Runtime/Online/BuildPatchServices/Private/Data/ChunkData.cpp
--- a/Engine/Source/Runtime/Online/BuildPatchServices/Private/Data/ChunkData.cpp
+++ b/Engine/Source/Runtime/Online/BuildPatchServices/Private/Data/ChunkData.cpp
@@ -527,6 +527,7 @@ namespace BuildPatchServices
 
 		virtual IChunkDataAccess* LoadFromArchive(FArchive& Archive, EChunkLoadResult& OutLoadResult) const override
 		{
+			TRACE_CPUPROFILER_EVENT_SCOPE(Chunk_LoadFromArchive);
 			if (Archive.IsLoading())
 			{
 				return Load(Archive, OutLoadResult);
@@ -629,7 +630,10 @@ namespace BuildPatchServices
 							*Header = HeaderCheck;
 
 							// Read the data.
-							Reader.Serialize(Data, Header->DataSizeCompressed);
+							{
+								TRACE_CPUPROFILER_EVENT_SCOPE(Chunk_Serialize);
+								Reader.Serialize(Data, Header->DataSizeCompressed);
+							}
 							if (Reader.IsError() == false)
 							{
 								OutLoadResult = EChunkLoadResult::Success;
@@ -671,7 +675,12 @@ namespace BuildPatchServices
 									}
 								}
 								// Verify.
-								if (OutLoadResult == EChunkLoadResult::Success && (Header->HashType & EChunkHashFlags::RollingPoly64) != EChunkHashFlags::None)
+
+								// If we have both hashes, only check Sha1 instead of both.
+								bool bHasRollingHash = EnumHasAnyFlags(Header->HashType, EChunkHashFlags::RollingPoly64);
+								bool bHasSha1Hash = EnumHasAnyFlags(Header->HashType, EChunkHashFlags::Sha1);
+
+								if (OutLoadResult == EChunkLoadResult::Success && bHasRollingHash && !bHasSha1Hash) // Don't bother verifying the rolling hash if we have a sha1
 								{
 									if (Header->DataSizeCompressed != Header->DataSizeUncompressed || Header->RollingHash != FRollingHash::GetHashForDataSet(Data, Header->DataSizeUncompressed))
 									{
@@ -679,8 +688,9 @@ namespace BuildPatchServices
 									}
 								}
 								FSHAHash ShaHashCheck;
-								if (OutLoadResult == EChunkLoadResult::Success && (Header->HashType & EChunkHashFlags::Sha1) != EChunkHashFlags::None)
+								if (OutLoadResult == EChunkLoadResult::Success && bHasSha1Hash)
 								{
+									TRACE_CPUPROFILER_EVENT_SCOPE(Chunk_ShaHash);
 									FSHA1::HashBuffer(Data, Header->DataSizeUncompressed, ShaHashCheck.Hash);
 									if (!(ShaHashCheck == Header->SHAHash))
 									{
