BPT: Don't bother locking the reference stack if we don't have anything to download.

#rnx

diff --git a/Engine/Source/Runtime/Online/BuildPatchServices/Private/Installer/CloudChunkSource.cpp b/Engine/Source/Runtime/Online/BuildPatchServices/Private/Installer/CloudChunkSource.cpp
--- a/Engine/Source/Runtime/Online/BuildPatchServices/Private/Installer/CloudChunkSource.cpp
+++ b/Engine/Source/Runtime/Online/BuildPatchServices/Private/Installer/CloudChunkSource.cpp
@@ -290,8 +290,9 @@ namespace BuildPatchServices
 		IChunkDataAccess* ChunkData = ChunkStore->Get(DataId);
 		if (ChunkData == nullptr)
 		{
-			// Ensure this chunk is on the list.
+			// Make sure we are trying to download this chunk before waiting for it to complete.
 			EnsureAquiring(DataId);
+
 			// Wait for the chunk to be available.
 			while ((ChunkData = ChunkStore->Get(DataId)) == nullptr && !bShouldAbort)
 			{
@@ -490,9 +491,17 @@ namespace BuildPatchServices
 			// Select the next X chunks that are for downloading, so we can request URIs.
 			TFunction<bool(const FGuid&)> SelectPredicate = [&TotalRequiredChunks, &RequestedChunkUris](const FGuid& ChunkId) 
 			{ 
+				// if we requre it and we haven't already requested it.
 				return TotalRequiredChunks.Contains(ChunkId) && !RequestedChunkUris.Contains(ChunkId); 
 			};
-			TArray<FGuid> ChunkUrisToRequest = ChunkReferenceTracker->SelectFromNextReferences(Configuration.PreFetchMaximum, SelectPredicate);
+			TArray<FGuid> ChunkUrisToRequest;
+
+			// Don't take the lock over the reference stack if we can't ever pass our selection predicate
+			if (TotalRequiredChunks.Num())
+			{
+				ChunkUrisToRequest = ChunkReferenceTracker->SelectFromNextReferences(Configuration.PreFetchMaximum, SelectPredicate);
+			}
+			
 			for (const FGuid& ChunkUriToRequest : ChunkUrisToRequest)
 			{
 				RequestedChunkUris.Add(ChunkUriToRequest);
