Async chunk writes in BPT

#rnx

diff --git a/Engine/Source/Runtime/Online/BuildPatchServices/Private/BuildPatchFileConstructor.cpp b/Engine/Source/Runtime/Online/BuildPatchServices/Private/BuildPatchFileConstructor.cpp
--- a/Engine/Source/Runtime/Online/BuildPatchServices/Private/BuildPatchFileConstructor.cpp
+++ b/Engine/Source/Runtime/Online/BuildPatchServices/Private/BuildPatchFileConstructor.cpp
@@ -18,6 +18,7 @@
 #include "Installer/ChunkReferenceTracker.h"
 #include "Installer/InstallerError.h"
 #include "Installer/InstallerAnalytics.h"
+#include "Installer/InstallerSharedContext.h"
 #include "BuildPatchUtil.h"
 
 using namespace BuildPatchServices;
@@ -284,10 +285,80 @@ FBuildPatchFileConstructor::FBuildPatchFileConstructor(FFileConstructorConfig In
 		}
 		ConstructionStack[(ConstructListNum - 1) - ConstructListIdx] = ConstructListElem;
 	}
+
+	WriteBuffers[0].Reserve(WriteBufferSize);
+	WriteBuffers[1].Reserve(WriteBufferSize);
+
+	WriteJobThread = Configuration.SharedContext->CreateThread();
+	WriteJobCompleteEvent = FPlatformProcess::GetSynchEventFromPool();
+	WriteJobStartEvent = FPlatformProcess::GetSynchEventFromPool();
+	WriteJobThread->RunTask([this]() { WriteJobThreadRun(); });
 }
 
 FBuildPatchFileConstructor::~FBuildPatchFileConstructor()
 {
+	if (bWriteJobRunning)
+	{
+		GLog->Logf(TEXT("FBuildPatchFileConstructor: Write job active during destruction! Very bad."));
+	}
+
+	// Signal background thread to shut down.
+	Abort();
+	WriteJobStartEvent->Trigger();	
+	WriteJobCompleteEvent->Wait();
+
+	FPlatformProcess::ReturnSynchEventToPool(WriteJobCompleteEvent);
+	WriteJobCompleteEvent = nullptr;
+	FPlatformProcess::ReturnSynchEventToPool(WriteJobStartEvent);
+	WriteJobStartEvent = nullptr;
+
+	Configuration.SharedContext->ReleaseThread(WriteJobThread);
+}
+
+
+void FBuildPatchFileConstructor::WriteJobThreadRun()
+{
+	TRACE_CPUPROFILER_EVENT_SCOPE(WriteJobThread);
+
+	for (;;)
+	{
+		bool bSignalWasFired = WriteJobStartEvent->Wait(100 /* ms */);
+
+		if (bSignalWasFired)
+		{
+			// (got signal) -- they launched a job - init to failed job
+			bWriteJobCompleted = false;
+		}
+
+		if (bShouldAbort) // this is also used for graceful shutdown on completion.
+		{
+			// Leave WriteJobCompleted = false;
+			WriteJobCompleteEvent->Trigger();
+			return;
+		}
+
+		if (!bSignalWasFired)
+		{
+			// We hit the timeout checking for an abort signal, wait agian.
+			continue;
+		}
+
+		FileConstructorStat->OnBeforeWrite();
+		ISpeedRecorder::FRecord ActivityRecord;
+		ActivityRecord.CyclesStart = FStatsCollector::GetCycles();
+		
+		{
+			TRACE_CPUPROFILER_EVENT_SCOPE(WriteThread_Serialize)
+			WriteJobArchive->Serialize(WriteJobBufferToWrite->GetData(), WriteJobBufferToWrite->Num());
+		}
+
+		ActivityRecord.Size = WriteJobBufferToWrite->Num();
+		ActivityRecord.CyclesEnd = FStatsCollector::GetCycles();
+		FileConstructorStat->OnAfterWrite(ActivityRecord);
+
+		bWriteJobCompleted = true;
+		WriteJobCompleteEvent->Trigger();
+	}
 }
 
 void FBuildPatchFileConstructor::Run()
@@ -493,6 +564,8 @@ uint64 FBuildPatchFileConstructor::CalculateRequiredDiskSpace(const FFileManifes
 
 bool FBuildPatchFileConstructor::ConstructFileFromChunks(const FString& BuildFilename, const FFileManifest& FileManifest, bool bResumeExisting)
 {
+	TRACE_CPUPROFILER_EVENT_SCOPE(ConstructFileFromChunks);
+
 	bool bSuccess = true;
 	EConstructionError ConstructionError = EConstructionError::None;
 	uint32 LastError = 0;
@@ -601,6 +674,58 @@ bool FBuildPatchFileConstructor::ConstructFileFromChunks(const FString& BuildFil
 		FileConstructorStat->OnResumeCompleted();
 	}
 
+	auto FlushToAsyncWriter = [this](FArchive& DestinationFile, FSHA1& HashState)
+	{
+		if (bStallWhenFileSystemThrottled)
+		{
+			int64 AvailableBytes = FileSystem->GetAllowedBytesToWriteThrottledStorage(*DestinationFile.GetArchiveName());
+			while (WriteBuffers[CurrentFillBuffer].Num() > AvailableBytes)
+			{
+				UE_LOG(LogBuildPatchServices, Display, TEXT("Avaliable write bytes to write throttled storage exhausted (%s).  Sleeping %ds.  Bytes needed: %u, bytes available: %lld")
+					, *DestinationFile.GetArchiveName(), SleepTimeWhenFileSystemThrottledSeconds, WriteBuffers[CurrentFillBuffer].Num(), AvailableBytes);
+				FPlatformProcess::Sleep(SleepTimeWhenFileSystemThrottledSeconds);
+				AvailableBytes = FileSystem->GetAllowedBytesToWriteThrottledStorage(*DestinationFile.GetArchiveName());
+			}
+		}
+		
+		// Wait for the last write to complete.
+		if (bWriteJobRunning)
+		{
+			// We can potentially wait here a while if we are FS throttled.
+			// \todo old code didn't check for abort during throttling, should we add?
+			{
+				TRACE_CPUPROFILER_EVENT_SCOPE(ConstructFileFromChunks_WaitForLastWrite);
+				WriteJobCompleteEvent->Wait();
+			}
+			bWriteJobRunning = false;
+
+			if (DestinationFile.IsError())
+			{
+				return false;
+			}
+
+			// !CurrentFillBuffer is now available for use.
+		}
+
+		// Kick off the write on another thread while we hash the data here.
+		WriteJobBufferToWrite = &WriteBuffers[CurrentFillBuffer];
+		WriteJobArchive = &DestinationFile;
+		bWriteJobRunning = true;
+		WriteJobStartEvent->Trigger();
+
+		// Hash the buffer we are writing while it's writing.
+		{
+			TRACE_CPUPROFILER_EVENT_SCOPE(ConstructFileFromChunks_Hash);
+			HashState.Update(WriteBuffers[CurrentFillBuffer].GetData(), WriteBuffers[CurrentFillBuffer].Num());
+		}
+
+		// Start filling the next buffer.
+		CurrentFillBuffer = !CurrentFillBuffer;
+		WriteBuffers[CurrentFillBuffer].SetNumUninitialized(0, EAllowShrinking::No);
+
+		return true;
+	};
+
 	// Attempt to create the file
 	ISpeedRecorder::FRecord ActivityRecord;
 	FileConstructorStat->OnBeforeAdminister();
@@ -627,9 +752,26 @@ bool FBuildPatchFileConstructor::ConstructFileFromChunks(const FString& BuildFil
 		// For each chunk, load it, and place it's data into the file
 		for (int32 ChunkPartIdx = StartChunkPart; ChunkPartIdx < FileManifest.ChunkParts.Num() && bSuccess && !bShouldAbort; ++ChunkPartIdx)
 		{
+			TRACE_CPUPROFILER_EVENT_SCOPE(ConstructFileFromChunks_Chunk);
+
 			const FChunkPart& ChunkPart = FileManifest.ChunkParts[ChunkPartIdx];
-			bSuccess = InsertChunkData(ChunkPart, *NewFile, HashState, ConstructionError);
-			FileConstructorStat->OnFileProgress(BuildFilename, NewFile->Tell());
+
+			// If we can't fit in the buffer, flush. Conditional arranged to avoid overflow risk.
+			if (ChunkPart.Size > (WriteBufferSize - WriteBuffers[CurrentFillBuffer].Num()))
+			{
+				if (!FlushToAsyncWriter(*NewFile, HashState))
+				{
+					bSuccess = false;
+					InstallerAnalytics->RecordConstructionError(BuildFilename, INDEX_NONE, TEXT("Serialization Error"));
+					UE_LOG(LogBuildPatchServices, Error, TEXT("FBuildPatchFileConstructor: Failed %s due to serialization error"), *BuildFilename);
+					ConstructionError = EConstructionError::SerializeError;
+					break;
+				}
+			}
+
+			bSuccess = AppendChunkData(ChunkPart, WriteBuffers[CurrentFillBuffer], ConstructionError);
+
+			FileConstructorStat->OnFileProgress(BuildFilename, NewFile->Tell() + WriteBuffers[CurrentFillBuffer].Num());
 			if (bSuccess)
 			{
 				CountBytesProcessed(ChunkPart.Size);
@@ -652,9 +794,30 @@ bool FBuildPatchFileConstructor::ConstructFileFromChunks(const FString& BuildFil
 			}
 		}
 
+		if (WriteBuffers[CurrentFillBuffer].Num())
+		{
+			if (!FlushToAsyncWriter(*NewFile, HashState))
+			{
+				bSuccess = false;
+				InstallerAnalytics->RecordConstructionError(BuildFilename, INDEX_NONE, TEXT("Serialization Error"));
+				UE_LOG(LogBuildPatchServices, Error, TEXT("FBuildPatchFileConstructor: Failed %s due to serialization error"), *BuildFilename);
+				ConstructionError = EConstructionError::SerializeError;
+			}
+		}
+
+		// Wait for the last write if there is one
+		if (bWriteJobRunning)
+		{
+			WriteJobCompleteEvent->Wait();
+			bWriteJobRunning = false;
+		}
+
+		bSuccess = !NewFile->IsError();
+
 		// Close the file writer
 		FileConstructorStat->OnBeforeAdminister();
 		ActivityRecord.CyclesStart = FStatsCollector::GetCycles();
+
 		const bool bArchiveSuccess = NewFile->Close();
 		NewFile.Reset();
 		ActivityRecord.CyclesEnd = FStatsCollector::GetCycles();
@@ -782,45 +945,31 @@ bool FBuildPatchFileConstructor::ConstructFileFromChunks(const FString& BuildFil
 	return bSuccess;
 }
 
-bool FBuildPatchFileConstructor::InsertChunkData(const FChunkPart& ChunkPart, FArchive& DestinationFile, FSHA1& HashState, EConstructionError& ConstructionError)
+bool FBuildPatchFileConstructor::AppendChunkData(const FChunkPart& ChunkPart, TArray<uint8>& DestinationBuffer, EConstructionError& ConstructionError)
 {
-	if (bStallWhenFileSystemThrottled)
-	{
-		int64 AvailableBytes = FileSystem->GetAllowedBytesToWriteThrottledStorage(*DestinationFile.GetArchiveName());
-		while (ChunkPart.Size > AvailableBytes)
-		{
-			UE_LOG(LogBuildPatchServices, Display, TEXT("Avaliable write bytes to write throttled storage exhausted (%s).  Sleeping %ds.  Bytes needed: %u, bytes available: %lld")
-				, *DestinationFile.GetArchiveName(), SleepTimeWhenFileSystemThrottledSeconds, ChunkPart.Size, AvailableBytes);
-			FPlatformProcess::Sleep(SleepTimeWhenFileSystemThrottledSeconds);
-			AvailableBytes = FileSystem->GetAllowedBytesToWriteThrottledStorage(*DestinationFile.GetArchiveName());
-		}
-	}
-
-	uint8* Data;
-	uint8* DataStart;
 	ConstructionError = EConstructionError::None;
-	ISpeedRecorder::FRecord ActivityRecord;
+	
 	FileConstructorStat->OnChunkGet(ChunkPart.Guid);
-	IChunkDataAccess* ChunkDataAccess = ChunkSource->Get(ChunkPart.Guid);
+	IChunkDataAccess* ChunkDataAccess = nullptr;
+	{
+		TRACE_CPUPROFILER_EVENT_SCOPE(GetChunkData);
+		ChunkDataAccess = ChunkSource->Get(ChunkPart.Guid);
+	}
 	if (ChunkDataAccess != nullptr)
 	{
+		uint8* Data;
 		ChunkDataAccess->GetDataLock(&Data, nullptr);
-		FileConstructorStat->OnBeforeWrite();
-		ActivityRecord.CyclesStart = FStatsCollector::GetCycles();
-		DataStart = &Data[ChunkPart.Offset];
-		HashState.Update(DataStart, ChunkPart.Size);
-		DestinationFile.Serialize(DataStart, ChunkPart.Size);
-		const bool bSerializeOk = !DestinationFile.IsError();
-		ActivityRecord.Size = ChunkPart.Size;
-		ActivityRecord.CyclesEnd = FStatsCollector::GetCycles();
-		FileConstructorStat->OnAfterWrite(ActivityRecord);
-		ChunkDataAccess->ReleaseDataLock();
-		const bool bPopReferenceOk = ChunkReferenceTracker->PopReference(ChunkPart.Guid);
-		if (!bSerializeOk)
+
+		uint8* DataStart = &Data[ChunkPart.Offset];
+
 		{
-			ConstructionError = EConstructionError::SerializeError;
+			TRACE_CPUPROFILER_EVENT_SCOPE(GetChunkData);
+			DestinationBuffer.Append(DataStart, ChunkPart.Size);
 		}
-		else if (!bPopReferenceOk)
+
+		ChunkDataAccess->ReleaseDataLock();
+		const bool bPopReferenceOk = ChunkReferenceTracker->PopReference(ChunkPart.Guid);
+		if (!bPopReferenceOk)
 		{
 			ConstructionError = EConstructionError::TrackingError;
 		}
diff --git a/Engine/Source/Runtime/Online/BuildPatchServices/Private/BuildPatchFileConstructor.h b/Engine/Source/Runtime/Online/BuildPatchServices/Private/BuildPatchFileConstructor.h
--- a/Engine/Source/Runtime/Online/BuildPatchServices/Private/BuildPatchFileConstructor.h
+++ b/Engine/Source/Runtime/Online/BuildPatchServices/Private/BuildPatchFileConstructor.h
@@ -15,8 +15,12 @@
 class FBuildPatchAppManifest;
 enum class EConstructionError : uint8;
 
+class IBuildInstallerSharedContext;
+
 namespace BuildPatchServices
 {
+	
+
 	struct FChunkPart;
 	class IFileSystem;
 	class IChunkSource;
@@ -25,6 +29,7 @@ namespace BuildPatchServices
 	class IInstallerAnalytics;
 	class IFileConstructorStat;
 	class IBuildManifestSet;
+	class IBuildInstallerThread;
 
 	/**
 	 * A struct containing the configuration values for a file constructor.
@@ -48,6 +53,8 @@ namespace BuildPatchServices
 
 		// The install mode used for this installation.
 		EInstallMode InstallMode;
+
+		IBuildInstallerSharedContext* SharedContext;
 	};
 
 	/**
@@ -147,14 +154,12 @@ namespace BuildPatchServices
 		bool ConstructFileFromChunks(const FString& BuildFilename, const FFileManifest& FileManifest, bool bResumeExisting);
 
 		/**
-		 * Inserts the data data from a chunk into the destination file according to the chunk part info
+		 * Adds the data from a chunk to the given buffer.
 		 * @param ChunkPart          The chunk part details.
-		 * @param DestinationFile    The Filename for the file being constructed.
-		 * @param HashState          An FSHA1 hash state to update with the data going into the destination file.
 		 * @param ConstructionError  Will be set to the error type that ocurred or EConstructionError::None.
 		 * @return true if no errors were detected
 		 */
-		bool InsertChunkData(const FChunkPart& ChunkPart, FArchive& DestinationFile, FSHA1& HashState, EConstructionError& ConstructionError);
+		bool AppendChunkData(const FChunkPart& ChunkPart, TArray<uint8>& DestinationBuffer, EConstructionError& ConstructionError);
 
 		/**
 		 * Delete all contents of a directory
@@ -216,6 +221,23 @@ namespace BuildPatchServices
 
 		// Event executed before deleting an old installation file.
 		FOnBeforeDeleteFile BeforeDeleteFileEvent;
+
+		//
+		// Async write management.
+		IBuildInstallerThread* WriteJobThread;
+		void WriteJobThreadRun();
+
+		// We ping pong between two buffers, filling/hashing one, and writing the other.
+		TArray<uint8> WriteBuffers[2];
+		int32 CurrentFillBuffer = 0;
+		uint32 WriteBufferSize = (4 << 20); // Default write buffer size 4MB.
+
+		FEvent* WriteJobCompleteEvent = nullptr;
+		FEvent* WriteJobStartEvent = nullptr;
+		TArray<uint8>* WriteJobBufferToWrite = nullptr;
+		FArchive* WriteJobArchive = nullptr;
+		std::atomic_bool bWriteJobCompleted = false; // Only set to true if the Serialize() call was completed.
+		bool bWriteJobRunning = false; // Foreground thread only - have we dispatched a job?
 	};
 
 	/**
diff --git a/Engine/Source/Runtime/Online/BuildPatchServices/Private/BuildPatchInstaller.cpp b/Engine/Source/Runtime/Online/BuildPatchServices/Private/BuildPatchInstaller.cpp
--- a/Engine/Source/Runtime/Online/BuildPatchServices/Private/BuildPatchInstaller.cpp
+++ b/Engine/Source/Runtime/Online/BuildPatchServices/Private/BuildPatchInstaller.cpp
@@ -1366,20 +1366,20 @@ namespace BuildPatchServices
 			ChunkSources.Add(CloudChunkSource.Get());
 			TUniquePtr<IChainedChunkSource> ChainedChunkSource(FChainedChunkSourceFactory::Create(
 				ChunkSources));
-			TUniquePtr<FBuildPatchFileConstructor> FileConstructor(new FBuildPatchFileConstructor(
-				FFileConstructorConfig({
-					ManifestSet.Get(),
-					Configuration.InstallDirectory,
-					InstallStagingDir,
-					MetaStagingDir,
-					FilesToConstruct.Array(),
-					Configuration.InstallMode}),
-				FileSystem.Get(),
-				ChainedChunkSource.Get(),
-				ChunkReferenceTracker.Get(),
-				InstallerError.Get(),
-				InstallerAnalytics.Get(),
-				FileConstructorStatistics.Get()));
+
+			FFileConstructorConfig FCC;
+			FCC.ConstructList = FilesToConstruct.Array();
+			FCC.InstallDirectory = Configuration.InstallDirectory;
+			FCC.InstallMode = Configuration.InstallMode;
+			FCC.ManifestSet = ManifestSet.Get();
+			FCC.MetaDirectory = MetaStagingDir;
+			FCC.SharedContext = Configuration.SharedContext.Get();
+			FCC.StagingDirectory = InstallStagingDir;
+
+
+			TUniquePtr<FBuildPatchFileConstructor> FileConstructor(
+				new FBuildPatchFileConstructor(FCC, FileSystem.Get(), ChainedChunkSource.Get(), ChunkReferenceTracker.Get(), InstallerError.Get(),InstallerAnalytics.Get(),FileConstructorStatistics.Get())
+			);
 			FDelegateHandle OnBeforeDeleteFileHandle = FileConstructor->OnBeforeDeleteFile().AddLambda([this, &InstallChunkSource](const FString& FilePath)
 			{
 				FString BuildRelativeFilename = FilePath;
