Delete chunkdbs after we are done with them.

#rnx

diff --git a/Engine/Source/Runtime/Online/BuildPatchServices/Private/BuildPatchFileConstructor.cpp b/Engine/Source/Runtime/Online/BuildPatchServices/Private/BuildPatchFileConstructor.cpp
--- a/Engine/Source/Runtime/Online/BuildPatchServices/Private/BuildPatchFileConstructor.cpp
+++ b/Engine/Source/Runtime/Online/BuildPatchServices/Private/BuildPatchFileConstructor.cpp
@@ -925,6 +925,11 @@ bool FBuildPatchFileConstructor::ConstructFileFromChunks(const FString& BuildFil
 	}
 #endif
 
+	if (bSuccess)
+	{
+		ChunkSource->ReportFileCompletion();
+	}
+
 	// Delete the staging file if unsuccessful by means of any failure that could leave the file in unknown state.
 	if (!bSuccess)
 	{
diff --git a/Engine/Source/Runtime/Online/BuildPatchServices/Private/BuildPatchInstaller.cpp b/Engine/Source/Runtime/Online/BuildPatchServices/Private/BuildPatchInstaller.cpp
--- a/Engine/Source/Runtime/Online/BuildPatchServices/Private/BuildPatchInstaller.cpp
+++ b/Engine/Source/Runtime/Online/BuildPatchServices/Private/BuildPatchInstaller.cpp
@@ -1022,6 +1022,8 @@ namespace BuildPatchServices
 		GConfig->GetFloat(TEXT("Portal.BuildPatch"), TEXT("ChunkDbSourceChunkDbOpenRetryTime"), ChunkDbSourceConfig.ChunkDbOpenRetryTime, GEngineIni);
 		ChunkDbSourceConfig.ChunkDbOpenRetryTime = FMath::Clamp<float>(ChunkDbSourceConfig.ChunkDbOpenRetryTime, 0.5f, 60.0f);
 
+		ChunkDbSourceConfig.bDeleteChunkDBAfterUse = Configuration.bDeleteChunkDbFilesAfterUse;
+
 		return ChunkDbSourceConfig;
 	}
 
@@ -1302,6 +1304,8 @@ namespace BuildPatchServices
 		{
 			TUniquePtr<IChunkDataSerialization> ChunkDataSerialization(FChunkDataSerializationFactory::Create(
 				FileSystem.Get()));
+
+			// Generate the list of chunks we will need and the order in which we will need them.
 			TUniquePtr<IChunkReferenceTracker> ChunkReferenceTracker(FChunkReferenceTrackerFactory::Create(
 				ManifestSet.Get(),
 				FilesToConstruct));
@@ -1323,6 +1327,8 @@ namespace BuildPatchServices
 				MemoryEvictionPolicy.Get(),
 				DiskOverflowStore.Get(),
 				MemoryChunkStoreStatistics.Get()));
+
+			// Add a source for pulling from chunk "databases", basically tarballs of chunks.
 			TUniquePtr<IChunkDbChunkSource> ChunkDbChunkSource(FChunkDbChunkSourceFactory::Create(
 				BuildChunkDbSourceConfig(),
 				Platform.Get(),
@@ -1333,6 +1339,8 @@ namespace BuildPatchServices
 				MessagePump.Get(),
 				InstallerError.Get(),
 				ChunkDbChunkSourceStatistics.Get()));
+
+			// Add a source for pulling from the existing directory (for patching).
 			TUniquePtr<IInstallChunkSource> InstallChunkSource(FInstallChunkSourceFactory::Create(
 				BuildInstallSourceConfig(ChunkDbChunkSource->GetAvailableChunks()),
 				FileSystem.Get(),
@@ -1342,10 +1350,14 @@ namespace BuildPatchServices
 				InstallChunkSourceStatistics.Get(),
 				InstallationInfo,
 				ManifestSet.Get()));
+
+			// Sub out the chunks we have available on disk - anything we don't we have to download.
 			TSet<FGuid> InitialDownloadChunks = ReferencedChunks.Difference(InstallChunkSource->GetAvailableChunks()).Difference(ChunkDbChunkSource->GetAvailableChunks());
 			FileOperationTracker->OnDataStateUpdate(ReferencedChunks.Intersect(ChunkDbChunkSource->GetAvailableChunks()), EFileOperationState::PendingLocalChunkDbData);
 			FileOperationTracker->OnDataStateUpdate(ReferencedChunks.Intersect(InstallChunkSource->GetAvailableChunks()).Difference(ChunkDbChunkSource->GetAvailableChunks()), EFileOperationState::PendingLocalInstallData);
 			FileOperationTracker->OnDataStateUpdate(InitialDownloadChunks, EFileOperationState::PendingRemoteCloudData);
+
+			// Add the source for downloading chunks we don't have on disk.
 			TUniquePtr<IDownloadConnectionCount> DownloadConnectionCount(FDownloadConnectionCountFactory::Create(BuildConnectionCountConfig(), DownloadServiceStatistics.Get()));
 			TUniquePtr<ICloudChunkSource> CloudChunkSource(FCloudChunkSourceFactory::Create(
 				BuildCloudSourceConfig(),
@@ -1360,6 +1372,8 @@ namespace BuildPatchServices
 				CloudChunkSourceStatistics.Get(),
 				ManifestSet.Get(),
 				MoveTemp(InitialDownloadChunks)));
+
+			// Set up a fallback chain of sources.
 			TArray<IChunkSource*> ChunkSources;
 			ChunkSources.Add(ChunkDbChunkSource.Get());
 			ChunkSources.Add(InstallChunkSource.Get());
@@ -1367,6 +1381,7 @@ namespace BuildPatchServices
 			TUniquePtr<IChainedChunkSource> ChainedChunkSource(FChainedChunkSourceFactory::Create(
 				ChunkSources));
 
+			// Set up the class that actually requests chunks and writes them to disk.
 			FFileConstructorConfig FCC;
 			FCC.ConstructList = FilesToConstruct.Array();
 			FCC.InstallDirectory = Configuration.InstallDirectory;
diff --git a/Engine/Source/Runtime/Online/BuildPatchServices/Private/Installer/ChainedChunkSource.cpp b/Engine/Source/Runtime/Online/BuildPatchServices/Private/Installer/ChainedChunkSource.cpp
--- a/Engine/Source/Runtime/Online/BuildPatchServices/Private/Installer/ChainedChunkSource.cpp
+++ b/Engine/Source/Runtime/Online/BuildPatchServices/Private/Installer/ChainedChunkSource.cpp
@@ -15,6 +15,13 @@ namespace BuildPatchServices
 		virtual TSet<FGuid> AddRuntimeRequirements(TSet<FGuid> NewRequirements) override;
 		virtual void SetUnavailableChunksCallback(TFunction<void(TSet<FGuid>)> Callback) override;
 		virtual bool AddRepeatRequirement(const FGuid& RepeatRequirement) override;
+		virtual void ReportFileCompletion() override
+		{
+			for (IChunkSource* ChunkSource : ChunkSources)
+			{
+				ChunkSource->ReportFileCompletion();
+			}
+		}
 		// IChunkSource interface end.
 
 	private:
diff --git a/Engine/Source/Runtime/Online/BuildPatchServices/Private/Installer/ChunkDbChunkSource.cpp b/Engine/Source/Runtime/Online/BuildPatchServices/Private/Installer/ChunkDbChunkSource.cpp
--- a/Engine/Source/Runtime/Online/BuildPatchServices/Private/Installer/ChunkDbChunkSource.cpp
+++ b/Engine/Source/Runtime/Online/BuildPatchServices/Private/Installer/ChunkDbChunkSource.cpp
@@ -15,6 +15,7 @@
 #include "Installer/MessagePump.h"
 #include "Installer/InstallerError.h"
 #include "Installer/InstallerSharedContext.h"
+#include "Misc/OutputDeviceRedirector.h"
 
 #if PLATFORM_WINDOWS
 #include "Windows/WindowsHWrapper.h"
@@ -75,6 +76,30 @@ namespace BuildPatchServices
 	{
 		FChunkDatabaseHeader Header;
 		TUniquePtr<FArchive> Archive;
+		FString ChunkDbFileName;
+
+		// When the reference tracker gets below this watermark, then we know we are done with this file and we can
+		// close/retire it.
+		int32 RetireAt = 0;
+
+		// If we're retired then any access is invalid and fatal as the file has been closed and could be deleted.
+		bool bIsRetired = false;
+
+		void Retire(IMessagePump* MessagePump, IFileSystem* FileSystemIfDeleting, bool bDelete)
+		{
+			bIsRetired = true;
+			FString ArchiveName = Archive->GetArchiveName();
+			Archive.Reset();
+			MessagePump->SendMessage(FChunkSourceEvent{ FChunkSourceEvent::EType::Retired, ArchiveName });
+
+			if (bDelete && FileSystemIfDeleting)
+			{
+				if (!FileSystemIfDeleting->DeleteFile(*ChunkDbFileName))
+				{
+					GLog->Logf(TEXT("Failed to delete chunkdb upon retirement: %s"), *ChunkDbFileName);
+				}
+			}
+		}
 	};
 
 	/**
@@ -115,6 +140,22 @@ namespace BuildPatchServices
 		virtual TSet<FGuid> AddRuntimeRequirements(TSet<FGuid> NewRequirements) override;
 		virtual bool AddRepeatRequirement(const FGuid& RepeatRequirement) override;
 		virtual void SetUnavailableChunksCallback(TFunction<void(TSet<FGuid>)> Callback) override;
+		virtual void ReportFileCompletion() override
+		{
+			//
+			// Since we've completed a file we know we won't need to resume/retry it and can delete
+			// the source chunkdb that it used.
+			//
+			int32 RemainingChunkCount = ChunkReferenceTracker->GetRemainingChunkCount();
+			for (FChunkDbDataAccess& ChunkDbDataAccess : ChunkDbDataAccesses)
+			{
+				if (!ChunkDbDataAccess.bIsRetired &&
+					ChunkDbDataAccess.RetireAt > RemainingChunkCount)
+				{
+					ChunkDbDataAccess.Retire(MessagePump, FileSystem, Configuration.bDeleteChunkDBAfterUse);
+				}
+			}
+		}
 		// IChunkSource interface end.
 
 		// IInstallChunkSource interface begin.
@@ -178,6 +219,8 @@ namespace BuildPatchServices
 		, ChunkDbChunkSourceStat(InChunkDbChunkSourceStat)
 		, bStartedLoading(!InConfiguration.bBeginLoadsOnFirstGet)
 	{
+		TRACE_CPUPROFILER_EVENT_SCOPE(FChunkDbChunkSource_ctor);
+
 		// Allow OS intervention only once.
 		bool bResetOsIntervention = false;
 		int32 PreviousOsIntervention = 0;
@@ -190,10 +233,22 @@ namespace BuildPatchServices
 				// Load header.
 				FChunkDatabaseHeader Header;
 				*ChunkDbFile << Header;
-				if (!ChunkDbFile->IsError() && Header.Contents.Num() > 0)
+				if (ChunkDbFile->IsError())
+				{
+					GLog->Logf(TEXT("Failed to load chunkdb header for %s"), *ChunkDbFilename);
+				}
+				else if (Header.Contents.Num() == 0)
+				{
+					GLog->Logf(TEXT("Loaded empty chunkdb %s"), *ChunkDbFilename);
+				}
+				else
 				{
 					// Hold on to the handle and header info.
-					ChunkDbDataAccesses.Add({MoveTemp(Header), MoveTemp(ChunkDbFile)});
+					FChunkDbDataAccess DataSource;
+					DataSource.Archive = MoveTemp(ChunkDbFile);
+					DataSource.ChunkDbFileName = ChunkDbFilename;
+					DataSource.Header = MoveTemp(Header);
+					ChunkDbDataAccesses.Add(MoveTemp(DataSource));
 				}
 			}
 			else if(!bResetOsIntervention)
@@ -221,6 +276,40 @@ namespace BuildPatchServices
 			}
 		}
 
+		// Once we have all the guids here we can predict when we'll be done with each
+		// chunkdb file by looking at the reference tracker.
+		TArray<FGuid> OrderedGuidList = InChunkReferenceTracker->GetNextReferences(0x7fffffff, [](const FGuid&) { return true; });
+
+		TMap<FString, int32> FileLastSeenAt;
+
+		int32 GuidIndex = 0;
+		for (; GuidIndex < OrderedGuidList.Num(); GuidIndex++)
+		{
+			const FGuid& Guid = OrderedGuidList[GuidIndex];
+
+			FChunkAccessLookup* SourceForGuid = ChunkDbDataAccessLookup.Find(Guid);
+			if (!SourceForGuid)
+			{
+				continue;
+			}
+
+			FileLastSeenAt.FindOrAdd(SourceForGuid->DbFile->Archive->GetArchiveName()) = GuidIndex;
+		}
+
+		for (FChunkDbDataAccess& ChunkDbDataAccess : ChunkDbDataAccesses)
+		{
+			int32* LastAt = FileLastSeenAt.Find(ChunkDbDataAccess.Archive->GetArchiveName());
+			if (LastAt == nullptr)
+			{
+				// If the file isn't going to be used, it's immediately retired.
+				ChunkDbDataAccess.Retire(MessagePump, FileSystem, Configuration.bDeleteChunkDBAfterUse);
+			}
+			else
+			{
+				ChunkDbDataAccess.RetireAt = OrderedGuidList.Num() - *LastAt - 1;
+			}
+		}
+
 		// Start threaded load worker.
 		if (ChunkDbDataAccesses.Num() > 0)
 		{
@@ -257,17 +346,27 @@ namespace BuildPatchServices
 
 	IChunkDataAccess* FChunkDbChunkSource::Get(const FGuid& DataId)
 	{
+		TRACE_CPUPROFILER_EVENT_SCOPE(ChunkDb_Get);
+
 		// Get from our store
 		IChunkDataAccess* ChunkData = ChunkStore->Get(DataId);
 		if (ChunkData == nullptr)
 		{
 			bStartedLoading = true;
-			if (ChunkDbDataAccessLookup.Contains(DataId) && AvailableChunks.Contains(DataId))
+			FChunkAccessLookup* ChunkInfo = ChunkDbDataAccessLookup.Find(DataId);
+			if (ChunkInfo && AvailableChunks.Contains(DataId))
 			{
+				if (ChunkInfo->DbFile->bIsRetired)
+				{
+					GLog->Logf(TEXT("BAD!!!!!!!! Accessed chunk after chunkdb %s retired - invalid chunk reference order."), *ChunkInfo->DbFile->Archive->GetArchiveName());
+				}
+
 				// Wait for the chunk to be available.
 				while (!HasFailed(DataId) && (ChunkData = ChunkStore->Get(DataId)) == nullptr && !bShouldAbort)
 				{
-					Platform->Sleep(0.01f);
+					TRACE_CPUPROFILER_EVENT_SCOPE(ChunkDb_GetChunkData_Sleep);
+					// 10ms was wayyy too long here. (def needs to be an event)
+					Platform->Sleep(0.001f);
 				}
 
 				// Dump out unavailable chunks on the incoming IO thread.
@@ -308,6 +407,10 @@ namespace BuildPatchServices
 
 	void FChunkDbChunkSource::ThreadRun()
 	{
+		TRACE_CPUPROFILER_EVENT_SCOPE(ChunkDb_Thread);
+
+		int32 LastChunkReferenceWatermark = ChunkReferenceTracker->GetRemainingChunkCount();
+
 		while (!bShouldAbort)
 		{
 			bool bWorkPerformed = false;
@@ -347,7 +450,10 @@ namespace BuildPatchServices
 			// If we had nothing to do, rest a little.
 			if(!bWorkPerformed)
 			{
-				Platform->Sleep(0.1f);
+				// This sleep is hit during startup since we don't do work before we start loading. For fresh installs
+				// this means we usually ask for the first chunk immediately after this sleep starts. So I changed this from 100ms
+				// to 10ms. Ideally this would be waiting on an event...
+				Platform->Sleep(0.01f);
 			}
 		}
 
@@ -359,6 +465,7 @@ namespace BuildPatchServices
 		bool bChunkGood = false;
 		if (ChunkDbDataAccessLookup.Contains(DataId))
 		{
+			TRACE_CPUPROFILER_EVENT_SCOPE(ChunkDb_LoadChunk);
 			ChunkDbChunkSourceStat->OnLoadStarted(DataId);
 			FChunkAccessLookup& ChunkAccessLookup = ChunkDbDataAccessLookup[DataId];
 			FChunkLocation& ChunkLocation = *ChunkAccessLookup.Location;
diff --git a/Engine/Source/Runtime/Online/BuildPatchServices/Private/Installer/ChunkDbChunkSource.h b/Engine/Source/Runtime/Online/BuildPatchServices/Private/Installer/ChunkDbChunkSource.h
--- a/Engine/Source/Runtime/Online/BuildPatchServices/Private/Installer/ChunkDbChunkSource.h
+++ b/Engine/Source/Runtime/Online/BuildPatchServices/Private/Installer/ChunkDbChunkSource.h
@@ -55,6 +55,9 @@ namespace BuildPatchServices
 		// The context for allocating shared resources.
 		IBuildInstallerSharedContext* SharedContext = nullptr;
 
+		// If true, once we complete a file we delete all the chunkdbs used to create it.
+		bool bDeleteChunkDBAfterUse = false;
+
 		/**
 		 * Constructor which sets usual defaults, and takes params for values that cannot use a default.
 		 * @param InChunkDbFiles    The chunkdb filename array.
diff --git a/Engine/Source/Runtime/Online/BuildPatchServices/Private/Installer/ChunkReferenceTracker.cpp b/Engine/Source/Runtime/Online/BuildPatchServices/Private/Installer/ChunkReferenceTracker.cpp
--- a/Engine/Source/Runtime/Online/BuildPatchServices/Private/Installer/ChunkReferenceTracker.cpp
+++ b/Engine/Source/Runtime/Online/BuildPatchServices/Private/Installer/ChunkReferenceTracker.cpp
@@ -26,10 +26,16 @@ namespace BuildPatchServices
 		virtual TArray<FGuid> GetNextReferences(int32 Count, const TFunction<bool(const FGuid&)>& SelectPredicate) const override;
 		virtual TArray<FGuid> SelectFromNextReferences(int32 Count, const TFunction<bool(const FGuid&)>& SelectPredicate) const override;
 		virtual bool PopReference(const FGuid& ChunkId) override;
+		virtual int32 GetRemainingChunkCount() const override;
 		// IChunkReferenceTracker interface end.
 
 	private:
 		TMap<FGuid, FThreadSafeCounter> ReferenceCount;
+
+		// This is a list of the chunk guids that we are going to use _in order_ during installation.
+		// it's reversed so that as we finish with chunks we can Pop instead of doing a full memmove.
+		// \todo just make this an array with an index in to current position so we don't have to do the reorder
+		// shenanegans everywhere and can grab the current position without a mutex.
 		TArray<FGuid> UseStack;
 		mutable FCriticalSection UseStackCs;
 	};
@@ -146,6 +152,12 @@ namespace BuildPatchServices
 		return NextReferences;
 	}
 
+	int32 FChunkReferenceTracker::GetRemainingChunkCount() const
+	{
+		FScopeLock ThreadLock(&UseStackCs);
+		return UseStack.Num();
+	}
+
 	TArray<FGuid> FChunkReferenceTracker::SelectFromNextReferences(int32 Count, const TFunction<bool(const FGuid&)>& SelectPredicate) const
 	{
 		// Thread lock to protect access to UseStack.
diff --git a/Engine/Source/Runtime/Online/BuildPatchServices/Private/Installer/ChunkReferenceTracker.h b/Engine/Source/Runtime/Online/BuildPatchServices/Private/Installer/ChunkReferenceTracker.h
--- a/Engine/Source/Runtime/Online/BuildPatchServices/Private/Installer/ChunkReferenceTracker.h
+++ b/Engine/Source/Runtime/Online/BuildPatchServices/Private/Installer/ChunkReferenceTracker.h
@@ -60,6 +60,9 @@ namespace BuildPatchServices
 		virtual TArray<FGuid> SelectFromNextReferences(int32 Count, const TFunction<bool(const FGuid&)>& SelectPredicate) const = 0;
 
 		/**
+		* Poorly named - Should be named "DoneWithChunk". The expectation is that we finish with chunks in the expected
+		* use order.
+		* 
 		 * Pop the top reference from the tracker, indicating that operation has been performed.
 		 * It is not valid to pop anything but the top guid, so it must be provided for verification of behavior.
 		 * @param ChunkId           The id of the top chunk, this is used to verify behavior.
@@ -67,6 +70,9 @@ namespace BuildPatchServices
 		 *         was provided and thus no change was made.
 		 */
 		virtual bool PopReference(const FGuid& ChunkId) = 0;
+
+		// Return how many chunks are left to process.
+		virtual int32 GetRemainingChunkCount() const = 0;
 	};
 
 	/**
diff --git a/Engine/Source/Runtime/Online/BuildPatchServices/Private/Installer/ChunkSource.h b/Engine/Source/Runtime/Online/BuildPatchServices/Private/Installer/ChunkSource.h
--- a/Engine/Source/Runtime/Online/BuildPatchServices/Private/Installer/ChunkSource.h
+++ b/Engine/Source/Runtime/Online/BuildPatchServices/Private/Installer/ChunkSource.h
@@ -45,5 +45,9 @@ namespace BuildPatchServices
 		 * @param Callback          The function to call with the set of chunks no longer still available.
 		 */
 		virtual void SetUnavailableChunksCallback(TFunction<void(TSet<FGuid>)> Callback) = 0;
+
+		// Inform the chunk source(s) that a file has been completed so we can clean up since we know we
+		// won't be going backwards for retries/resumes.
+		virtual void ReportFileCompletion() {}
 	};
 }
diff --git a/Engine/Source/Runtime/Online/BuildPatchServices/Private/Tests/Fake/ChunkReferenceTracker.fake.h b/Engine/Source/Runtime/Online/BuildPatchServices/Private/Tests/Fake/ChunkReferenceTracker.fake.h
--- a/Engine/Source/Runtime/Online/BuildPatchServices/Private/Tests/Fake/ChunkReferenceTracker.fake.h
+++ b/Engine/Source/Runtime/Online/BuildPatchServices/Private/Tests/Fake/ChunkReferenceTracker.fake.h
@@ -68,6 +68,8 @@ namespace BuildPatchServices
 			}
 			return false;
 		}
+		virtual int32 GetRemainingChunkCount() const { return NextReferences.Num(); };
+
 
 	public:
 		TSet<FGuid> ReferencedChunks;
diff --git a/Engine/Source/Runtime/Online/BuildPatchServices/Public/BuildPatchMessage.h b/Engine/Source/Runtime/Online/BuildPatchServices/Public/BuildPatchMessage.h
--- a/Engine/Source/Runtime/Online/BuildPatchServices/Public/BuildPatchMessage.h
+++ b/Engine/Source/Runtime/Online/BuildPatchServices/Public/BuildPatchMessage.h
@@ -18,6 +18,8 @@ namespace BuildPatchServices
 			AccessLost = 0,
 			// Access has been regained after being lost.
 			AccessRegained,
+			// The chunk source has been retired and is not expected to be accessed again (i.e. can be deleted).
+			Retired
 		};
 
 		// The type of event that occurred.
diff --git a/Engine/Source/Runtime/Online/BuildPatchServices/Public/BuildPatchSettings.h b/Engine/Source/Runtime/Online/BuildPatchServices/Public/BuildPatchSettings.h
--- a/Engine/Source/Runtime/Online/BuildPatchServices/Public/BuildPatchSettings.h
+++ b/Engine/Source/Runtime/Online/BuildPatchServices/Public/BuildPatchSettings.h
@@ -267,10 +267,15 @@ namespace BuildPatchServices
 		FString StagingDirectory;
 		// The directory for placing files that are believed to have local changes, before we overwrite them. Empty string will use module's global setting. If both empty, the feature disables.
 		FString BackupDirectory;
-		// The list of chunk database filenames that will be used to pull patch data from.
+		// The list of chunk database filenames that will be used to pull patch data from. These must be local (i.e. no URLs)
 		TArray<FString> ChunkDatabaseFiles;
-		// The list of cloud directory roots that will be used to pull patch data from. Empty array will use module's global setting.
+		// If set, when we've read everything we're going to read from a chunkdb file, we delete it to minimize required disk space for install/patch.
+		bool bDeleteChunkDbFilesAfterUse = false;
+
+		// The list of cloud directory roots that will be used to pull patch data from. Empty array will use module's global setting. This is only hit if a chunk can not be satisfied by
+		// other sources (i.e. chunkdbs or install directory for patches). If the chunk can not be serviced from here, it's a failure.
 		TArray<FString> CloudDirectories;
+
 		// The mode for installation.
 		EInstallMode InstallMode;
 		// The mode for verification. This will be overridden for any files referenced by a repair action.
