Require size calculation, whlie not called, at least is written and tested.
This also contains some fixes for the retirement code.

#rnx

diff --git a/Engine/Source/Runtime/Online/BuildPatchServices/Private/BuildPatchInstaller.cpp b/Engine/Source/Runtime/Online/BuildPatchServices/Private/BuildPatchInstaller.cpp
--- a/Engine/Source/Runtime/Online/BuildPatchServices/Private/BuildPatchInstaller.cpp
+++ b/Engine/Source/Runtime/Online/BuildPatchServices/Private/BuildPatchInstaller.cpp
@@ -178,6 +178,7 @@ namespace InstallerHelpers
 		UE_LOG(LogBuildPatchServices, Log, TEXT("Build Stat: FailureReasonText: %s"), *BuildStats.FailureReasonText.BuildSourceString());
 		UE_LOG(LogBuildPatchServices, Log, TEXT("Build Stat: FailureType: %s"), *EnumToString(BuildStats.FailureType));
 		UE_LOG(LogBuildPatchServices, Log, TEXT("Build Stat: NumInstallRetries: %u"), BuildStats.NumInstallRetries);
+		UE_LOG(LogBuildPatchServices, Log, TEXT("Build Stat: MaxDiskSpaceNeededWhenDeletingChunkDbsIfRequested: %s"), *FText::AsNumber(BuildStats.MaxDiskSpaceNeededWhenDeletingChunkDbsIfRequested).ToString());
 		check(BuildStats.NumInstallRetries == BuildStats.RetryFailureTypes.Num() && BuildStats.NumInstallRetries == BuildStats.RetryErrorCodes.Num());
 		for (uint32 RetryIdx = 0; RetryIdx < BuildStats.NumInstallRetries; ++RetryIdx)
 		{
@@ -421,6 +422,8 @@ namespace BuildPatchServices
 		Cloud
 	};
 
+	static uint64 DetermineInstallMaxDiskSizeIfDeletingChunkDbs(const TArray<FString>& ChunkDbFiles, IBuildManifestSet* ManifestSet, IFileSystem* FileSystem, const FString& InstallDirectory, const TArray<FString>& CorruptFiles, bool bIsPrereqOnly);
+
 	/* FBuildPatchInstaller implementation
 	*****************************************************************************/
 	FBuildPatchInstaller::FBuildPatchInstaller(FBuildInstallerConfiguration InConfiguration, TMultiMap<FString, FBuildPatchAppManifestRef> InInstallationInfo, const FString& InLocalMachineConfigFile, TSharedPtr<IAnalyticsProvider> InAnalytics, FBuildPatchInstallerDelegate InStartDelegate, FBuildPatchInstallerDelegate InCompleteDelegate)
@@ -836,109 +839,120 @@ namespace BuildPatchServices
 		InitializeTimer.Start();
 		bool bProcessSuccess = Initialize();
 
-		// Run if successful init
-		if (bProcessSuccess)
+		uint64 MaxDiskSpaceNeededWhenDeletingChunkDbsIfRequested = 0;
+		if (Configuration.bCalculateDeleteChunkDbMaxDiskSpaceAndExit)
+		{		
+			if (bProcessSuccess)
+			{
+				MaxDiskSpaceNeededWhenDeletingChunkDbsIfRequested = DetermineInstallMaxDiskSizeIfDeletingChunkDbs(Configuration.ChunkDatabaseFiles, ManifestSet.Get(), FileSystem.Get(), Configuration.InstallDirectory, {}, Configuration.InstallMode == EInstallMode::PrereqOnly);
+			}
+		}
+		else
 		{
-			// Keep track of files that failed verify
-			TArray<FString> CorruptFiles;
-
-			// Keep retrying the install while it is not canceled, or caused by download error
-			bProcessSuccess = false;
-			bool bCanRetry = true;
-			int32 InstallRetries = ConfigHelpers::NumInstallerRetries();
-			while (!bProcessSuccess && bCanRetry)
+			// Run if successful init 
+			if (bProcessSuccess)
 			{
-				// Inform file operation tracker of the selected manifest.
-				FileOperationTracker->OnManifestSelection(ManifestSet.Get());
-				MemoryChunkStoreStatistics->SetMultipleReferencedChunk(InstallerHelpers::GetMultipleReferencedChunks(ManifestSet.Get()));
+				// Keep track of files that failed verify
+				TArray<FString> CorruptFiles;
+
+				// Keep retrying the install while it is not canceled, or caused by download error
+				bProcessSuccess = false;
+				bool bCanRetry = true;
+				int32 InstallRetries = ConfigHelpers::NumInstallerRetries();
+				while (!bProcessSuccess && bCanRetry)
+				{
+					// Inform file operation tracker of the selected manifest.
+					FileOperationTracker->OnManifestSelection(ManifestSet.Get());
+					MemoryChunkStoreStatistics->SetMultipleReferencedChunk(InstallerHelpers::GetMultipleReferencedChunks(ManifestSet.Get()));
 
-				// Run the install
-				bool bInstallSuccess = RunInstallation(CorruptFiles);
+					// Run the install
+					bool bInstallSuccess = RunInstallation(CorruptFiles);
 
-				// Backup local changes then move generated files
-				bInstallSuccess = bInstallSuccess && RunBackupAndMove();
+					// Backup local changes then move generated files
+					bInstallSuccess = bInstallSuccess && RunBackupAndMove();
 
-				// Setup file attributes
-				bInstallSuccess = bInstallSuccess && RunFileAttributes();
+					// Setup file attributes
+					bInstallSuccess = bInstallSuccess && RunFileAttributes();
 
-				// Run Verification
-				CorruptFiles.Empty();
-				bProcessSuccess = bInstallSuccess && RunVerification(CorruptFiles);
+					// Run Verification
+					CorruptFiles.Empty();
+					bProcessSuccess = bInstallSuccess && RunVerification(CorruptFiles);
 
-				// Clean staging if INSTALL success, we still do cleanup if we failed at the verify stage.
-				if (bInstallSuccess)
-				{
-					CleanUpTimer.Start();
-					if (Configuration.InstallMode == EInstallMode::StageFiles)
+					// Clean staging if INSTALL success, we still do cleanup if we failed at the verify stage.
+					if (bInstallSuccess)
 					{
-						UE_LOG(LogBuildPatchServices, Log, TEXT("Deleting litter from staging area."));
-						IFileManager::Get().DeleteDirectory(*DataStagingDir, false, true);
+						CleanUpTimer.Start();
+						if (Configuration.InstallMode == EInstallMode::StageFiles)
+						{
+							UE_LOG(LogBuildPatchServices, Log, TEXT("Deleting litter from staging area."));
+							IFileManager::Get().DeleteDirectory(*DataStagingDir, false, true);
+						}
+						else
+						{
+							UE_LOG(LogBuildPatchServices, Log, TEXT("Deleting staging area."));
+							IFileManager::Get().DeleteDirectory(*Configuration.StagingDirectory, false, true);
+							CleanupEmptyDirectories(Configuration.InstallDirectory);
+						}
+						CleanUpTimer.Stop();
 					}
-					else
+					BuildProgress.SetStateProgress(EBuildPatchState::CleanUp, 1.0f);
+
+					// Set if we can retry
+					--InstallRetries;
+					bCanRetry = InstallRetries > 0 && !InstallerError->IsCancelled() && InstallerError->CanRetry();
+					const bool bWillRetry = !bProcessSuccess && bCanRetry;
+
+					// If successful or we will retry, remove the moved files marker
+					if (bProcessSuccess || bCanRetry)
 					{
-						UE_LOG(LogBuildPatchServices, Log, TEXT("Deleting staging area."));
-						IFileManager::Get().DeleteDirectory(*Configuration.StagingDirectory, false, true);
-						CleanupEmptyDirectories(Configuration.InstallDirectory);
+						UE_LOG(LogBuildPatchServices, Log, TEXT("Reset MM."));
+						IFileManager::Get().Delete(*PreviousMoveMarker, false, true);
 					}
-					CleanUpTimer.Stop();
-				}
-				BuildProgress.SetStateProgress(EBuildPatchState::CleanUp, 1.0f);
 
-				// Set if we can retry
-				--InstallRetries;
-				bCanRetry = InstallRetries > 0 && !InstallerError->IsCancelled() && InstallerError->CanRetry();
-				const bool bWillRetry = !bProcessSuccess && bCanRetry;
-
-				// If successful or we will retry, remove the moved files marker
-				if (bProcessSuccess || bCanRetry)
-				{
-					UE_LOG(LogBuildPatchServices, Log, TEXT("Reset MM."));
-					IFileManager::Get().Delete(*PreviousMoveMarker, false, true);
-				}
+					// Setup end of attempt stats
+					bFirstInstallIteration = false;
+					float TempFinalProgress = BuildProgress.GetProgressNoMarquee();
+					{
+						FScopeLock Lock(&ThreadLock);
+						BuildStats.NumInstallRetries = ConfigHelpers::NumInstallerRetries() - (InstallRetries + 1);
+						BuildStats.FinalProgress = TempFinalProgress;
+						// If we failed, and will retry, record this failure type and reset the abort flag
+						if (bWillRetry)
+						{
+							BuildStats.RetryFailureTypes.Add(InstallerError->GetErrorType());
+							BuildStats.RetryErrorCodes.Add(InstallerError->GetErrorCode());
+							bShouldAbort = false;
+						}
+					}
 
-				// Setup end of attempt stats
-				bFirstInstallIteration = false;
-				float TempFinalProgress = BuildProgress.GetProgressNoMarquee();
-				{
-					FScopeLock Lock(&ThreadLock);
-					BuildStats.NumInstallRetries = ConfigHelpers::NumInstallerRetries() - (InstallRetries + 1);
-					BuildStats.FinalProgress = TempFinalProgress;
-					// If we failed, and will retry, record this failure type and reset the abort flag
+					// If we will retry the install, reset progress states.
 					if (bWillRetry)
 					{
-						BuildStats.RetryFailureTypes.Add(InstallerError->GetErrorType());
-						BuildStats.RetryErrorCodes.Add(InstallerError->GetErrorCode());
-						bShouldAbort = false;
+						InitializeTimer.Start();
+						BuildProgress.CancelAbort();
+						BuildProgress.SetStateProgress(EBuildPatchState::Initializing, 0.0f);
+						BuildProgress.SetStateProgress(EBuildPatchState::Resuming, 0.0f);
+						BuildProgress.SetStateProgress(EBuildPatchState::Downloading, 0.0f);
+						BuildProgress.SetStateProgress(EBuildPatchState::Installing, 0.0f);
+						BuildProgress.SetStateProgress(EBuildPatchState::MovingToInstall, 0.0f);
+						BuildProgress.SetStateProgress(EBuildPatchState::SettingAttributes, 0.0f);
+						BuildProgress.SetStateProgress(EBuildPatchState::BuildVerification, 0.0f);
+						BuildProgress.SetStateProgress(EBuildPatchState::CleanUp, 0.0f);
 					}
 				}
-
-				// If we will retry the install, reset progress states.
-				if (bWillRetry)
-				{
-					InitializeTimer.Start();
-					BuildProgress.CancelAbort();
-					BuildProgress.SetStateProgress(EBuildPatchState::Initializing, 0.0f);
-					BuildProgress.SetStateProgress(EBuildPatchState::Resuming, 0.0f);
-					BuildProgress.SetStateProgress(EBuildPatchState::Downloading, 0.0f);
-					BuildProgress.SetStateProgress(EBuildPatchState::Installing, 0.0f);
-					BuildProgress.SetStateProgress(EBuildPatchState::MovingToInstall, 0.0f);
-					BuildProgress.SetStateProgress(EBuildPatchState::SettingAttributes, 0.0f);
-					BuildProgress.SetStateProgress(EBuildPatchState::BuildVerification, 0.0f);
-					BuildProgress.SetStateProgress(EBuildPatchState::CleanUp, 0.0f);
-				}
 			}
-		}
 
-		if (bProcessSuccess)
-		{
-			// Run the prerequisites installer if this is our first install and the manifest has prerequisites info
-			if (bInstallPrereqs)
+			if (bProcessSuccess)
 			{
-				PrereqTimer.Start();
-				bProcessSuccess &= RunPrerequisites();
-				PrereqTimer.Stop();
+				// Run the prerequisites installer if this is our first install and the manifest has prerequisites info
+				if (bInstallPrereqs)
+				{
+					PrereqTimer.Start();
+					bProcessSuccess &= RunPrerequisites();
+					PrereqTimer.Stop();
+				}
 			}
-		}
+		} // end if doing normal installation
 
 		// Make sure all timers are stopped
 		InitializeTimer.Stop();
@@ -970,6 +984,7 @@ namespace BuildPatchServices
 			BuildStats.ErrorCode = InstallerError->GetErrorCode();
 			BuildStats.FailureReasonText = InstallerError->GetErrorText();
 			BuildStats.FailureType = InstallerError->GetErrorType();
+			BuildStats.MaxDiskSpaceNeededWhenDeletingChunkDbsIfRequested = MaxDiskSpaceNeededWhenDeletingChunkDbsIfRequested;
 		}
 
 		// Mark that we are done
@@ -1167,6 +1182,128 @@ namespace BuildPatchServices
 		return ConnectionCountConfiguration;
 	}
 
+
+	static void GenerateFilesToConstruct(TSet<FString>& FilesToConstruct, IBuildManifestSet* ManifestSet, const TArray<FString>& CorruptFiles, const TSet<FString>& TaggedFiles, const TSet<FString>& OutdatedFiles, bool bIsPrereqOnly)
+	{
+		// Get the list of files actually needing construction
+		FilesToConstruct.Empty();
+		if (CorruptFiles.Num())
+		{
+			FilesToConstruct.Append(CorruptFiles);
+		}
+		else if (bIsPrereqOnly)
+		{
+			TArray<FPreReqInfo> PreReqInfos;
+			ManifestSet->GetPreReqInfo(PreReqInfos);
+			for (FPreReqInfo PreReqInfo : PreReqInfos)
+			{
+				FilesToConstruct.Add(PreReqInfo.Path);
+			}
+		}
+		else
+		{
+			FilesToConstruct = OutdatedFiles.Intersect(TaggedFiles);
+		}
+		FilesToConstruct.Sort(TLess<FString>());
+	}
+
+	static uint64 DetermineInstallMaxDiskSizeIfDeletingChunkDbs(const TArray<FString>& ChunkDbFiles, IBuildManifestSet* ManifestSet, IFileSystem* FileSystem, const FString& InstallDirectory, const TArray<FString>& CorruptFiles, bool bIsPrereqOnly)
+	{
+		TSet<FString> FilesToConstruct;
+
+		TSet<FString> TaggedFiles;
+		ManifestSet->GetExpectedFiles(TaggedFiles);
+		TSet<FString> OutdatedFiles;
+		ManifestSet->GetOutdatedFiles(InstallDirectory, OutdatedFiles);
+
+		GenerateFilesToConstruct(FilesToConstruct, ManifestSet, CorruptFiles, TaggedFiles, OutdatedFiles, bIsPrereqOnly);
+
+		TUniquePtr<IChunkReferenceTracker> ChunkReferenceTracker(FChunkReferenceTrackerFactory::Create(
+			ManifestSet,
+			FilesToConstruct));
+
+
+		//
+		// Calculating the disk size requried to install:
+		//
+		// We have a bunch of chunkdbs on disk that sum to the size of the end install
+		// We know that we can delete many - but not all - of the input chunkdbs when a file is
+		// completed.
+		//
+		// After each file completion we have:
+		// 
+		// Size of all remaining chunkdbs + size of installed files thus far.
+		//
+		// So we get the actual required size by iterating over the file list
+		// and asking the chunk db system how many chunkdbs are left at the current
+		// reference level.
+		//
+		TArray<FGuid> ReferenceChain;
+		ChunkReferenceTracker->CopyOutOrderedUseList(ReferenceChain);
+		int32 CurrentPosition = 0;
+
+		TArray<int32> FileCompletionPositions;
+		TArray<FString> OrderedFiles;
+
+		for (const FString& FileToConstruct : FilesToConstruct)
+		{
+			const FFileManifest* FileManifest = ManifestSet->GetNewFileManifest(FileToConstruct);
+			if (!FileManifest)
+			{
+				continue;
+			}
+
+			// We will be advancing the chunk reference tracker by this many chunks.
+			int32 AdvanceCount = FileManifest->ChunkParts.Num();
+
+			CurrentPosition += AdvanceCount;
+
+			OrderedFiles.Add(FileToConstruct);
+			FileCompletionPositions.Add(CurrentPosition);
+		}
+
+		TArray<uint64> ChunkDbSizesAtPosition;
+		uint64 TotalChunkDbSize = IChunkDbChunkSource::GetChunkDbSizesAtIndexes(ChunkDbFiles, FileSystem, ReferenceChain, FileCompletionPositions, ChunkDbSizesAtPosition);
+
+		uint64 TotalWrittenSize = 0;
+
+		uint64 MaxDiskSize = 0;
+		int32 AtFile = 0;
+
+		int32 i = 0;
+
+		uint64 LastFileSize = TotalChunkDbSize;
+
+		for (uint64 FileSize : ChunkDbSizesAtPosition)
+		{
+			// We've completed this file:
+			const FFileManifest* FileManifest = ManifestSet->GetNewFileManifest(OrderedFiles[i]);
+			uint64 ThisFileSize = 0;
+			if (FileManifest)
+			{
+				ThisFileSize = FileManifest->FileSize;
+			}
+
+			TotalWrittenSize += ThisFileSize;
+
+			// We delete the chunkdbs _after_ we write the output so we can't use this size until 
+			// the next file gets done.
+			if (LastFileSize + TotalWrittenSize > MaxDiskSize)
+			{
+				MaxDiskSize = LastFileSize + TotalWrittenSize;
+				AtFile = i;
+			}
+
+			i++;
+			LastFileSize = FileSize;
+		}
+
+		//UE_LOG(LogTemp, Display, TEXT("Max disk use %llu (install size: %llu: +%llu) after file %s"), MaxDiskSize, TotalWrittenSize, MaxDiskSize - TotalWrittenSize, *OrderedFiles[AtFile]);
+
+		return MaxDiskSize;
+	}
+
+
 	bool FBuildPatchInstaller::RunInstallation(TArray<FString>& CorruptFiles)
 	{
 		UE_LOG(LogBuildPatchServices, Log, TEXT("Starting Installation"));
@@ -1193,26 +1330,7 @@ namespace BuildPatchServices
 		const bool bIsPrereqOnly = Configuration.InstallMode == EInstallMode::PrereqOnly;
 		const bool bHasCorruptFiles = CorruptFiles.Num() > 0;
 
-		// Get the list of files actually needing construction
-		FilesToConstruct.Empty();
-		if (bHasCorruptFiles)
-		{
-			FilesToConstruct.Append(CorruptFiles);
-		}
-		else if (bIsPrereqOnly)
-		{
-			TArray<FPreReqInfo> PreReqInfos;
-			ManifestSet->GetPreReqInfo(PreReqInfos);
-			for (FPreReqInfo PreReqInfo : PreReqInfos)
-			{
-				FilesToConstruct.Add(PreReqInfo.Path);
-			}
-		}
-		else
-		{
-			FilesToConstruct = OutdatedFiles.Intersect(TaggedFiles);
-		}
-		FilesToConstruct.Sort(TLess<FString>());
+		GenerateFilesToConstruct(FilesToConstruct, ManifestSet.Get(), CorruptFiles, TaggedFiles, OutdatedFiles, bIsPrereqOnly);
 		UE_LOG(LogBuildPatchServices, Log, TEXT("Requiring %d files"), FilesToConstruct.Num());
 
 		// Check if we should skip out of this process due to existing installation,
diff --git a/Engine/Source/Runtime/Online/BuildPatchServices/Private/Installer/ChunkDbChunkSource.cpp b/Engine/Source/Runtime/Online/BuildPatchServices/Private/Installer/ChunkDbChunkSource.cpp
--- a/Engine/Source/Runtime/Online/BuildPatchServices/Private/Installer/ChunkDbChunkSource.cpp
+++ b/Engine/Source/Runtime/Online/BuildPatchServices/Private/Installer/ChunkDbChunkSource.cpp
@@ -162,6 +162,10 @@ namespace BuildPatchServices
 		virtual const TSet<FGuid>& GetAvailableChunks() const override;
 		// IInstallChunkSource interface end.
 
+		static void LoadChunkDbFiles(
+			const TArray<FString>& ChunkDbFiles, IFileSystem* FileSystem, const TArray<FGuid>& ChunkAccessOrderedList,
+			TArray<FChunkDbDataAccess>& OutChunkDbDataAccesses, TMap<FGuid, FChunkAccessLookup>& OutChunkDbDataAccessLookup, TSet<FGuid>* OutOptionalAvailableStore);
+
 	private:
 		void ThreadRun();
 		void LoadChunk(const FGuid& DataId);
@@ -207,25 +211,17 @@ namespace BuildPatchServices
 		TQueue<FGuid, EQueueMode::Mpsc> RepeatRequirementMessages;
 	};
 
-	FChunkDbChunkSource::FChunkDbChunkSource(FChunkDbSourceConfig InConfiguration, IPlatform* InPlatform, IFileSystem* InFileSystem, IChunkStore* InChunkStore, IChunkReferenceTracker* InChunkReferenceTracker, IChunkDataSerialization* InChunkDataSerialization, IMessagePump* InMessagePump, IInstallerError* InInstallerError, IChunkDbChunkSourceStat* InChunkDbChunkSourceStat)
-		: Configuration(MoveTemp(InConfiguration))
-		, Platform(InPlatform)
-		, FileSystem(InFileSystem)
-		, ChunkStore(InChunkStore)
-		, ChunkReferenceTracker(InChunkReferenceTracker)
-		, ChunkDataSerialization(InChunkDataSerialization)
-		, MessagePump(InMessagePump)
-		, InstallerError(InInstallerError)
-		, ChunkDbChunkSourceStat(InChunkDbChunkSourceStat)
-		, bStartedLoading(!InConfiguration.bBeginLoadsOnFirstGet)
+	// Read in the headers, evalutate the list of chunks, and determine when we'll be done with our chunk dbs.
+	void FChunkDbChunkSource::LoadChunkDbFiles(
+		const TArray<FString>& ChunkDbFiles, IFileSystem* FileSystem, const TArray<FGuid>& ChunkAccessOrderedList, 
+		TArray<FChunkDbDataAccess>& OutChunkDbDataAccesses, TMap<FGuid, FChunkAccessLookup>& OutChunkDbDataAccessLookup, TSet<FGuid>* OutOptionalAvailableStore)
 	{
-		TRACE_CPUPROFILER_EVENT_SCOPE(FChunkDbChunkSource_ctor);
 
 		// Allow OS intervention only once.
 		bool bResetOsIntervention = false;
 		int32 PreviousOsIntervention = 0;
 		// Load each chunkdb's TOC to enumerate available chunks.
-		for (const FString& ChunkDbFilename : Configuration.ChunkDbFiles)
+		for (const FString& ChunkDbFilename : ChunkDbFiles)
 		{
 			TUniquePtr<FArchive> ChunkDbFile(FileSystem->CreateFileReader(*ChunkDbFilename));
 			if (ChunkDbFile.IsValid())
@@ -248,10 +244,10 @@ namespace BuildPatchServices
 					DataSource.Archive = MoveTemp(ChunkDbFile);
 					DataSource.ChunkDbFileName = ChunkDbFilename;
 					DataSource.Header = MoveTemp(Header);
-					ChunkDbDataAccesses.Add(MoveTemp(DataSource));
+					OutChunkDbDataAccesses.Add(MoveTemp(DataSource));
 				}
 			}
-			else if(!bResetOsIntervention)
+			else if (!bResetOsIntervention)
 			{
 				bResetOsIntervention = true;
 				PreviousOsIntervention = ChunkDbSourceHelpers::DisableOsIntervention();
@@ -264,30 +260,30 @@ namespace BuildPatchServices
 		}
 
 		// Index all chunks to their location info.
-		for (FChunkDbDataAccess& ChunkDbDataAccess : ChunkDbDataAccesses)
+		for (FChunkDbDataAccess& ChunkDbDataAccess : OutChunkDbDataAccesses)
 		{
 			for (FChunkLocation& ChunkLocation : ChunkDbDataAccess.Header.Contents)
 			{
-				if (!ChunkDbDataAccessLookup.Contains(ChunkLocation.ChunkId))
+				if (!OutChunkDbDataAccessLookup.Contains(ChunkLocation.ChunkId))
 				{
-					ChunkDbDataAccessLookup.Add(ChunkLocation.ChunkId, {&ChunkLocation, &ChunkDbDataAccess});
-					AvailableChunks.Add(ChunkLocation.ChunkId);
+					OutChunkDbDataAccessLookup.Add(ChunkLocation.ChunkId, { &ChunkLocation, &ChunkDbDataAccess });
+
+					if (OutOptionalAvailableStore)
+					{
+						OutOptionalAvailableStore->Add(ChunkLocation.ChunkId);
+					}
 				}
 			}
 		}
 
-		// Once we have all the guids here we can predict when we'll be done with each
-		// chunkdb file by looking at the reference tracker.
-		TArray<FGuid> OrderedGuidList = InChunkReferenceTracker->GetNextReferences(0x7fffffff, [](const FGuid&) { return true; });
-
 		TMap<FString, int32> FileLastSeenAt;
 
 		int32 GuidIndex = 0;
-		for (; GuidIndex < OrderedGuidList.Num(); GuidIndex++)
+		for (; GuidIndex < ChunkAccessOrderedList.Num(); GuidIndex++)
 		{
-			const FGuid& Guid = OrderedGuidList[GuidIndex];
+			const FGuid& Guid = ChunkAccessOrderedList[GuidIndex];
 
-			FChunkAccessLookup* SourceForGuid = ChunkDbDataAccessLookup.Find(Guid);
+			FChunkAccessLookup* SourceForGuid = OutChunkDbDataAccessLookup.Find(Guid);
 			if (!SourceForGuid)
 			{
 				continue;
@@ -296,19 +292,74 @@ namespace BuildPatchServices
 			FileLastSeenAt.FindOrAdd(SourceForGuid->DbFile->Archive->GetArchiveName()) = GuidIndex;
 		}
 
-		for (FChunkDbDataAccess& ChunkDbDataAccess : ChunkDbDataAccesses)
+		for (FChunkDbDataAccess& ChunkDbDataAccess : OutChunkDbDataAccesses)
 		{
+			// Default to retired immediately (this is functionally after the first file completes)
+			ChunkDbDataAccess.RetireAt = ChunkAccessOrderedList.Num();
 			int32* LastAt = FileLastSeenAt.Find(ChunkDbDataAccess.Archive->GetArchiveName());
-			if (LastAt == nullptr)
+			if (LastAt)
 			{
-				// If the file isn't going to be used, it's immediately retired.
-				ChunkDbDataAccess.Retire(MessagePump, FileSystem, Configuration.bDeleteChunkDBAfterUse);
+				// The reference stack gets popped rather than advanced, so we need to reverse the ordering.
+				// LastAt is the chunk index that last used the file, which means when we are at
+				// LastAt + 1 through the stack we can delete the file.
+				ChunkDbDataAccess.RetireAt = ChunkAccessOrderedList.Num() - (*LastAt + 1);
 			}
-			else
+		}
+	}
+
+
+	uint64 IChunkDbChunkSource::GetChunkDbSizesAtIndexes(const TArray<FString>& ChunkDbFiles, IFileSystem* FileSystem, const TArray<FGuid>& ChunkAccessOrderedList, const TArray<int32>& FileCompletionIndexes, TArray<uint64>& OutChunkDbSizesAtCompletion)
+	{
+		TArray<FChunkDbDataAccess> ChunkFiles;
+		TMap<FGuid, FChunkAccessLookup> ChunkGuidLookup;
+
+		FChunkDbChunkSource::LoadChunkDbFiles(ChunkDbFiles, FileSystem, ChunkAccessOrderedList, ChunkFiles, ChunkGuidLookup, nullptr);
+
+		uint64 AllChunkDbSize = 0;
+		for (const FChunkDbDataAccess& ChunkFile : ChunkFiles)
+		{
+			AllChunkDbSize += ChunkFile.Archive->TotalSize();
+		}
+
+		
+		// Go over the list of completions and evalute how many chunk dbs are left over.
+		for (int32 FileCompletionIndex : FileCompletionIndexes)
+		{
+			// retiring happens as the list is _popped_, so everything is backwards.
+			int32 RetireAtEquivalent = ChunkAccessOrderedList.Num() - FileCompletionIndex;
+
+			uint64 TotalSizeAtIndex = 0;
+			for (const FChunkDbDataAccess& ChunkFile : ChunkFiles)
 			{
-				ChunkDbDataAccess.RetireAt = OrderedGuidList.Num() - *LastAt - 1;
+				if (ChunkFile.RetireAt < RetireAtEquivalent)
+				{
+					TotalSizeAtIndex += ChunkFile.Archive->TotalSize();
+				}
 			}
+
+			OutChunkDbSizesAtCompletion.Add(TotalSizeAtIndex);
 		}
+		return AllChunkDbSize;
+	}
+
+	FChunkDbChunkSource::FChunkDbChunkSource(FChunkDbSourceConfig InConfiguration, IPlatform* InPlatform, IFileSystem* InFileSystem, IChunkStore* InChunkStore, IChunkReferenceTracker* InChunkReferenceTracker, IChunkDataSerialization* InChunkDataSerialization, IMessagePump* InMessagePump, IInstallerError* InInstallerError, IChunkDbChunkSourceStat* InChunkDbChunkSourceStat)
+		: Configuration(MoveTemp(InConfiguration))
+		, Platform(InPlatform)
+		, FileSystem(InFileSystem)
+		, ChunkStore(InChunkStore)
+		, ChunkReferenceTracker(InChunkReferenceTracker)
+		, ChunkDataSerialization(InChunkDataSerialization)
+		, MessagePump(InMessagePump)
+		, InstallerError(InInstallerError)
+		, ChunkDbChunkSourceStat(InChunkDbChunkSourceStat)
+		, bStartedLoading(!InConfiguration.bBeginLoadsOnFirstGet)
+	{
+		TRACE_CPUPROFILER_EVENT_SCOPE(FChunkDbChunkSource_ctor);
+
+		TArray<FGuid> OrderedGuidList;
+		InChunkReferenceTracker->CopyOutOrderedUseList(OrderedGuidList);
+
+		LoadChunkDbFiles(Configuration.ChunkDbFiles, FileSystem, OrderedGuidList, ChunkDbDataAccesses, ChunkDbDataAccessLookup, &AvailableChunks);
 
 		// Start threaded load worker.
 		if (ChunkDbDataAccesses.Num() > 0)
@@ -409,8 +460,6 @@ namespace BuildPatchServices
 	{
 		TRACE_CPUPROFILER_EVENT_SCOPE(ChunkDb_Thread);
 
-		int32 LastChunkReferenceWatermark = ChunkReferenceTracker->GetRemainingChunkCount();
-
 		while (!bShouldAbort)
 		{
 			bool bWorkPerformed = false;
diff --git a/Engine/Source/Runtime/Online/BuildPatchServices/Private/Installer/ChunkDbChunkSource.h b/Engine/Source/Runtime/Online/BuildPatchServices/Private/Installer/ChunkDbChunkSource.h
--- a/Engine/Source/Runtime/Online/BuildPatchServices/Private/Installer/ChunkDbChunkSource.h
+++ b/Engine/Source/Runtime/Online/BuildPatchServices/Private/Installer/ChunkDbChunkSource.h
@@ -34,6 +34,10 @@ namespace BuildPatchServices
 		 * @return the set of chunks available.
 		 */
 		virtual const TSet<FGuid>& GetAvailableChunks() const = 0;
+
+		// Fill out how many bytes of chunkdbs are left if we delete all the ones that
+		// are no longer necessary at the given FileCompletionIndexes in to ChunkAccessOrderedList
+		static uint64 GetChunkDbSizesAtIndexes(const TArray<FString>& ChunkDbFiles, IFileSystem* FileSystem, const TArray<FGuid>& ChunkAccessOrderedList, const TArray<int32>& FileCompletionIndexes, TArray<uint64>& OutChunkDbSizesAtCompletion);
 	};
 
 	/**
diff --git a/Engine/Source/Runtime/Online/BuildPatchServices/Private/Installer/ChunkReferenceTracker.cpp b/Engine/Source/Runtime/Online/BuildPatchServices/Private/Installer/ChunkReferenceTracker.cpp
--- a/Engine/Source/Runtime/Online/BuildPatchServices/Private/Installer/ChunkReferenceTracker.cpp
+++ b/Engine/Source/Runtime/Online/BuildPatchServices/Private/Installer/ChunkReferenceTracker.cpp
@@ -20,6 +20,13 @@ namespace BuildPatchServices
 		~FChunkReferenceTracker();
 
 		// IChunkReferenceTracker interface begin.
+		virtual void CopyOutOrderedUseList(TArray<FGuid>& OutUseList) const override
+		{
+			UseStackCs.Lock();
+			OutUseList.Append(UseStack);
+			UseStackCs.Unlock();
+			Algo::Reverse(OutUseList);
+		}
 		virtual TSet<FGuid> GetReferencedChunks() const override;
 		virtual int32 GetReferenceCount(const FGuid& ChunkId) const override;
 		virtual void SortByUseOrder(TArray<FGuid>& ChunkList, ESortDirection Direction) const override;
diff --git a/Engine/Source/Runtime/Online/BuildPatchServices/Private/Installer/ChunkReferenceTracker.h b/Engine/Source/Runtime/Online/BuildPatchServices/Private/Installer/ChunkReferenceTracker.h
--- a/Engine/Source/Runtime/Online/BuildPatchServices/Private/Installer/ChunkReferenceTracker.h
+++ b/Engine/Source/Runtime/Online/BuildPatchServices/Private/Installer/ChunkReferenceTracker.h
@@ -73,6 +73,10 @@ namespace BuildPatchServices
 
 		// Return how many chunks are left to process.
 		virtual int32 GetRemainingChunkCount() const = 0;
+
+		// Return the ordered list we are consuming.
+		// copies the _remaining_ items!
+		virtual void CopyOutOrderedUseList(TArray<FGuid>& OutUseList) const = 0;
 	};
 
 	/**
diff --git a/Engine/Source/Runtime/Online/BuildPatchServices/Private/Tests/Fake/ChunkReferenceTracker.fake.h b/Engine/Source/Runtime/Online/BuildPatchServices/Private/Tests/Fake/ChunkReferenceTracker.fake.h
--- a/Engine/Source/Runtime/Online/BuildPatchServices/Private/Tests/Fake/ChunkReferenceTracker.fake.h
+++ b/Engine/Source/Runtime/Online/BuildPatchServices/Private/Tests/Fake/ChunkReferenceTracker.fake.h
@@ -59,7 +59,7 @@ namespace BuildPatchServices
 			});
 		}
 
-		virtual bool PopReference(const FGuid& ChunkId)
+		virtual bool PopReference(const FGuid& ChunkId) override
 		{
 			if (NextReferences.Num() && NextReferences[0] == ChunkId)
 			{
@@ -68,8 +68,12 @@ namespace BuildPatchServices
 			}
 			return false;
 		}
-		virtual int32 GetRemainingChunkCount() const { return NextReferences.Num(); };
+		virtual int32 GetRemainingChunkCount() const override { return NextReferences.Num(); };
 
+		virtual void CopyOutOrderedUseList(TArray<FGuid>& OutUseList) const
+		{
+			OutUseList = NextReferences;
+		}
 
 	public:
 		TSet<FGuid> ReferencedChunks;
diff --git a/Engine/Source/Runtime/Online/BuildPatchServices/Public/BuildPatchSettings.h b/Engine/Source/Runtime/Online/BuildPatchServices/Public/BuildPatchSettings.h
--- a/Engine/Source/Runtime/Online/BuildPatchServices/Public/BuildPatchSettings.h
+++ b/Engine/Source/Runtime/Online/BuildPatchServices/Public/BuildPatchSettings.h
@@ -272,6 +272,10 @@ namespace BuildPatchServices
 		// If set, when we've read everything we're going to read from a chunkdb file, we delete it to minimize required disk space for install/patch.
 		bool bDeleteChunkDbFilesAfterUse = false;
 
+		// If set, don't actually do an installation, just figure out how much disk space is required if we are using entirely
+		// chunkdbs and we delete them as we go.
+		bool bCalculateDeleteChunkDbMaxDiskSpaceAndExit = false;
+
 		// The list of cloud directory roots that will be used to pull patch data from. Empty array will use module's global setting. This is only hit if a chunk can not be satisfied by
 		// other sources (i.e. chunkdbs or install directory for patches). If the chunk can not be serviced from here, it's a failure.
 		TArray<FString> CloudDirectories;
diff --git a/Engine/Source/Runtime/Online/BuildPatchServices/Public/Interfaces/IBuildInstaller.h b/Engine/Source/Runtime/Online/BuildPatchServices/Public/Interfaces/IBuildInstaller.h
--- a/Engine/Source/Runtime/Online/BuildPatchServices/Public/Interfaces/IBuildInstaller.h
+++ b/Engine/Source/Runtime/Online/BuildPatchServices/Public/Interfaces/IBuildInstaller.h
@@ -234,6 +234,11 @@ struct FBuildInstallStats
 	double PeakDiskWriteSpeed;
 	// The total number of files constructed.
 	uint32 NumFilesConstructed;
+	
+	// We don't actually do any installation when requesting this value, we just start up an installer, compute it, then bail. See
+	// bCalculateDeleteChunkDbMaxDiskSpaceAndExit in the installer configuration. Will be zero all other times.
+	uint64 MaxDiskSpaceNeededWhenDeletingChunkDbsIfRequested = 0;
+
 	// The time spent during the initialization stage.
 	float InitializeTime;
 	// The time spent during the construction stage.
