Multithreaded BPT validation.

#rnx

diff --git a/Engine/Config/Android/BaseAndroidEngine.ini b/Engine/Config/Android/BaseAndroidEngine.ini
--- a/Engine/Config/Android/BaseAndroidEngine.ini
+++ b/Engine/Config/Android/BaseAndroidEngine.ini
@@ -14,4 +14,7 @@ DefaultExternalCredentialTypeStr=Google
 
 ;Uncomment to turn on AndroidBackgroundHttp through the AndroidFetchBackgroundDownload and AndroidBackgroundService modules
 ;[BackgroundHttp]
-;PlatformModularFeatureName=AndroidFetchBackgroundDownload
\ No newline at end of file
+;PlatformModularFeatureName=AndroidFetchBackgroundDownload
+
+[BuildPatchTool]
+VerificationThreadCount=3
\ No newline at end of file
diff --git a/Engine/Config/BaseEngine.ini b/Engine/Config/BaseEngine.ini
--- a/Engine/Config/BaseEngine.ini
+++ b/Engine/Config/BaseEngine.ini
@@ -3701,3 +3701,6 @@ UsageType=Editor
 [/Script/Engine.WorldPartitionSettings]
 EditorHashDefaultClass=/Script/CoreUObject.Class'/Script/Engine.WorldPartitionEditorSpatialHash'
 RuntimeHashDefaultClass=/Script/CoreUObject.Class'/Script/Engine.WorldPartitionRuntimeHashSet'
+
+[BuildPatchTool]
+VerificationThreadCount=4
\ No newline at end of file
diff --git a/Engine/Config/IOS/BaseIOSEngine.ini b/Engine/Config/IOS/BaseIOSEngine.ini
--- a/Engine/Config/IOS/BaseIOSEngine.ini
+++ b/Engine/Config/IOS/BaseIOSEngine.ini
@@ -124,3 +124,6 @@ fx.Niagara.QualityLevel.Max=1
 [TextureStreaming]
 ; PoolSizeVRAMPercentage is how much percentage of GPU Dedicated VRAM should be used as a TexturePool cache for streaming textures (0 - unlimited streaming)
 PoolSizeVRAMPercentage=70
+
+[BuildPatchTool]
+VerificationThreadCount=3
\ No newline at end of file
diff --git a/Engine/Source/Runtime/Online/BuildPatchServices/Private/BuildPatchInstaller.cpp b/Engine/Source/Runtime/Online/BuildPatchServices/Private/BuildPatchInstaller.cpp
--- a/Engine/Source/Runtime/Online/BuildPatchServices/Private/BuildPatchInstaller.cpp
+++ b/Engine/Source/Runtime/Online/BuildPatchServices/Private/BuildPatchInstaller.cpp
@@ -755,7 +755,7 @@ namespace BuildPatchServices
 
 		// We can now build out any systems that need late construction but can survive between retries.
 		ManifestSet.Reset(FBuildManifestSetFactory::Create(InstallerActions));
-		Verifier.Reset(FVerifierFactory::Create(FileSystem.Get(), VerifierStatistics.Get(), Configuration.VerifyMode, ManifestSet.Get(), Configuration.InstallDirectory, Configuration.InstallMode == EInstallMode::StageFiles ? InstallStagingDir : FString()));
+		Verifier.Reset(FVerifierFactory::Create(FileSystem.Get(), VerifierStatistics.Get(), Configuration.VerifyMode, Configuration.SharedContext, ManifestSet.Get(), Configuration.InstallDirectory, Configuration.InstallMode == EInstallMode::StageFiles ? InstallStagingDir : FString()));
 
 		// Add systems to controllables.
 		{
diff --git a/Engine/Source/Runtime/Online/BuildPatchServices/Private/Installer/Verifier.cpp b/Engine/Source/Runtime/Online/BuildPatchServices/Private/Installer/Verifier.cpp
--- a/Engine/Source/Runtime/Online/BuildPatchServices/Private/Installer/Verifier.cpp
+++ b/Engine/Source/Runtime/Online/BuildPatchServices/Private/Installer/Verifier.cpp
@@ -7,15 +7,13 @@
 #include "BuildPatchVerify.h"
 #include "BuildPatchUtil.h"
 #include "IBuildManifestSet.h"
+#include "Misc/ConfigCacheIni.h"
 
 DECLARE_LOG_CATEGORY_EXTERN(LogVerifier, Warning, All);
 DEFINE_LOG_CATEGORY(LogVerifier);
 
 namespace BuildPatchServices
 {
-	// 4MB Read buffer.
-	const int32 VerifyBufferSize = 1024 * 1024 * 4;
-
 	bool TryConvertToVerifyResult(EVerifyError InVerifyError, EVerifyResult& OutVerifyResult)
 	{
 		switch (InVerifyError)
@@ -43,7 +41,7 @@ namespace BuildPatchServices
 	class FVerifier : public IVerifier
 	{
 	public:
-		FVerifier(IFileSystem* FileSystem, IVerifierStat* InVerificationStat, EVerifyMode InVerifyMode, IBuildManifestSet* InManifestSet, FString InVerifyDirectory, FString InStagedFileDirectory);
+		FVerifier(IFileSystem* FileSystem, IVerifierStat* InVerificationStat, EVerifyMode InVerifyMode, IBuildInstallerSharedContextPtr InSharedContext, IBuildManifestSet* InManifestSet, FString InVerifyDirectory, FString InStagedFileDirectory);
 		~FVerifier() {}
 
 		// IControllable interface begin.
@@ -58,7 +56,7 @@ namespace BuildPatchServices
 
 	private:
 		FString SelectFullFilePath(const FString& BuildFile);
-		EVerifyResult VerfiyFileSha(const FString& BuildFile, const FFileManifest& BuildFileManifest);
+		EVerifyResult VerfiyFileSha(TArray<uint8>& ReadBuffer, const FString& BuildFile, const FFileManifest& BuildFileManifest);
 		EVerifyResult VerfiyFileSize(const FString& BuildFile, const FFileManifest& BuildFileManifest);
 
 	private:
@@ -67,30 +65,92 @@ namespace BuildPatchServices
 		IFileSystem* const FileSystem;
 		IVerifierStat* const VerifierStat;
 		IBuildManifestSet* const ManifestSet;
+		IBuildInstallerSharedContextPtr SharedContext;
 
-		TArray<uint8> FileReadBuffer;
 		EVerifyMode VerifyMode;
 		TSet<FString> FilesToVerify;
 		TSet<FString> FilesPassedVerify;
 		FThreadSafeBool bIsPaused;
 		FThreadSafeBool bShouldAbort;
-		int64 ProcessedBytes;
+
+		struct FThreadVerifyJob
+		{
+			const FFileManifest* BuildFileManifest;
+			bool bVerifySha;
+			FString FileName;
+		};
+
+		FEvent* AllThreadsDoneEvent;
+		std::atomic_int32_t ThreadsCreated;
+		std::atomic_int32_t ThreadsCompleted;
+		std::atomic_int64_t ThreadProcessedBytes;
+		std::atomic_int32_t ThreadPullIndex;
+		TArray<EVerifyResult> ThreadJobResults;
+		TArray<FThreadVerifyJob> ThreadJobList;
+
+		TArray<IBuildInstallerThread*> SecondaryProcessThreads;
+
+		void ProcessVerifyJobs();
 	};
 
-	FVerifier::FVerifier(IFileSystem* InFileSystem, IVerifierStat* InVerificationStat, EVerifyMode InVerifyMode, IBuildManifestSet* InManifestSet, FString InVerifyDirectory, FString InStagedFileDirectory)
+	FVerifier::FVerifier(IFileSystem* InFileSystem, IVerifierStat* InVerificationStat, EVerifyMode InVerifyMode, IBuildInstallerSharedContextPtr InSharedContext, IBuildManifestSet* InManifestSet, FString InVerifyDirectory, FString InStagedFileDirectory)
 		: VerifyDirectory(MoveTemp(InVerifyDirectory))
 		, StagedFileDirectory(MoveTemp(InStagedFileDirectory))
 		, FileSystem(InFileSystem)
 		, VerifierStat(InVerificationStat)
 		, ManifestSet(InManifestSet)
+		, SharedContext(InSharedContext)
 		, VerifyMode(InVerifyMode)
 		, bIsPaused(false)
 		, bShouldAbort(false)
-		, ProcessedBytes(0)
 	{
 		ManifestSet->GetFilesTaggedForRepair(FilesToVerify);
-		FileReadBuffer.Reserve(VerifyBufferSize);
-		FileReadBuffer.AddUninitialized(VerifyBufferSize);
+
+	}
+
+	void FVerifier::ProcessVerifyJobs()
+	{
+		TArray<uint8> FileReadBuffer;
+		FileReadBuffer.SetNumUninitialized(4 << 20); // 4MB chunked reads.
+
+		for (;;)
+		{
+			if (bShouldAbort)
+			{
+				break;
+			}
+
+			// Get a thread index.
+			int32 PossiblePullIndex = ThreadPullIndex.load();
+			if (!ThreadPullIndex.compare_exchange_strong(PossiblePullIndex, PossiblePullIndex +1))
+			{
+				// Couldn't claim this job, try again
+				continue;
+			}
+
+			int32 OurJobIndex = PossiblePullIndex;
+			if (OurJobIndex >= ThreadJobList.Num())
+			{
+				// Ran out of jobs to do.
+				break;
+			}
+
+			FThreadVerifyJob& OurJob = ThreadJobList[OurJobIndex];
+
+			VerifierStat->OnFileStarted(OurJob.FileName, OurJob.BuildFileManifest->FileSize);
+
+			EVerifyResult FileVerifyResult = OurJob.bVerifySha ? VerfiyFileSha(FileReadBuffer, OurJob.FileName, *OurJob.BuildFileManifest) : VerfiyFileSize(OurJob.FileName, *OurJob.BuildFileManifest);
+
+			VerifierStat->OnFileCompleted(OurJob.FileName, FileVerifyResult);
+
+			ThreadJobResults[OurJobIndex] = FileVerifyResult;
+		}
+
+		if ((ThreadsCompleted.fetch_add(1) + 1) == ThreadsCreated.load(std::memory_order_relaxed))
+		{
+			// We're the last one.
+			AllThreadsDoneEvent->Trigger();
+		}
 	}
 
 	void FVerifier::SetPaused(bool bInIsPaused)
@@ -105,6 +165,7 @@ namespace BuildPatchServices
 
 	EVerifyResult FVerifier::Verify(TArray<FString>& CorruptFiles)
 	{
+		TRACE_CPUPROFILER_EVENT_SCOPE(Verify);
 		bShouldAbort = false;
 		EVerifyResult VerifyResult = EVerifyResult::Success;
 		CorruptFiles.Empty();
@@ -123,44 +184,98 @@ namespace BuildPatchServices
 		// Select verify function.
 		const bool bVerifyShaMode = VerifyMode == EVerifyMode::ShaVerifyAllFiles || VerifyMode == EVerifyMode::ShaVerifyTouchedFiles;
 
-		// For each required file, perform the selected verification.
-		ProcessedBytes = 0;
+
+		ThreadJobList.Reset();
+		ThreadJobResults.Reset();
+		ThreadProcessedBytes = 0;
+		ThreadPullIndex = 0;
+
 		for (const FString& BuildFile : VerifyList)
 		{
-			// Break if quitting
-			if (bShouldAbort)
+			FThreadVerifyJob Job;
+			Job.FileName = BuildFile;
+			Job.BuildFileManifest = ManifestSet->GetNewFileManifest(BuildFile);
+			Job.bVerifySha = bVerifyShaMode || ManifestSet->IsFileRepairAction(BuildFile);
+
+			ThreadJobList.Add(MoveTemp(Job));
+			ThreadJobResults.Add(EVerifyResult::Aborted);
+		}
+
+		ThreadsCompleted = 0;
+		AllThreadsDoneEvent = FPlatformProcess::GetSynchEventFromPool();
+
+		if (SharedContext)
+		{
+			// 3 cores should saturate any modern drive, and on anything that isn't we just sit in the wait
+			// state anyway. 
+			int32 ThreadsToCreate = 1;
+			GConfig->GetInt(TEXT("BuildPatchTool"), TEXT("VerificationThreadCount"), ThreadsToCreate, GEngineIni);
+
+			ThreadsCreated = ThreadsToCreate;
+
+			// We wait to create our threads so that the shared threads from file construction are returned.
+			for (int i = 0; i < ThreadsToCreate; i++)
 			{
-				break;
+				SecondaryProcessThreads.Add(SharedContext->CreateThread());
+				SecondaryProcessThreads[i]->RunTask([this]() { ProcessVerifyJobs(); });
 			}
 
-			// Get file details.
-			const FFileManifest* BuildFileManifest = ManifestSet->GetNewFileManifest(BuildFile);
+			for (;;)
+			{
+				if (AllThreadsDoneEvent->Wait(100))
+				{
+					// triggered which means we're all done.
+					break;
+				}
+
+				if (bShouldAbort)
+				{
+					// if we aborted without a trigger, then the threads could still be working and
+					// we need to wait until they are done.
+					AllThreadsDoneEvent->Wait();
+					break;
+				}
+
+				// Update our overall progress tracker.
+				VerifierStat->OnProcessedDataUpdated(ThreadProcessedBytes.load(std::memory_order_acquire));
+			}
+
+			for (int i = 0; i < ThreadsToCreate; i++)
+			{
+				SharedContext->ReleaseThread(SecondaryProcessThreads[i]);
+			}
+			SecondaryProcessThreads.Empty();
+		}
+		else
+		{
+			// Can't create threads, just do the work here.
+			ThreadsCreated = 1;
+			ProcessVerifyJobs();
+		}
+
+		FPlatformProcess::ReturnSynchEventToPool(AllThreadsDoneEvent);
+		AllThreadsDoneEvent = nullptr;
+
+		// move results over.
+		for (int32 i=0; i<ThreadJobList.Num(); i++)
+		{
+			EVerifyResult FileVerifyResult = ThreadJobResults[i];
 
-			// Verify the file.
-			const bool bVerifySha = bVerifyShaMode || ManifestSet->IsFileRepairAction(BuildFile);
-			VerifierStat->OnFileStarted(BuildFile, BuildFileManifest->FileSize);
-			EVerifyResult FileVerifyResult = bVerifySha ? VerfiyFileSha(BuildFile, *BuildFileManifest) : VerfiyFileSize(BuildFile, *BuildFileManifest);
-			VerifierStat->OnFileCompleted(BuildFile, FileVerifyResult);
 			if (FileVerifyResult != EVerifyResult::Success)
 			{
-				CorruptFiles.Add(BuildFile);
+				CorruptFiles.Add(ThreadJobList[i].FileName);
 				if (VerifyResult == EVerifyResult::Success)
 				{
 					VerifyResult = FileVerifyResult;
 				}
 			}
 			// If success, and it was an SHA verify, cache the result so we don't repeat an SHA verify.
-			else if (bVerifySha)
+			else if (ThreadJobList[i].bVerifySha)
 			{
-				FilesPassedVerify.Add(BuildFile);
+				FilesPassedVerify.Add(ThreadJobList[i].FileName);
 			}
 		}
 
-		if (bShouldAbort && VerifyResult == EVerifyResult::Success)
-		{
-			VerifyResult = EVerifyResult::Aborted;
-		}
-
 		return VerifyResult;
 	}
 
@@ -186,10 +301,10 @@ namespace BuildPatchServices
 		return FullFilePath;
 	}
 
-	EVerifyResult FVerifier::VerfiyFileSha(const FString& BuildFile, const FFileManifest& BuildFileManifest)
+	EVerifyResult FVerifier::VerfiyFileSha(TArray<uint8>& ReadBuffer, const FString& BuildFile, const FFileManifest& BuildFileManifest)
 	{
+		TRACE_CPUPROFILER_EVENT_SCOPE(VerifyFileSha);
 		ISpeedRecorder::FRecord ActivityRecord;
-		const int64 PrevProcessedBytes = ProcessedBytes;
 		EVerifyResult VerifyResult;
 		const FString FileToVerify = SelectFullFilePath(BuildFile);
 		uint8 ReturnValue = 0;
@@ -216,14 +331,23 @@ namespace BuildPatchServices
 					ActivityRecord.CyclesStart = FStatsCollector::GetCycles();
 					// Read file and update hash state
 					const int64 SizeLeft = FileSize - FileReader->Tell();
-					ActivityRecord.Size = FMath::Min<int64>(VerifyBufferSize, SizeLeft);
-					FileReader->Serialize(FileReadBuffer.GetData(), ActivityRecord.Size);
-					HashState.Update(FileReadBuffer.GetData(), ActivityRecord.Size);
-					ProcessedBytes = PrevProcessedBytes + FileReader->Tell();
+					ActivityRecord.Size = FMath::Min<int64>(ReadBuffer.Num(), SizeLeft);
+
+					{
+						TRACE_CPUPROFILER_EVENT_SCOPE(VerifyFileSha_Serialize);
+						FileReader->Serialize(ReadBuffer.GetData(), ActivityRecord.Size);
+					}
+					{
+						TRACE_CPUPROFILER_EVENT_SCOPE(VerifyFileSha_Hash);
+						
+						HashState.Update(ReadBuffer.GetData(), ActivityRecord.Size);
+					}
+
 					ActivityRecord.CyclesEnd = FStatsCollector::GetCycles();
 					VerifierStat->OnFileRead(ActivityRecord);
 					VerifierStat->OnFileProgress(BuildFile, FileReader->Tell());
-					VerifierStat->OnProcessedDataUpdated(ProcessedBytes);
+
+					ThreadProcessedBytes.fetch_add(ActivityRecord.Size, std::memory_order_release);
 				}
 				HashState.Final();
 				HashState.GetHash(HashValue.Hash);
@@ -250,11 +374,9 @@ namespace BuildPatchServices
 		{
 			VerifyResult = EVerifyResult::FileMissing;
 		}
-		ProcessedBytes = PrevProcessedBytes + BuildFileManifest.FileSize;
 		if (VerifyResult != EVerifyResult::Success)
 		{
 			VerifierStat->OnFileProgress(BuildFile, BuildFileManifest.FileSize);
-			VerifierStat->OnProcessedDataUpdated(ProcessedBytes);
 		}
 
 		return VerifyResult;
@@ -263,13 +385,11 @@ namespace BuildPatchServices
 	EVerifyResult FVerifier::VerfiyFileSize(const FString& BuildFile, const FFileManifest& BuildFileManifest)
 	{
 		// Pause if necessary.
-		const double PrePauseTime = FStatsCollector::GetSeconds();
-		double PostPauseTime = PrePauseTime;
 		while (bIsPaused && !bShouldAbort)
 		{
 			FPlatformProcess::Sleep(0.1f);
-			PostPauseTime = FStatsCollector::GetSeconds();
 		}
+
 		VerifierStat->OnFileProgress(BuildFile, 0);
 		int64 FileSize;
 		EVerifyResult VerifyResult;
@@ -289,15 +409,15 @@ namespace BuildPatchServices
 			VerifyResult = EVerifyResult::FileMissing;
 		}
 		VerifierStat->OnFileProgress(BuildFile, BuildFileManifest.FileSize);
-		ProcessedBytes += BuildFileManifest.FileSize;
-		VerifierStat->OnProcessedDataUpdated(ProcessedBytes);
+
+		ThreadProcessedBytes.fetch_add(BuildFileManifest.FileSize, std::memory_order_release);
 		return VerifyResult;
 	}
 
-	IVerifier* FVerifierFactory::Create(IFileSystem* FileSystem, IVerifierStat* VerifierStat, EVerifyMode VerifyMode, IBuildManifestSet* ManifestSet, FString VerifyDirectory, FString StagedFileDirectory)
+	IVerifier* FVerifierFactory::Create(IFileSystem* FileSystem, IVerifierStat* VerifierStat, EVerifyMode VerifyMode, IBuildInstallerSharedContextPtr SharedContext, IBuildManifestSet* ManifestSet, FString VerifyDirectory, FString StagedFileDirectory)
 	{
 		check(FileSystem != nullptr);
 		check(VerifierStat != nullptr);
-		return new FVerifier(FileSystem, VerifierStat, VerifyMode, ManifestSet, MoveTemp(VerifyDirectory), MoveTemp(StagedFileDirectory));
+		return new FVerifier(FileSystem, VerifierStat, VerifyMode, SharedContext, ManifestSet, MoveTemp(VerifyDirectory), MoveTemp(StagedFileDirectory));
 	}
 }
diff --git a/Engine/Source/Runtime/Online/BuildPatchServices/Private/Installer/Verifier.h b/Engine/Source/Runtime/Online/BuildPatchServices/Private/Installer/Verifier.h
--- a/Engine/Source/Runtime/Online/BuildPatchServices/Private/Installer/Verifier.h
+++ b/Engine/Source/Runtime/Online/BuildPatchServices/Private/Installer/Verifier.h
@@ -4,7 +4,7 @@
 #include "Installer/Controllable.h"
 #include "Common/SpeedRecorder.h"
 #include "BuildPatchManifest.h"
-
+#include "Interfaces/IBuildInstallerSharedContext.h"
 
 namespace BuildPatchServices
 {
@@ -91,7 +91,7 @@ namespace BuildPatchServices
 		 * @param StagedFileDirectory   A stage directory for updated files, ignored if empty string. If a file exists here, it will be checked instead of the one in VerifyDirectory.
 		 * @return     Ref of an object that can be used to perform the operation.
 		 */
-		static IVerifier* Create(IFileSystem* FileSystem, IVerifierStat* VerifierStat, EVerifyMode VerifyMode, IBuildManifestSet* ManifestSet, FString VerifyDirectory, FString StagedFileDirectory);
+		static IVerifier* Create(IFileSystem* FileSystem, IVerifierStat* VerifierStat, EVerifyMode VerifyMode, IBuildInstallerSharedContextPtr SharedContext, IBuildManifestSet* ManifestSet, FString VerifyDirectory, FString StagedFileDirectory);
 	};
 
 	/**
diff --git a/Engine/Source/Runtime/Online/BuildPatchServices/Private/Tests/Unit/Verifier.spec.cpp b/Engine/Source/Runtime/Online/BuildPatchServices/Private/Tests/Unit/Verifier.spec.cpp
--- a/Engine/Source/Runtime/Online/BuildPatchServices/Private/Tests/Unit/Verifier.spec.cpp
+++ b/Engine/Source/Runtime/Online/BuildPatchServices/Private/Tests/Unit/Verifier.spec.cpp
@@ -556,6 +556,7 @@ void FVerifierSpec::MakeUnit(BuildPatchServices::EVerifyMode Mode)
 		FakeFileSystem.Get(),
 		MockVerificationStat.Get(),
 		Mode,
+		nullptr,
 		/*TouchedFiles,*/
 		ManifestSet.Get(),
 		VerifyDirectory,
