Fix missing iOS push notifications for background download

diff --git a/Engine/Source/Runtime/Online/BackgroundHTTP/Private/BackgroundHttpNotificationObject.cpp b/Engine/Source/Runtime/Online/BackgroundHTTP/Private/BackgroundHttpNotificationObject.cpp
--- a/Engine/Source/Runtime/Online/BackgroundHTTP/Private/BackgroundHttpNotificationObject.cpp
+++ b/Engine/Source/Runtime/Online/BackgroundHTTP/Private/BackgroundHttpNotificationObject.cpp
@@ -10,6 +10,10 @@
 
 #include "LocalNotification.h"
 
+#if PLATFORM_IOS || PLATFORM_TVOS
+#import <UIKit/UIKit.h>
+#endif
+
 FBackgroundHttpNotificationObject::FBackgroundHttpNotificationObject(FText InNotificationTitle, FText InNotificationBody, FText InNotificationAction, const FString& InNotificationActivationString, bool InNotifyOnlyOnFullSuccess)
 	: FBackgroundHttpNotificationObject(InNotificationTitle, InNotificationBody, InNotificationAction, InNotificationActivationString, InNotifyOnlyOnFullSuccess, true, -1)
 {
@@ -70,6 +74,12 @@ FBackgroundHttpNotificationObject::~FBackgroundHttpNotificationObject()
 		FCoreDelegates::ApplicationWillEnterBackgroundDelegate.Remove(OnApp_EnteringBackgroundHandle);
 		FCoreDelegates::ApplicationHasEnteredForegroundDelegate.Remove(OnApp_EnteringForegroundHandle);
 
+#if PLATFORM_IOS || PLATFORM_TVOS
+		// Temp workaround of ApplicationWillEnterBackgroundDelegate not correctly invoked
+		// TODO remove workaround
+		bIsInBackground = [UIApplication sharedApplication].applicationState != UIApplicationStateActive;
+#endif
+
 		//If we have flagged as only sending notifications when we are in the BG, and we are not in the BG, just early out
 		//so we don't send a notification
 		if (!bIsInBackground)
