Speculative fix for crashes related to curl easy handle's dns hostcache being null in some cases

diff --git a/Engine/Source/Runtime/Online/HTTP/Private/Curl/CurlHttp.cpp b/Engine/Source/Runtime/Online/HTTP/Private/Curl/CurlHttp.cpp
--- a/Engine/Source/Runtime/Online/HTTP/Private/Curl/CurlHttp.cpp
+++ b/Engine/Source/Runtime/Online/HTTP/Private/Curl/CurlHttp.cpp
@@ -1016,6 +1016,11 @@ bool FCurlHttpRequest::SetupRequestHttpThread()
 	return true;
 }
 
+void FCurlHttpRequest::CleanupRequestHttpThread()
+{
+	curl_easy_setopt(EasyHandle, CURLOPT_SHARE, nullptr);
+}
+
 bool FCurlHttpRequest::ProcessRequest()
 {
 	QUICK_SCOPE_CYCLE_COUNTER(STAT_FCurlHttpRequest_ProcessRequest);
@@ -1342,8 +1347,6 @@ void FCurlHttpRequest::FinishRequest()
 
 void FCurlHttpRequest::CleanupRequest()
 {
-	curl_easy_setopt(EasyHandle, CURLOPT_SHARE, nullptr);
-
 	CloseRequestPayloadDefaultImpl();
 }
 
diff --git a/Engine/Source/Runtime/Online/HTTP/Private/Curl/CurlHttp.h b/Engine/Source/Runtime/Online/HTTP/Private/Curl/CurlHttp.h
--- a/Engine/Source/Runtime/Online/HTTP/Private/Curl/CurlHttp.h
+++ b/Engine/Source/Runtime/Online/HTTP/Private/Curl/CurlHttp.h
@@ -182,6 +182,11 @@ public:
 	 */
 	bool SetupRequestHttpThread();
 
+	/**
+	 * Perform the http-thread cleanup of the request
+	 */
+	void  CleanupRequestHttpThread();
+
 	/**
 	 * Returns libcurl's easy handle - needed for HTTP manager.
 	 *
diff --git a/Engine/Source/Runtime/Online/HTTP/Private/Curl/CurlHttpThread.cpp b/Engine/Source/Runtime/Online/HTTP/Private/Curl/CurlHttpThread.cpp
--- a/Engine/Source/Runtime/Online/HTTP/Private/Curl/CurlHttpThread.cpp
+++ b/Engine/Source/Runtime/Online/HTTP/Private/Curl/CurlHttpThread.cpp
@@ -107,6 +107,8 @@ void FCurlHttpThread::CompleteThreadedRequest(FHttpRequestCommon* Request)
 		curl_multi_remove_handle(FCurlHttpManager::GMultiHandle, EasyHandle);
 		HandlesToRequests.Remove(EasyHandle);
 	}
+
+	CurlRequest->CleanupRequestHttpThread();
 }
 PRAGMA_ENABLE_DEPRECATION_WARNINGS
 
diff --git a/Engine/Source/Runtime/Online/HTTP/Private/Curl/CurlMultiPollEventLoopHttpThread.cpp b/Engine/Source/Runtime/Online/HTTP/Private/Curl/CurlMultiPollEventLoopHttpThread.cpp
--- a/Engine/Source/Runtime/Online/HTTP/Private/Curl/CurlMultiPollEventLoopHttpThread.cpp
+++ b/Engine/Source/Runtime/Online/HTTP/Private/Curl/CurlMultiPollEventLoopHttpThread.cpp
@@ -105,6 +105,8 @@ void FCurlMultiPollEventLoopHttpThread::CompleteThreadedRequest(FHttpRequestComm
 		curl_multi_remove_handle(FCurlHttpManager::GMultiHandle, EasyHandle);
 		HandlesToRequests.Remove(EasyHandle);
 	}
+
+	CurlRequest->CleanupRequestHttpThread();
 }
 PRAGMA_ENABLE_DEPRECATION_WARNINGS
 
diff --git a/Engine/Source/Runtime/Online/HTTP/Private/Curl/CurlMultiWaitEventLoopHttpThread.cpp b/Engine/Source/Runtime/Online/HTTP/Private/Curl/CurlMultiWaitEventLoopHttpThread.cpp
--- a/Engine/Source/Runtime/Online/HTTP/Private/Curl/CurlMultiWaitEventLoopHttpThread.cpp
+++ b/Engine/Source/Runtime/Online/HTTP/Private/Curl/CurlMultiWaitEventLoopHttpThread.cpp
@@ -123,6 +123,8 @@ void FCurlMultiWaitEventLoopHttpThread::CompleteThreadedRequest(FHttpRequestComm
 	{
 		curl_multi_remove_handle(FCurlHttpManager::GMultiHandle, EasyHandle);
 	}
+
+	CurlRequest->CleanupRequestHttpThread();
 }
 PRAGMA_ENABLE_DEPRECATION_WARNINGS
 
diff --git a/Engine/Source/Runtime/Online/HTTP/Private/Curl/CurlSocketEventLoopHttpThread.cpp b/Engine/Source/Runtime/Online/HTTP/Private/Curl/CurlSocketEventLoopHttpThread.cpp
--- a/Engine/Source/Runtime/Online/HTTP/Private/Curl/CurlSocketEventLoopHttpThread.cpp
+++ b/Engine/Source/Runtime/Online/HTTP/Private/Curl/CurlSocketEventLoopHttpThread.cpp
@@ -96,6 +96,8 @@ void FCurlSocketEventLoopHttpThread::CompleteThreadedRequest(FHttpRequestCommon*
 		curl_multi_remove_handle(FCurlHttpManager::GMultiHandle, EasyHandle);
 		HandlesToRequests.Remove(EasyHandle);
 	}
+
+	CurlRequest->CleanupRequestHttpThread();
 }
 PRAGMA_ENABLE_DEPRECATION_WARNINGS
 
