CloudChunkSource shared CDN failover.

#rnx

diff --git a/Engine/Source/Runtime/Online/BuildPatchServices/Private/Installer/CloudChunkSource.cpp b/Engine/Source/Runtime/Online/BuildPatchServices/Private/Installer/CloudChunkSource.cpp
--- a/Engine/Source/Runtime/Online/BuildPatchServices/Private/Installer/CloudChunkSource.cpp
+++ b/Engine/Source/Runtime/Online/BuildPatchServices/Private/Installer/CloudChunkSource.cpp
@@ -123,6 +123,9 @@ namespace BuildPatchServices
 			FTaskInfo();
 
 		public:
+			// Are we currently trying to downloads?
+			bool bQueuedForDownload = false;
+			int32 CloudDirUsed = 0;
 			FString UrlUsed;
 			int32 RetryNum;
 			int32 ExpectedSize;
@@ -152,7 +155,7 @@ namespace BuildPatchServices
 
 	private:
 		void EnsureAquiring(const FGuid& DataId);
-		const FString& GetCloudRoot(int32 RetryNum) const;
+
 		float GetRetryDelay(int32 RetryNum);
 		EBuildPatchDownloadHealth GetDownloadHealth(bool bIsDisconnected, float ChunkSuccessRate);
 		FGuid GetNextTask(const TMap<FGuid, FTaskInfo>& TaskInfos, const TMap<int32, FGuid>& InFlightDownloads, const TSet<FGuid>& TotalRequiredChunks, const TSet<FGuid>& PriorityRequests, const TSet<FGuid>& FailedDownloads, const TSet<FGuid>& Stored, TArray<FGuid>& DownloadQueue, EBuildPatchDownloadHealth DownloadHealth);
@@ -202,6 +205,11 @@ namespace BuildPatchServices
 		// Determine if additional download requests should be initiated.
 		IDownloadConnectionCount* DownloadCount;
 
+		// If we start getting failures on our downloads, we track which ones
+		// fail and avoid them until everything goes bad. Initially we just hit the
+		// first directory.
+		int32 CurrentBestCloudDir = 0;
+		TArray<int32> CloudDirFailureCount;
 	};
 
 	FCloudChunkSource::FDownloadDelegates::FDownloadDelegates(FCloudChunkSource& InCloudChunkSource)
@@ -254,6 +262,8 @@ namespace BuildPatchServices
 		, RequestedDownloads()
 		, DownloadCount(InDownloadConnectionCount)
 	{
+		CloudDirFailureCount.SetNumZeroed(Configuration.CloudRoots.Num());
+
 		Future = Promise.GetFuture();
 		if (Configuration.bRunOwnThread)
 		{
@@ -328,11 +338,6 @@ namespace BuildPatchServices
 		RequestedDownloads.Add(DataId);
 	}
 
-	const FString& FCloudChunkSource::GetCloudRoot(int32 RetryNum) const
-	{
-		return Configuration.CloudRoots[RetryNum % Configuration.CloudRoots.Num()];
-	}
-
 	float FCloudChunkSource::GetRetryDelay(int32 RetryNum)
 	{
 		const int32 RetryTimeIndex = FMath::Clamp<int32>(RetryNum - 1, 0, Configuration.RetryDelayTimes.Num() - 1);
@@ -411,7 +416,12 @@ namespace BuildPatchServices
 				int32 SearchLength = FMath::Max(ChunkStore->GetSize(), Configuration.PreFetchMinimum);
 				DownloadQueue = ChunkReferenceTracker->SelectFromNextReferences(SearchLength, SelectPredicate);
 				// Remove already downloading or complete chunks.
-				TFunction<bool(const FGuid&)> RemovePredicate = [&TaskInfos, &FailedDownloads, &Stored](const FGuid& ChunkId) { return TaskInfos.Contains(ChunkId) || FailedDownloads.Contains(ChunkId) || Stored.Contains(ChunkId); };
+				TFunction<bool(const FGuid&)> RemovePredicate = [&TaskInfos, &FailedDownloads, &Stored](const FGuid& ChunkId)
+				{
+					const FTaskInfo* TaskInfo = TaskInfos.Find(ChunkId);
+					return (TaskInfo && TaskInfo->bQueuedForDownload) || FailedDownloads.Contains(ChunkId) || Stored.Contains(ChunkId);
+				};
+
 				DownloadQueue.RemoveAll(RemovePredicate);
 				// Clamp to configured max.
 				DownloadQueue.SetNum(FMath::Min(DownloadQueue.Num(), Configuration.PreFetchMaximum), EAllowShrinking::No);
@@ -507,8 +517,9 @@ namespace BuildPatchServices
 				RequestedChunkUris.Add(ChunkUriToRequest);
 				FChunkUriRequest ChunkUriRequest;
 
-				const FTaskInfo* Info = TaskInfos.Find(ChunkUriToRequest);
-				ChunkUriRequest.CloudDirectory = GetCloudRoot(Info ? Info->RetryNum : 0 );
+				FTaskInfo& Info = TaskInfos.FindOrAdd(ChunkUriToRequest);
+				Info.CloudDirUsed = CurrentBestCloudDir;
+				ChunkUriRequest.CloudDirectory = Configuration.CloudRoots[CurrentBestCloudDir];
 				ChunkUriRequest.RelativePath = ManifestSet->GetDataFilename(ChunkUriToRequest);
 				ChunkUriRequest.RelativePath.RemoveFromStart(TEXT("/"));
 
@@ -571,6 +582,7 @@ namespace BuildPatchServices
 				const FDownloadRef& Download = FrameCompletedDownload.Value;
 				const FGuid& DownloadId = InFlightDownloads[RequestId];
 				FTaskInfo& TaskInfo = TaskInfos.FindOrAdd(DownloadId);
+				TaskInfo.bQueuedForDownload = false;
 				bool bDownloadSuccess = Download->ResponseSuccessful();
 				if (bDownloadSuccess)
 				{
@@ -609,7 +621,7 @@ namespace BuildPatchServices
 					UE_LOG(LogCloudChunkSource, Error, TEXT("FAILED: %s"), *TaskInfo.UrlUsed);
 				}
 
-				// Handle failed
+				// Handle failed (note this also launches a retry on a bad serialization, not just download.
 				if (!bDownloadSuccess)
 				{
 					ChunkSuccessRate.AddFail();
@@ -620,6 +632,34 @@ namespace BuildPatchServices
 						bShouldAbort = true;
 					}
 					++TaskInfo.RetryNum;
+
+					// Mark this CDN as failed.
+					{
+						CloudDirFailureCount[TaskInfo.CloudDirUsed]++;
+						if (CloudDirFailureCount[TaskInfo.CloudDirUsed] > (100 << 20))
+						{
+							// Cap to prevent wrap. I think this is technically impossible due to the time it would take to 
+							// get here but...
+							CloudDirFailureCount[TaskInfo.CloudDirUsed] = 100 << 20;
+						}
+
+						// Find who has failed the least, be sure to take equivalents in the initial specified order.
+						// We expect this to be like 5 entries.
+						int32 MinFailCount = CloudDirFailureCount[0];
+						int32 MinAtIndex = 0;
+						for (int32 CloudDirSeek = 1; CloudDirSeek < CloudDirFailureCount.Num(); CloudDirSeek++)
+						{
+							if (CloudDirFailureCount[CloudDirSeek] < MinFailCount)
+							{
+								MinFailCount = CloudDirFailureCount[CloudDirSeek];
+								MinAtIndex = CloudDirSeek;
+							}
+						}
+
+						CurrentBestCloudDir = MinAtIndex;
+						UE_LOG(LogCloudChunkSource, Warning, TEXT("CDN %s failed download, updating CDN selection to: %s"), *Configuration.CloudRoots[TaskInfo.CloudDirUsed], *Configuration.CloudRoots[MinAtIndex]);
+					}
+
 					TaskInfo.SecondsAtFail = FStatsCollector::GetSeconds();
 
 					RequestedChunkUris.Remove(DownloadId);
@@ -644,9 +684,9 @@ namespace BuildPatchServices
 				}
 			}
 			const double SecondsSinceData = FStatsCollector::CyclesToSeconds(FStatsCollector::GetCycles() - CyclesAtLastData);
-			const bool bDisconnect = (bAllDownloadsRetrying && SecondsSinceData > Configuration.DisconnectedDelay);
+			const bool bReportAsDisconnected = (bAllDownloadsRetrying && SecondsSinceData > Configuration.DisconnectedDelay);
 			const float SuccessRate = ChunkSuccessRate.GetOverall();
-			EBuildPatchDownloadHealth OverallDownloadHealth = GetDownloadHealth(bDisconnect, SuccessRate);
+			EBuildPatchDownloadHealth OverallDownloadHealth = GetDownloadHealth(bReportAsDisconnected, SuccessRate);
 			if (TrackedDownloadHealth != OverallDownloadHealth)
 			{
 				TrackedDownloadHealth = OverallDownloadHealth;
@@ -657,7 +697,7 @@ namespace BuildPatchServices
 				CloudChunkSourceStat->OnSuccessRateUpdated(SuccessRate);
 			}
 			const float ImmediateSuccessRate = ChunkSuccessRate.GetImmediate();
-			EBuildPatchDownloadHealth ImmediateDownloadHealth = GetDownloadHealth(bDisconnect, ImmediateSuccessRate);
+			EBuildPatchDownloadHealth ImmediateDownloadHealth = GetDownloadHealth(bReportAsDisconnected, ImmediateSuccessRate);
 			// Kick off new downloads.
 			if (bDownloadsStarted)
 			{
@@ -668,6 +708,7 @@ namespace BuildPatchServices
 					if (ChunkUri)
 					{
 						FTaskInfo& TaskInfo = TaskInfos.FindOrAdd(NextTask);
+						TaskInfo.bQueuedForDownload = true;
 						TaskInfo.UrlUsed = ChunkUri->Uri;
 						TaskInfo.ExpectedSize = ManifestSet->GetDownloadSize(NextTask);
 						TaskInfo.SecondsAtRequested = FStatsCollector::GetSeconds();
diff --git a/Engine/Source/Runtime/Online/BuildPatchServices/Private/Installer/CloudChunkSource.h b/Engine/Source/Runtime/Online/BuildPatchServices/Private/Installer/CloudChunkSource.h
--- a/Engine/Source/Runtime/Online/BuildPatchServices/Private/Installer/CloudChunkSource.h
+++ b/Engine/Source/Runtime/Online/BuildPatchServices/Private/Installer/CloudChunkSource.h
@@ -53,7 +53,9 @@ namespace BuildPatchServices
 		// The minimum number of chunks to request ahead of what is required, depending on store slack.
 		int32 PreFetchMinimum  = 16;
 		// The maximum number of chunks to request ahead of what is required, depending on store slack.
-		int32 PreFetchMaximum = 256;
+		// This shouldn't be too much because we'll queue too many on the same CDN which might not be good and we have to chew 
+		// though a ton of requests before we swap over, and we only really want there to be one queued + one ready to go per connection anyway.
+		int32 PreFetchMaximum = 32;
 		// Array of times in seconds, representing the time between each retry upon failure. The last entry will be used
 		// indefinitely once it is reached.
 		TArray<float> RetryDelayTimes;
