Fix the test failure by increasing the connect timeout deviation from 4s to 4.5s for xCurl.

diff --git a/Engine/Source/Runtime/Online/HTTP/HTTP.Build.cs b/Engine/Source/Runtime/Online/HTTP/HTTP.Build.cs
--- a/Engine/Source/Runtime/Online/HTTP/HTTP.Build.cs
+++ b/Engine/Source/Runtime/Online/HTTP/HTTP.Build.cs
@@ -139,7 +139,7 @@ public class HTTP : ModuleRules
 		float PlatformConnectionTimeoutMaxDeviation = 0.5f;
 		if (bPlatformSupportsXCurl)
 		{
-			PlatformConnectionTimeoutMaxDeviation = 4.0f;
+			PlatformConnectionTimeoutMaxDeviation = 4.5f;
 		}
 		if (Target.IsInPlatformGroup(UnrealPlatformGroup.Apple))
 		{
