Fix the in-progress disk space check when delete-during-install is enabled.
Fix for incorrect handling of first file under normal installation.

#rnx

diff --git a/Engine/Source/Runtime/Online/BuildPatchServices/Private/BuildPatchFileConstructor.cpp b/Engine/Source/Runtime/Online/BuildPatchServices/Private/BuildPatchFileConstructor.cpp
--- a/Engine/Source/Runtime/Online/BuildPatchServices/Private/BuildPatchFileConstructor.cpp
+++ b/Engine/Source/Runtime/Online/BuildPatchServices/Private/BuildPatchFileConstructor.cpp
@@ -15,6 +15,7 @@
 #include "Common/SpeedRecorder.h"
 #include "Common/FileSystem.h"
 #include "Installer/ChunkSource.h"
+#include "Installer/ChunkDbChunkSource.h"
 #include "Installer/ChunkReferenceTracker.h"
 #include "Installer/InstallerError.h"
 #include "Installer/InstallerAnalytics.h"
@@ -66,6 +67,12 @@ namespace FileConstructorHelpers
 				bContinueConstruction = false;
 			}
 		}
+		else
+		{
+			// If we can't get the disk space free then the most likely reason is the drive is no longer around...
+			bContinueConstruction = false;
+		}
+
 		return bContinueConstruction;
 	}
 
@@ -252,7 +259,10 @@ public:
 
 /* FBuildPatchFileConstructor implementation
  *****************************************************************************/
-FBuildPatchFileConstructor::FBuildPatchFileConstructor(FFileConstructorConfig InConfiguration, IFileSystem* InFileSystem, IChunkSource* InChunkSource, IChunkReferenceTracker* InChunkReferenceTracker, IInstallerError* InInstallerError, IInstallerAnalytics* InInstallerAnalytics, IFileConstructorStat* InFileConstructorStat)
+FBuildPatchFileConstructor::FBuildPatchFileConstructor(
+	FFileConstructorConfig InConfiguration, IFileSystem* InFileSystem, IChunkSource* InChunkSource, 
+	IChunkDbChunkSource* InChunkDbChunkSource, IChunkReferenceTracker* InChunkReferenceTracker, IInstallerError* InInstallerError, 
+	IInstallerAnalytics* InInstallerAnalytics, IFileConstructorStat* InFileConstructorStat)
 	: Configuration(MoveTemp(InConfiguration))
 	, bIsDownloadStarted(false)
 	, bInitialDiskSizeCheck(false)
@@ -262,6 +272,7 @@ FBuildPatchFileConstructor::FBuildPatchFileConstructor(FFileConstructorConfig In
 	, ConstructionStack()
 	, FileSystem(InFileSystem)
 	, ChunkSource(InChunkSource)
+	, ChunkDbSource(InChunkDbChunkSource)
 	, ChunkReferenceTracker(InChunkReferenceTracker)
 	, InstallerError(InInstallerError)
 	, InstallerAnalytics(InInstallerAnalytics)
@@ -275,6 +286,11 @@ FBuildPatchFileConstructor::FBuildPatchFileConstructor(FFileConstructorConfig In
 	const int32 ConstructListNum = Configuration.ConstructList.Num();
 	ConstructionStack.Reserve(ConstructListNum);
 	ConstructionStack.AddDefaulted(ConstructListNum);
+
+	// Track when we will complete files in the reference chain.
+	int32 CurrentPosition = 0;
+	FileCompletionPositions.Reserve(ConstructListNum);
+
 	for (int32 ConstructListIdx = 0; ConstructListIdx < ConstructListNum ; ++ConstructListIdx)
 	{
 		const FString& ConstructListElem = Configuration.ConstructList[ConstructListIdx];
@@ -282,7 +298,14 @@ FBuildPatchFileConstructor::FBuildPatchFileConstructor(FFileConstructorConfig In
 		if (FileManifest)
 		{
 			TotalJobSize += FileManifest->FileSize;
+		
+			// We will be advancing the chunk reference tracker by this many chunks.
+			int32 AdvanceCount = FileManifest->ChunkParts.Num();
+			CurrentPosition += AdvanceCount;
+
+			FileCompletionPositions.Add(CurrentPosition);
 		}
+
 		ConstructionStack[(ConstructListNum - 1) - ConstructListIdx] = ConstructListElem;
 	}
 
@@ -486,14 +509,12 @@ void FBuildPatchFileConstructor::Run()
 
 uint64 FBuildPatchFileConstructor::GetRequiredDiskSpace()
 {
-	FScopeLock Lock(&ThreadLock);
-	return RequiredDiskSpace;
+	return RequiredDiskSpace.load(std::memory_order_relaxed);
 }
 
 uint64 FBuildPatchFileConstructor::GetAvailableDiskSpace()
 {
-	FScopeLock Lock(&ThreadLock);
-	return AvailableDiskSpace;
+	return AvailableDiskSpace.load(std::memory_order_relaxed);
 }
 
 FBuildPatchFileConstructor::FOnBeforeDeleteFile& FBuildPatchFileConstructor::OnBeforeDeleteFile()
@@ -524,16 +545,29 @@ int64 FBuildPatchFileConstructor::GetRemainingBytes()
 	return Configuration.ManifestSet->GetTotalNewFileSize(ConstructionStack);
 }
 
-uint64 FBuildPatchFileConstructor::CalculateRequiredDiskSpace(const FFileManifest& InProgressFileManifest, uint64 InProgressFileSize)
+uint64 FBuildPatchFileConstructor::CalculateInProgressDiskSpaceRequired(const FFileManifest& InProgressFileManifest, uint64 InProgressFileAmountWritten)
 {
-	int64 DiskSpaceDeltaPeak = InProgressFileSize;
 	if (Configuration.InstallMode == EInstallMode::DestructiveInstall)
 	{
 		// The simplest method will be to run through each high level file operation, tracking peak disk usage delta.
-		int64 DiskSpaceDelta = InProgressFileSize;
 
-		// Can remove old in progress file.
-		DiskSpaceDelta -= InProgressFileManifest.FileSize;
+		// We know we need enough space to finish writing this file
+		uint64 RemainingThisFileSpace = InProgressFileManifest.FileSize - InProgressFileAmountWritten;
+		
+		int64 DiskSpaceDeltaPeak = RemainingThisFileSpace;
+		int64 DiskSpaceDelta = RemainingThisFileSpace;
+
+		// Then we move this file over.
+		{
+			const FFileManifest* OldFileManifest = Configuration.ManifestSet->GetCurrentFileManifest(InProgressFileManifest.Filename);
+			if (OldFileManifest)
+			{
+				DiskSpaceDelta -= OldFileManifest->FileSize;
+			}
+
+			// We've already accounted for the new file above, so we could be pretty negative if we resumed the file
+			// almost at the end and had an existing file we're deleting.
+		}
 
 		// Loop through all files to be made next, in order.
 		for (int32 ConstructionStackIdx = ConstructionStack.Num() - 1; ConstructionStackIdx >= 0; --ConstructionStackIdx)
@@ -553,13 +587,49 @@ uint64 FBuildPatchFileConstructor::CalculateRequiredDiskSpace(const FFileManifes
 				DiskSpaceDelta -= OldFileManifest->FileSize;
 			}
 		}
+		return DiskSpaceDeltaPeak;
 	}
 	else
 	{
 		// When not destructive, we always stage all new and changed files.
-		DiskSpaceDeltaPeak += Configuration.ManifestSet->GetTotalNewFileSize(ConstructionStack);
+		uint64 RemainingFilesSpace = Configuration.ManifestSet->GetTotalNewFileSize(ConstructionStack);
+		uint64 RemainingThisFileSpace = InProgressFileManifest.FileSize - InProgressFileAmountWritten;
+		return RemainingFilesSpace + RemainingThisFileSpace;
+	}
+}
+
+
+
+uint64 FBuildPatchFileConstructor::CalculateDiskSpaceRequirementsWithDeleteDuringInstall(const TArray<FString>& InBackwardsFilesLeftToConstruct)
+{
+	if (ChunkDbSource == nullptr)
+	{
+		// invalid use.
+		return 0;
 	}
-	return FMath::Max<int64>(DiskSpaceDeltaPeak, 0);
+
+	// These are the sizes at after each file that we _started_ with. This is the size after retirement for the
+	// file at those positions.
+	TArray<uint64> ChunkDbSizesAtPosition;
+	uint64 TotalChunkDbSize = ChunkDbSource->GetChunkDbSizesAtIndexes(FileCompletionPositions, ChunkDbSizesAtPosition);
+
+	// Strip off the files we've completed.
+	int32 CompletedFileCount = Configuration.ConstructList.Num() - InBackwardsFilesLeftToConstruct.Num();
+
+	// Since we are called after the first file is popped (but before it's actually done), we have one less completed.
+	CompletedFileCount--;
+
+	uint64 MaxDiskSize = FBuildPatchUtils::CalculateDiskSpaceRequirementsWithDeleteDuringInstall(
+		Configuration.ConstructList, CompletedFileCount, Configuration.ManifestSet, ChunkDbSizesAtPosition, TotalChunkDbSize);
+
+	// Strip off the data we already have on disk.
+	uint64 PostDlSize = 0;
+	if (MaxDiskSize > TotalChunkDbSize)
+	{
+		PostDlSize = MaxDiskSize - TotalChunkDbSize;
+	}
+
+	return PostDlSize;
 }
 
 bool FBuildPatchFileConstructor::ConstructFileFromChunks(const FString& BuildFilename, const FFileManifest& FileManifest, bool bResumeExisting)
@@ -651,22 +721,43 @@ bool FBuildPatchFileConstructor::ConstructFileFromChunks(const FString& BuildFil
 		}
 	}
 
-	// If we haven't done so yet, make the initial disk space check
+	// If we haven't done so yet, make the initial disk space check. We do this after resume
+	// so that we know how much to discount from our current file size.
 	if (!bInitialDiskSizeCheck)
 	{
 		bInitialDiskSizeCheck = true;
-		const uint64 RequiredSpace = CalculateRequiredDiskSpace(FileManifest, FileManifest.FileSize - StartPosition);
-		// ThreadLock protects access to members RequiredDiskSpace and AvailableDiskSpace;
-		FScopeLock Lock(&ThreadLock);
-		RequiredDiskSpace = RequiredSpace;
-		if (!FileConstructorHelpers::CheckRemainingDiskSpace(Configuration.InstallDirectory, RequiredSpace, AvailableDiskSpace))
+
+		// Normal operation can just use the classic calculation
+		uint64 LocalDiskSpaceRequired = CalculateInProgressDiskSpaceRequired(FileManifest, StartPosition);
+
+		// If we are delete-during-install this gets more complicated because we'll be freeing up
+		// space as we add.
+		if (Configuration.bDeleteChunkDBFilesAfterUse)
+		{
+			LocalDiskSpaceRequired = CalculateDiskSpaceRequirementsWithDeleteDuringInstall(ConstructionStack);
+		}
+
+		uint64 LocalDiskSpaceAvailable = 0;
+		{
+			uint64 TotalSize = 0;
+			uint64 AvailableSpace = 0;
+			if (FPlatformMisc::GetDiskTotalAndFreeSpace(Configuration.InstallDirectory, TotalSize, AvailableSpace))
+			{
+				LocalDiskSpaceAvailable = AvailableSpace;
+			}
+		}
+
+		AvailableDiskSpace.store(LocalDiskSpaceAvailable, std::memory_order_release);
+		RequiredDiskSpace.store(LocalDiskSpaceRequired, std::memory_order_release);	
+
+		if (!FileConstructorHelpers::CheckRemainingDiskSpace(Configuration.InstallDirectory, LocalDiskSpaceRequired, LocalDiskSpaceAvailable))
 		{
-			UE_LOG(LogBuildPatchServices, Error, TEXT("Out of HDD space. Needs %llu bytes, Free %llu bytes"), RequiredSpace, AvailableDiskSpace);
+			UE_LOG(LogBuildPatchServices, Error, TEXT("Out of HDD space. Needs %llu bytes, Free %llu bytes"), LocalDiskSpaceRequired, LocalDiskSpaceAvailable);
 			InstallerError->SetError(
 				EBuildPatchInstallError::OutOfDiskSpace,
 				DiskSpaceErrorCodes::InitialSpaceCheck,
 				0,
-				BuildPatchServices::GetDiskSpaceMessage(Configuration.InstallDirectory, RequiredSpace, AvailableDiskSpace));
+				BuildPatchServices::GetDiskSpaceMessage(Configuration.InstallDirectory, LocalDiskSpaceRequired, LocalDiskSpaceAvailable));
 			return false;
 		}
 	}
@@ -678,6 +769,7 @@ bool FBuildPatchFileConstructor::ConstructFileFromChunks(const FString& BuildFil
 		FileConstructorStat->OnResumeCompleted();
 	}
 
+	// Returns false if the write failed in some way (almost certainly disk space, could be drive disconnection)
 	auto FlushToAsyncWriter = [this](FArchive& DestinationFile, FSHA1& HashState)
 	{
 		if (bStallWhenFileSystemThrottled)
@@ -818,6 +910,9 @@ bool FBuildPatchFileConstructor::ConstructFileFromChunks(const FString& BuildFil
 
 		bSuccess = !NewFile->IsError();
 
+		// Update this for disk space requirements tracking below on error
+		StartPosition = NewFile->Tell();
+
 		// Close the file writer
 		FileConstructorStat->OnBeforeAdminister();
 		ActivityRecord.CyclesStart = FStatsCollector::GetCycles();
@@ -843,33 +938,75 @@ bool FBuildPatchFileConstructor::ConstructFileFromChunks(const FString& BuildFil
 	// Check for error state
 	if (!bSuccess)
 	{
-		// Recalculate disk space first
-		int64 InProgressFileSize = FileManifest.FileSize;
-		FileSystem->GetFileSize(*NewFilename, InProgressFileSize);
-		const uint64 RemainingRequiredSpace = CalculateRequiredDiskSpace(FileManifest, FileManifest.FileSize - InProgressFileSize);
-		uint64 RemainingAvailableDiskSpace = 0;
-		if (!FileConstructorHelpers::CheckRemainingDiskSpace(Configuration.InstallDirectory, RemainingRequiredSpace, RemainingAvailableDiskSpace))
-		{
-			// ThreadLock protects access to members RequiredDiskSpace and AvailableDiskSpace
-			ThreadLock.Lock();
-			RequiredDiskSpace = RemainingRequiredSpace;
-			AvailableDiskSpace = RemainingAvailableDiskSpace;
-			ThreadLock.Unlock();
-			// Convert error to disk space rather than reported.
-			UE_LOG(LogBuildPatchServices, Error, TEXT("Out of HDD space. Needs %llu bytes, Free %llu bytes"), RemainingRequiredSpace, RemainingAvailableDiskSpace);
-			InstallerError->SetError(
-				EBuildPatchInstallError::OutOfDiskSpace,
-				DiskSpaceErrorCodes::DuringInstallation,
-				0,
-				BuildPatchServices::GetDiskSpaceMessage(Configuration.InstallDirectory, RemainingRequiredSpace, RemainingAvailableDiskSpace));
-			ConstructionError = EConstructionError::OutOfDiskSpace;
+		if (ConstructionError == EConstructionError::SerializeError)
+		{
+			// Serialize error is our catchall file error right now. This should probably get
+			// migrated such that it's when we fail to load an existing chunk (i.e. corruption)
+			// but instead that shows up as a missing chunk.
+
+
+		// HOTFIX - we disable this check for delete during install
+			
+			uint64 TotalSize = 0;
+			uint64 FreeSize = 0;
+			if (FPlatformMisc::GetDiskTotalAndFreeSpace(Configuration.InstallDirectory, TotalSize, FreeSize))
+			{
+				// We're responding to an actual failure, which would have happened because we literally weren't able
+				// to write our write butter. Because of transient stuff this might not be correct so we double our
+				// write buffer size for this check.
+				if (FreeSize < (2 * WriteBufferSize))
+				{
+					// We've already failed so it makes sense to reevaluate how much extra we need. 
+					// I'm not sure I like using the same error wording for initial and ongoing disk space failure, but whatevs
+					{
+						uint64 LocalDiskSpaceRequired = CalculateInProgressDiskSpaceRequired(FileManifest, StartPosition);
+
+						// If we are delete-during-install this gets more complicated because we'll be freeing up
+						// space as we add.
+		if (!Configuration.bDeleteChunkDBFilesAfterUse &&
+			!FileConstructorHelpers::CheckRemainingDiskSpace(Configuration.InstallDirectory, RemainingRequiredSpace, RemainingAvailableDiskSpace))
+						{
+							LocalDiskSpaceRequired = CalculateDiskSpaceRequirementsWithDeleteDuringInstall(ConstructionStack);
+						}
+
+						AvailableDiskSpace.store(FreeSize, std::memory_order_release);
+						RequiredDiskSpace.store(LocalDiskSpaceRequired, std::memory_order_release);
+
+					}
+
+					ConstructionError = EConstructionError::OutOfDiskSpace;
+				}
+				else
+				{
+					// If it looks like we had enough disk space to write the last buffer, then 
+					// leave it as serialize.
+				}
+			}
+			else
+			{
+				// If we can't get the free space then likely the disk has disconnected or otherwise had a Bad Error, leave
+				// as serialize.
+			}
 		}
-		else
+
+		// \todo not exactly sure why this only reports on a file creation error?
+		const bool bReportAnalytic = InstallerError->HasError() == false;
+		switch (ConstructionError)
 		{
-			const bool bReportAnalytic = InstallerError->HasError() == false;
-			switch (ConstructionError)
+		case EConstructionError::OutOfDiskSpace:
+			{
+				uint64 LocalAvailableDiskSpace = AvailableDiskSpace.load(std::memory_order_acquire);
+				uint64 LocalRequiredDiskSpace = RequiredDiskSpace.load(std::memory_order_acquire);
+				UE_LOG(LogBuildPatchServices, Error, TEXT("Out of HDD space. Needs %llu bytes, Free %llu bytes"), LocalRequiredDiskSpace, LocalAvailableDiskSpace);
+				InstallerError->SetError(
+					EBuildPatchInstallError::OutOfDiskSpace,
+					DiskSpaceErrorCodes::DuringInstallation,
+					0,
+					BuildPatchServices::GetDiskSpaceMessage(Configuration.InstallDirectory, LocalRequiredDiskSpace, LocalAvailableDiskSpace));
+				break;
+			}
+		case EConstructionError::CannotCreateFile:
 			{
-			case EConstructionError::CannotCreateFile:
 				if (bReportAnalytic)
 				{
 					InstallerAnalytics->RecordConstructionError(BuildFilename, LastError, TEXT("Could Not Create File"));
@@ -877,13 +1014,19 @@ bool FBuildPatchFileConstructor::ConstructFileFromChunks(const FString& BuildFil
 				}
 				InstallerError->SetError(EBuildPatchInstallError::FileConstructionFail, ConstructionErrorCodes::FileCreateFail, LastError);
 				break;
-			case EConstructionError::MissingChunk:
+			}
+		case EConstructionError::MissingChunk:
+			{
 				InstallerError->SetError(EBuildPatchInstallError::FileConstructionFail, ConstructionErrorCodes::MissingChunkData);
 				break;
-			case EConstructionError::SerializeError:
+			}
+		case EConstructionError::SerializeError:
+			{
 				InstallerError->SetError(EBuildPatchInstallError::FileConstructionFail, ConstructionErrorCodes::SerializationError);
 				break;
-			case EConstructionError::TrackingError:
+			}
+		case EConstructionError::TrackingError:
+			{
 				InstallerError->SetError(EBuildPatchInstallError::FileConstructionFail, ConstructionErrorCodes::TrackingError);
 				break;
 			}
@@ -985,6 +1128,7 @@ bool FBuildPatchFileConstructor::AppendChunkData(const FChunkPart& ChunkPart, TA
 	}
 	else
 	{
+		// We'd really like to know if this was because it's missing or because it failed in some way.
 		ConstructionError = EConstructionError::MissingChunk;
 	}
 	return ConstructionError == EConstructionError::None;
diff --git a/Engine/Source/Runtime/Online/BuildPatchServices/Private/BuildPatchFileConstructor.h b/Engine/Source/Runtime/Online/BuildPatchServices/Private/BuildPatchFileConstructor.h
--- a/Engine/Source/Runtime/Online/BuildPatchServices/Private/BuildPatchFileConstructor.h
+++ b/Engine/Source/Runtime/Online/BuildPatchServices/Private/BuildPatchFileConstructor.h
@@ -30,6 +30,7 @@ namespace BuildPatchServices
 	class IFileConstructorStat;
 	class IBuildManifestSet;
 	class IBuildInstallerThread;
+	class IChunkDbChunkSource;
 
 	/**
 	 * A struct containing the configuration values for a file constructor.
@@ -75,7 +76,10 @@ namespace BuildPatchServices
 		 * @param InstallerAnalytics        Pointer to the installer analytics handler for reporting events.
 		 * @param FileConstructorStat       Pointer to the stat class for receiving updates.
 		 */
-		FBuildPatchFileConstructor(FFileConstructorConfig Configuration, IFileSystem* FileSystem, IChunkSource* ChunkSource, IChunkReferenceTracker* ChunkReferenceTracker, IInstallerError* InstallerError, IInstallerAnalytics* InstallerAnalytics, IFileConstructorStat* FileConstructorStat);
+		FBuildPatchFileConstructor(
+			FFileConstructorConfig Configuration, IFileSystem* FileSystem, IChunkSource* ChunkSource, 
+			IChunkDbChunkSource* ChunkDbChunkSource, IChunkReferenceTracker* ChunkReferenceTracker, IInstallerError* InstallerError, 
+			IInstallerAnalytics* InstallerAnalytics, IFileConstructorStat* FileConstructorStat);
 
 		/**
 		 * Default Destructor, will delete the allocated Thread
@@ -97,14 +101,19 @@ namespace BuildPatchServices
 		// IControllable interface end.
 
 		/**
-		 * Get the disk space that was required to perform the installation
-		 * @return	the disk space required to perform the installation in bytes
+		 * Get the disk space that was required to perform the installation. This can change over time and indicates the required
+		 * space to _finish_ the installation from the current state. It is not initialized until after resume is processed and returns
+		 * zero until that time. Note that since this and GetAvailableDiskSpace are separate accessors there's no guarantee that they
+		 * match - e.g. if you call GetRequiredDiskSpace and then GetAvailableDiskSpace immediately afterwards, it's possible the Available
+		 * Disk Space value is from a later call. This is highly unlikely due to how rare these updates are, but it's possible. Use these
+		 * for UI purposes only.
 		 */
 		uint64 GetRequiredDiskSpace();
 
 		/**
-		 * Get the disk space that was available at the time of checking for the required disk space
-		 * @return	the disk space that was available in bytes
+		 * Get the disk space that was available when last updating RequiredDiskSpace. See notes with GetRequiredDiskSpace.
+		 * It's possible for this to return 0 due to the underlying operating system being unable to report a value in cases of
+		 * e.g. the drive being disconnected.
 		 */
 		uint64 GetAvailableDiskSpace();
 
@@ -141,7 +150,10 @@ namespace BuildPatchServices
 		 * @param InProgressFileSize		The remaining size required for the file currently being constructed.
 		 * @return the number of bytes required on disk to complete the installation.
 		 */
-		uint64 CalculateRequiredDiskSpace(const FFileManifest& InProgressFileManifest, uint64 InProgressFileSize);
+		uint64 CalculateInProgressDiskSpaceRequired(const FFileManifest& InProgressFileManifest, uint64 InProgressFileSize);
+
+		// Calculates the amount of disk space we need to finish the install, needs to be called on file boundaries.
+		uint64 CalculateDiskSpaceRequirementsWithDeleteDuringInstall(const TArray<FString>& InConstructionStack);
 
 		/**
 		 * Constructs a particular file referenced by the given BuildManifest. The function takes an interface to a class that can provide availability information of chunks so that this
@@ -194,6 +206,7 @@ namespace BuildPatchServices
 
 		// Pointer to chunk source.
 		IChunkSource* ChunkSource;
+		IChunkDbChunkSource* ChunkDbSource; // can be null if not using.
 
 		// Pointer to the chunk reference tracker.
 		IChunkReferenceTracker* ChunkReferenceTracker;
@@ -214,10 +227,10 @@ namespace BuildPatchServices
 		int64 ByteProcessed;
 
 		// The amount of disk space requirement that was calculated when beginning the process. 0 if the install process was not started, or no additional space was needed.
-		uint64 RequiredDiskSpace;
+		std::atomic_uint64_t RequiredDiskSpace;
 
 		// The amount of disk space available when beginning the process. 0 if the install process was not started.
-		uint64 AvailableDiskSpace;
+		std::atomic_uint64_t AvailableDiskSpace;
 
 		// Event executed before deleting an old installation file.
 		FOnBeforeDeleteFile BeforeDeleteFileEvent;
@@ -238,6 +251,9 @@ namespace BuildPatchServices
 		FArchive* WriteJobArchive = nullptr;
 		std::atomic_bool bWriteJobCompleted = false; // Only set to true if the Serialize() call was completed.
 		bool bWriteJobRunning = false; // Foreground thread only - have we dispatched a job?
+
+		// Where we are in the chunk consumption list after each file.
+		TArray<int32> FileCompletionPositions;
 	};
 
 	/**
diff --git a/Engine/Source/Runtime/Online/BuildPatchServices/Private/BuildPatchInstaller.cpp b/Engine/Source/Runtime/Online/BuildPatchServices/Private/BuildPatchInstaller.cpp
--- a/Engine/Source/Runtime/Online/BuildPatchServices/Private/BuildPatchInstaller.cpp
+++ b/Engine/Source/Runtime/Online/BuildPatchServices/Private/BuildPatchInstaller.cpp
@@ -1256,13 +1256,12 @@ namespace BuildPatchServices
 		int32 CurrentPosition = 0;
 
 		TArray<int32> FileCompletionPositions;
-		TArray<FString> OrderedFiles;
-
 		for (const FString& FileToConstruct : FilesToConstruct)
 		{
 			const FFileManifest* FileManifest = ManifestSet->GetNewFileManifest(FileToConstruct);
 			if (!FileManifest)
 			{
+				// This is an error condition that will fail to install, but we don't want to crash here...
 				continue;
 			}
 
@@ -1271,61 +1270,15 @@ namespace BuildPatchServices
 
 			CurrentPosition += AdvanceCount;
 
-			OrderedFiles.Add(FileToConstruct);
 			FileCompletionPositions.Add(CurrentPosition);
 		}
 
 		TArray<uint64> ChunkDbSizesAtPosition;
 		uint64 TotalChunkDbSize = IChunkDbChunkSource::GetChunkDbSizesAtIndexes(ChunkDbFiles, FileSystem, ReferenceChain, FileCompletionPositions, ChunkDbSizesAtPosition);
 
-		uint64 TotalDeletedSize = 0;
-		uint64 TotalWrittenSize = 0;
-
-		uint64 MaxDiskSize = 0;
-
-		int32 MaxDiskSizeFileIndex = 0;
-		int32 FileIndex = 0;
-
-		// We start off with full chunk db size.
-		uint64 TotalChunkDbSizeAtLastFile = TotalChunkDbSize;
-
-		for (uint64 FileSize : ChunkDbSizesAtPosition)
-		{
-			// We've completed this file:
-			const FFileManifest* IncomingFileManifest = ManifestSet->GetNewFileManifest(OrderedFiles[FileIndex]);
-			uint64 NewFileSize = 0;
-			if (IncomingFileManifest)
-			{
-				NewFileSize = IncomingFileManifest->FileSize;
-			}
-
-			TotalWrittenSize += NewFileSize;
-
-			// We delete the chunkdbs _after_ we write the output so we can't use this size until 
-			// the next file gets done. Be sure to handle the case where we've deleted so much data
-			// we are below the waterline.
-			if (TotalDeletedSize < (TotalWrittenSize + TotalChunkDbSizeAtLastFile) &&
-				(TotalChunkDbSizeAtLastFile + TotalWrittenSize - TotalDeletedSize) > MaxDiskSize)
-			{
-				MaxDiskSize = TotalChunkDbSizeAtLastFile + TotalWrittenSize - TotalDeletedSize;
-				MaxDiskSizeFileIndex = FileIndex;
-			}
-
-			//UE_LOG(LogTemp, Display, TEXT("Max disk use %llu (total written: %llu -- total deleted: +%llu)"), MaxDiskSize, TotalWrittenSize, TotalDeletedSize);
-
-			// If we are patching, we can now delete the output file, which decreases our disk presence, however
-			// we update this after the check because we can't delete until after we have the file fully constructed.
-			const FFileManifest* OnDiskFileManifest = ManifestSet->GetCurrentFileManifest(FilesToConstruct[FileIndex]);
-			if (OnDiskFileManifest)
-			{
-				TotalDeletedSize += OnDiskFileManifest->FileSize;
-			}
-
-			FileIndex++;
-			TotalChunkDbSizeAtLastFile = FileSize;
-		}
-
-		//UE_LOG(LogTemp, Display, TEXT("Max disk use %llu (install size: %llu: +%llu) after file %s"), MaxDiskSize, TotalWrittenSize, MaxDiskSize - TotalWrittenSize, *OrderedFiles[AtFile]);
+		uint64 MaxDiskSize = FBuildPatchUtils::CalculateDiskSpaceRequirementsWithDeleteDuringInstall(
+			FilesToConstruct, 0, ManifestSet,
+			ChunkDbSizesAtPosition, TotalChunkDbSize);
 
 		return MaxDiskSize;
 	}
@@ -1582,9 +1535,8 @@ namespace BuildPatchServices
 			FCC.SharedContext = Configuration.SharedContext.Get();
 			FCC.StagingDirectory = InstallStagingDir;
 
-
 			TUniquePtr<FBuildPatchFileConstructor> FileConstructor(
-				new FBuildPatchFileConstructor(FCC, FileSystem.Get(), ChainedChunkSource.Get(), ChunkReferenceTracker.Get(), InstallerError.Get(),InstallerAnalytics.Get(),FileConstructorStatistics.Get())
+				new FBuildPatchFileConstructor(FCC, FileSystem.Get(), ChainedChunkSource.Get(), ChunkDbChunkSource.Get(), ChunkReferenceTracker.Get(), InstallerError.Get(),InstallerAnalytics.Get(),FileConstructorStatistics.Get())
 			);
 			FDelegateHandle OnBeforeDeleteFileHandle = FileConstructor->OnBeforeDeleteFile().AddLambda([this, &InstallChunkSource](const FString& FilePath)
 			{
diff --git a/Engine/Source/Runtime/Online/BuildPatchServices/Private/BuildPatchUtil.cpp b/Engine/Source/Runtime/Online/BuildPatchServices/Private/BuildPatchUtil.cpp
--- a/Engine/Source/Runtime/Online/BuildPatchServices/Private/BuildPatchUtil.cpp
+++ b/Engine/Source/Runtime/Online/BuildPatchServices/Private/BuildPatchUtil.cpp
@@ -254,3 +254,65 @@ uint8 FBuildPatchUtils::VerifyFile(IFileSystem* FileSystem, const FString& FileT
 	ProgressDelegate.ExecuteIfBound(1.0f);
 	return ReturnValue;
 }
+
+
+uint64 FBuildPatchUtils::CalculateDiskSpaceRequirementsWithDeleteDuringInstall(
+	const TArray<FString>& InFilesToConstruct, 
+	int32 InCompletedFileCount, 
+	IBuildManifestSet* InManifestSet,
+	const TArray<uint64>& InChunkDbSizesAtPosition,
+	uint64 InTotalChunkDbSize
+	)
+{
+	uint64 TotalDeletedSize = 0;
+	uint64 TotalWrittenSize = 0;
+
+	uint64 MaxDiskSize = 0;
+	int32 MaxDiskSizeFileIndex = 0;
+
+	// We start off with full chunk db size.
+	uint64 TotalChunkDbSizeAtLastFile = InTotalChunkDbSize;
+
+	for (int32 FileIndex = InCompletedFileCount; FileIndex < InFilesToConstruct.Num(); FileIndex++)
+	{
+		uint64 ChunkDbSize = InChunkDbSizesAtPosition[FileIndex];
+
+		// We've completed this file
+		uint64 NewFileSize = 0;
+		{
+			const FFileManifest* IncomingFileManifest = InManifestSet->GetNewFileManifest(InFilesToConstruct[FileIndex]);			
+			if (IncomingFileManifest)
+			{
+				NewFileSize = IncomingFileManifest->FileSize;
+			}
+		}
+
+		TotalWrittenSize += NewFileSize;
+
+		// We delete the chunkdbs _after_ we write the output so we can't use this size until 
+		// the next file gets done. Be sure to handle the case where we've deleted so much data
+		// we are below the waterline.
+		if (TotalDeletedSize < (TotalWrittenSize + TotalChunkDbSizeAtLastFile) &&
+			(TotalChunkDbSizeAtLastFile + TotalWrittenSize - TotalDeletedSize) > MaxDiskSize)
+		{
+			MaxDiskSize = TotalChunkDbSizeAtLastFile + TotalWrittenSize - TotalDeletedSize;
+			MaxDiskSizeFileIndex = FileIndex;
+		}
+
+		//UE_LOG(LogTemp, Display, TEXT("...@ file %d chunks = %llu, install = %llu"), FileIndex, TotalChunkDbSizeAtLastFile, TotalWrittenSize);
+
+		// If we are patching, we can now delete the output file, which decreases our disk presence, however
+		// we update this after the check because we can't delete until after we have the file fully constructed.
+		const FFileManifest* OnDiskFileManifest = InManifestSet->GetCurrentFileManifest(InFilesToConstruct[FileIndex]);
+		if (OnDiskFileManifest)
+		{
+			TotalDeletedSize += OnDiskFileManifest->FileSize;
+		}
+
+		TotalChunkDbSizeAtLastFile = ChunkDbSize;
+	}
+
+	//UE_LOG(LogTemp, Display, TEXT("Max disk use %llu (+%llu after d/l) (install size: %llu: +%llu) after file %d"), MaxDiskSize, PostDlSize, TotalWrittenSize, MaxDiskSize - TotalWrittenSize, MaxDiskSizeFileIndex);
+
+	return MaxDiskSize;
+}
\ No newline at end of file
diff --git a/Engine/Source/Runtime/Online/BuildPatchServices/Private/BuildPatchUtil.h b/Engine/Source/Runtime/Online/BuildPatchServices/Private/BuildPatchUtil.h
--- a/Engine/Source/Runtime/Online/BuildPatchServices/Private/BuildPatchUtil.h
+++ b/Engine/Source/Runtime/Online/BuildPatchServices/Private/BuildPatchUtil.h
@@ -164,4 +164,11 @@ struct FBuildPatchUtils
 	 * @return		0 if no match, 1 for match with Hash1, and 2 for match with Hash2
 	 */
 	static uint8 VerifyFile(BuildPatchServices::IFileSystem* FileSystem, const FString& FileToVerify, const FSHAHash& Hash1, const FSHAHash& Hash2);
+
+	static uint64 CalculateDiskSpaceRequirementsWithDeleteDuringInstall(
+		const TArray<FString>& InFilesToConstruct, 
+		int32 InCompletedFileCount, 
+		BuildPatchServices::IBuildManifestSet* InManifestSet, 
+		const TArray<uint64>& InChunkDbSizesAtPosition, 
+		uint64 InTotalChunkDbSize);
 };
diff --git a/Engine/Source/Runtime/Online/BuildPatchServices/Private/Installer/ChunkDbChunkSource.cpp b/Engine/Source/Runtime/Online/BuildPatchServices/Private/Installer/ChunkDbChunkSource.cpp
--- a/Engine/Source/Runtime/Online/BuildPatchServices/Private/Installer/ChunkDbChunkSource.cpp
+++ b/Engine/Source/Runtime/Online/BuildPatchServices/Private/Installer/ChunkDbChunkSource.cpp
@@ -160,6 +160,7 @@ namespace BuildPatchServices
 
 		// IInstallChunkSource interface begin.
 		virtual const TSet<FGuid>& GetAvailableChunks() const override;
+		virtual uint64 GetChunkDbSizesAtIndexes(const TArray<int32>& FileCompletionIndexes, TArray<uint64>& OutChunkDbSizesAtCompletion) const override;
 		// IInstallChunkSource interface end.
 
 		static void LoadChunkDbFiles(
@@ -209,6 +210,9 @@ namespace BuildPatchServices
 		TSet<FGuid> FailedToLoad;
 		// Communication between the incoming AddRepeatRequirement calls, and our worker thread.
 		TQueue<FGuid, EQueueMode::Mpsc> RepeatRequirementMessages;
+
+		// Number of chunks to process in this manifest when we started.
+		int32 OriginalChunkCount = 0;
 	};
 
 	// Read in the headers, evalutate the list of chunks, and determine when we'll be done with our chunk dbs.
@@ -306,32 +310,31 @@ namespace BuildPatchServices
 			}
 		}
 	}
-
-
-	uint64 IChunkDbChunkSource::GetChunkDbSizesAtIndexes(const TArray<FString>& ChunkDbFiles, IFileSystem* FileSystem, const TArray<FGuid>& ChunkAccessOrderedList, const TArray<int32>& FileCompletionIndexes, TArray<uint64>& OutChunkDbSizesAtCompletion)
+	
+	// Get how many bytes of chunkdbs will still exist on disk after the given indexes (i.e. will not have been retired)
+	// Return is the total size of given chunkdbs.
+	static uint64 GetChunkDbSizesAtIndexesInternal(const TArray<FChunkDbDataAccess>& InOpenedChunkDbs, int32 InOriginalChunkCount, const TArray<int32>& InFileCompletionIndexes, TArray<uint64>& OutChunkDbSizesAtCompletion)
 	{
-		TArray<FChunkDbDataAccess> ChunkFiles;
-		TMap<FGuid, FChunkAccessLookup> ChunkGuidLookup;
-
-		FChunkDbChunkSource::LoadChunkDbFiles(ChunkDbFiles, FileSystem, ChunkAccessOrderedList, ChunkFiles, ChunkGuidLookup, nullptr);
-
 		uint64 AllChunkDbSize = 0;
-		for (const FChunkDbDataAccess& ChunkFile : ChunkFiles)
+		for (const FChunkDbDataAccess& ChunkFile : InOpenedChunkDbs)
 		{
-			AllChunkDbSize += ChunkFile.Archive->TotalSize();
+			if (!ChunkFile.bIsRetired)
+			{
+				AllChunkDbSize += ChunkFile.Archive->TotalSize();
+			}
 		}
-
 		
 		// Go over the list of completions and evalute how many chunk dbs are left over.
-		for (int32 FileCompletionIndex : FileCompletionIndexes)
+		for (int32 FileCompletionIndex : InFileCompletionIndexes)
 		{
 			// retiring happens as the list is _popped_, so everything is backwards.
-			int32 RetireAtEquivalent = ChunkAccessOrderedList.Num() - FileCompletionIndex;
+			int32 RetireAtEquivalent = InOriginalChunkCount - FileCompletionIndex;
 
 			uint64 TotalSizeAtIndex = 0;
-			for (const FChunkDbDataAccess& ChunkFile : ChunkFiles)
+			for (const FChunkDbDataAccess& ChunkFile : InOpenedChunkDbs)
 			{
-				if (ChunkFile.RetireAt < RetireAtEquivalent)
+				if (!ChunkFile.bIsRetired &&
+					ChunkFile.RetireAt < RetireAtEquivalent)
 				{
 					TotalSizeAtIndex += ChunkFile.Archive->TotalSize();
 				}
@@ -342,6 +345,22 @@ namespace BuildPatchServices
 		return AllChunkDbSize;
 	}
 
+
+	uint64 FChunkDbChunkSource::GetChunkDbSizesAtIndexes(const TArray<int32>& FileCompletionIndexes, TArray<uint64>& OutChunkDbSizesAtCompletion) const
+	{
+		return GetChunkDbSizesAtIndexesInternal(ChunkDbDataAccesses, OriginalChunkCount, FileCompletionIndexes, OutChunkDbSizesAtCompletion);
+	}
+
+	uint64 IChunkDbChunkSource::GetChunkDbSizesAtIndexes(const TArray<FString>& ChunkDbFiles, IFileSystem* FileSystem, const TArray<FGuid>& ChunkAccessOrderedList, const TArray<int32>& FileCompletionIndexes, TArray<uint64>& OutChunkDbSizesAtCompletion)
+	{
+		TArray<FChunkDbDataAccess> ChunkFiles;
+		TMap<FGuid, FChunkAccessLookup> ChunkGuidLookup;
+
+		FChunkDbChunkSource::LoadChunkDbFiles(ChunkDbFiles, FileSystem, ChunkAccessOrderedList, ChunkFiles, ChunkGuidLookup, nullptr);
+
+		return GetChunkDbSizesAtIndexesInternal(ChunkFiles, ChunkAccessOrderedList.Num(), FileCompletionIndexes, OutChunkDbSizesAtCompletion);
+	}
+
 	FChunkDbChunkSource::FChunkDbChunkSource(FChunkDbSourceConfig InConfiguration, IPlatform* InPlatform, IFileSystem* InFileSystem, IChunkStore* InChunkStore, IChunkReferenceTracker* InChunkReferenceTracker, IChunkDataSerialization* InChunkDataSerialization, IMessagePump* InMessagePump, IInstallerError* InInstallerError, IChunkDbChunkSourceStat* InChunkDbChunkSourceStat)
 		: Configuration(MoveTemp(InConfiguration))
 		, Platform(InPlatform)
@@ -359,8 +378,13 @@ namespace BuildPatchServices
 		TArray<FGuid> OrderedGuidList;
 		InChunkReferenceTracker->CopyOutOrderedUseList(OrderedGuidList);
 
+		OriginalChunkCount = OrderedGuidList.Num();
+
 		LoadChunkDbFiles(Configuration.ChunkDbFiles, FileSystem, OrderedGuidList, ChunkDbDataAccesses, ChunkDbDataAccessLookup, &AvailableChunks);
 
+		// Immediately retire any chunkdbs we don't need so they don't eat disk space during the first file.
+		FChunkDbChunkSource::ReportFileCompletion();
+
 		// Start threaded load worker.
 		if (ChunkDbDataAccesses.Num() > 0)
 		{
diff --git a/Engine/Source/Runtime/Online/BuildPatchServices/Private/Installer/ChunkDbChunkSource.h b/Engine/Source/Runtime/Online/BuildPatchServices/Private/Installer/ChunkDbChunkSource.h
--- a/Engine/Source/Runtime/Online/BuildPatchServices/Private/Installer/ChunkDbChunkSource.h
+++ b/Engine/Source/Runtime/Online/BuildPatchServices/Private/Installer/ChunkDbChunkSource.h
@@ -38,6 +38,9 @@ namespace BuildPatchServices
 		// Fill out how many bytes of chunkdbs are left if we delete all the ones that
 		// are no longer necessary at the given FileCompletionIndexes in to ChunkAccessOrderedList
 		static uint64 GetChunkDbSizesAtIndexes(const TArray<FString>& ChunkDbFiles, IFileSystem* FileSystem, const TArray<FGuid>& ChunkAccessOrderedList, const TArray<int32>& FileCompletionIndexes, TArray<uint64>& OutChunkDbSizesAtCompletion);
+
+		// As above, except use the remaining open chunkdbs for progressive disk space checking.		
+		virtual uint64 GetChunkDbSizesAtIndexes(const TArray<int32>& FileCompletionIndexes, TArray<uint64>& OutChunkDbSizesAtCompletion) const = 0;
 	};
 
 	/**
