CIS fix due to robomerge over a deadended change

diff --git a/Engine/Source/Runtime/Online/BuildPatchServices/Private/BuildPatchFileConstructor.cpp b/Engine/Source/Runtime/Online/BuildPatchServices/Private/BuildPatchFileConstructor.cpp
--- a/Engine/Source/Runtime/Online/BuildPatchServices/Private/BuildPatchFileConstructor.cpp
+++ b/Engine/Source/Runtime/Online/BuildPatchServices/Private/BuildPatchFileConstructor.cpp
@@ -945,8 +945,6 @@ bool FBuildPatchFileConstructor::ConstructFileFromChunks(const FString& BuildFil
 			// but instead that shows up as a missing chunk.
 
 
-		// HOTFIX - we disable this check for delete during install
-			
 			uint64 TotalSize = 0;
 			uint64 FreeSize = 0;
 			if (FPlatformMisc::GetDiskTotalAndFreeSpace(Configuration.InstallDirectory, TotalSize, FreeSize))
@@ -963,8 +961,7 @@ bool FBuildPatchFileConstructor::ConstructFileFromChunks(const FString& BuildFil
 
 						// If we are delete-during-install this gets more complicated because we'll be freeing up
 						// space as we add.
-		if (!Configuration.bDeleteChunkDBFilesAfterUse &&
-			!FileConstructorHelpers::CheckRemainingDiskSpace(Configuration.InstallDirectory, RemainingRequiredSpace, RemainingAvailableDiskSpace))
+						if (Configuration.bDeleteChunkDBFilesAfterUse)
 						{
 							LocalDiskSpaceRequired = CalculateDiskSpaceRequirementsWithDeleteDuringInstall(ConstructionStack);
 						}
diff --git a/Engine/Source/Runtime/Online/BuildPatchServices/Private/BuildPatchFileConstructor.h b/Engine/Source/Runtime/Online/BuildPatchServices/Private/BuildPatchFileConstructor.h
--- a/Engine/Source/Runtime/Online/BuildPatchServices/Private/BuildPatchFileConstructor.h
+++ b/Engine/Source/Runtime/Online/BuildPatchServices/Private/BuildPatchFileConstructor.h
@@ -56,6 +56,8 @@ namespace BuildPatchServices
 		EInstallMode InstallMode;
 
 		IBuildInstallerSharedContext* SharedContext;
+
+		bool bDeleteChunkDBFilesAfterUse = false;
 	};
 
 	/**
