Fixes for install source tests.

#rnx

diff --git a/Engine/Source/Runtime/Online/BuildPatchServices/Private/Tests/Unit/InstallChunkSource.spec.cpp b/Engine/Source/Runtime/Online/BuildPatchServices/Private/Tests/Unit/InstallChunkSource.spec.cpp
--- a/Engine/Source/Runtime/Online/BuildPatchServices/Private/Tests/Unit/InstallChunkSource.spec.cpp
+++ b/Engine/Source/Runtime/Online/BuildPatchServices/Private/Tests/Unit/InstallChunkSource.spec.cpp
@@ -166,10 +166,13 @@ void FInstallChunkSourceSpec::Define()
 						{
 							InstallChunkSource->Get(SomeChunk);
 							TEST_TRUE(FakeChunkStore->Store.Contains(SomeChunk));
-							for (int32 NextReferenceIdx = 0; NextReferenceIdx < Configuration.BatchFetchMaximum; ++NextReferenceIdx)
-							{
-								TEST_TRUE(FakeChunkStore->Store.Contains(MockChunkReferenceTracker->NextReferences[NextReferenceIdx]));
-							}
+
+							// disable this one because we no longer prefetch in the install thread, so only the one chunk
+							// we "get" should exist.
+//							for (int32 NextReferenceIdx = 0; NextReferenceIdx < Configuration.BatchFetchMaximum; ++NextReferenceIdx)
+//							{
+//								TEST_TRUE(FakeChunkStore->Store.Contains(MockChunkReferenceTracker->NextReferences[NextReferenceIdx]));
+//							}
 						});
 					});
 
@@ -290,11 +293,23 @@ void FInstallChunkSourceSpec::Define()
 
 			It("should delay the chunk load process.", [this]()
 			{
-				InstallChunkSource->Get(SomeChunk);
+				// Pausing happens in the LoadFromBuild and we only load 1 per Get, so get two and be sure
+				// to check complete time.
+				int32 GetCounter = 2;
+				for (FGuid& AvailableChunk : SomeAvailableChunks)
+				{
+					InstallChunkSource->Get(AvailableChunk);
+					GetCounter--;
+					if (GetCounter == 0)
+					{
+						break;
+					}
+				}
+
 				double LongestDelay = 0.0f;
-				for (int32 Idx = 1; Idx < MockInstallChunkSourceStat->RxLoadStarted.Num(); ++Idx)
+				for (int32 Idx = 1; Idx < MockInstallChunkSourceStat->RxLoadComplete.Num(); ++Idx)
 				{
-					double ThisDelay = MockInstallChunkSourceStat->RxLoadStarted[Idx].Get<0>() - MockInstallChunkSourceStat->RxLoadStarted[Idx - 1].Get<0>();
+					double ThisDelay = MockInstallChunkSourceStat->RxLoadComplete[Idx].Get<0>() - MockInstallChunkSourceStat->RxLoadComplete[Idx - 1].Get<0>();
 					if (ThisDelay > LongestDelay)
 					{
 						LongestDelay = ThisDelay;
