Fixes for the chunk reference tracker test suite.

#rnx

diff --git a/Engine/Source/Runtime/Online/BuildPatchServices/Private/BuildPatchManifest.h b/Engine/Source/Runtime/Online/BuildPatchServices/Private/BuildPatchManifest.h
--- a/Engine/Source/Runtime/Online/BuildPatchServices/Private/BuildPatchManifest.h
+++ b/Engine/Source/Runtime/Online/BuildPatchServices/Private/BuildPatchManifest.h
@@ -351,23 +351,28 @@ private:
 	 */
 	void DestroyData();
 
+protected:
 	/**
 	 * Setups the lookup maps that optimize data access, should be called when Data changes
 	 */
 	void InitLookups();
+private:
 
 	/**
 	 * Helper for the public EnumerateProducibleChunks functions.
 	 */
 	int32 EnumerateProducibleChunks_Internal(const TFunction<bool(const FString&)>& FileAccessChecker, const TSet<FGuid>& ChunksRequired, TSet<FGuid>& ChunksAvailable) const;
 
-private:
+protected:
 	/** Holds the actual manifest data. Some other variables point to the memory held by these objects */
+
 	BuildPatchServices::FManifestMeta ManifestMeta;
 	BuildPatchServices::FChunkDataList ChunkDataList;
 	BuildPatchServices::FFileManifestList FileManifestList;
 	BuildPatchServices::FCustomFields CustomFields;
 
+private:
+
 	/** Holds the handle to our PreExit delegate */
 	FDelegateHandle OnPreExitHandle;
 
diff --git a/Engine/Source/Runtime/Online/BuildPatchServices/Private/Tests/Mock/Manifest.mock.h b/Engine/Source/Runtime/Online/BuildPatchServices/Private/Tests/Mock/Manifest.mock.h
--- a/Engine/Source/Runtime/Online/BuildPatchServices/Private/Tests/Mock/Manifest.mock.h
+++ b/Engine/Source/Runtime/Online/BuildPatchServices/Private/Tests/Mock/Manifest.mock.h
@@ -380,6 +380,18 @@ namespace BuildPatchServices
 			return true;
 		}
 
+		void SyncInternalManifestStructures()
+		{
+			// Normal manifests have some lookup structures we need to be fixed up in order for
+			// things to work.
+			for (const TPair<FString, FFileManifest>& FileManifest : FileManifests)
+			{
+				FileManifestList.FileList.Add(FileManifest.Value);
+			}
+
+			InitLookups();
+		}
+
 	public:
 		uint64 AppId;
 		FString AppName;
diff --git a/Engine/Source/Runtime/Online/BuildPatchServices/Private/Tests/Unit/ChunkReferenceTracker.spec.cpp b/Engine/Source/Runtime/Online/BuildPatchServices/Private/Tests/Unit/ChunkReferenceTracker.spec.cpp
--- a/Engine/Source/Runtime/Online/BuildPatchServices/Private/Tests/Unit/ChunkReferenceTracker.spec.cpp
+++ b/Engine/Source/Runtime/Online/BuildPatchServices/Private/Tests/Unit/ChunkReferenceTracker.spec.cpp
@@ -82,7 +82,7 @@ void FChunkReferenceTrackerSpec::Define()
 	}
 
 	// Specs.
-	xDescribe("ChunkReferenceTracker", [this]()
+	Describe("ChunkReferenceTracker", [this]()
 	{
 		Describe("GetReferencedChunks", [this]()
 		{
@@ -92,6 +92,10 @@ void FChunkReferenceTrackerSpec::Define()
 				{
 					MockManifest = MakeShareable(new FMockManifest());
 					MockManifest->FileManifests = FileManifests;
+					MockManifest->BuildFileList = FileList;
+					MockManifest->TaggedFileList = TSet<FString>(FileList);
+					MockManifest->SyncInternalManifestStructures();
+
 					TSet<FString> FilesToConstruct(FileList);
 					FilesToConstruct.Sort(TLess<FString>());
 					ManifestSet.Reset(FBuildManifestSetFactory::Create({ BuildPatchServices::FInstallerAction::MakeInstall(MockManifest.ToSharedRef()) }));
@@ -136,6 +140,10 @@ void FChunkReferenceTrackerSpec::Define()
 				{
 					MockManifest = MakeShareable(new FMockManifest());
 					MockManifest->FileManifests = FileManifests;
+					MockManifest->BuildFileList = FileList;
+					MockManifest->TaggedFileList = TSet<FString>(FileList);
+					MockManifest->SyncInternalManifestStructures();
+
 					TSet<FString> FilesToConstruct(SubsetFileList);
 					FilesToConstruct.Sort(TLess<FString>());
 					ManifestSet.Reset(FBuildManifestSetFactory::Create({ BuildPatchServices::FInstallerAction::MakeInstall(MockManifest.ToSharedRef()) }));
@@ -181,6 +189,10 @@ void FChunkReferenceTrackerSpec::Define()
 			{
 				MockManifest = MakeShareable(new FMockManifest());
 				MockManifest->FileManifests = FileManifests;
+				MockManifest->BuildFileList = FileList;
+				MockManifest->TaggedFileList = TSet<FString>(FileList);
+				MockManifest->SyncInternalManifestStructures();
+
 				TSet<FString> FilesToConstruct(FileList);
 				FilesToConstruct.Sort(TLess<FString>());
 				ManifestSet.Reset(FBuildManifestSetFactory::Create({ BuildPatchServices::FInstallerAction::MakeInstall(MockManifest.ToSharedRef()) }));
@@ -230,6 +242,10 @@ void FChunkReferenceTrackerSpec::Define()
 			{
 				MockManifest = MakeShareable(new FMockManifest());
 				MockManifest->FileManifests = FileManifests;
+				MockManifest->BuildFileList = FileList;
+				MockManifest->TaggedFileList = TSet<FString>(FileList);
+				MockManifest->SyncInternalManifestStructures();
+
 				TSet<FString> FilesToConstruct(FileList);
 				FilesToConstruct.Sort(TLess<FString>());
 				ManifestSet.Reset(FBuildManifestSetFactory::Create({ BuildPatchServices::FInstallerAction::MakeInstall(MockManifest.ToSharedRef()) }));
@@ -298,6 +314,10 @@ void FChunkReferenceTrackerSpec::Define()
 			{
 				MockManifest = MakeShareable(new FMockManifest());
 				MockManifest->FileManifests = FileManifests;
+				MockManifest->BuildFileList = FileList;
+				MockManifest->TaggedFileList = TSet<FString>(FileList);
+				MockManifest->SyncInternalManifestStructures();
+
 				TSet<FString> FilesToConstruct(FileList);
 				FilesToConstruct.Sort(TLess<FString>());
 				ManifestSet.Reset(FBuildManifestSetFactory::Create({ BuildPatchServices::FInstallerAction::MakeInstall(MockManifest.ToSharedRef()) }));
@@ -348,6 +368,10 @@ void FChunkReferenceTrackerSpec::Define()
 			{
 				MockManifest = MakeShareable(new FMockManifest());
 				MockManifest->FileManifests = FileManifests;
+				MockManifest->BuildFileList = FileList;
+				MockManifest->TaggedFileList = TSet<FString>(FileList);
+				MockManifest->SyncInternalManifestStructures();
+
 				TSet<FString> FilesToConstruct(FileList);
 				FilesToConstruct.Sort(TLess<FString>());
 				ManifestSet.Reset(FBuildManifestSetFactory::Create({ BuildPatchServices::FInstallerAction::MakeInstall(MockManifest.ToSharedRef()) }));
