Update chunk reference tracker to not use a critical section, and to avoid the linear lookup during sort.

#rnx

diff --git a/Engine/Source/Runtime/Online/BuildPatchServices/Private/Installer/ChunkReferenceTracker.cpp b/Engine/Source/Runtime/Online/BuildPatchServices/Private/Installer/ChunkReferenceTracker.cpp
--- a/Engine/Source/Runtime/Online/BuildPatchServices/Private/Installer/ChunkReferenceTracker.cpp
+++ b/Engine/Source/Runtime/Online/BuildPatchServices/Private/Installer/ChunkReferenceTracker.cpp
@@ -14,18 +14,19 @@ namespace BuildPatchServices
 	class FChunkReferenceTracker : public IChunkReferenceTracker
 	{
 	public:
+		// Construct the list of chunks from a manifest and an ordered list of files to construct. 
 		FChunkReferenceTracker(const IBuildManifestSet* ManifestSet, const TSet<FString>& FilesToConstruct);
-		FChunkReferenceTracker(TArray<FGuid> CustomUseStack);
+
+		// Pass in a direct ordered list of guids to use as chunks.
+		FChunkReferenceTracker(TArray<FGuid> CustomUseList);
 
 		~FChunkReferenceTracker();
 
 		// IChunkReferenceTracker interface begin.
 		virtual void CopyOutOrderedUseList(TArray<FGuid>& OutUseList) const override
 		{
-			UseStackCs.Lock();
-			OutUseList.Append(UseStack);
-			UseStackCs.Unlock();
-			Algo::Reverse(OutUseList);
+			int32 LocalCurrentPosition = CurrentPosition.load(std::memory_order_acquire);
+			OutUseList.Append(UseList.GetData() + LocalCurrentPosition, UseList.Num() - LocalCurrentPosition);
 		}
 		virtual TSet<FGuid> GetReferencedChunks() const override;
 		virtual int32 GetReferenceCount(const FGuid& ChunkId) const override;
@@ -37,22 +38,23 @@ namespace BuildPatchServices
 		// IChunkReferenceTracker interface end.
 
 	private:
-		TMap<FGuid, FThreadSafeCounter> ReferenceCount;
+		// Index of the next chunk to be used in UseList.
+		std::atomic<int> CurrentPosition = 0;
+
+		// Ordered list of guids in order of consumption.
+		TArray<FGuid> UseList;
 
-		// This is a list of the chunk guids that we are going to use _in order_ during installation.
-		// it's reversed so that as we finish with chunks we can Pop instead of doing a full memmove.
-		// \todo just make this an array with an index in to current position so we don't have to do the reorder
-		// shenanegans everywhere and can grab the current position without a mutex.
-		TArray<FGuid> UseStack;
-		mutable FCriticalSection UseStackCs;
-	};
+		// A sorted array (guid then index), where the indeces are the location of the guid in UseList, in ascending
+		// order.
+		// i.e. for all X in GuidUsagePositions, UseList[X.Value] = X.Key;
+		TArray<TPair<FGuid, int32>> GuidUsagePositions;		
+	};	
 
 	FChunkReferenceTracker::FChunkReferenceTracker(const IBuildManifestSet* ManifestSet, const TSet<FString>& FilesToConstruct)
-		: ReferenceCount()
-		, UseStack()
-		, UseStackCs()
 	{
-		// Create our full list of chunks, including dupe references, and track the reference count of each chunk.
+		// Iterate each file in reference order to construct the ordered list of chunks we
+		// will need to construct the files, and track when we are going to use them.
+		int32 UsageIndex = 0;
 		for (const FString& File : FilesToConstruct)
 		{
 			const FFileManifest* NewFileManifest = ManifestSet->GetNewFileManifest(File);
@@ -60,144 +62,169 @@ namespace BuildPatchServices
 			{
 				for (const FChunkPart& ChunkPart : NewFileManifest->ChunkParts)
 				{
-					ReferenceCount.FindOrAdd(ChunkPart.Guid).Increment();
-					UseStack.Add(ChunkPart.Guid);
+					UseList.Add(ChunkPart.Guid);
+					GuidUsagePositions.Add({ChunkPart.Guid, UsageIndex});
+					UsageIndex++;
 				}
 			}
 		}
-		// Reverse the order of UseStack so it can be used as a stack.
-		Algo::Reverse(UseStack);
-		UE_LOG(LogChunkReferenceTracker, VeryVerbose, TEXT("Created. Total references:%d. Unique chunks:%d"), UseStack.Num(), ReferenceCount.Num());
+
+		Algo::Sort(GuidUsagePositions);
 	}
 
 	FChunkReferenceTracker::FChunkReferenceTracker(TArray<FGuid> CustomChunkReferences)
-		: ReferenceCount()
-		, UseStack(MoveTemp(CustomChunkReferences))
-		, UseStackCs()
+		: UseList(MoveTemp(CustomChunkReferences))
 	{
-		ReferenceCount.Reserve(UseStack.Num());
-		for (const FGuid& Chunk : UseStack)
+		int32 ChunkIndex = 0;
+		for (const FGuid& Chunk : UseList)
 		{
-			ReferenceCount.FindOrAdd(Chunk).Increment();
+			GuidUsagePositions.Add({Chunk, ChunkIndex});
+			ChunkIndex++;
 		}
-		// Reverse the order of UseStack so it can be used as a stack.
-		Algo::Reverse(UseStack);
-		UE_LOG(LogChunkReferenceTracker, VeryVerbose, TEXT("Created. Total references:%d. Unique chunks:%d"), UseStack.Num(), ReferenceCount.Num());
+
+		Algo::Sort(GuidUsagePositions);
 	}
 
 	FChunkReferenceTracker::~FChunkReferenceTracker()
 	{
 	}
 
+
+	
 	TSet<FGuid> FChunkReferenceTracker::GetReferencedChunks() const
 	{
+		int32 LocalCurrentPosition = CurrentPosition.load(std::memory_order_acquire);
+
 		TSet<FGuid> ReferencedChunks;
-		for (const TPair<FGuid, FThreadSafeCounter>& Pair : ReferenceCount)
+
+		FGuid CurrentGuid;
+		for (const TPair<FGuid, int32>& ReferencedChunk : GuidUsagePositions)
 		{
-			if (Pair.Value.GetValue() > 0)
+			if (ReferencedChunk.Value < LocalCurrentPosition)
+			{
+				continue;
+			}
+			if (CurrentGuid != ReferencedChunk.Key)
 			{
-				ReferencedChunks.Add(Pair.Key);
+				ReferencedChunks.Add(ReferencedChunk.Key);
+				CurrentGuid = ReferencedChunk.Key;
 			}
 		}
+		
 		return ReferencedChunks;
 	}
 
 	int32 FChunkReferenceTracker::GetReferenceCount(const FGuid& ChunkId) const
 	{
-		return ReferenceCount.Contains(ChunkId) ? ReferenceCount[ChunkId].GetValue() : 0;
+		int32 LocalCurrentPosition = CurrentPosition.load(std::memory_order_acquire);
+		int32 GuidPosition = Algo::LowerBound(GuidUsagePositions, TPair<FGuid,int32> { ChunkId, CurrentPosition });
+
+		int32 StartGuidPosition = GuidPosition;
+		while (GuidPosition < GuidUsagePositions.Num() && GuidUsagePositions[GuidPosition].Key == ChunkId)
+		{
+			GuidPosition++;
+		}
+
+		return GuidPosition - StartGuidPosition;
 	}
 
 	void FChunkReferenceTracker::SortByUseOrder(TArray<FGuid>& ChunkList, ESortDirection Direction) const
 	{
-		// Thread lock to protect access to UseStack.
-		FScopeLock ThreadLock(&UseStackCs);
-		struct FIndexCache
+		int32 LocalCurrentPosition = CurrentPosition.load(std::memory_order_acquire);
+
+		// Get the next index for each chunk.
+		TMap<FGuid, int32> NextUsageIndexes;
+		for (FGuid& Guid : ChunkList)
 		{
-			FIndexCache(const TArray<FGuid>& InArray)
-				: Array(InArray)
-			{}
+			int32 GuidPosition = Algo::LowerBound(GuidUsagePositions, TPair<FGuid,int32> { Guid, CurrentPosition });
 
-			int32 GetIndex(const FGuid& Id)
+			int32 UsageIndex = TNumericLimits<int32>::Max(); // Unused chunks need to be sorted as though they are never used
+			if (GuidPosition < GuidUsagePositions.Num() &&
+				GuidUsagePositions[GuidPosition].Key == Guid)
 			{
-				if (!IndexCache.Contains(Id))
-				{
-					IndexCache.Add(Id, Array.FindLast(Id));
-				}
-				return IndexCache[Id];
+				UsageIndex =  GuidUsagePositions[GuidPosition].Value;
 			}
 
-			const TArray<FGuid>& Array;
-			TMap<FGuid, int32> IndexCache;
-		};
-		FIndexCache ChunkUseIndexes(UseStack);
+			NextUsageIndexes.Add(Guid, UsageIndex);
+		}
+
 		switch (Direction)
 		{
-			case ESortDirection::Ascending:
-				Algo::SortBy(ChunkList, [&ChunkUseIndexes](const FGuid& Id) { return ChunkUseIndexes.GetIndex(Id); }, TGreater<int32>());
+		case ESortDirection::Ascending:
+			{
+				Algo::SortBy(ChunkList, [&NextUsageIndexes](const FGuid& Id) { return NextUsageIndexes[Id]; }, TLess<int32>());
 				break;
-			case ESortDirection::Descending:
-				Algo::SortBy(ChunkList, [&ChunkUseIndexes](const FGuid& Id) { return ChunkUseIndexes.GetIndex(Id); }, TLess<int32>());
+			}
+		case ESortDirection::Descending:
+			{
+				Algo::SortBy(ChunkList, [&NextUsageIndexes](const FGuid& Id) { return NextUsageIndexes[Id]; }, TGreater<int32>());
 				break;
+			}
 		}
 	}
 
 	TArray<FGuid> FChunkReferenceTracker::GetNextReferences(int32 Count, const TFunction<bool(const FGuid&)>& SelectPredicate) const
 	{
-		// Thread lock to protect access to UseStack.
-		FScopeLock ThreadLock(&UseStackCs);
+		// Returns "Count" unique references that match SelectPredicate.
+
+		int32 LocalCurrentPosition = CurrentPosition.load(std::memory_order_acquire);
 		TSet<FGuid> AddedIds;
 		TArray<FGuid> NextReferences;
-		for (int32 UseStackIdx = UseStack.Num() - 1; UseStackIdx >= 0 && Count > NextReferences.Num(); --UseStackIdx)
+
+		for (; LocalCurrentPosition < UseList.Num() && NextReferences.Num() < Count; LocalCurrentPosition++)
 		{
-			const FGuid& UseId = UseStack[UseStackIdx];
+			const FGuid& UseId = UseList[LocalCurrentPosition];
 			if (AddedIds.Contains(UseId) == false && SelectPredicate(UseId))
 			{
 				AddedIds.Add(UseId);
 				NextReferences.Add(UseId);
 			}
 		}
+		
 		return NextReferences;
 	}
 
 	int32 FChunkReferenceTracker::GetRemainingChunkCount() const
 	{
-		FScopeLock ThreadLock(&UseStackCs);
-		return UseStack.Num();
+		int32 LocalCurrentPosition = CurrentPosition.load(std::memory_order_acquire);
+		return UseList.Num() - LocalCurrentPosition;
 	}
 
 	TArray<FGuid> FChunkReferenceTracker::SelectFromNextReferences(int32 Count, const TFunction<bool(const FGuid&)>& SelectPredicate) const
 	{
-		// Thread lock to protect access to UseStack.
-		FScopeLock ThreadLock(&UseStackCs);
-		TSet<FGuid> AddedIds;
-		TArray<FGuid> NextReferences;
-		for (int32 UseStackIdx = UseStack.Num() - 1; UseStackIdx >= 0 && Count > 0; --UseStackIdx)
-		{
-			const FGuid& UseId = UseStack[UseStackIdx];
-			if (AddedIds.Contains(UseId) == false)
+		// Original code 
+		/*
+			for (int32 UseStackIdx = UseStack.Num() - 1; UseStackIdx >= 0 && Count > 0; --UseStackIdx)
 			{
-				--Count;
-				if (SelectPredicate(UseId))
+				const FGuid& UseId = UseStack[UseStackIdx];
+				if (AddedIds.Contains(UseId) == false)
 				{
-					AddedIds.Add(UseId);
-					NextReferences.Add(UseId);
+					--Count;
+					if (SelectPredicate(UseId))
+					{
+						AddedIds.Add(UseId);
+						NextReferences.Add(UseId);
+					}
 				}
 			}
-		}
-		return NextReferences;
+		*/
+		// This is obfuscated a bit but is actually exactly the same as GetNextReferences. Since we can only decrease Count
+		// on a new unique chunk and we only add a new unique chunk when we add a reference, this is the same.
+		return GetNextReferences(Count, SelectPredicate);
 	}
 
 	bool FChunkReferenceTracker::PopReference(const FGuid& ChunkId)
 	{
-		// Thread lock to protect access to UseStack.
-		FScopeLock ThreadLock(&UseStackCs);
-		if (UseStack.Last() == ChunkId)
+		int32 LocalCurrentPosition = CurrentPosition.load(std::memory_order_acquire);
+		do
 		{
-			ReferenceCount[ChunkId].Decrement();
-			UseStack.Pop();
-			return true;
-		}
-		return false;
+			if (LocalCurrentPosition >= UseList.Num() || UseList[LocalCurrentPosition] != ChunkId)
+			{
+				return false;
+			}
+		} while (!CurrentPosition.compare_exchange_weak(LocalCurrentPosition, LocalCurrentPosition + 1, std::memory_order_acq_rel));
+
+		return true;
 	}
 
 	IChunkReferenceTracker* FChunkReferenceTrackerFactory::Create(const IBuildManifestSet* ManifestSet, const TSet<FString>& FilesToConstruct)
