Add explicit timestamps to content download events with best effort background completion timestamp

diff --git a/Engine/Source/Runtime/Online/BackgroundHTTP/Private/BackgroundHttpNotificationObject.cpp b/Engine/Source/Runtime/Online/BackgroundHTTP/Private/BackgroundHttpNotificationObject.cpp
--- a/Engine/Source/Runtime/Online/BackgroundHTTP/Private/BackgroundHttpNotificationObject.cpp
+++ b/Engine/Source/Runtime/Online/BackgroundHTTP/Private/BackgroundHttpNotificationObject.cpp
@@ -19,13 +19,14 @@ FBackgroundHttpNotificationObject::FBackgroundHttpNotificationObject(FText InNot
 {
 }
 
-FBackgroundHttpNotificationObject::FBackgroundHttpNotificationObject(FText InNotificationTitle, FText InNotificationBody, FText InNotificationAction, const FString& InNotificationActivationString, bool InNotifyOnlyOnFullSuccess, bool InbOnlySendNotificationInBackground, int32 InIdOverride)
+FBackgroundHttpNotificationObject::FBackgroundHttpNotificationObject(FText InNotificationTitle, FText InNotificationBody, FText InNotificationAction, const FString& InNotificationActivationString, bool InNotifyOnlyOnFullSuccess, bool InbOnlySendNotificationInBackground, int32 InIdOverride, bool InRecordNotificationTimestamp)
 	: NotificationTitle(InNotificationTitle)
     , NotificationAction(InNotificationAction)
     , NotificationBody(InNotificationBody)
 	, NotificationActivationString(InNotificationActivationString)
     , bOnlySendNotificationInBackground(InbOnlySendNotificationInBackground)
     , bNotifyOnlyOnFullSuccess(InNotifyOnlyOnFullSuccess)
+	, bRecordNotificationTimestamp(InRecordNotificationTimestamp)
 	, bIsInBackground(false)
 	, NumFailedDownloads(0)
 	, IdOverride(InIdOverride)
@@ -68,6 +69,12 @@ void FBackgroundHttpNotificationObject::OnApp_EnteringBackground()
 
 FBackgroundHttpNotificationObject::~FBackgroundHttpNotificationObject()
 {
+	if (bRecordNotificationTimestamp)
+	{
+		FString Timestamp = FDateTime::UtcNow().ToString();
+		FPlatformMisc::SetStoredValue(TEXT("Epic Games"), TEXT("UEBackgroundHTTPNotification"), TEXT("ScheduleTime"), Timestamp);
+	}
+
 	if (bOnlySendNotificationInBackground)
 	{
 		//These should only be registered if bOnlySendNotificationInBackground is set
diff --git a/Engine/Source/Runtime/Online/BackgroundHTTP/Public/BackgroundHttpNotificationObject.h b/Engine/Source/Runtime/Online/BackgroundHTTP/Public/BackgroundHttpNotificationObject.h
--- a/Engine/Source/Runtime/Online/BackgroundHTTP/Public/BackgroundHttpNotificationObject.h
+++ b/Engine/Source/Runtime/Online/BackgroundHTTP/Public/BackgroundHttpNotificationObject.h
@@ -14,7 +14,7 @@ class FBackgroundHttpNotificationObject
 {
 public:
 	BACKGROUNDHTTP_API FBackgroundHttpNotificationObject(FText InNotificationTitle, FText InNotificationBody, FText InNotificationAction, const FString& InNotificationActivationString, bool InNotifyOnlyOnFullSuccess);
-	BACKGROUNDHTTP_API FBackgroundHttpNotificationObject(FText InNotificationTitle, FText InNotificationBody, FText InNotificationAction, const FString& InNotificationActivationString, bool InNotifyOnlyOnFullSuccess, bool bOnlySendNotificationInBackground, int32 InIdOverride);
+	BACKGROUNDHTTP_API FBackgroundHttpNotificationObject(FText InNotificationTitle, FText InNotificationBody, FText InNotificationAction, const FString& InNotificationActivationString, bool InNotifyOnlyOnFullSuccess, bool bOnlySendNotificationInBackground, int32 InIdOverride, bool InRecordNotificationTimestamp = false);
 	BACKGROUNDHTTP_API ~FBackgroundHttpNotificationObject();
 
 	BACKGROUNDHTTP_API void NotifyOfDownloadResult(bool bWasSuccess);
@@ -31,7 +31,8 @@ private:
 
 	bool bOnlySendNotificationInBackground;
 	bool bNotifyOnlyOnFullSuccess;
-	
+	bool bRecordNotificationTimestamp;
+
 	volatile bool bIsInBackground;
 	volatile int32 NumFailedDownloads;
 
