Fix the data race issue related to http stats: the csv profiler read the http stats from main thread while http thread write the http stats.

diff --git a/Engine/Source/Runtime/Online/HTTP/Private/HttpManager.cpp b/Engine/Source/Runtime/Online/HTTP/Private/HttpManager.cpp
--- a/Engine/Source/Runtime/Online/HTTP/Private/HttpManager.cpp
+++ b/Engine/Source/Runtime/Online/HTTP/Private/HttpManager.cpp
@@ -522,15 +522,15 @@ PRAGMA_ENABLE_DEPRECATION_WARNINGS
 	}
 	
 	// Report csv stats.
-	int32 TotalDownloadedMB = int32(HttpStats.TotalDownloadedBytes >> 20);
-	CSV_CUSTOM_STAT_DEFINED(RequestsInQueue, HttpStats.RequestsInQueue, ECsvCustomStatOp::Set);
-	CSV_CUSTOM_STAT_DEFINED(MaxRequestsInQueue, int32(HttpStats.MaxRequestsInQueue), ECsvCustomStatOp::Set);
-	CSV_CUSTOM_STAT_DEFINED(RequestsInFlight, HttpStats.RequestsInFlight, ECsvCustomStatOp::Set);
-	CSV_CUSTOM_STAT_DEFINED(MaxRequestsInFlight, int32(HttpStats.MaxRequestsInFlight), ECsvCustomStatOp::Set);
-	CSV_CUSTOM_STAT_DEFINED(MaxTimeToWaitInQueue, int32(HttpStats.MaxTimeToWaitInQueue), ECsvCustomStatOp::Set);
+	int32 TotalDownloadedMB = int32(HttpStats.TotalDownloadedBytes.load() >> 20);
+	CSV_CUSTOM_STAT_DEFINED(RequestsInQueue, HttpStats.RequestsInQueue.load(), ECsvCustomStatOp::Set);
+	CSV_CUSTOM_STAT_DEFINED(MaxRequestsInQueue, int32(HttpStats.MaxRequestsInQueue.load()), ECsvCustomStatOp::Set);
+	CSV_CUSTOM_STAT_DEFINED(RequestsInFlight, HttpStats.RequestsInFlight.load(), ECsvCustomStatOp::Set);
+	CSV_CUSTOM_STAT_DEFINED(MaxRequestsInFlight, int32(HttpStats.MaxRequestsInFlight.load()), ECsvCustomStatOp::Set);
+	CSV_CUSTOM_STAT_DEFINED(MaxTimeToWaitInQueue, int32(HttpStats.MaxTimeToWaitInQueue.load()), ECsvCustomStatOp::Set);
 	CSV_CUSTOM_STAT_DEFINED(DownloadedMB, TotalDownloadedMB, ECsvCustomStatOp::Set);
-	CSV_CUSTOM_STAT_DEFINED(BandwidthMbps, int32(HttpStats.BandwidthMbps), ECsvCustomStatOp::Set);
-	CSV_CUSTOM_STAT_DEFINED(DurationMsAvg, int32(HttpStats.HttpDurationMsAvg), ECsvCustomStatOp::Set);
+	CSV_CUSTOM_STAT_DEFINED(BandwidthMbps, int32(HttpStats.BandwidthMbps.load()), ECsvCustomStatOp::Set);
+	CSV_CUSTOM_STAT_DEFINED(DurationMsAvg, int32(HttpStats.HttpDurationMsAvg.load()), ECsvCustomStatOp::Set);
 
 	// keep ticking
 	return true;
@@ -655,24 +655,24 @@ FHttpThreadBase* FHttpManager::GetThread()
 
 void FHttpManager::RecordStatTimeToConnect(float Duration)
 {
-	HttpStats.MaxTimeToConnect = FGenericPlatformMath::Max(Duration, HttpStats.MaxTimeToConnect);
+	HttpStats.MaxTimeToConnect = FGenericPlatformMath::Max(Duration, HttpStats.MaxTimeToConnect.load());
 }
 
 void FHttpManager::RecordStatRequestsInFlight(uint32 RequestsInFlight)
 {
 	HttpStats.RequestsInFlight = RequestsInFlight;
-	HttpStats.MaxRequestsInFlight = FGenericPlatformMath::Max(RequestsInFlight, HttpStats.MaxRequestsInFlight);
+	HttpStats.MaxRequestsInFlight = FGenericPlatformMath::Max(RequestsInFlight, HttpStats.MaxRequestsInFlight.load());
 }
 
 void FHttpManager::RecordStatRequestsInQueue(uint32 RequestsInQueue)
 {
 	HttpStats.RequestsInQueue = RequestsInQueue;
-	HttpStats.MaxRequestsInQueue = FGenericPlatformMath::Max(RequestsInQueue, HttpStats.MaxRequestsInQueue);
+	HttpStats.MaxRequestsInQueue = FGenericPlatformMath::Max(RequestsInQueue, HttpStats.MaxRequestsInQueue.load());
 }
 
 void FHttpManager::RecordMaxTimeToWaitInQueue(float Duration)
 {
-	HttpStats.MaxTimeToWaitInQueue = FGenericPlatformMath::Max(Duration, HttpStats.MaxTimeToWaitInQueue);
+	HttpStats.MaxTimeToWaitInQueue = FGenericPlatformMath::Max(Duration, HttpStats.MaxTimeToWaitInQueue.load());
 }
 
 void FHttpManager::RecordPlatformStats(const FHttpStatsPlatform& PlatformStats)
diff --git a/Engine/Source/Runtime/Online/HTTP/Public/HttpManager.h b/Engine/Source/Runtime/Online/HTTP/Public/HttpManager.h
--- a/Engine/Source/Runtime/Online/HTTP/Public/HttpManager.h
+++ b/Engine/Source/Runtime/Online/HTTP/Public/HttpManager.h
@@ -68,24 +68,27 @@ struct FHttpStatsPlatform
 
 struct FHttpStats
 {
+	// Use atomic for the following fields because csv profiler reads them from game thread while 
+	// http thread record them from http thread
+
 	/** The number of requests waiting in queue in http manager */
-	int32 RequestsInQueue = 0;
+	std::atomic<int32> RequestsInQueue = 0;
 	/** The number of requests in flight in http manager */
-	int32 RequestsInFlight = 0;
+	std::atomic<int32> RequestsInFlight = 0;
 	/** The max time to successfully connect the backend */
-	float MaxTimeToConnect = -1.0f;
+	std::atomic<float> MaxTimeToConnect = -1.0f;
 	/** The max waiting queue in http manager */
-	uint32 MaxRequestsInQueue = 0;
+	std::atomic<uint32> MaxRequestsInQueue = 0;
 	/** The max number of requests in flight in http manager */
-	uint32 MaxRequestsInFlight = 0;
+	std::atomic<uint32> MaxRequestsInFlight = 0;
 	/** The max waiting time in queue of http manager */
-	float MaxTimeToWaitInQueue = 0.0f;
+	std::atomic<float> MaxTimeToWaitInQueue = 0.0f;
 	/** The total bytes downloaded so far */
-	int64 TotalDownloadedBytes = 0;
+	std::atomic<int64> TotalDownloadedBytes = 0;
 	/** Approximate download bandwidth used */
-	int64 BandwidthMbps = 0;
+	std::atomic<int64> BandwidthMbps = 0;
 	/** Avg duration (in milliseconds) from request to response */
-	int64 HttpDurationMsAvg = 0;
+	std::atomic<int64> HttpDurationMsAvg = 0;
 	/** The optional http stats on specific platform */
 	TOptional<FHttpStatsPlatform> PlatformStats;
 
@@ -102,6 +105,24 @@ struct FHttpStats
 			&& FMath::IsNearlyEqual(MaxTimeToConnect, Other.MaxTimeToConnect)
 			&& FMath::IsNearlyEqual(MaxTimeToWaitInQueue, Other.MaxTimeToWaitInQueue);
 	}
+
+	FHttpStats()
+	{
+	}
+
+	FHttpStats(const FHttpStats& Other)
+		: RequestsInQueue(Other.RequestsInQueue.load())
+		, RequestsInFlight(Other.RequestsInFlight.load())
+		, MaxTimeToConnect(Other.MaxTimeToConnect.load())
+		, MaxRequestsInQueue(Other.MaxRequestsInQueue.load())
+		, MaxRequestsInFlight(Other.MaxRequestsInFlight.load())
+		, MaxTimeToWaitInQueue(Other.MaxTimeToWaitInQueue.load())
+		, TotalDownloadedBytes(Other.TotalDownloadedBytes.load())
+		, BandwidthMbps(Other.BandwidthMbps.load())
+		, HttpDurationMsAvg(Other.HttpDurationMsAvg.load())
+		, PlatformStats(Other.PlatformStats)
+	{
+	}
 };
 
 /**
