Enable event loop by default on Windows.

diff --git a/Engine/Source/Runtime/Online/HTTP/HTTP.Build.cs b/Engine/Source/Runtime/Online/HTTP/HTTP.Build.cs
--- a/Engine/Source/Runtime/Online/HTTP/HTTP.Build.cs
+++ b/Engine/Source/Runtime/Online/HTTP/HTTP.Build.cs
@@ -4,8 +4,7 @@ using UnrealBuildTool;
 
 public class HTTP : ModuleRules
 {
-	// Currently there is a random event loop crash when shutdown HTTP manager on PC
-	protected virtual bool bPlatformEventLoopEnabledByDefault { get { return !Target.Platform.IsInGroup(UnrealPlatformGroup.Windows); } }
+	protected virtual bool bPlatformEventLoopEnabledByDefault { get { return true; } }
 
 	protected virtual bool bPlatformSupportToIncreaseMaxRequestsAtRuntime { get { return true; } }
 
