Fix for Horde issue 596943. Removed deprecated method from Apple platform file.

diff --git a/Engine/Source/Runtime/Online/HTTP/Private/Apple/AppleHttp.cpp b/Engine/Source/Runtime/Online/HTTP/Private/Apple/AppleHttp.cpp
--- a/Engine/Source/Runtime/Online/HTTP/Private/Apple/AppleHttp.cpp
+++ b/Engine/Source/Runtime/Online/HTTP/Private/Apple/AppleHttp.cpp
@@ -954,9 +954,6 @@ void FAppleHttpRequest::CheckProgressDelegate()
 		const uint64 BytesRead = Response->GetNumBytesReceived();
 		if (BytesWritten != LastReportedBytesWritten || BytesRead != LastReportedBytesRead)
 		{
-PRAGMA_DISABLE_DEPRECATION_WARNINGS
-			OnRequestProgress().ExecuteIfBound(SharedThis(this), BytesWritten, BytesRead);
-PRAGMA_ENABLE_DEPRECATION_WARNINGS
 			OnRequestProgress64().ExecuteIfBound(SharedThis(this), BytesWritten, BytesRead);
 			LastReportedBytesWritten = BytesWritten;
 			LastReportedBytesRead = BytesRead;
