#UE BPS can now be used before the fork on servers. It will be allowed to create threads and must delete all threads by the time the fork arrives.

diff --git a/Engine/Source/Runtime/Online/BuildPatchServices/Private/Installer/InstallerSharedContext.cpp b/Engine/Source/Runtime/Online/BuildPatchServices/Private/Installer/InstallerSharedContext.cpp
--- a/Engine/Source/Runtime/Online/BuildPatchServices/Private/Installer/InstallerSharedContext.cpp
+++ b/Engine/Source/Runtime/Online/BuildPatchServices/Private/Installer/InstallerSharedContext.cpp
@@ -27,19 +27,11 @@ namespace BuildPatchServices
 
 	bool FBuildInstallerThread::StartThread(const TCHAR* DebugName)
 	{
-		// Ideally this would check if we were forkable or a forked child process but there is 
-		// currently no in-engine way to check that.  Since BPS does not currently support
-		// FRunnableThread::ThreadType::Fake or FRunnableThread::ThreadType::Forkable 
-		// this check ends up being equivalent for now.  We most likely *never* want support 
-		// forking while an installer is running.
-		if (FPlatformProcess::SupportsMultithreading() || FForkProcessHelper::IsForkedMultithreadInstance())
-		{
-			DoWorkEvent = FPlatformProcess::GetSynchEventFromPool();
+		// We most likely *never* want support forking while an installer is running.
+		DoWorkEvent = FPlatformProcess::GetSynchEventFromPool();
 
-			Thread = FForkProcessHelper::CreateForkableThread(this, DebugName);
-			check(Thread != nullptr);
-			check(Thread->GetThreadType() == FRunnableThread::ThreadType::Real);
-		}
+		Thread = FForkProcessHelper::CreateForkableThread(this, DebugName, 0, TPri_Normal, FGenericPlatformAffinity::GetNoAffinityMask(), EThreadCreateFlags::None, true);
+		check(Thread != nullptr);
 
 		return DoWorkEvent && Thread;
 	}
