BuildPatchServices: Disable core.uninitialized.ArraySubscript clang checker

#rnx

diff --git a/Engine/Source/Runtime/Online/BuildPatchServices/BuildPatchServices.Build.cs b/Engine/Source/Runtime/Online/BuildPatchServices/BuildPatchServices.Build.cs
--- a/Engine/Source/Runtime/Online/BuildPatchServices/BuildPatchServices.Build.cs
+++ b/Engine/Source/Runtime/Online/BuildPatchServices/BuildPatchServices.Build.cs
@@ -7,6 +7,8 @@ public class BuildPatchServices : ModuleRules
 {
 	public BuildPatchServices(ReadOnlyTargetRules Target) : base(Target)
 	{
+		StaticAnalyzerDisabledCheckers.Add("core.uninitialized.ArraySubscript");
+
 		PublicDependencyModuleNames.AddRange(
 			new string[]
 			{
