LwsWebSockets: Fix issue where websocket was being downgraded to not use SSL in wss+inscure:// endpoints
 - Allow self-signed certs with bDisableVerification.
 - Adds flag lws-insecure=true to disable websocket validation.

#tests QuickSmoke tested
[RN] minor

diff --git a/Engine/Source/Runtime/Online/WebSockets/Private/Lws/LwsWebSocket.cpp b/Engine/Source/Runtime/Online/WebSockets/Private/Lws/LwsWebSocket.cpp
--- a/Engine/Source/Runtime/Online/WebSockets/Private/Lws/LwsWebSocket.cpp
+++ b/Engine/Source/Runtime/Online/WebSockets/Private/Lws/LwsWebSocket.cpp
@@ -732,12 +732,18 @@ void FLwsWebSocket::ConnectInternal(struct lws_context &LwsContext)
 	// Use SSL and require a valid cerver cert
 	if (FCStringAnsi::Stricmp(UrlProtocol, "wss") == 0)
 	{
-		SslConnection = 1;
+		SslConnection = 1; // LCCSCF_USE_SSL
+
+		// Use ssl, and allow self-signed certs with bDisableCertValidation.
+		if (FLwsWebSocketsManager::Get().bDisableCertValidation)
+		{
+			SslConnection = 3; // LCCSCF_USE_SSL | LCCSCF_ALLOW_SELFSIGNED;
+		}
 	}
 	// Use SSL, and allow self-signed certs
 	else if (FCStringAnsi::Stricmp(UrlProtocol, "wss+insecure") == 0)
 	{
-		SslConnection = 2;
+		SslConnection = 3; // LCCSCF_USE_SSL | LCCSCF_ALLOW_SELFSIGNED;
 	}
 	// No encryption
 	else if (FCStringAnsi::Stricmp(UrlProtocol, "ws") == 0)
diff --git a/Engine/Source/Runtime/Online/WebSockets/Private/Lws/LwsWebSocketsManager.cpp b/Engine/Source/Runtime/Online/WebSockets/Private/Lws/LwsWebSocketsManager.cpp
--- a/Engine/Source/Runtime/Online/WebSockets/Private/Lws/LwsWebSocketsManager.cpp
+++ b/Engine/Source/Runtime/Online/WebSockets/Private/Lws/LwsWebSocketsManager.cpp
@@ -15,6 +15,7 @@
 #include "HAL/PlatformTime.h"
 #include "HAL/PlatformProcess.h"
 #include "HAL/LowLevelMemTracker.h"
+#include "Misc/CommandLine.h"
 #include "Misc/ConfigCacheIni.h"
 #include "Misc/CoreDelegates.h"
 #include "Misc/Fork.h"
@@ -57,7 +58,13 @@ FLwsWebSocketsManager::FLwsWebSocketsManager()
 	GConfig->GetDouble(TEXT("WebSockets.LibWebSockets"), TEXT("ThreadMinimumSleepTimeInSeconds"), ThreadMinimumSleepTimeInSeconds, GEngineIni);
 
 	GConfig->GetBool(TEXT("LwsWebSocket"), TEXT("bDisableDomainAllowlist"), bDisableDomainAllowlist, GEngineIni);
+
 	GConfig->GetBool(TEXT("LwsWebSocket"), TEXT("bDisableCertValidation"), bDisableCertValidation, GEngineIni);
+	
+#if !UE_BUILD_SHIPPING
+	// Allow -lws-insecure= to override ini flags
+	FParse::Bool(FCommandLine::Get(), TEXT("lws-insecure="), bDisableCertValidation);
+#endif
 }
 
 FLwsWebSocketsManager& FLwsWebSocketsManager::Get()
