BPS | Fix clang SA warnings in BuildPatchHash files

diff --git a/Engine/Source/Runtime/Online/BuildPatchServices/Private/BuildPatchHash.cpp b/Engine/Source/Runtime/Online/BuildPatchServices/Private/BuildPatchHash.cpp
--- a/Engine/Source/Runtime/Online/BuildPatchServices/Private/BuildPatchHash.cpp
+++ b/Engine/Source/Runtime/Online/BuildPatchServices/Private/BuildPatchHash.cpp
@@ -73,7 +73,7 @@ void FRollingHash::RollForward(const uint8& NewByte)
 {
 	// We must have consumed enough bytes to function correctly
 	check(NumBytesConsumed == WindowSize);
-	uint8 OldByte;
+	uint8 OldByte = 0;
 	WindowData.Dequeue(OldByte);
 	WindowData.Enqueue(NewByte);
 	// Update our HashState
@@ -97,15 +97,11 @@ uint64 FRollingHash::GetHashForDataSet(const uint8* DataSet, uint32 WindowSize)
 
 FRollingHash::FRollingHash(uint32 InWindowSize)
 	: WindowSize(InWindowSize)
-	, HashState(0)
-	, NumBytesConsumed(0)
 	, WindowData(WindowSize)
 {
 }
 
 FRollingHash::FRollingHash()
-	: WindowSize(0)
-	, WindowData(0)
 {
 }
 
diff --git a/Engine/Source/Runtime/Online/BuildPatchServices/Private/BuildPatchHash.h b/Engine/Source/Runtime/Online/BuildPatchServices/Private/BuildPatchHash.h
--- a/Engine/Source/Runtime/Online/BuildPatchServices/Private/BuildPatchHash.h
+++ b/Engine/Source/Runtime/Online/BuildPatchServices/Private/BuildPatchHash.h
@@ -102,13 +102,13 @@ public:
 private:
 	FRollingHash();
 	// The data size that we roll over.
-	const uint32 WindowSize;
+	const uint32 WindowSize = 0;
 	// The current hash value
-	uint64 HashState;
+	uint64 HashState = 0;
 	// The number of bytes we have consumed so far, used in hash function and to check validity of calls
-	uint32 NumBytesConsumed;
+	uint32 NumBytesConsumed = 0;
 	// Store the data to make access and rolling easier
-	HashRingBuffer WindowData;
+	HashRingBuffer WindowData = 0;
 };
 
 /**
