Make the locks in LwsWebSocket tranactionally safe

diff --git a/Engine/Source/Runtime/Online/WebSockets/Private/Lws/LwsWebSocket.cpp b/Engine/Source/Runtime/Online/WebSockets/Private/Lws/LwsWebSocket.cpp
--- a/Engine/Source/Runtime/Online/WebSockets/Private/Lws/LwsWebSocket.cpp
+++ b/Engine/Source/Runtime/Online/WebSockets/Private/Lws/LwsWebSocket.cpp
@@ -16,6 +16,7 @@
 #include "HttpModule.h"
 #include "HttpManager.h"
 #include "PlatformHttp.h"
+#include "Misc/TransactionallySafeScopeLock.h"
 
 // FLwsSendBuffer 
 FLwsSendBuffer::FLwsSendBuffer(const uint8* Data, SIZE_T Size, bool bInIsBinary)
@@ -160,7 +161,7 @@ void FLwsWebSocket::Close(int32 Code, const FString& Reason)
 	UE_LOG(LogWebSockets, Verbose, TEXT("FLwsWebSocket[%d]::Close: Close queued with code=%d reason=%s"), Identifier, Code, *Reason);
 
 	{
-		FScopeLock ScopeLock(&StateLock);
+		FTransactionallySafeScopeLock ScopeLock(&StateLock);
 		CloseRequest.Code = Code;
 		CloseRequest.Reason = ANSIReason;
 	}
@@ -283,7 +284,7 @@ int FLwsWebSocket::LwsCallback(lws* Instance, lws_callback_reasons Reason, void*
 		// No lock because we are the only thread that writes to State
 		EState PreviousState = State;
 		{
-			FScopeLock ScopeLock(&StateLock);
+			FTransactionallySafeScopeLock ScopeLock(&StateLock);
 			State = EState::Connected;
 		}
 		UE_LOG(LogWebSockets, Verbose, TEXT("FLwsWebSocket[%d]::LwsCallback: Received LWS_CALLBACK_CLIENT_ESTABLISHED, setting State=%s PreviousState=%s"), Identifier, ToString(State), ToString(PreviousState));
@@ -370,7 +371,7 @@ int FLwsWebSocket::LwsCallback(lws* Instance, lws_callback_reasons Reason, void*
 			UE_LOG(LogWebSockets, Verbose, TEXT("FLwsWebSocket[%d]::LwsCallback: Received LWS_CALLBACK_WS_PEER_INITIATED_CLOSE, setting State=%s CloseStatus=%d Reason=%s bPeerSpecifiedReason=%d PreviousState=%s"),
 				Identifier, ToString(EState::Closed), (int32)CloseStatus, *CloseReasonString, (int32)bPeerSpecifiedReason, ToString(State));
 			{
-				FScopeLock ScopeLock(&StateLock);
+				FTransactionallySafeScopeLock ScopeLock(&StateLock);
 				State = EState::Closed;
 				ClosedReason.Reason = MoveTemp(CloseReasonString);
 				ClosedReason.CloseStatus = CloseStatus;
@@ -394,7 +395,7 @@ int FLwsWebSocket::LwsCallback(lws* Instance, lws_callback_reasons Reason, void*
 			UE_LOG(LogWebSockets, Verbose, TEXT("FLwsWebSocket[%d]::LwsCallback: Received LWS_CALLBACK_WSI_DESTROY, setting State=%s PreviousState=%s"), Identifier, ToString(EState::Error), ToString(State));
 		
 			{
-				FScopeLock ScopeLock(&StateLock);
+				FTransactionallySafeScopeLock ScopeLock(&StateLock);
 				State = EState::Error;
 				ClosedReason.Reason = MoveTemp(CloseReasonString);
 			}
@@ -405,7 +406,7 @@ int FLwsWebSocket::LwsCallback(lws* Instance, lws_callback_reasons Reason, void*
 			if (State == EState::Connected)
 			{
 				{
-					FScopeLock ScopeLock(&StateLock);
+					FTransactionallySafeScopeLock ScopeLock(&StateLock);
 					State = EState::Closed;
 					ClosedReason.bWasClean = false;
 					ClosedReason.CloseStatus = LWS_CLOSE_STATUS_ABNORMAL_CLOSE;
@@ -429,7 +430,7 @@ int FLwsWebSocket::LwsCallback(lws* Instance, lws_callback_reasons Reason, void*
 				}
 
 				{
-					FScopeLock ScopeLock(&StateLock);
+					FTransactionallySafeScopeLock ScopeLock(&StateLock);
 					State = EState::Closed;
 					ClosedReason.Reason = MoveTemp(ClosedReasonString);
 					ClosedReason.CloseStatus = ClosedStatus;
@@ -456,7 +457,7 @@ int FLwsWebSocket::LwsCallback(lws* Instance, lws_callback_reasons Reason, void*
 			UE_LOG(LogWebSockets, Verbose, TEXT("FLwsWebSocket[%d]::LwsCallback: Received LWS_CALLBACK_CLOSED, setting State=%s CloseReason=%s PreviousState=%s"),
 				Identifier, ToString(EState::Closed), *CloseReasonString, ToString(State));
 			{
-				FScopeLock ScopeLock(&StateLock);
+				FTransactionallySafeScopeLock ScopeLock(&StateLock);
 				State = EState::Closed;
 				ClosedReason.Reason = MoveTemp(CloseReasonString);
 				ClosedReason.CloseStatus = LWS_CLOSE_STATUS_NORMAL;
@@ -481,7 +482,7 @@ int FLwsWebSocket::LwsCallback(lws* Instance, lws_callback_reasons Reason, void*
 		UE_LOG(LogWebSockets, Verbose, TEXT("FLwsWebSocket[%d]::LwsCallback: Received LWS_CALLBACK_CLIENT_CONNECTION_ERROR, setting State=%s CloseReason=%s PreviousState=%s"),
 			Identifier, ToString(EState::Error), *CloseReasonString, ToString(State));
 		{
-			FScopeLock ScopeLock(&StateLock);
+			FTransactionallySafeScopeLock ScopeLock(&StateLock);
 			if (State == EState::Connected ||
 				State == EState::ClosingByRequest)
 			{
@@ -517,7 +518,7 @@ int FLwsWebSocket::LwsCallback(lws* Instance, lws_callback_reasons Reason, void*
 			ANSICHAR* CloseReasonCopy;
 			{
 				// TODO:  Is this scope lock necessary?  We only got here because we recognized CloseRequest.Code was valid, and we do not allow it to be modified afer being set on the game thread, so shouldn't this thread have the correct value?
-				FScopeLock ScopeLock(&StateLock);
+				FTransactionallySafeScopeLock ScopeLock(&StateLock);
 				CloseCodeCopy = CloseRequest.Code;
 				CloseReasonCopy = CloseRequest.Reason;
 			}
@@ -579,7 +580,7 @@ void FLwsWebSocket::GameThreadTick()
 {
 	EState CurrentState;
 	{
-		FScopeLock ScopeLock(&StateLock);
+		FTransactionallySafeScopeLock ScopeLock(&StateLock);
 		CurrentState = State;
 	}
 	if (CurrentState != LastGameThreadState)
@@ -624,7 +625,7 @@ void FLwsWebSocket::GameThreadFinalize()
 
 	{
 		// TODO:  The contract requires the libwebsockets be done with this object prior to this being called.  Is there a better way to ensure we get the last set value by the libwebsockets thread?  FPlatformMisc::MemoryBarrier maybe?
-		FScopeLock ScopeLock(&StateLock);
+		FTransactionallySafeScopeLock ScopeLock(&StateLock);
 		PreviousState = State;
 		State = EState::None; // Will be re-usable on final delegate triggering
 		LastGameThreadState = State;
@@ -660,7 +661,7 @@ void FLwsWebSocket::LwsThreadTick()
 		check(LwsConnection);
 		{
 			// Check if we want to close
-			FScopeLock ScopeLock(&StateLock);
+			FTransactionallySafeScopeLock ScopeLock(&StateLock);
 			if (CloseRequest.Code != 0)
 			{
 				State = EState::ClosingByRequest;
@@ -697,7 +698,7 @@ void FLwsWebSocket::ConnectInternal(struct lws_context &LwsContext)
 	UE_LOG(LogWebSockets, Verbose, TEXT("FLwsWebSocket[%d]::ConnectInternal: setting State=%s PreviousState=%s"),
 		Identifier, ToString(EState::Connecting), ToString(EState::StartConnecting));
 	{
-		FScopeLock ScopeLock(&StateLock);
+		FTransactionallySafeScopeLock ScopeLock(&StateLock);
 		State = EState::Connecting;
 	}
 
@@ -713,7 +714,7 @@ void FLwsWebSocket::ConnectInternal(struct lws_context &LwsContext)
 		UE_LOG(LogWebSockets, Verbose, TEXT("FLwsWebSocket[%d]::ConnectInternal: setting State=%s PreviousState=%s Reason=%s"),
 			Identifier, ToString(EState::Error), ToString(EState::Connecting), *Reason);
 		{
-			FScopeLock ScopeLock(&StateLock);
+			FTransactionallySafeScopeLock ScopeLock(&StateLock);
 			State = EState::Error;
 			ClosedReason.Reason = MoveTemp(Reason);
 		}
@@ -758,7 +759,7 @@ void FLwsWebSocket::ConnectInternal(struct lws_context &LwsContext)
 		UE_LOG(LogWebSockets, Verbose, TEXT("FLwsWebSocket[%d]::ConnectInternal: setting State=%s PreviousState=%s Reason=%s"),
 			Identifier, ToString(EState::Error), ToString(EState::Connecting), *Reason);
 		{
-			FScopeLock ScopeLock(&StateLock);
+			FTransactionallySafeScopeLock ScopeLock(&StateLock);
 			State = EState::Error;
 			ClosedReason.Reason = MoveTemp(Reason);
 		}
@@ -814,7 +815,7 @@ void FLwsWebSocket::ConnectInternal(struct lws_context &LwsContext)
 		UE_LOG(LogWebSockets, Verbose, TEXT("FLwsWebSocket[%d]::ConnectInternal: setting State=%s PreviousState=%s Reason=%s"),
 			Identifier, ToString(EState::Error), ToString(EState::Connecting), *Reason);
 		{
-			FScopeLock ScopeLock(&StateLock);
+			FTransactionallySafeScopeLock ScopeLock(&StateLock);
 			State = EState::Error;
 			ClosedReason.Reason = MoveTemp(Reason);
 		}
diff --git a/Engine/Source/Runtime/Online/WebSockets/Private/Lws/LwsWebSocket.h b/Engine/Source/Runtime/Online/WebSockets/Private/Lws/LwsWebSocket.h
--- a/Engine/Source/Runtime/Online/WebSockets/Private/Lws/LwsWebSocket.h
+++ b/Engine/Source/Runtime/Online/WebSockets/Private/Lws/LwsWebSocket.h
@@ -5,6 +5,7 @@
 #include "CoreMinimal.h"
 #include "Containers/Ticker.h"
 #include "Containers/Queue.h"
+#include "Misc/TransactionallySafeCriticalSection.h"
 
 #if WITH_WEBSOCKETS && WITH_LIBWEBSOCKETS
 
@@ -245,7 +246,7 @@ private:
 private:
 
 	/** Critical section to lock access to state and close request variables */
-	FCriticalSection StateLock;
+	FTransactionallySafeCriticalSection StateLock;
 
 	/** Possible state values */
 	enum class EState : uint8
