[AutoRTFM] Remove FTransactionallySafeScopeLock entirely.

UE::TScopeLock works equally well for this purpose.

This CL also fixes a few newly-added references to the FTransactionallySafeScopeLock class.

diff --git a/Engine/Plugins/Animation/ControlRig/Source/ControlRig/Private/ControlRigComponent.cpp b/Engine/Plugins/Animation/ControlRig/Source/ControlRig/Private/ControlRigComponent.cpp
--- a/Engine/Plugins/Animation/ControlRig/Source/ControlRig/Private/ControlRigComponent.cpp
+++ b/Engine/Plugins/Animation/ControlRig/Source/ControlRig/Private/ControlRigComponent.cpp
@@ -5,7 +5,7 @@
 #include "Units/Execution/RigUnit_BeginExecution.h"
 #include "Units/Execution/RigUnit_InverseExecution.h"
 #include "Units/Execution/RigUnit_Hierarchy.h"
-#include "Misc/TransactionallySafeScopeLock.h"
+#include "Misc/ScopeLock.h"
 
 #include "SkeletalDebugRendering.h"
 #include "Components/InstancedStaticMeshComponent.h"
@@ -104,7 +104,7 @@ void UControlRigComponent::BeginDestroy()
 {
 	Super::BeginDestroy();
 
-	FTransactionallySafeScopeLock Lock(&gPendingSkeletalMeshesLock);
+	UE::TScopeLock Lock(gPendingSkeletalMeshesLock);
 	gPendingSkeletalMeshes.Remove(this);
 }
 
@@ -131,7 +131,7 @@ void UControlRigComponent::OnRegister()
 	ControlRig = nullptr;
 
 	{
-		FTransactionallySafeScopeLock Lock(&gPendingSkeletalMeshesLock);
+		UE::TScopeLock Lock(gPendingSkeletalMeshesLock);
 		gPendingSkeletalMeshes.FindOrAdd(this);
 	}
 
@@ -188,7 +188,7 @@ void UControlRigComponent::OnUnregister()
 	}
 	else
 	{
-		FTransactionallySafeScopeLock Lock(&gPendingSkeletalMeshesLock);
+		UE::TScopeLock Lock(gPendingSkeletalMeshesLock);
 		gPendingSkeletalMeshes.Remove(this);
 	}
 
@@ -210,7 +210,7 @@ void UControlRigComponent::TickComponent(float DeltaTime, enum ELevelTick TickTy
 
 	if (UControlRig* CR = SetupControlRigIfRequired())
 	{
-		FTransactionallySafeScopeLock Lock(&gPendingSkeletalMeshesLock);
+		UE::TScopeLock Lock(gPendingSkeletalMeshesLock);
 		TArray<FSkeletalMeshToMap>* PendingSkeletalMeshes = gPendingSkeletalMeshes.Find(this);
 
 		if (PendingSkeletalMeshes != nullptr && PendingSkeletalMeshes->Num() > 0)
@@ -657,7 +657,7 @@ void UControlRigComponent::AddMappedSkeletalMesh(USkeletalMeshComponent* Skeleta
 		PendingMesh.Bones = Bones;
 		PendingMesh.Curves = Curves;
 
-		FTransactionallySafeScopeLock Lock(&gPendingSkeletalMeshesLock);
+		UE::TScopeLock Lock(gPendingSkeletalMeshesLock);
 		TArray<FSkeletalMeshToMap>& PendingSkeletalMeshes = gPendingSkeletalMeshes.FindOrAdd(this);
 		PendingSkeletalMeshes.Add(PendingMesh);
 		return;
diff --git a/Engine/Source/Runtime/Core/Public/Misc/TransactionallySafeScopeLock.h b/Engine/Source/Runtime/Core/Public/Misc/TransactionallySafeScopeLock.h
deleted file mode 100644
diff --git a/Engine/Source/Runtime/Online/WebSockets/Private/Lws/LwsWebSocket.cpp b/Engine/Source/Runtime/Online/WebSockets/Private/Lws/LwsWebSocket.cpp
--- a/Engine/Source/Runtime/Online/WebSockets/Private/Lws/LwsWebSocket.cpp
+++ b/Engine/Source/Runtime/Online/WebSockets/Private/Lws/LwsWebSocket.cpp
@@ -16,7 +16,6 @@
 #include "HttpModule.h"
 #include "HttpManager.h"
 #include "PlatformHttp.h"
-#include "Misc/TransactionallySafeScopeLock.h"
 
 // FLwsSendBuffer 
 FLwsSendBuffer::FLwsSendBuffer(const uint8* Data, SIZE_T Size, bool bInIsBinary)
@@ -161,7 +160,7 @@ void FLwsWebSocket::Close(int32 Code, const FString& Reason)
 	UE_LOG(LogWebSockets, Verbose, TEXT("FLwsWebSocket[%d]::Close: Close queued with code=%d reason=%s"), Identifier, Code, *Reason);
 
 	{
-		FTransactionallySafeScopeLock ScopeLock(&StateLock);
+		UE::TScopeLock ScopeLock(StateLock);
 		CloseRequest.Code = Code;
 		CloseRequest.Reason = ANSIReason;
 	}
@@ -284,7 +283,7 @@ int FLwsWebSocket::LwsCallback(lws* Instance, lws_callback_reasons Reason, void*
 		// No lock because we are the only thread that writes to State
 		EState PreviousState = State;
 		{
-			FTransactionallySafeScopeLock ScopeLock(&StateLock);
+			UE::TScopeLock ScopeLock(StateLock);
 			State = EState::Connected;
 		}
 		UE_LOG(LogWebSockets, Verbose, TEXT("FLwsWebSocket[%d]::LwsCallback: Received LWS_CALLBACK_CLIENT_ESTABLISHED, setting State=%s PreviousState=%s"), Identifier, ToString(State), ToString(PreviousState));
@@ -371,7 +370,7 @@ int FLwsWebSocket::LwsCallback(lws* Instance, lws_callback_reasons Reason, void*
 			UE_LOG(LogWebSockets, Verbose, TEXT("FLwsWebSocket[%d]::LwsCallback: Received LWS_CALLBACK_WS_PEER_INITIATED_CLOSE, setting State=%s CloseStatus=%d Reason=%s bPeerSpecifiedReason=%d PreviousState=%s"),
 				Identifier, ToString(EState::Closed), (int32)CloseStatus, *CloseReasonString, (int32)bPeerSpecifiedReason, ToString(State));
 			{
-				FTransactionallySafeScopeLock ScopeLock(&StateLock);
+				UE::TScopeLock ScopeLock(StateLock);
 				State = EState::Closed;
 				ClosedReason.Reason = MoveTemp(CloseReasonString);
 				ClosedReason.CloseStatus = CloseStatus;
@@ -395,7 +394,7 @@ int FLwsWebSocket::LwsCallback(lws* Instance, lws_callback_reasons Reason, void*
 			UE_LOG(LogWebSockets, Verbose, TEXT("FLwsWebSocket[%d]::LwsCallback: Received LWS_CALLBACK_WSI_DESTROY, setting State=%s PreviousState=%s"), Identifier, ToString(EState::Error), ToString(State));
 		
 			{
-				FTransactionallySafeScopeLock ScopeLock(&StateLock);
+				UE::TScopeLock ScopeLock(StateLock);
 				State = EState::Error;
 				ClosedReason.Reason = MoveTemp(CloseReasonString);
 			}
@@ -406,7 +405,7 @@ int FLwsWebSocket::LwsCallback(lws* Instance, lws_callback_reasons Reason, void*
 			if (State == EState::Connected)
 			{
 				{
-					FTransactionallySafeScopeLock ScopeLock(&StateLock);
+					UE::TScopeLock ScopeLock(StateLock);
 					State = EState::Closed;
 					ClosedReason.bWasClean = false;
 					ClosedReason.CloseStatus = LWS_CLOSE_STATUS_ABNORMAL_CLOSE;
@@ -430,7 +429,7 @@ int FLwsWebSocket::LwsCallback(lws* Instance, lws_callback_reasons Reason, void*
 				}
 
 				{
-					FTransactionallySafeScopeLock ScopeLock(&StateLock);
+					UE::TScopeLock ScopeLock(StateLock);
 					State = EState::Closed;
 					ClosedReason.Reason = MoveTemp(ClosedReasonString);
 					ClosedReason.CloseStatus = ClosedStatus;
@@ -457,7 +456,7 @@ int FLwsWebSocket::LwsCallback(lws* Instance, lws_callback_reasons Reason, void*
 			UE_LOG(LogWebSockets, Verbose, TEXT("FLwsWebSocket[%d]::LwsCallback: Received LWS_CALLBACK_CLOSED, setting State=%s CloseReason=%s PreviousState=%s"),
 				Identifier, ToString(EState::Closed), *CloseReasonString, ToString(State));
 			{
-				FTransactionallySafeScopeLock ScopeLock(&StateLock);
+				UE::TScopeLock ScopeLock(StateLock);
 				State = EState::Closed;
 				ClosedReason.Reason = MoveTemp(CloseReasonString);
 				ClosedReason.CloseStatus = LWS_CLOSE_STATUS_NORMAL;
@@ -482,7 +481,7 @@ int FLwsWebSocket::LwsCallback(lws* Instance, lws_callback_reasons Reason, void*
 		UE_LOG(LogWebSockets, Verbose, TEXT("FLwsWebSocket[%d]::LwsCallback: Received LWS_CALLBACK_CLIENT_CONNECTION_ERROR, setting State=%s CloseReason=%s PreviousState=%s"),
 			Identifier, ToString(EState::Error), *CloseReasonString, ToString(State));
 		{
-			FTransactionallySafeScopeLock ScopeLock(&StateLock);
+			UE::TScopeLock ScopeLock(StateLock);
 			if (State == EState::Connected ||
 				State == EState::ClosingByRequest)
 			{
@@ -518,7 +517,7 @@ int FLwsWebSocket::LwsCallback(lws* Instance, lws_callback_reasons Reason, void*
 			ANSICHAR* CloseReasonCopy;
 			{
 				// TODO:  Is this scope lock necessary?  We only got here because we recognized CloseRequest.Code was valid, and we do not allow it to be modified afer being set on the game thread, so shouldn't this thread have the correct value?
-				FTransactionallySafeScopeLock ScopeLock(&StateLock);
+				UE::TScopeLock ScopeLock(StateLock);
 				CloseCodeCopy = CloseRequest.Code;
 				CloseReasonCopy = CloseRequest.Reason;
 			}
@@ -580,7 +579,7 @@ void FLwsWebSocket::GameThreadTick()
 {
 	EState CurrentState;
 	{
-		FTransactionallySafeScopeLock ScopeLock(&StateLock);
+		UE::TScopeLock ScopeLock(StateLock);
 		CurrentState = State;
 	}
 	if (CurrentState != LastGameThreadState)
@@ -625,7 +624,7 @@ void FLwsWebSocket::GameThreadFinalize()
 
 	{
 		// TODO:  The contract requires the libwebsockets be done with this object prior to this being called.  Is there a better way to ensure we get the last set value by the libwebsockets thread?  FPlatformMisc::MemoryBarrier maybe?
-		FTransactionallySafeScopeLock ScopeLock(&StateLock);
+		UE::TScopeLock ScopeLock(StateLock);
 		PreviousState = State;
 		State = EState::None; // Will be re-usable on final delegate triggering
 		LastGameThreadState = State;
@@ -661,7 +660,7 @@ void FLwsWebSocket::LwsThreadTick()
 		check(LwsConnection);
 		{
 			// Check if we want to close
-			FTransactionallySafeScopeLock ScopeLock(&StateLock);
+			UE::TScopeLock ScopeLock(StateLock);
 			if (CloseRequest.Code != 0)
 			{
 				State = EState::ClosingByRequest;
@@ -698,7 +697,7 @@ void FLwsWebSocket::ConnectInternal(struct lws_context &LwsContext)
 	UE_LOG(LogWebSockets, Verbose, TEXT("FLwsWebSocket[%d]::ConnectInternal: setting State=%s PreviousState=%s"),
 		Identifier, ToString(EState::Connecting), ToString(EState::StartConnecting));
 	{
-		FTransactionallySafeScopeLock ScopeLock(&StateLock);
+		UE::TScopeLock ScopeLock(StateLock);
 		State = EState::Connecting;
 	}
 
@@ -714,7 +713,7 @@ void FLwsWebSocket::ConnectInternal(struct lws_context &LwsContext)
 		UE_LOG(LogWebSockets, Verbose, TEXT("FLwsWebSocket[%d]::ConnectInternal: setting State=%s PreviousState=%s Reason=%s"),
 			Identifier, ToString(EState::Error), ToString(EState::Connecting), *Reason);
 		{
-			FTransactionallySafeScopeLock ScopeLock(&StateLock);
+			UE::TScopeLock ScopeLock(StateLock);
 			State = EState::Error;
 			ClosedReason.Reason = MoveTemp(Reason);
 		}
@@ -759,7 +758,7 @@ void FLwsWebSocket::ConnectInternal(struct lws_context &LwsContext)
 		UE_LOG(LogWebSockets, Verbose, TEXT("FLwsWebSocket[%d]::ConnectInternal: setting State=%s PreviousState=%s Reason=%s"),
 			Identifier, ToString(EState::Error), ToString(EState::Connecting), *Reason);
 		{
-			FTransactionallySafeScopeLock ScopeLock(&StateLock);
+			UE::TScopeLock ScopeLock(StateLock);
 			State = EState::Error;
 			ClosedReason.Reason = MoveTemp(Reason);
 		}
@@ -815,7 +814,7 @@ void FLwsWebSocket::ConnectInternal(struct lws_context &LwsContext)
 		UE_LOG(LogWebSockets, Verbose, TEXT("FLwsWebSocket[%d]::ConnectInternal: setting State=%s PreviousState=%s Reason=%s"),
 			Identifier, ToString(EState::Error), ToString(EState::Connecting), *Reason);
 		{
-			FTransactionallySafeScopeLock ScopeLock(&StateLock);
+			UE::TScopeLock ScopeLock(StateLock);
 			State = EState::Error;
 			ClosedReason.Reason = MoveTemp(Reason);
 		}
