BPS File Constructor Rework.

#rnx

diff --git a/Engine/Source/Runtime/Online/BuildPatchServices/Private/BuildPatchFileConstructor.cpp b/Engine/Source/Runtime/Online/BuildPatchServices/Private/BuildPatchFileConstructor.cpp
--- a/Engine/Source/Runtime/Online/BuildPatchServices/Private/BuildPatchFileConstructor.cpp
+++ b/Engine/Source/Runtime/Online/BuildPatchServices/Private/BuildPatchFileConstructor.cpp
@@ -2,9 +2,12 @@
 
 #include "BuildPatchFileConstructor.h"
 #include "IBuildManifestSet.h"
+#include "HAL/UESemaphore.h"
 #include "HAL/PlatformFileManager.h"
 #include "HAL/FileManager.h"
 #include "HAL/RunnableThread.h"
+#include "Hash/xxhash.h"
+#include "Misc/ConfigCacheIni.h"
 #include "Misc/Paths.h"
 #include "Misc/ScopeLock.h"
 #include "Misc/OutputDeviceRedirector.h"
@@ -14,20 +17,20 @@
 #include "Common/StatsCollector.h"
 #include "Common/SpeedRecorder.h"
 #include "Common/FileSystem.h"
+#include "Compression/CompressionUtil.h"
 #include "Installer/ChunkSource.h"
 #include "Installer/ChunkDbChunkSource.h"
+#include "Installer/InstallChunkSource.h"
+#include "Installer/CloudChunkSource.h"
 #include "Installer/ChunkReferenceTracker.h"
 #include "Installer/InstallerError.h"
 #include "Installer/InstallerAnalytics.h"
 #include "Installer/InstallerSharedContext.h"
+#include "Templates/Greater.h"
 #include "BuildPatchUtil.h"
 
 using namespace BuildPatchServices;
 
-// This define the number of bytes on a half-finished file that we ignore from the end
-// incase of previous partial write.
-#define NUM_BYTES_RESUME_IGNORE     1024
-
 static int32 SleepTimeWhenFileSystemThrottledSeconds = 15;
 static FAutoConsoleVariableRef CVarSleepTimeWhenFileSystemThrottledSeconds(
 	TEXT("BuildPatchFileConstructor.SleepTimeWhenFileSystemThrottledSeconds"),
@@ -99,7 +102,8 @@
 			int64 DiskSpaceDelta = 0;
 
 			// Loop through all files to be made next, in order.
-			FilesToConstruct.Sort(TLess<FString>());
+			// This is sorted coming in and needs to stay in that order to pass BPT test suite
+			//FilesToConstruct.Sort(TLess<FString>());
 			for (const FString& FileToConstruct : FilesToConstruct)
 			{
 				// First we would need to make the new file.
@@ -129,7 +133,9 @@
 	MissingChunk,
 	SerializeError,
 	TrackingError,
-	OutboundDataError
+	OutboundDataError,
+	InternalConsistencyError,
+	Aborted
 };
 
 /**
@@ -257,12 +263,958 @@
 	}
 };
 
+static FString FormatNumber(uint64 Value) { return FText::AsNumber(Value).ToString(); }
+
+
+// We need a place to put chunks if we know we're going to need them again after
+// their source retires. We don't want to route everything through here. We also
+// want to be able to optionally overflow to disk, and ideally this would be persistent
+// across resumes so that we don't have to re-download a huge amount of harvested chunks
+// from install files.
+class FChunkBackingStore : IConstructorChunkSource
+{
+	FBuildPatchFileConstructor* ParentConstructor = nullptr;
+	const FString& InstallDirectory;
+	FBuildPatchFileConstructor::FBackingStoreStats& Stats;
+
+	struct FStoredChunk
+	{
+		TArray<uint8> ChunkData;
+		uint32 ChunkSize = 0;
+		int32 NextUsageIndex;
+		int32 LastUsageIndex;
+
+		// Has the data made it into memory yet?
+		bool bCommitted = false;
+
+		// Have we been evicted for memory concerns? We are in UsedEntries if so.
+		bool bBackedByDisk = false;
+
+		// When something is reading/writing to us we mark that so we don't
+		// evict during async operations.
+		uint16 LockCount = 0;
+	};
+
+	TMap<FGuid, FStoredChunk> StoredChunks;
+	uint64 CurrentMemoryLoad = 0;
+	uint64 PeakMemoryLoad = 0;
+
+	uint64 MaxMemoryBytes = 1 << 30; // 1GB.
+
+
+	// 0 disables
+	uint64 MaxDiskSpaceBytes = 1 << 30; // 1GB.
+	uint64 AdditionalDiskSpaceHeadroomBytes = 0;
+	uint64 InstallationFreeSpaceRequired = 0;
+
+	// The backing store allocates in 128kb chunks.
+	static constexpr uint32 BitsPerEntry = 17;
+	struct FBackingStoreFreeSpan
+	{
+		uint32 StartEntryIndex = 0;
+		uint32 EndEntryIndex = 0;
+		uint64 Size() const { return EntryCount() << BitsPerEntry; }
+		uint64 Offset() const { return (uint64)StartEntryIndex << BitsPerEntry; }
+		uint64 EntryCount() const { return EndEntryIndex - StartEntryIndex; }
+	};
+	struct FBackingStoreUsedSpan
+	{
+		uint32 StartEntryIndex = 0;
+		uint32 EndEntryIndex = 0;
+		uint32 UsedBytes = 0;
+		FXxHash64 Hash;
+		uint64 ReservedSize() const { return EntryCount() << BitsPerEntry; }
+		uint64 Offset() const { return (uint64)StartEntryIndex << BitsPerEntry; }
+		uint64 EntryCount() const { return EndEntryIndex - StartEntryIndex; }
+	};
+
+	uint32 BackingStoreMaxEntries = 0;
+	uint32 BackingStoreEntryCount = 0;
+	uint64 BackingStoreUsedSpace = 0;
+	uint64 BackingStoreWastedSpace = 0;
+	TMap<FGuid, FBackingStoreUsedSpan> UsedDiskSpans;
+	TArray<FBackingStoreFreeSpan> FreeDiskSpans;
+	uint64 CurrentDiskLoad() const { return (uint64)(BackingStoreEntryCount) << BitsPerEntry; }
+
+	TUniquePtr<IFileHandle> BackingStoreFileHandle;
+	FString BackingStoreFileName;
+
+	void DumpFreeDiskSpans()
+	{
+		UE_LOG(LogBuildPatchServices, Display, TEXT("Backing Store (Entries / Max): %d / %d"), BackingStoreEntryCount, BackingStoreMaxEntries);
+		UE_LOG(LogBuildPatchServices, Display, TEXT("Dumping Free Disk Spans..."));		
+
+		for (int32 FreeIndex = 0; FreeIndex < FreeDiskSpans.Num(); FreeIndex++)
+		{
+			UE_LOG(LogBuildPatchServices, Display, TEXT("   %d/%d: %d - %d"), FreeIndex, FreeDiskSpans.Num(), FreeDiskSpans[FreeIndex].StartEntryIndex, FreeDiskSpans[FreeIndex].EndEntryIndex);
+		}
+	}
+	// false means consistency failure
+	[[nodiscard]] bool ConsistencyCheck()
+	{
+		bool bSuccess = true;
+
+		uint64 UsedEntryReservedBytes = 0;
+		uint64 UsedEntryUsedBytes = 0;
+		for (TPair<FGuid, FBackingStoreUsedSpan>& Pair : UsedDiskSpans)
+		{
+			FBackingStoreUsedSpan& Span = Pair.Value;
+			UsedEntryReservedBytes += Span.ReservedSize();
+			UsedEntryUsedBytes += Span.UsedBytes;
+		}
+
+		uint64 FreeBytes = 0;
+
+		for (int32 FreeIndex = 0; FreeIndex < FreeDiskSpans.Num(); FreeIndex++)
+		{
+			FBackingStoreFreeSpan& Span = FreeDiskSpans[FreeIndex];
+
+			FreeBytes += Span.Size();
+
+			if (Span.Size() == 0)
+			{
+				UE_LOG(LogBuildPatchServices, Error, TEXT("BackingStore FreeList Merge Fail: empty entry failed to get deleted"));
+				bSuccess = false;
+			}
+			if (Span.EndEntryIndex > BackingStoreEntryCount)
+			{
+				UE_LOG(LogBuildPatchServices, Error, TEXT("BackingStore FreeList Merge Fail: entry exceeds backing store size"));
+				bSuccess = false;
+			}
+			if (Span.StartEntryIndex > Span.EndEntryIndex)
+			{
+				UE_LOG(LogBuildPatchServices, Error, TEXT("BackingStore FreeList Merge Fail: negative sized entry"));
+				bSuccess = false;
+			}
+
+			if (FreeIndex)
+			{
+				if (Span.StartEntryIndex == FreeDiskSpans[FreeIndex-1].EndEntryIndex)
+				{
+					UE_LOG(LogBuildPatchServices, Error, TEXT("BackingStore FreeList Merge Fail: adjacent entries failed to merge"));
+					bSuccess = false;
+				}
+				if (Span.StartEntryIndex < FreeDiskSpans[FreeIndex-1].EndEntryIndex)
+				{
+					UE_LOG(LogBuildPatchServices, Error, TEXT("BackingStore FreeList Merge Fail: adjacent entries overlap or are out of order"));
+					bSuccess = false;
+				}
+			}
+		}
+
+		if (!bSuccess)
+		{
+			DumpFreeDiskSpans();
+		}
+
+		if (UsedEntryReservedBytes != (BackingStoreUsedSpace + BackingStoreWastedSpace) ||
+			UsedEntryUsedBytes != BackingStoreUsedSpace ||
+			CurrentDiskLoad() - UsedEntryReservedBytes != FreeBytes)
+		{
+			UE_LOG(LogBuildPatchServices, Error, TEXT("Disk Backing Store Consistency Fail:"));
+			UE_LOG(LogBuildPatchServices, Error, TEXT("Actual:"));
+			UE_LOG(LogBuildPatchServices, Error, TEXT("    ReservedBytes: %s"), *FormatNumber(UsedEntryReservedBytes));
+			UE_LOG(LogBuildPatchServices, Error, TEXT("    WastedBytes: %s"), *FormatNumber(UsedEntryReservedBytes - UsedEntryUsedBytes));
+			UE_LOG(LogBuildPatchServices, Error, TEXT("    FreeBytes: %s"), *FormatNumber(FreeBytes));
+			UE_LOG(LogBuildPatchServices, Error, TEXT("Expected:"));
+			UE_LOG(LogBuildPatchServices, Error, TEXT("    ReservedBytes: %s"), *FormatNumber(BackingStoreUsedSpace + BackingStoreWastedSpace));
+			UE_LOG(LogBuildPatchServices, Error, TEXT("    WastedBytes: %s"), *FormatNumber(BackingStoreWastedSpace));
+			UE_LOG(LogBuildPatchServices, Error, TEXT("    FreeBytes: %s"), *FormatNumber(CurrentDiskLoad() - BackingStoreUsedSpace - BackingStoreWastedSpace));
+			bSuccess = false;
+		}
+
+		return bSuccess;
+	}
+
+	// false means consistency or write failure
+	// constructor thread
+	[[nodiscard]] bool PageOut(const FGuid& InGuid)
+	{
+		TRACE_CPUPROFILER_EVENT_SCOPE(BackingStore_PageOut);
+
+		// InChunk must be valid for paging out at this point!
+		FStoredChunk& InChunk = StoredChunks[InGuid];
+
+		// It's possible we loaded back into memory but are already backed by disk so we can
+		// just free.
+		if (InChunk.bBackedByDisk)
+		{
+			InChunk.ChunkData.Empty();
+			CurrentMemoryLoad -= InChunk.ChunkSize;
+			ParentConstructor->SetChunkLocation(InGuid, EConstructorChunkLocation::DiskOverflow);
+			return true;
+		}
+
+		uint32 EntriesRequired = Align(InChunk.ChunkSize, BitsPerEntry) >> BitsPerEntry;
+
+		int32 SpanIndex = 0;
+		while (SpanIndex < FreeDiskSpans.Num() && 
+			FreeDiskSpans[SpanIndex].EntryCount() < EntriesRequired)
+		{
+			SpanIndex++;
+		}
+
+		bool bAppendingToFile = false;
+		if (SpanIndex == FreeDiskSpans.Num())
+		{
+			// Check our disk space limitations.
+			uint64 UseMaxDiskSpaceBytes = MaxDiskSpaceBytes;
+
+			// If we have a headroom value, we want to dynamically adjust
+			// our max disk space so that we always leave that amount of disk space free.
+			// This is expected to almost always be enabled in order to prevent the backing store from
+			// eating into space reserved for the actual installation. 
+			if (InstallationFreeSpaceRequired || AdditionalDiskSpaceHeadroomBytes)
+			{
+				uint64 TotalDiskBytes, FreeDiskBytes;
+				if (!FPlatformMisc::GetDiskTotalAndFreeSpace(InstallDirectory, TotalDiskBytes, FreeDiskBytes))
+				{
+					// If we fail to get disk space then disable it since we don't really know what we're doing
+					// at that point.
+					InstallationFreeSpaceRequired = 0;
+					AdditionalDiskSpaceHeadroomBytes = 0;
+				}
+				else
+				{
+					// The free space we got is counting any bytes we've already written to disk, so adjust for that.
+					uint64 FreeSizeBytesWithoutBackingStore = FreeDiskBytes + CurrentDiskLoad();
+
+					uint64 HeadroomRequiredBytes = InstallationFreeSpaceRequired + AdditionalDiskSpaceHeadroomBytes;
+
+					// By default we aren't allowed any space due to headroom limitations.
+					uint64 HeadroomRestrictedMaxDiskSpace = 0;
+
+					// If we have enough space above the headroom, then we can talk.
+					if (FreeSizeBytesWithoutBackingStore > HeadroomRequiredBytes)
+					{
+						HeadroomRestrictedMaxDiskSpace = FreeSizeBytesWithoutBackingStore - HeadroomRequiredBytes;
+						if (UseMaxDiskSpaceBytes == 0 || // Note if UseMaxDiskSpaceBytes is 0 then we they have no other limitation.
+							HeadroomRestrictedMaxDiskSpace < UseMaxDiskSpaceBytes)
+						{
+							UseMaxDiskSpaceBytes = HeadroomRestrictedMaxDiskSpace;
+						}
+					}
+				}
+			}
+
+			// Need to expand the backing store.
+			if (UseMaxDiskSpaceBytes && ((uint64)(BackingStoreEntryCount + EntriesRequired) << BitsPerEntry) > UseMaxDiskSpaceBytes)
+			{
+				// We can't expand - this fails to page out and gets evicted from the backing store.
+				Stats.DiskLostChunkCount++;
+				ParentConstructor->SetChunkLocation(InGuid, EConstructorChunkLocation::Cloud);
+				CurrentMemoryLoad -= InChunk.ChunkSize;
+				InChunk.ChunkData.Empty();
+				StoredChunks.Remove(InGuid);
+				return true;
+			}
+
+			FBackingStoreFreeSpan& NewSpan = FreeDiskSpans.AddDefaulted_GetRef();
+			NewSpan.StartEntryIndex = BackingStoreEntryCount;
+			BackingStoreEntryCount += EntriesRequired;
+			if (BackingStoreEntryCount > BackingStoreMaxEntries)
+			{
+				BackingStoreMaxEntries = BackingStoreEntryCount;
+			}
+			NewSpan.EndEntryIndex = NewSpan.StartEntryIndex + EntriesRequired;
+			bAppendingToFile = true;
+
+			Stats.DiskPeakUsageBytes = CurrentDiskLoad();
+		}
+
+		// SpanIndex is the one we just added or found to reuse
+		FBackingStoreFreeSpan& FreeSpan = FreeDiskSpans[SpanIndex];
+		FBackingStoreUsedSpan& UsedSpan = UsedDiskSpans.FindOrAdd(InGuid);
+		if (!UsedSpan.Hash.IsZero())
+		{
+			// Can't be paging out twice!
+			UE_LOG(LogBuildPatchServices, Error, TEXT("Consistency failure: Backing store used entry already existed for %s"), *WriteToString<40>(InGuid));
+			return false;
+		}
+
+		UsedSpan.StartEntryIndex = FreeSpan.StartEntryIndex;
+		UsedSpan.EndEntryIndex = UsedSpan.StartEntryIndex + EntriesRequired;
+		UsedSpan.UsedBytes = InChunk.ChunkSize;
+
+		FreeSpan.StartEntryIndex += EntriesRequired;
+
+		// LONGTERM - we should be using XX has for all chunk consistency checking, so we would have this value already.
+		{
+			TRACE_CPUPROFILER_EVENT_SCOPE(BSPO_Hash);
+			UsedSpan.Hash = FXxHash64::HashBuffer(InChunk.ChunkData.GetData(), InChunk.ChunkData.Num());
+		}
+
+		if (FreeSpan.EndEntryIndex == FreeSpan.StartEntryIndex)
+		{
+			FreeDiskSpans.RemoveAt(SpanIndex);
+		}
+
+		// Write
+		{
+			TRACE_CPUPROFILER_EVENT_SCOPE(BSPO_Write);
+			if (!BackingStoreFileHandle->Seek(UsedSpan.Offset()))
+			{
+				UE_LOG(LogBuildPatchServices, Error, TEXT("Failed to seek disk backing store to %llu"), UsedSpan.Offset());
+				return false;
+			}
+
+			if (!BackingStoreFileHandle->Write(InChunk.ChunkData.GetData(), InChunk.ChunkData.Num()))
+			{
+				UE_LOG(LogBuildPatchServices, Error, TEXT("Failed to write %llu bytes to disk backing store at %llu"), InChunk.ChunkData.Num(), UsedSpan.Offset());
+				return false;
+			}
+		}
+
+		uint64 Wastage = UsedSpan.ReservedSize() - InChunk.ChunkData.Num();
+
+		// If we just added to the end of the file, we only wrote the size of the chunk data, not necessarily
+		// the size of our reservation, so top off with zeroes.
+		if (bAppendingToFile)
+		{
+			// This is at most a 128kb allocation (1 << BitsPerEntry), but the chunk size currently in use
+			// is 3 bytes shy of being perfectly aligned to a multiple of that meaning we only expect to need to write
+			// 3 bytes of zeroes. So we make sure we have enough space on the stack for that.
+			TArray<uint8, TInlineAllocator<128>> Zeroes;
+			Zeroes.AddZeroed(Wastage);
+			if (!BackingStoreFileHandle->Write(Zeroes.GetData(), Zeroes.Num()))
+			{
+				UE_LOG(LogBuildPatchServices, Error, TEXT("Failed to write %llu bytes of zeroes to disk backing store at %llu"), Zeroes.Num(), UsedSpan.Offset());
+				return false;
+			}
+		}
+
+		// Done.
+		CurrentMemoryLoad -= InChunk.ChunkSize;
+		BackingStoreUsedSpace += InChunk.ChunkSize;
+		BackingStoreWastedSpace += Wastage;
+		InChunk.ChunkData.Empty();
+		InChunk.bBackedByDisk = true;
+		ParentConstructor->SetChunkLocation(InGuid, EConstructorChunkLocation::DiskOverflow);
+
+		return ConsistencyCheck();
+	}
+
+	// false means consistency failure
+	[[nodiscard]] bool ReleaseBackingStoreEntry(const FGuid& InGuid, FStoredChunk& InChunk)
+	{
+		// Return the bits back to the free list.
+		FBackingStoreFreeSpan NewFreeSpan;
+		int32 FreeSpanIndex;
+		{
+			FBackingStoreUsedSpan& UsedEntry = UsedDiskSpans[InGuid];
+
+			BackingStoreUsedSpace -= InChunk.ChunkSize;
+			BackingStoreWastedSpace -= UsedEntry.ReservedSize() - InChunk.ChunkSize;
+
+			FreeSpanIndex = Algo::LowerBoundBy(FreeDiskSpans, UsedEntry.StartEntryIndex, &FBackingStoreFreeSpan::StartEntryIndex);
+			
+			NewFreeSpan.StartEntryIndex = UsedEntry.StartEntryIndex;
+			NewFreeSpan.EndEntryIndex = UsedEntry.EndEntryIndex;
+		
+			InChunk.bBackedByDisk = false;
+
+			UsedDiskSpans.Remove(InGuid);
+		}
+
+		// Merge into an adjacent entry without adding and having to do a linear
+		// pass to coalesce. 
+		bool bMerged = false;
+		if (FreeSpanIndex < FreeDiskSpans.Num())
+		{
+			if (NewFreeSpan.EndEntryIndex == FreeDiskSpans[FreeSpanIndex].StartEntryIndex)
+			{
+				// We are right before to the one after us - just extend them lower
+				FreeDiskSpans[FreeSpanIndex].StartEntryIndex = NewFreeSpan.StartEntryIndex;
+				bMerged = true;
+
+				// OK they got merged down, see if they can connect with the one below
+				if (FreeSpanIndex > 0)
+				{
+					if (FreeDiskSpans[FreeSpanIndex - 1].EndEntryIndex == FreeDiskSpans[FreeSpanIndex].StartEntryIndex)
+					{
+						// We fill a gap, we can connect and remove.
+						FreeDiskSpans[FreeSpanIndex - 1].EndEntryIndex = FreeDiskSpans[FreeSpanIndex].EndEntryIndex;
+						FreeDiskSpans.RemoveAt(FreeSpanIndex);
+					}
+				}
+			}
+		}
+		
+		if (!bMerged && FreeSpanIndex > 0) // if we missed, this will be Num(). If we have nothing yet it'll fail this and be fine.
+		{
+			if (FreeDiskSpans[FreeSpanIndex - 1].EndEntryIndex == NewFreeSpan.StartEntryIndex)
+			{
+				// We are right after the one before us, extend them farther.
+				FreeDiskSpans[FreeSpanIndex - 1].EndEntryIndex = NewFreeSpan.EndEntryIndex;
+				bMerged = true;
+
+				// They got merged up, see if we filled a gap.
+				if (FreeSpanIndex < FreeDiskSpans.Num())
+				{
+					if (FreeDiskSpans[FreeSpanIndex - 1].EndEntryIndex == FreeDiskSpans[FreeSpanIndex].StartEntryIndex)
+					{
+						FreeDiskSpans[FreeSpanIndex - 1].EndEntryIndex = FreeDiskSpans[FreeSpanIndex].EndEntryIndex;
+						FreeDiskSpans.RemoveAt(FreeSpanIndex);
+					}
+				}
+			}
+		}
+
+		// If we didn't merge, we need to insert
+		if (!bMerged)
+		{
+			FreeDiskSpans.Insert(NewFreeSpan, FreeSpanIndex);
+		}
+
+		// Check and see if the free space is at the end of the file. If so, we can truncate
+		// and free up disk space. 
+		if (FreeDiskSpans.Top().EndEntryIndex == BackingStoreEntryCount)
+		{
+			uint64 TruncateToSize = (uint64)FreeDiskSpans.Top().StartEntryIndex << BitsPerEntry;
+			if (BackingStoreFileHandle->Truncate(TruncateToSize)) // truncate can fail... don't update structures if we can't!
+			{
+				BackingStoreEntryCount = FreeDiskSpans.Top().StartEntryIndex;
+				FreeDiskSpans.Pop();
+			}
+		}
+
+		return ConsistencyCheck();
+	}
+
+	[[nodiscard]] bool ReleaseEntryInternal(const FGuid& InGuid, FStoredChunk* InStoredChunk)
+	{
+		if (InStoredChunk->LockCount == 0)
+		{
+			UE_LOG(LogBuildPatchServices, Error, TEXT("Consistency Failure: Releasing memory entry that isn't locked! %s"), *WriteToString<40>(InGuid));
+			return false;
+		}
+		else
+		{
+			InStoredChunk->LockCount--;
+			return true;
+		}
+	}
+
+public:
+
+	FChunkBackingStore(FBuildPatchFileConstructor* InParentConstructor, const FString& InInstallDirectory, FBuildPatchFileConstructor::FBackingStoreStats& InStats) 
+	: ParentConstructor(InParentConstructor)
+	, InstallDirectory(InInstallDirectory)
+	, Stats(InStats)
+	{
+		BackingStoreFileName = InParentConstructor->Configuration.BackingStoreDirectory / TEXT("backingstore");
+
+		bool bUseDiskOverflowStore = true;
+		GConfig->GetBool(TEXT("BuildPatchServices"), TEXT("bEnableDiskOverflowStore"), bUseDiskOverflowStore, GEngineIni);
+		
+		// Is there a hard limit on how much disk space we can use?
+		// negative = no
+		// 0 = disable disk overflow
+		int64 MaxDiskMB = 0;
+		GConfig->GetInt64(TEXT("BuildPatchServices"), TEXT("DiskOverflowStoreLimitMB"), MaxDiskMB, GEngineIni);
+		if (MaxDiskMB < 0)
+		{
+			MaxDiskSpaceBytes = TNumericLimits<int64>::Max();
+		}
+		else
+		{
+			MaxDiskSpaceBytes = MaxDiskMB << 20;
+		}
+
+		// Do we want to always try and keep some disk space available, no matter what our limit is?
+		// note that independent of this we try and prevent the disk backing store from eating into space
+		// we have reserved for the actual install so they don't compete.
+		//
+		// this checks the free space after each file and updates the space limit correspondingly.
+		int64 AdditionalDiskSpaceHeadroomMB = 0;
+		GConfig->GetInt64(TEXT("BuildPatchServices"), TEXT("DiskOverflowStoreAdditionalHeadroomMB"), AdditionalDiskSpaceHeadroomMB, GEngineIni);
+		if (AdditionalDiskSpaceHeadroomMB >= 0)
+		{
+			AdditionalDiskSpaceHeadroomBytes = AdditionalDiskSpaceHeadroomMB << 20;
+		}
+
+		if (bUseDiskOverflowStore)
+		{
+			BackingStoreFileHandle.Reset(FPlatformFileManager::Get().GetPlatformFile().OpenWrite(*BackingStoreFileName, false, true));
+		}
+
+		if (!BackingStoreFileHandle)
+		{
+			// Prevent any pageouts.
+			UE_LOG(LogBuildPatchServices, Warning, TEXT("Unable to open disk backing store at %s"), *BackingStoreFileName);
+			UE_LOG(LogBuildPatchServices, Warning, TEXT("Disk backing store will be disabled"));
+			MaxDiskSpaceBytes = 0;
+		}
+
+		UE_LOG(LogBuildPatchServices, Display, TEXT("DiskOverflowStore is: %s - MaxSize = %s, Additional Headroom = %s"), 
+			BackingStoreFileHandle ? TEXT("Enabled") : TEXT("Disabled"),
+			*FormatNumber(MaxDiskSpaceBytes),
+			*FormatNumber(AdditionalDiskSpaceHeadroomBytes)
+			);
+
+		
+		// Now memory limits.
+		uint64 ChunkStoreMemoryLimit = 0;
+
+		// Check old values and warn/assume
+		{
+			int32 ChunkStoreMemorySizeChunks;
+			const bool bLoadedStoreSize = GConfig->GetInt(TEXT("Portal.BuildPatch"), TEXT("ChunkStoreMemorySize"), ChunkStoreMemorySizeChunks, GEngineIni);
+			if (bLoadedStoreSize)
+			{
+				UE_LOG(LogBuildPatchServices, Warning, TEXT("Outdated memory size limitation found: ChunkStoreMemorySize. Assuming chunk size is 1MB, use ChunkStoreMemorySizeMB instead."));
+				ChunkStoreMemoryLimit = (uint64)FMath::Max(ChunkStoreMemorySizeChunks, 0) << 20;
+			}
+
+			int32 CloudChunkStoreMemorySizeChunks = 0;
+			int32 InstallChunkStoreMemorySizeChunks = 0;
+			const bool bLoadedCloudSize = GConfig->GetInt(TEXT("Portal.BuildPatch"), TEXT("CloudChunkStoreMemorySize"), CloudChunkStoreMemorySizeChunks, GEngineIni);
+			const bool bLoadedInstallSize = GConfig->GetInt(TEXT("Portal.BuildPatch"), TEXT("InstallChunkStoreMemorySize"), InstallChunkStoreMemorySizeChunks, GEngineIni);
+			if (bLoadedCloudSize || bLoadedInstallSize)
+			{
+				UE_LOG(LogBuildPatchServices, Warning, TEXT("Outdated memory size limitations found: CloudChunkStoreMemorySize or InstallChunkStoreMemorySize. Assuming chunk size is 1MB."));
+				UE_LOG(LogBuildPatchServices, Warning, TEXT("Use ChunkStoreMemorySizeMB and/or ChunkStoreMemoryHeadRoomMB."));
+				uint64 OldMemoryLimit = FMath::Max(CloudChunkStoreMemorySizeChunks + InstallChunkStoreMemorySizeChunks, 0);
+				OldMemoryLimit <<= 20;
+				ChunkStoreMemoryLimit = OldMemoryLimit;
+			}
+		}
+
+		// check current values. We expect this to override anything from above.
+		int32 ChunkStoreMemoryLimitMB = 0;
+		if (GConfig->GetInt(TEXT("Portal.BuildPatch"), TEXT("ChunkStoreMemorySizeMB"), ChunkStoreMemoryLimitMB, GEngineIni))
+		{
+			ChunkStoreMemoryLimit = (uint64)ChunkStoreMemoryLimitMB << 20;
+		}
+		
+
+		// Get headroom. Default to 2GB of headroom. If no config is entered then we expect a 0 chunk limit that
+		// gets updated off of this default headroom.
+		int32 ChunkStoreMemoryHeadRoomMB = 2000;
+		if (GConfig->GetInt(TEXT("Portal.BuildPatch"), TEXT("ChunkStoreMemoryHeadRoomMB"), ChunkStoreMemoryHeadRoomMB, GEngineIni))
+		{
+			if (ChunkStoreMemoryHeadRoomMB < 0)
+			{
+				// negative disables
+				ChunkStoreMemoryHeadRoomMB = -1;
+			}
+			else if (ChunkStoreMemoryHeadRoomMB < 500)
+			{
+				UE_LOG(LogBuildPatchServices, Warning, TEXT("ChunkStoreMemoryHeadRoomMB too low (%d), using min (500)"), ChunkStoreMemoryHeadRoomMB);
+				ChunkStoreMemoryHeadRoomMB = 500;
+			}
+		}		
+
+		// To be consistent with other limitations, negative disables it
+		// 0 is OK - we require locked data to be in memory so we'll go over the limit but it'll be a minimum.
+		if (ChunkStoreMemoryLimit < 0)
+		{
+			ChunkStoreMemoryLimit = TNumericLimits<uint64>::Max();
+		}
+
+
+		uint64 AvailableMem = FPlatformMemory::GetStats().AvailablePhysical;
+		if (ChunkStoreMemoryHeadRoomMB >= 0)
+		{
+			uint64 RequestedHeadRoom = ((uint64)ChunkStoreMemoryHeadRoomMB << 20);
+
+			if (RequestedHeadRoom < AvailableMem)
+			{
+				uint64 MemoryStoreMem = AvailableMem - RequestedHeadRoom;
+				ChunkStoreMemoryLimit = MemoryStoreMem;
+			}
+			else
+			{
+				// Cap at available.
+				if (ChunkStoreMemoryLimit > AvailableMem)
+				{
+					ChunkStoreMemoryLimit = AvailableMem;
+				}
+			}
+		}
+
+		if (ChunkStoreMemoryLimit == -1)
+		{
+			UE_LOG(LogBuildPatchServices, Display, TEXT("ChunkStoreMemoryLimits are disabled"));
+		}
+		else
+		{
+			UE_LOG(LogBuildPatchServices, Display, TEXT("ChunkStoreMemoryLimits: %s using headroom of %s (%s available memory)"), 
+				*FormatNumber(ChunkStoreMemoryLimit),
+				ChunkStoreMemoryHeadRoomMB >= 0 ? *FormatNumber((uint64)ChunkStoreMemoryHeadRoomMB << 20) : TEXT("<disabled>"),
+				*FormatNumber(AvailableMem));
+		}
+
+		MaxMemoryBytes = ChunkStoreMemoryLimit;
+		Stats.MemoryLimitBytes = MaxMemoryBytes;
+	}
+
+	~FChunkBackingStore()
+	{
+		BackingStoreFileHandle.Reset(); // have to close the file before we try to delete it.
+		if (!FPlatformFileManager::Get().GetPlatformFile().DeleteFile(*BackingStoreFileName))
+		{
+			UE_LOG(LogBuildPatchServices, Warning, TEXT("Unable to delete disk backing store: %s"), *BackingStoreFileName);
+		}
+	}
+
+	// Set the amount of disk space the installation needs to we can ensure that we don't 
+	// expand into that space no matter what our config disk space limits are.
+	void SetDynamicDiskSpaceHeadroom(uint64 InInstallationFreeSpaceRequired)
+	{
+		InstallationFreeSpaceRequired = InInstallationFreeSpaceRequired;
+	}
+	void DisableDynamicDiskSpaceHeadroom()
+	{
+		InstallationFreeSpaceRequired = 0;
+	}
+
+	// false means consistency failure
+	// Constructor thread
+	[[nodiscard]] bool DereserveHarvestingEntry(const FGuid& InGuid)
+	{
+		FStoredChunk* StoredChunk = StoredChunks.Find(InGuid);
+		if (StoredChunk)
+		{
+			if (StoredChunk->bCommitted || StoredChunk->bBackedByDisk)
+			{
+				UE_LOG(LogBuildPatchServices, Error, 
+					TEXT("Consistency Failure: deserve memory entry that's uncommitted or paged out! %s, paged out = %d committed = %d"), 
+					*WriteToString<40>(InGuid), StoredChunk->bBackedByDisk, StoredChunk->bCommitted);
+				return false;
+			}
+
+			CurrentMemoryLoad -= StoredChunk->ChunkData.Num();
+			StoredChunks.Remove(InGuid);
+		}
+		else
+		{
+			UE_LOG(LogBuildPatchServices, Error, TEXT("Consistency Failure: Cleared memory entry that doesn't exist! %s"), *WriteToString<40>(InGuid));
+			return false;
+		}
+		return true;
+	}
+
+	// false means consistency failure
+	// Constructor thread
+	[[nodiscard]] bool ReleaseEntry(const FGuid& InGuid)
+	{
+		FStoredChunk* StoredChunk = StoredChunks.Find(InGuid);
+		if (StoredChunk)
+		{
+			return ReleaseEntryInternal(InGuid, StoredChunk);
+		}
+		else
+		{
+			UE_LOG(LogBuildPatchServices, Error, TEXT("Consistency Failure: Releasing memory entry that doesn't exist! %s"), *WriteToString<40>(InGuid));
+			return false;
+		}
+	}
+
+	// false means consistency failure
+	// Constructor thread
+	[[nodiscard]] bool CommitAndReleaseEntry(const FGuid& InGuid)
+	{
+		FStoredChunk* StoredChunk = StoredChunks.Find(InGuid);
+		if (StoredChunk)
+		{
+			if (StoredChunk->bCommitted)
+			{
+				UE_LOG(LogBuildPatchServices, Error, TEXT("Consistency Failure: Committing memory entry that is already committed! %s"), *WriteToString<40>(InGuid));
+				return false;
+			}
+			StoredChunk->bCommitted = true;
+			return ReleaseEntryInternal(InGuid, StoredChunk);
+		}
+		else
+		{
+			UE_LOG(LogBuildPatchServices, Error, TEXT("Consistency Failure: Committed memory entry that doesn't exist! %s"), *WriteToString<40>(InGuid));
+			return false;
+		}
+	}
+
+	// false means consistency failure
+	// Constructor thread
+	[[nodiscard]] bool LockEntry(const FGuid& InGuid)
+	{
+		FStoredChunk* StoredChunk = StoredChunks.Find(InGuid);
+		if (!StoredChunk)
+		{
+			UE_LOG(LogBuildPatchServices, Error, TEXT("Consistency Failure: locking memory entry that doesn't exist! %s"), *WriteToString<40>(InGuid));
+			return false;
+		}
+		else
+		{
+			StoredChunk->LockCount++;
+		}
+		return true;
+	}
+
+	// Returns an empty view on consistency failure
+	// Constructor thread.
+	[[nodiscard]] FMutableMemoryView ReserveAndLockEntry(const FGuid& InGuid, uint32 InChunkSize, int32 LastUsageIndex)
+	{
+		// Note that this function can be called to reserve an entry already in the backing store
+		// because we need to read from the disk to memory for a sub-chunk (or otherwise). So it needs
+		// to be able to handle reserving on top of paged out chunks.
+
+		int32 CurrentUsageIndex = ParentConstructor->ChunkReferenceTracker->GetCurrentUsageIndex();
+
+		while (CurrentMemoryLoad + InChunkSize > MaxMemoryBytes)
+		{
+			// Gotta dump stuff to disk. If we fail to dump to disk we mark the chunk
+			// as only available via the cloud source.
+
+			// Evict the one that is the longest until we use it. Chunks are almost always the same size
+			// so we expect this to run once.
+			int32 FarthestNextUsage = -1;
+			FGuid EvictGuid;
+			for (TPair<FGuid, FStoredChunk>& Pair : StoredChunks)
+			{
+				FStoredChunk& Chunk = Pair.Value;
+
+				if (Chunk.LockCount || // Actively in use - don't evict.
+					!Chunk.ChunkData.Num()) // Not in memory - can't evict.
+				{					
+					continue;
+				}
+
+				if (Chunk.NextUsageIndex < CurrentUsageIndex)
+				{
+					Chunk.NextUsageIndex = ParentConstructor->ChunkReferenceTracker->GetNextUsageForChunk(Pair.Key, Chunk.LastUsageIndex);
+				}
+
+				if (Chunk.NextUsageIndex > FarthestNextUsage)
+				{
+					FarthestNextUsage = Pair.Value.NextUsageIndex;
+					EvictGuid = Pair.Key;
+				}
+			}
+
+			if (FarthestNextUsage == -1)
+			{
+				// This means we can't reserve and also keep our memory requirements. 
+				//
+				// We currently have a minimum memory requirement for construction: Partial or reused chunks
+				// are routed through the backing store for holding. Partial because we need a full allocation to decompress
+				// the chunk and reuse so we don't have to re-read from disk.
+				//
+				// For unpatched installs this is minimal as most chunks are able to write directly to
+				// the destination buffer. For optimized delta manifests this is not the case as BPT assembles
+				// chunks from all over the place - I've seen 80+ chunk references to assemble 16MB of data,
+				// resulting in total buffer usage of 80MB + 16MB > 100MB of memory use per batch. With two
+				// batches in flight this pushes 200MB total buffer allocation.
+				//
+				// This is not likely an issue in real life as 200MB isn't much, but it does mean if we have a 
+				// low memory usage limit we can hit this legitimately.
+				//
+				// So...we let the reservation continue in violation of our memory constraints and hope we don't OOM.
+				//
+				// Future work could be:
+				// 1. Limit batch creation based on inflight chunk in addition to write buffer size
+				// 2. Only allow 1 batch in flight when using an optimized delta + low memory constraints.
+				break;
+			}
+			else
+			{
+				if (!PageOut(EvictGuid))
+				{
+					return FMutableMemoryView();
+				}
+			}
+		}
+
+		FStoredChunk& StoredChunk = StoredChunks.FindOrAdd(InGuid);
+
+		//
+		// It's possible to request a reservation for a chunk already in the backing store when reading from the
+		// disk into memory (i.e. with a non-direct read). However it can't already be in memory.
+		//
+		if (StoredChunk.ChunkData.Num())
+		{
+			UE_LOG(LogBuildPatchServices, Error, TEXT("Consistency failure: Reserving read space for a chunk already in memory: %s"), *WriteToString<40>(InGuid));
+			return FMutableMemoryView();
+		}
+		else
+		{
+			StoredChunk.LastUsageIndex = LastUsageIndex;
+			StoredChunk.bCommitted = false;
+			StoredChunk.ChunkData.AddUninitialized(InChunkSize);
+			StoredChunk.ChunkSize = InChunkSize;
+			StoredChunk.LockCount = 1;
+			CurrentMemoryLoad += InChunkSize;
+			if (CurrentMemoryLoad > PeakMemoryLoad)
+			{
+				PeakMemoryLoad = CurrentMemoryLoad;
+				Stats.MemoryPeakUsageBytes = PeakMemoryLoad;
+			}
+		}
+
+		FMutableMemoryView ChunkBuffer(StoredChunk.ChunkData.GetData(), StoredChunk.ChunkData.Num());
+		return ChunkBuffer;
+	}
+	
+	// false means consistency failure
+	// Constructor thread
+	[[nodiscard]] bool CheckRetirements(int32 CurrentUsageIndex)
+	{
+		// 6 just chosen because we wouldn't expect a ton of things at once
+		// but we could get several. shrug.
+		TArray<FGuid, TInlineAllocator<6>> GuidsToDelete;
+
+		for (TPair<FGuid, FStoredChunk>& Pair : StoredChunks)
+		{
+			if (Pair.Value.LastUsageIndex < CurrentUsageIndex)
+			{
+				GuidsToDelete.Add(Pair.Key);
+			}
+		}
+
+		for (FGuid& ToDelete : GuidsToDelete)
+		{
+			FStoredChunk* Chunk = StoredChunks.Find(ToDelete);
+			if (Chunk->bBackedByDisk)
+			{
+				if (!ReleaseBackingStoreEntry(ToDelete, *Chunk))
+				{
+					return false;
+				}
+			}
+
+			if (Chunk->LockCount)
+			{
+				UE_LOG(LogBuildPatchServices, Error, TEXT("Consistency Failure: Retiring memory entry with lock count! %s %d"), *WriteToString<40>(ToDelete), Chunk->LockCount);
+				return false;
+			}
+			
+			if (Chunk->ChunkData.Num())
+			{
+				if (!Chunk->bCommitted)
+				{
+					UE_LOG(LogBuildPatchServices, Error, TEXT("Consistency Failure: Retiring memory entry that never got committed! %s"), *WriteToString<40>(ToDelete));
+					return false;
+				}
+
+				CurrentMemoryLoad -= Chunk->ChunkSize;
+			}
+
+			StoredChunks.Remove(ToDelete);
+
+			ParentConstructor->SetChunkLocation(ToDelete, EConstructorChunkLocation::Retired);
+		}
+
+		return true;
+	}
+	
+	// This should only be for paged out chunks... in-memory chunks should be handled directly.
+	// Consistency failures in this will pass as data read fails, which will end up redirecting to
+	// the cloud. However the post-file consistency check will catch it and fail the install.
+	// Constructor thread only for this chunk source
+	[[nodiscard]] FRequestProcessFn CreateRequest(const FGuid& DataId, FMutableMemoryView DestinationBuffer, void* UserPtr, FChunkRequestCompleteDelegate CompleteFn)
+	{
+		// LONGTERM - can we make this fully async once we have a pipe we can prevent reentrancy on or ReadAtOffset API?
+		TRACE_CPUPROFILER_EVENT_SCOPE(BackingStoreRead);
+
+		bool bSuccess = true;
+		// This should go to a generate IO dispatch with completion function I think? idk...
+		// This acts as any other IO source - we don't know where it's going, it might be going back into memory.
+		FStoredChunk* StoredChunk = StoredChunks.Find(DataId);
+		bSuccess = StoredChunk != nullptr;
+
+		if (bSuccess && !StoredChunk->bBackedByDisk)
+		{
+			UE_LOG(LogBuildPatchServices, Error, TEXT("Consistency Failure: Trying to page in a chunk that isn't paged out! %s"), *WriteToString<40>(DataId));
+			bSuccess = false;
+		}
+
+		FBackingStoreUsedSpan* UsedEntry = nullptr;
+		if (bSuccess)
+		{
+			UsedEntry = UsedDiskSpans.Find(DataId);
+			if (!UsedEntry)
+			{
+				UE_LOG(LogBuildPatchServices, Error, TEXT("Consistency Failure: Backing store entry not found for paged out chunk! %s"), *WriteToString<40>(DataId));
+				bSuccess = false;
+			}
+		}
+
+		if (bSuccess)
+		{
+			BackingStoreFileHandle->Seek(UsedEntry->Offset());
+			if (!BackingStoreFileHandle->Read((uint8*)DestinationBuffer.GetData(), DestinationBuffer.GetSize()))
+			{
+				UE_LOG(LogBuildPatchServices, Error, TEXT("Consistency Failure: Backing store page-in failed read! %s"), *WriteToString<40>(DataId));
+				bSuccess = false;
+			}
+		}
+
+		if (bSuccess)
+		{
+			TRACE_CPUPROFILER_EVENT_SCOPE(BSR_Hash);
+			if (FXxHash64::HashBuffer(DestinationBuffer) != UsedEntry->Hash)
+			{
+				UE_LOG(LogBuildPatchServices, Error, TEXT("Consistency Failure: Backing store page-in failed hash check! %s"), *WriteToString<40>(DataId));
+				bSuccess = false;
+			}
+		}
+
+		if (!bSuccess)
+		{
+			Stats.DiskLoadFailureCount++;
+		}
+
+		Stats.DiskChunkLoadCount++;
+		CompleteFn.Execute(DataId, false, !bSuccess, UserPtr);
+		return [](bool) { return; };
+	}
+
+	// false means consistency failure
+	// Constructor thread
+	[[nodiscard]] bool CheckNoLocks(bool bIsHarvest)
+	{
+		TRACE_CPUPROFILER_EVENT_SCOPE(BackingStore_CheckNoLocks);
+		// After a file there should be no locked chunks since there are no reads.
+		bool bSuccess = true;
+		for (TPair<FGuid, FStoredChunk>& Chunk : StoredChunks)
+		{
+			if (Chunk.Value.LockCount)
+			{
+				UE_LOG(LogBuildPatchServices, Error, TEXT("Chunk %s locked with count %d after %s!"), *WriteToString<40>(Chunk.Key), Chunk.Value.LockCount, bIsHarvest ? TEXT("harvest") : TEXT("file completion"));
+				bSuccess = false;
+			}
+		}
+
+		if (bSuccess)
+		{
+			bSuccess = ConsistencyCheck();
+		}
+
+		return bSuccess;
+	}
+
+	// This is for the disk store - we don't free the page entry until it retires, so it's always available and can
+	// be read direct.
+	virtual int32 GetChunkUnavailableAt(const FGuid& DataId) const override { return TNumericLimits<int32>::Max(); }
+
+	FMemoryView GetViewForChunk(const FGuid& DataId) const
+	{
+		FStoredChunk const* StoredChunk = StoredChunks.Find(DataId);
+		if (!StoredChunk)
+		{
+			return FMemoryView();
+		}
+
+		return FMemoryView(StoredChunk->ChunkData.GetData(), StoredChunk->ChunkData.Num());
+	}
+}; // FChunkBackingStore
+
 /* FBuildPatchFileConstructor implementation
  *****************************************************************************/
 FBuildPatchFileConstructor::FBuildPatchFileConstructor(
-	FFileConstructorConfig InConfiguration, IFileSystem* InFileSystem, IChunkSource* InChunkSource, 
-	IChunkDbChunkSource* InChunkDbChunkSource, IChunkReferenceTracker* InChunkReferenceTracker, IInstallerError* InInstallerError, 
-	IInstallerAnalytics* InInstallerAnalytics, IFileConstructorStat* InFileConstructorStat)
+	FFileConstructorConfig InConfiguration, IFileSystem* InFileSystem, 
+	IConstructorChunkDbChunkSource* InChunkDbChunkSource, IConstructorCloudChunkSource* InCloudChunkSource, IConstructorInstallChunkSource* InInstallChunkSource, 
+	IChunkReferenceTracker* InChunkReferenceTracker, IInstallerError* InInstallerError, IInstallerAnalytics* InInstallerAnalytics, 
+	IFileConstructorStat* InFileConstructorStat, TMap<FGuid, EConstructorChunkLocation>&& InChunkLocations)
 	: Configuration(MoveTemp(InConfiguration))
 	, bIsDownloadStarted(false)
 	, bInitialDiskSizeCheck(false)
@@ -271,8 +1223,10 @@
 	, ThreadLock()
 	, ConstructionStack()
 	, FileSystem(InFileSystem)
-	, ChunkSource(InChunkSource)
 	, ChunkDbSource(InChunkDbChunkSource)
+	, InstallSource(InInstallChunkSource)
+	, CloudSource(InCloudChunkSource)
+	, ChunkLocations(MoveTemp(InChunkLocations))
 	, ChunkReferenceTracker(InChunkReferenceTracker)
 	, InstallerError(InInstallerError)
 	, InstallerAnalytics(InInstallerAnalytics)
@@ -282,6 +1236,10 @@
 	, RequiredDiskSpace(0)
 	, AvailableDiskSpace(0)
 {
+	TRACE_CPUPROFILER_EVENT_SCOPE(FileConstructor_ctor);
+
+	BackingStore.Reset(new FChunkBackingStore(this, Configuration.InstallDirectory, BackingStoreStats));
+
 	// Count initial job size
 	const int32 ConstructListNum = Configuration.ConstructList.Num();
 	ConstructionStack.Reserve(ConstructListNum);
@@ -291,6 +1249,9 @@
 	int32 CurrentPosition = 0;
 	FileCompletionPositions.Reserve(ConstructListNum);
 
+	// The first index after the file is complete.
+	TMap<FString, int32> FileRetirementPositions;
+
 	for (int32 ConstructListIdx = 0; ConstructListIdx < ConstructListNum ; ++ConstructListIdx)
 	{
 		const FString& ConstructListElem = Configuration.ConstructList[ConstructListIdx];
@@ -304,90 +1265,232 @@
 			CurrentPosition += AdvanceCount;
 
 			FileCompletionPositions.Add(CurrentPosition);
+
+			FileRetirementPositions.Add(ConstructListElem, CurrentPosition);
 		}
 
 		ConstructionStack[(ConstructListNum - 1) - ConstructListIdx] = ConstructListElem;
 	}
 
+	// Let the install source know when we're going to be deleting their sources.
+	if (InstallSource)
+	{
+		InstallSource->SetFileRetirementPositions(MoveTemp(FileRetirementPositions));	
+	}
+
 	WriteBuffers[0].Reserve(WriteBufferSize);
 	WriteBuffers[1].Reserve(WriteBufferSize);
 
-	WriteJobThread = Configuration.SharedContext->CreateThread();
-	WriteJobCompleteEvent = FPlatformProcess::GetSynchEventFromPool();
-	WriteJobStartEvent = FPlatformProcess::GetSynchEventFromPool();
-	WriteJobThread->RunTask([this]() { WriteJobThreadRun(); });
+	//
+	// Create the threads we are allowed to.
+	//
+	bool bHasInstallSource = InstallSource && InstallSource->GetAvailableChunks().Num() != 0;
+	bool bHasChunkDbSource = ChunkDbSource->GetAvailableChunks().Num() != 0;
+	
+	// For now always create a write thread.
+	int32 ThreadCount = 1;
+	WriteThreadIndex = 0;
+
+	// Default everything to running synchronously.
+	for (int8& ThreadAssignment : ThreadAssignments)
+	{
+		ThreadAssignment = -1;
+	}
+
+	// For now we have to strictly assign jobs to threads so that we don't accidentally
+	// hit the same file handle on multiple threads. Once we have proper ReadAtOffset support
+	// we can go nuts (and just use UE::Tasks)
+	// LONGTERM try using UE::Pipe and just blasting everything on tasks?
+	if (bHasInstallSource)
+	{
+		ThreadAssignments[EConstructorChunkLocation::Install] = ThreadCount;
+		ThreadCount++;
+	}
+	if (bHasChunkDbSource)
+	{
+		ThreadAssignments[EConstructorChunkLocation::ChunkDb] = ThreadCount;
+		ThreadCount++;
+	}
+
+	
+	WakeUpDispatchThreadEvent = FPlatformProcess::GetSynchEventFromPool();
+	CloudSource->SetWakeupFunction([WakeUpDispatchThreadEvent = WakeUpDispatchThreadEvent]()
+		{
+			WakeUpDispatchThreadEvent->Trigger();
+		}
+	);
+
+	// Preallocate the arrays so we don't get any movement.
+	Threads.SetNum(ThreadCount);
+	ThreadWakeups.SetNum(ThreadCount);
+	ThreadCompleteEvents.SetNum(ThreadCount);
+	ThreadJobPostingLocks.SetNum(ThreadCount); // Default init is fine for these.
+	ThreadJobPostings.SetNum(ThreadCount);  // Default init is fine for these.
+
+	for (int32 ThreadIndex = 0; ThreadIndex < ThreadCount; ThreadIndex++)
+	{
+		Threads[ThreadIndex] = Configuration.SharedContext->CreateThread();
+		ThreadWakeups[ThreadIndex] = FPlatformProcess::GetSynchEventFromPool();
+		ThreadCompleteEvents[ThreadIndex] = FPlatformProcess::GetSynchEventFromPool();
+
+		Threads[ThreadIndex]->RunTask([this, ThreadIndex]() { GenericThreadFn(ThreadIndex); });
+	}
+
 }
 
 FBuildPatchFileConstructor::~FBuildPatchFileConstructor()
 {
-	if (bWriteJobRunning)
-	{
-		GLog->Logf(TEXT("FBuildPatchFileConstructor: Write job active during destruction! Very bad."));
-	}
-
-	// Signal background thread to shut down.
+	// Wait for threads to shut down.
 	Abort();
-	WriteJobStartEvent->Trigger();	
-	WriteJobCompleteEvent->Wait();
+	
+	for (int32 ThreadIndex = 0; ThreadIndex < Threads.Num(); ThreadIndex++)
+	{
+		FPlatformProcess::ReturnSynchEventToPool(ThreadWakeups[ThreadIndex]);
+		FPlatformProcess::ReturnSynchEventToPool(ThreadCompleteEvents[ThreadIndex]);
+		Configuration.SharedContext->ReleaseThread(Threads[ThreadIndex]);
+	}
 
-	FPlatformProcess::ReturnSynchEventToPool(WriteJobCompleteEvent);
-	WriteJobCompleteEvent = nullptr;
-	FPlatformProcess::ReturnSynchEventToPool(WriteJobStartEvent);
-	WriteJobStartEvent = nullptr;
+	FPlatformProcess::ReturnSynchEventToPool(WakeUpDispatchThreadEvent);
+	WakeUpDispatchThreadEvent = nullptr;
+}
 
-	Configuration.SharedContext->ReleaseThread(WriteJobThread);
+void FBuildPatchFileConstructor::SetPaused(bool bInIsPaused)
+{
+	bool bWasPaused = bIsPaused.AtomicSet(bInIsPaused);
+	if (bWasPaused && !bInIsPaused)
+	{
+		// If we unpaused, the dispatch thread might be waiting in an event for us
+		// to tell it to unpark.
+		WakeUpDispatchThreadEvent->Trigger();
+	}
+}
+
+void FBuildPatchFileConstructor::Abort()
+{
+	bool bAlreadyAborted = bShouldAbort.AtomicSet(true);
+	if (bAlreadyAborted)
+	{
+		return;
+	}
+
+	// Make sure to wake up any threads that might be parked so they can bail.
+	WakeUpDispatchThreadEvent->Trigger();
+	
+	for (int32 ThreadIndex = 0; ThreadIndex < Threads.Num(); ThreadIndex++)
+	{
+		ThreadWakeups[ThreadIndex]->Trigger();
+	}
+
+	for (int32 ThreadIndex = 0; ThreadIndex < Threads.Num(); ThreadIndex++)
+	{
+		ThreadCompleteEvents[ThreadIndex]->Wait();
+	}
+
+}
+
+void FBuildPatchFileConstructor::SetChunkLocation(const FGuid& InGuid, EConstructorChunkLocation InNewLocation)
+{
+	// NOTE - reentrant. We assume a) that ChunkLocations is filled before threading
+	// and that b) no guids are used across concurrent actions.
+	EConstructorChunkLocation* Location = ChunkLocations.Find(InGuid);
+	if (!Location)
+	{
+		UE_LOG(LogBuildPatchServices, Error, TEXT("Consistency failure: setting chunk location for non existent chunk %s"), *WriteToString<40>(InGuid));
+	}
+	else
+	{
+		if (*Location != EConstructorChunkLocation::Cloud &&
+			InNewLocation == EConstructorChunkLocation::Cloud)
+		{
+			uint64 ChunkSize = Configuration.ManifestSet->GetDownloadSize(InGuid);
+
+			DownloadRequirementLock.Lock();
+			DownloadRequirement += ChunkSize;
+			CloudSource->PostRequiredByteCount(DownloadRequirement);
+			DownloadRequirementLock.Unlock();
+		}
+
+		*Location = InNewLocation;
+	}
+}
+
+void FBuildPatchFileConstructor::QueueGenericThreadTask(int32 ThreadIndex, IConstructorChunkSource::FRequestProcessFn&& Task)
+{
+	// No thread for this task - run synchronously
+	if (ThreadIndex == -1 || ThreadIndex >= Threads.Num())
+	{
+		Task(false);
+		WakeUpDispatchThreadEvent->Trigger();
+		return;
+	}
+
+	bool bPosted = false;
+	ThreadJobPostingLocks[ThreadIndex].Lock();
+	if (!bShouldAbort)
+	{
+		ThreadJobPostings[ThreadIndex].Add(MoveTemp(Task));
+		bPosted = true;
+	}
+	ThreadJobPostingLocks[ThreadIndex].Unlock();
+	ThreadWakeups[ThreadIndex]->Trigger();
+
+	if (!bPosted)
+	{
+		// This means we aborted during the queue - make sure to run
+		Task(true);
+	}
 }
 
 
-void FBuildPatchFileConstructor::WriteJobThreadRun()
+void FBuildPatchFileConstructor::GenericThreadFn(int32 ThreadIndex)
 {
-	TRACE_CPUPROFILER_EVENT_SCOPE(WriteJobThread);
-
 	for (;;)
 	{
-		bool bSignalWasFired = WriteJobStartEvent->Wait(100 /* ms */);
+		ThreadWakeups[ThreadIndex]->Wait();
 
-		if (bSignalWasFired)
+		TArray<IConstructorChunkSource::FRequestProcessFn> GrabbedJobs;
+		
 		{
-			// (got signal) -- they launched a job - init to failed job
-			bWriteJobCompleted = false;
+			ThreadJobPostingLocks[ThreadIndex].Lock();
+			GrabbedJobs = MoveTemp(ThreadJobPostings[ThreadIndex]);
+			ThreadJobPostingLocks[ThreadIndex].Unlock();
 		}
 
-		if (bShouldAbort) // this is also used for graceful shutdown on completion.
+		if (bShouldAbort)
 		{
-			// Leave WriteJobCompleted = false;
-			WriteJobCompleteEvent->Trigger();
-			return;
+			for (IConstructorChunkSource::FRequestProcessFn& AbortJob : GrabbedJobs)
+			{
+				AbortJob(true);
+			}
+			WakeUpDispatchThreadEvent->Trigger();
+			break;
 		}
 
-		if (!bSignalWasFired)
+		for (IConstructorChunkSource::FRequestProcessFn& Job : GrabbedJobs)
 		{
-			// We hit the timeout checking for an abort signal, wait agian.
-			continue;
+			Job(false);
 		}
 
-		FileConstructorStat->OnBeforeWrite();
-		ISpeedRecorder::FRecord ActivityRecord;
-		ActivityRecord.CyclesStart = FStatsCollector::GetCycles();
-		
-		{
-			TRACE_CPUPROFILER_EVENT_SCOPE(WriteThread_Serialize)
-			WriteJobArchive->Serialize(WriteJobBufferToWrite->GetData(), WriteJobBufferToWrite->Num());
-		}
+		WakeUpDispatchThreadEvent->Trigger();
+	}
 
-		ActivityRecord.Size = WriteJobBufferToWrite->Num();
-		ActivityRecord.CyclesEnd = FStatsCollector::GetCycles();
-		FileConstructorStat->OnAfterWrite(ActivityRecord);
-
-		bWriteJobCompleted = true;
-		WriteJobCompleteEvent->Trigger();
-	}
+	ThreadCompleteEvents[ThreadIndex]->Trigger();
 }
 
 void FBuildPatchFileConstructor::Run()
 {
 	FileConstructorStat->OnTotalRequiredUpdated(TotalJobSize);
 
+	DownloadRequirement = 0;
+	for (const TPair<FGuid, EConstructorChunkLocation>& ChunkLocation : ChunkLocations)
+	{
+		if (ChunkLocation.Value == EConstructorChunkLocation::Cloud)
+		{
+			DownloadRequirement += Configuration.ManifestSet->GetDownloadSize(ChunkLocation.Key);
+		}
+	}
+	CloudSource->PostRequiredByteCount(DownloadRequirement);
+
 	// Check for resume data, we need to also look for a legacy resume file to use instead in case we are resuming from an install of previous code version.
 	const FString LegacyResumeDataFilename = Configuration.StagingDirectory / TEXT("$resumeData");
 	const FString ResumeDataFilename = Configuration.MetaDirectory / TEXT("$resumeData");
@@ -481,15 +1584,167 @@
 				FPaths::CollapseRelativeDirectories(FileToDelete);
 				if (FileSystem->FileExists(*FileToDelete))
 				{
+					bool bHarvestSuccess = true;
+					if (InstallSource)
+					{
+						TRACE_CPUPROFILER_EVENT_SCOPE(Harvest);
+						// We need to grab any chunks from install sources that are no longer available.
+						// Anything that's already been loaded is already placed into the memory store appropriately,
+						// but anything that _hasn't_ needs to be pulled out.	
+						TSet<FGuid> FileChunks;
+						InstallSource->GetChunksForFile(FileToDelete, FileChunks);
+
+						struct FNeededChunk
+						{
+							FGuid Id;
+							int32 LastUsageIndex;
+							int32 NextUsageIndex;
+							int32 ChunkSize;
+						};
+
+						TArray<FNeededChunk> ChunksFromFileWeNeed;
+
+						int32 CurrentUsageIndex = ChunkReferenceTracker->GetCurrentUsageIndex();
+						for (const FGuid& FileChunk : FileChunks)
+						{
+							EConstructorChunkLocation* Location = ChunkLocations.Find(FileChunk);
+							if (Location && *Location == EConstructorChunkLocation::Install)
+							{
+								int32 LastUsageIndex = 0;
+								int32 NextUsageIndex = ChunkReferenceTracker->GetNextUsageForChunk(FileChunk, LastUsageIndex);
+
+								if (NextUsageIndex == -1 ||
+									LastUsageIndex < CurrentUsageIndex)
+								{
+									// The chunk is no longer needed
+									*Location = EConstructorChunkLocation::Retired;
+									continue;
+								}
+
+								int32 ChunkSize = Configuration.ManifestSet->GetChunkInfo(FileChunk)->WindowSize;
+
+								ChunksFromFileWeNeed.Add({FileChunk, LastUsageIndex, NextUsageIndex, ChunkSize});
+							}
+						}
+
+						if (ChunksFromFileWeNeed.Num())
+						{
+							// Try to load all the chunks that are about to go away. If it fails we don't particularly
+							// care since we would have fallen back to cloud anyway.
+
+							// There's some care here - if we just kick off a ton of reads, all those backing store
+							// entries are locked during the reads so we have to allocate space and can't page anything
+							// out. This is fine if we can load the whole file, but under constrained memory we want to
+							// only keep stuff that's going to get used soon - so we load the stuff we _aren't_ going to use
+							// soon so it can get paged out. Then we load in batches so we can release the locks and let them
+							// page out. 
+							// LONGTERM - detect this condition and write directly to the disk backing store? Ideally this
+							// would be something we can retain across restarts as right now any harvested chunks get lost
+							// on abort and cause a download from the cloud source (not chunkdb!)
+							Algo::SortBy(ChunksFromFileWeNeed, &FNeededChunk::NextUsageIndex, TGreater<int32>());
+
+							constexpr int32 HarvestBatchSize = 16 << 20; // 16 MB
+
+							for (int32 ChunkIndex = 0; ChunkIndex < ChunksFromFileWeNeed.Num();)
+							{
+								int32 BatchSize = 0;
+								int32 ChunkEndIndex = ChunkIndex + 1;
+								for (; ChunkEndIndex < ChunksFromFileWeNeed.Num(); ChunkEndIndex++)
+								{
+									if (BatchSize + ChunksFromFileWeNeed[ChunkEndIndex].ChunkSize > HarvestBatchSize)
+									{
+										break;
+									}
+									BatchSize += ChunksFromFileWeNeed[ChunkEndIndex].ChunkSize;
+								}
+
+								PendingHarvestRequests.store(ChunkEndIndex - ChunkIndex, std::memory_order_release);
+								
+								for (int32 DispatchIndex = ChunkIndex; DispatchIndex < ChunkEndIndex; DispatchIndex++)
+								{
+									FNeededChunk& Chunk = ChunksFromFileWeNeed[DispatchIndex];
+
+									SetChunkLocation(Chunk.Id, EConstructorChunkLocation::Memory);
+
+									FMutableMemoryView Destination = BackingStore->ReserveAndLockEntry(Chunk.Id, Chunk.ChunkSize, Chunk.LastUsageIndex);
+									if (Destination.GetSize() == 0)
+									{
+										// Call the completion function so we decrement the request count,
+										// but this is a consistency failure so we can't use the results.
+										ChunkHarvestCompletedFn(Chunk.Id, false, false, 0);
+										bHarvestSuccess = false;
+									}
+
+									IConstructorChunkSource::FRequestProcessFn HarvestFn = InstallSource->CreateRequest(Chunk.Id, Destination, 0, IConstructorChunkSource::FChunkRequestCompleteDelegate::CreateRaw(this, &FBuildPatchFileConstructor::ChunkHarvestCompletedFn));
+									
+									HarvestFn(false);
+								}
+
+								// The read is synchronous but the verification is not, so we still need to do the wait.
+								{
+									TRACE_CPUPROFILER_EVENT_SCOPE(HarvestWait);
+									WakeUpDispatchThreadEvent->Wait();
+								}
+							
+								if (bHarvestSuccess)
+								{
+									// Unlock any memory store entries
+									for (int32 DispatchIndex = ChunkIndex; DispatchIndex < ChunkEndIndex; DispatchIndex++)
+									{
+										// The read could have failed - in which case the location was switched to cloud from Memory.
+										// We handle this here so we don't have to deal with synchronization in the completion function.
+										if (ChunkLocations[ChunksFromFileWeNeed[DispatchIndex].Id] == EConstructorChunkLocation::Cloud)
+										{
+											if (!BackingStore->DereserveHarvestingEntry(ChunksFromFileWeNeed[DispatchIndex].Id))
+											{
+												bHarvestSuccess = false;
+											}
+										}
+										else
+										{
+											if (!BackingStore->CommitAndReleaseEntry(ChunksFromFileWeNeed[DispatchIndex].Id))
+											{
+												bHarvestSuccess = false;
+											}
+										}
+									}
+								}
+								
+								if (bHarvestSuccess)
+								{
+									bHarvestSuccess = BackingStore->CheckNoLocks(true);
+								}
+								
+								ChunkIndex = ChunkEndIndex;
+
+								if (bShouldAbort || !bHarvestSuccess)
+								{
+									break;
+								}
+							}						
+						} // end if we need to harvest
+
+					} // end if install source exists.
+
+					if (!bHarvestSuccess)
+					{
+						Abort();
+						InstallerError->SetError(EBuildPatchInstallError::FileConstructionFail, ConstructionErrorCodes::InternalConsistencyFailure);
+					}
+
 					OnBeforeDeleteFile().Broadcast(FileToDelete);
-					IFileManager::Get().Delete(*FileToDelete, bRequireExists, bEvenReadOnly);
+					{
+						// This can take forever due to file system filters
+						TRACE_CPUPROFILER_EVENT_SCOPE(Delete);
+						IFileManager::Get().Delete(*FileToDelete, bRequireExists, bEvenReadOnly);
+					}
 				}
 			}
 		}
 		else
 		{
 			// This will only record and log if a failure was not already registered.
-			bShouldAbort = true;
+			Abort();
 			InstallerError->SetError(EBuildPatchInstallError::FileConstructionFail, ConstructionErrorCodes::UnknownFail);
 			UE_LOG(LogBuildPatchServices, Error, TEXT("FBuildPatchFileConstructor: Failed to build %s "), *FileToConstruct);
 		}
@@ -598,8 +1853,6 @@
 	}
 }
 
-
-
 uint64 FBuildPatchFileConstructor::CalculateDiskSpaceRequirementsWithDeleteDuringInstall(const TArray<FString>& InBackwardsFilesLeftToConstruct)
 {
 	if (ChunkDbSource == nullptr)
@@ -632,4 +1885,473 @@
 	return PostDlSize;
 }
 
+struct FBatchState;
+struct FRequestInfo
+{
+	struct FRequestSplat
+	{
+		int32 DestinationOffset;
+		int32 OffsetInChunk;
+		int32 BytesToCopy;
+	};
+
+	FGuid Guid;
+
+	// We can do a lot of shortcuts if we are working with the entire chunk
+	int32 ChunkSize = 0;
+
+	// We could request the same guid multiple times
+	// for the same buffer... in this case we want one
+	// request but we need to remember to splat it
+	// afterwards.
+	TArray<FRequestSplat, TInlineAllocator<1>> Splats;
+	
+	// The read goes here - this is usually directly into the write buffer,
+	// but we might need to duplicate out of this (and this might be memory
+	// store owned if we don't use the whole chunk)
+	FMutableMemoryView ReadBuffer;
+
+	FBatchState* Batch = nullptr;
+
+	// Splats offset into this.
+	FMutableMemoryView DestinationBuffer;
+
+	// We can only read direct in some cases.
+	bool bReadIntoMemoryStore = false;
+
+	// Never save back to memory store if it came from it.
+	bool bSourceIsMemoryStore = false;
+	bool bAborted = false;
+	bool bFailedToRead = false;
+	bool bLaunchedFallback = false;
+
+	int32 ChunkUnavailableAt = TNumericLimits<int32>::Max();
+};
+
+
+struct FBatchState
+{
+	int32_t BatchId = 0;
+
+	TMap<FGuid, FRequestInfo> Requests;			
+	std::atomic_int32_t PendingRequestCount = 0;
+	std::atomic_int32_t FailedRequestCount = 0;			
+
+	int32 BatchIndex = 0;
+	EConstructionError Error = EConstructionError::None;
+	FGuid ErrorContextGuid;
+
+	int32 StartChunkFileIndex = 0;
+	int32 ChunkCount = 0;
+
+	bool bNeedsWrite = false;
+
+	bool bIsReading = false;
+
+	bool BufferIsAvailableForRead() const { return !bIsWriting && !bIsReading && !bNeedsWrite; }
+
+	bool bIsWriting = false;
+
+	// Set by the writing thread when the write completed prior to waking up the dispatch thread.
+	std::atomic<bool> bIsFinished = true;
+	
+};
+
+// Called from basically any thread.
+void FBuildPatchFileConstructor::ChunkHarvestCompletedFn(const FGuid& Guid, bool bAborted, bool bFailedToRead, void* UserPtr)
+{
+	if (bFailedToRead)
+	{
+		// We tell the main thread this failed by setting the location since that's thread
+		// safe.
+		SetChunkLocation(Guid, EConstructorChunkLocation::Cloud);
+	}
+
+	if (PendingHarvestRequests.fetch_add(-1, std::memory_order_relaxed) == 1)
+	{
+		WakeUpDispatchThreadEvent->Trigger();
+	}
+}
+
+// Called from basically any thread.
+void FBuildPatchFileConstructor::RequestCompletedFn(const FGuid& Guid, bool bAborted, bool bFailedToRead, void* UserPtr)
+{
+	FRequestInfo* Request = (FRequestInfo*)UserPtr;
+	Request->bAborted = bAborted;
+	Request->bFailedToRead = bFailedToRead;
+
+	// If failed but didn't abort and we haven't already kicked the fallback, kick the fallback
+	if (!bAborted && bFailedToRead && !Request->bLaunchedFallback)
+	{		
+		Request->bLaunchedFallback = true;
+		Request->bFailedToRead = false;
+
+		// We couldn't read the expected source. This means we need to fall back to the cloud source.
+
+		// This should be safe because the values already exist in the map and we only
+		// ever have 1 request for a Guid active at one time. However if we've already read into the
+		// memory store then it's already updated to memory which is where the cloud source will read
+		// it to. If it's not reading into the memory store, then we need to remember to grab it
+		// from the cloud next time.
+		if (!Request->bReadIntoMemoryStore)
+		{
+			SetChunkLocation(Request->Guid, EConstructorChunkLocation::Cloud);
+		}
+		else
+		{
+			// Normally SetChunkLocation will update the download amount, but we aren't actually
+			// changing the chunk's location since it's going into the memory store. We do still
+			// need to tell the user about the download requirement though:
+			uint64 ChunkSize = Configuration.ManifestSet->GetDownloadSize(Guid);
+
+			DownloadRequirementLock.Lock();
+			DownloadRequirement += ChunkSize;
+			CloudSource->PostRequiredByteCount(DownloadRequirement);
+			DownloadRequirementLock.Unlock();
+
+		}
+
+		QueueGenericThreadTask(ThreadAssignments[EConstructorChunkLocation::Cloud], CloudSource->CreateRequest(
+			Request->Guid, 
+			Request->ReadBuffer, 
+			Request,
+			IConstructorChunkSource::FChunkRequestCompleteDelegate::CreateRaw(this, &FBuildPatchFileConstructor::RequestCompletedFn)));
+	}
+	else
+	{
+		if (bFailedToRead)
+		{
+			Request->Batch->FailedRequestCount.fetch_add(1, std::memory_order_relaxed);
+		}
+		else if (!bAborted && Request->bReadIntoMemoryStore)
+		{
+			// If the read went to memory, we need to copy splats. Otherwise it was a single
+			// direct read so do nothing.
+			TRACE_CPUPROFILER_EVENT_SCOPE(CFFC_Splats);
+			for (const FRequestInfo::FRequestSplat& Splat : Request->Splats)
+			{
+				FMemory::Memcpy(
+					(uint8*)Request->DestinationBuffer.GetData() + Splat.DestinationOffset,
+					(uint8*)Request->ReadBuffer.GetData() + Splat.OffsetInChunk,
+					Splat.BytesToCopy);
+			}
+		}
+		
+		if (Request->Batch->PendingRequestCount.fetch_add(-1, std::memory_order_relaxed) == 1)
+		{
+			Request->Batch->bIsFinished.store(true, std::memory_order_release);
+			WakeUpDispatchThreadEvent->Trigger();
+		}
+	}
+}
+
+// Return a function that writes the data for the batch to the given file.
+IConstructorChunkSource::FRequestProcessFn FBuildPatchFileConstructor::CreateWriteRequest(FArchive* File, FBatchState& Batch)
+{
+	return [this, File, Batch = &Batch](bool bIsAbort)
+	{
+		if (bIsAbort)
+		{
+			return;
+		}
+
+		TArray<uint8>& WriteBuffer = WriteBuffers[Batch->BatchIndex];
+
+		// Manage write limits.
+		if (bStallWhenFileSystemThrottled)
+		{
+			int64 AvailableBytes = FileSystem->GetAllowedBytesToWriteThrottledStorage(*File->GetArchiveName());
+			while (WriteBuffer.Num() > AvailableBytes)
+			{
+				UE_LOG(LogBuildPatchServices, Display, TEXT("Avaliable write bytes to write throttled storage exhausted (%s).  Sleeping %ds.  Bytes needed: %u, bytes available: %lld")
+					, *File->GetArchiveName(), SleepTimeWhenFileSystemThrottledSeconds, WriteBuffer.Num(), AvailableBytes);
+				FPlatformProcess::Sleep(SleepTimeWhenFileSystemThrottledSeconds);
+				AvailableBytes = FileSystem->GetAllowedBytesToWriteThrottledStorage(*File->GetArchiveName());
+
+				if (bShouldAbort)
+				{
+					return;
+				}
+			}
+		}
+
+		FileConstructorStat->OnBeforeWrite();
+		ISpeedRecorder::FRecord ActivityRecord;
+		ActivityRecord.CyclesStart = FStatsCollector::GetCycles();
+		
+		{
+			TRACE_CPUPROFILER_EVENT_SCOPE(WriteThread_Serialize)
+			File->Serialize(WriteBuffer.GetData(), WriteBuffer.Num());
+		}
+
+		ActivityRecord.Size = WriteBuffer.Num();
+		ActivityRecord.CyclesEnd = FStatsCollector::GetCycles();
+		FileConstructorStat->OnAfterWrite(ActivityRecord);
+		
+		Batch->bIsFinished.store(true, std::memory_order_release);
+	};
+}
+					
+
+// Always called on Constructor thread.
+void FBuildPatchFileConstructor::CompleteReadBatch(const FFileManifest& FileManifest, FBatchState& Batch)
+{
+	TRACE_CPUPROFILER_EVENT_SCOPE(CFFC_CompleteRead);
+
+	if (Batch.Error == EConstructionError::None && Batch.FailedRequestCount.load(std::memory_order_acquire))
+	{
+		Batch.Error = EConstructionError::MissingChunk;
+	}
+
+	//UE_LOG(LogTemp, Warning, TEXT("Completing batch: %d"), Batch.BatchId);
+
+	// We have to copy the memory source chunks after the reads are done
+	// because if we have two buffer's reads queued, the first one could be
+	// filling the memory source. If we copy these after we are done, we guarantee
+	// that the previous buffer has completed its reads so we know we are
+	// working with valid memory.
+	if (Batch.Error == EConstructionError::None)
+	{
+		for (TPair<FGuid, FRequestInfo>& RequestPair : Batch.Requests)
+		{					
+			FRequestInfo& Request = RequestPair.Value;
+
+			if (Request.bSourceIsMemoryStore)
+			{
+				TRACE_CPUPROFILER_EVENT_SCOPE(CFFC_DirectCopy);
+				// Just copy what we need directly.
+				FMemoryView ChunkData = BackingStore->GetViewForChunk(Request.Guid);
+				bool bFailedToGetChunk = ChunkData.GetSize() == 0;
+
+				if (!bFailedToGetChunk)
+				{
+					for (const FRequestInfo::FRequestSplat& Splat : Request.Splats)
+					{
+						FMemory::Memcpy(
+							(uint8*)Request.DestinationBuffer.GetData() + Splat.DestinationOffset,
+							(uint8*)ChunkData.GetData() + Splat.OffsetInChunk,
+							Splat.BytesToCopy);
+					}
+				}
+				else
+				{
+					Batch.Error = EConstructionError::MissingChunk;
+					Batch.ErrorContextGuid = Request.Guid;
+				}
+
+				// Mark that we're done with the memory so it can get evicted if necessary.
+				if (!BackingStore->ReleaseEntry(Request.Guid))
+				{
+					Batch.Error = EConstructionError::InternalConsistencyError;
+					Batch.ErrorContextGuid = Request.Guid;
+				}
+			}
+			else if (Request.bReadIntoMemoryStore)
+			{
+				// Commit the read so the memory store knows its safe to evict if necessary.
+				if (!BackingStore->CommitAndReleaseEntry(Request.Guid))
+				{
+					Batch.Error = EConstructionError::InternalConsistencyError;
+					Batch.ErrorContextGuid = Request.Guid;
+				}
+			}
+		}
+	}
+
+	if (Batch.Error == EConstructionError::None)
+	{
+		// Retire the chunks we've used. This has to be in order 
+		for (int32 i = 0; i < Batch.ChunkCount; i++)
+		{
+			if (!ChunkReferenceTracker->PopReference(FileManifest.ChunkParts[Batch.StartChunkFileIndex + i].Guid))
+			{
+				Batch.Error = EConstructionError::TrackingError;
+				Batch.ErrorContextGuid = FileManifest.ChunkParts[Batch.StartChunkFileIndex + i].Guid;
+			}
+		}
+	}
+
+	if (Batch.Error == EConstructionError::None)
+	{
+		TRACE_CPUPROFILER_EVENT_SCOPE(CFFC_MemRetire);
+		// Has to be after the splats because this might free the memory we need!				
+		if (!BackingStore->CheckRetirements(ChunkReferenceTracker->GetCurrentUsageIndex()))
+		{
+			Batch.Error = EConstructionError::InternalConsistencyError;
+		}
+	}
+}
+
+// Always called from Constructor thread
+// This return -1 when we run into an error that requires stopping the installation.
+// Note that there could be outstanding reads and we can no longer rely on the completion functions
+// to wake up the dispatch thread (which we must run on), so we can't ever wait if this returns -1.
+// Check Batch.Error for the error on -1 return.
+int32 FBuildPatchFileConstructor::StartReadBatch(const FFileManifest& FileManifest, int32 StartFileChunkPartIndex, FBatchState& Batch)
+{
+	static std::atomic_int32_t UniqueBatchId = 1;
+
+	Batch.Requests.Reset();
+	Batch.BatchId = UniqueBatchId.fetch_add(1);
+	Batch.StartChunkFileIndex = StartFileChunkPartIndex;
+
+	//UE_LOG(LogTemp, Warning, TEXT("Starting batch: %d"), Batch.BatchId);
+
+	int32 EndChunkIdx = StartFileChunkPartIndex;
+	{
+		TRACE_CPUPROFILER_EVENT_SCOPE(CFFC_Count);
+		uint32 BufferFillLevel = 0;
+		for (; EndChunkIdx < FileManifest.ChunkParts.Num(); ++EndChunkIdx)
+		{
+			const FChunkPart& ChunkPart = FileManifest.ChunkParts[EndChunkIdx];
+			if (ChunkPart.Size + BufferFillLevel > WriteBufferSize)
+			{
+				break;
+			}
+
+			FRequestInfo& Request = Batch.Requests.FindOrAdd(ChunkPart.Guid);
+			Request.Guid = ChunkPart.Guid;
+			Request.Batch = &Batch;
+			Request.ChunkSize = Configuration.ManifestSet->GetChunkInfo(ChunkPart.Guid)->WindowSize;
+			
+			FRequestInfo::FRequestSplat& Splat = Request.Splats.AddDefaulted_GetRef();
+			Splat.DestinationOffset = BufferFillLevel;
+			Splat.OffsetInChunk = ChunkPart.Offset;
+			Splat.BytesToCopy = ChunkPart.Size;
+				
+			BufferFillLevel += ChunkPart.Size;
+		}
+
+		WriteBuffers[Batch.BatchIndex].SetNumUninitialized(BufferFillLevel, EAllowShrinking::No);
+	}
+
+	Batch.ChunkCount = EndChunkIdx - StartFileChunkPartIndex;
+	Batch.PendingRequestCount = Batch.Requests.Num();
+	Batch.FailedRequestCount = 0;
+	
+	//
+	// IMPORTANT!!!
+	//
+	// We MUST call the completion routine for each request here so that we can know
+	// when all outstanding requests are done. If we bail early then we can't know when
+	// all requests are done during cancelation/abort. If we hit a consistency error during then
+	// we just need to call it with a failure.
+	//
+	for (TPair<FGuid, FRequestInfo>& RequestPair : Batch.Requests)
+	{
+		// Can come from:
+		// -- memory. This chunk had already been loaded and we needed it again _after it's source expired_.
+		// -- chunkdb. async IO + decompress
+		// -- cloud. async download + decompress
+		// -- disk. we ran out of memory store and had to dump to disk. async IO + maybe decompress
+		// -- install. async IO
+
+		// Note -- most chunks will be used more than once.
+
+		// sources that expire: note that expired chunks can always be redownloaded via the cloud source
+		//  -- cloud (they expire immediately, but can be redownloaded)
+		//  -- install
+
+		// We aren't necessarily using the whole chunk - e.g. if we are a small file, we'll only
+		// be a tiny part of the chunk and the rest will need to be used by the next file. In this case
+		// we have to read into a memory store destination so that it can be carried over into the
+		// next file.
+
+
+		FRequestInfo& Request = RequestPair.Value;
+		Request.DestinationBuffer = MakeMemoryView(WriteBuffers[Batch.BatchIndex]);
+
+		
+		EConstructorChunkLocation* ChunkLocationPtr = ChunkLocations.Find(Request.Guid);
+		if (ChunkLocationPtr == nullptr)
+		{
+			Batch.Error = EConstructionError::InternalConsistencyError;
+			RequestCompletedFn(Request.Guid, true, false, &Request);
+			continue;
+		}
+		EConstructorChunkLocation ChunkLocation = *ChunkLocationPtr;
+
+		Request.bSourceIsMemoryStore = ChunkLocation == EConstructorChunkLocation::Memory;
+		if (Request.bSourceIsMemoryStore)
+		{
+			// We copy after the reads are done since the memory might not be ready.
+			// Make sure we don't evict it in the meantime
+			if (!BackingStore->LockEntry(Request.Guid))
+			{
+				Batch.Error = EConstructionError::InternalConsistencyError;
+				RequestCompletedFn(Request.Guid, true, false, &Request);
+				continue;
+			}
+			RequestCompletedFn(Request.Guid, false, false, &Request);
+		}
+		else
+		{
+			IConstructorChunkSource* ThisChunkSource = nullptr;
+			switch (ChunkLocation)
+			{
+				case EConstructorChunkLocation::Install: ThisChunkSource = InstallSource; break;
+				case EConstructorChunkLocation::Cloud: ThisChunkSource = CloudSource; break;
+				case EConstructorChunkLocation::ChunkDb: ThisChunkSource = ChunkDbSource; break;
+				case EConstructorChunkLocation::Memory: /* already handled above */ break;
+				case EConstructorChunkLocation::DiskOverflow: ThisChunkSource = (IConstructorChunkSource*)BackingStore.Get(); break;
+			}
+
+			if (ThisChunkSource == CloudSource)
+			{
+				// If we are already downloading from the cloud, then failures shouldn't try to 
+				// fall back to the cloud.
+				Request.bLaunchedFallback = true;
+			}
+
+			// We need to kick a request. The question is whether we can request direct
+			// or need to route through the memory store.
+			int32 LastUsageIndex = 0;
+			ChunkReferenceTracker->GetNextUsageForChunk(Request.Guid, LastUsageIndex);
+			Request.ChunkUnavailableAt = ThisChunkSource->GetChunkUnavailableAt(Request.Guid);
+				
+			const bool bNeedsEntireChunk = Request.Splats.Num() == 1 && Request.Splats[0].BytesToCopy == Request.ChunkSize && Request.Splats[0].OffsetInChunk == 0;
+			const bool bNeededAfterRetirement = LastUsageIndex >= Request.ChunkUnavailableAt;
+
+			if (bNeedsEntireChunk && !bNeededAfterRetirement)
+			{
+				// Read direct.
+				Request.ReadBuffer = MakeMemoryView((uint8*)Request.DestinationBuffer.GetData() + Request.Splats[0].DestinationOffset, Request.Splats[0].BytesToCopy);
+			}
+			else
+			{
+				// Route through memory store.
+				Request.ReadBuffer = BackingStore->ReserveAndLockEntry(Request.Guid, Request.ChunkSize, LastUsageIndex);
+				if (Request.ReadBuffer.GetSize() == 0)
+				{
+					Batch.Error = EConstructionError::InternalConsistencyError;
+					RequestCompletedFn(Request.Guid, true, false, &Request);
+					continue;
+				}
+				Request.bReadIntoMemoryStore = true;
+
+				// Note that when we set this, the next batch read could want this chunk before its read is done.
+				// Hence reads for memory sources are done _after_ reads are done, because we retire reads in order,
+				// we then know this memory is populated.
+				SetChunkLocation(RequestPair.Value.Guid, EConstructorChunkLocation::Memory);
+			}
+			
+
+			IConstructorChunkSource::FRequestProcessFn RequestProcess = ThisChunkSource->CreateRequest(
+				Request.Guid, 
+				Request.ReadBuffer,
+				(void*)&Request,
+				IConstructorChunkSource::FChunkRequestCompleteDelegate::CreateRaw(this, &FBuildPatchFileConstructor::RequestCompletedFn));
+						
+			QueueGenericThreadTask(ThreadAssignments[ChunkLocation], MoveTemp(RequestProcess));
+		}
+	}
+
+	return EndChunkIdx;
+}
+
+void FBuildPatchFileConstructor::WakeUpDispatch()
+{
+	WakeUpDispatchThreadEvent->Trigger();
+}
+
 bool FBuildPatchFileConstructor::ConstructFileFromChunks(const FString& BuildFilename, const FFileManifest& FileManifest, bool bResumeExisting)
@@ -637,13 +2359,14 @@
 	TRACE_CPUPROFILER_EVENT_SCOPE(ConstructFileFromChunks);
 
 	bool bSuccess = true;
+
+	FGuid ErrorContextGuid;
 	EConstructionError ConstructionError = EConstructionError::None;
-	uint32 LastError = 0;
+
 	FString NewFilename = Configuration.StagingDirectory / BuildFilename;
 
 	// Calculate the hash as we write the data
 	FSHA1 HashState;
-	FSHAHash HashValue;
 
 	if (!FileManifest.SymlinkTarget.IsEmpty())
 	{
@@ -744,6 +2467,11 @@
 			if (FPlatformMisc::GetDiskTotalAndFreeSpace(Configuration.InstallDirectory, TotalSize, AvailableSpace))
 			{
 				LocalDiskSpaceAvailable = AvailableSpace;
+				BackingStore->SetDynamicDiskSpaceHeadroom(LocalDiskSpaceRequired);
 			}
+			else
+			{
+				BackingStore->DisableDynamicDiskSpaceHeadroom();
+			}
 		}
 
@@ -769,65 +2497,12 @@
 		FileConstructorStat->OnResumeCompleted();
 	}
 
-	// Returns false if the write failed in some way (almost certainly disk space, could be drive disconnection)
-	auto FlushToAsyncWriter = [this](FArchive& DestinationFile, FSHA1& HashState)
-	{
-		if (bStallWhenFileSystemThrottled)
-		{
-			int64 AvailableBytes = FileSystem->GetAllowedBytesToWriteThrottledStorage(*DestinationFile.GetArchiveName());
-			while (WriteBuffers[CurrentFillBuffer].Num() > AvailableBytes)
-			{
-				UE_LOG(LogBuildPatchServices, Display, TEXT("Avaliable write bytes to write throttled storage exhausted (%s).  Sleeping %ds.  Bytes needed: %u, bytes available: %lld")
-					, *DestinationFile.GetArchiveName(), SleepTimeWhenFileSystemThrottledSeconds, WriteBuffers[CurrentFillBuffer].Num(), AvailableBytes);
-				FPlatformProcess::Sleep(SleepTimeWhenFileSystemThrottledSeconds);
-				AvailableBytes = FileSystem->GetAllowedBytesToWriteThrottledStorage(*DestinationFile.GetArchiveName());
-			}
-		}
-		
-		// Wait for the last write to complete.
-		if (bWriteJobRunning)
-		{
-			// We can potentially wait here a while if we are FS throttled.
-			// \todo old code didn't check for abort during throttling, should we add?
-			{
-				TRACE_CPUPROFILER_EVENT_SCOPE(ConstructFileFromChunks_WaitForLastWrite);
-				WriteJobCompleteEvent->Wait();
-			}
-			bWriteJobRunning = false;
-
-			if (DestinationFile.IsError())
-			{
-				return false;
-			}
-
-			// !CurrentFillBuffer is now available for use.
-		}
-
-		// Kick off the write on another thread while we hash the data here.
-		WriteJobBufferToWrite = &WriteBuffers[CurrentFillBuffer];
-		WriteJobArchive = &DestinationFile;
-		bWriteJobRunning = true;
-		WriteJobStartEvent->Trigger();
-
-		// Hash the buffer we are writing while it's writing.
-		{
-			TRACE_CPUPROFILER_EVENT_SCOPE(ConstructFileFromChunks_Hash);
-			HashState.Update(WriteBuffers[CurrentFillBuffer].GetData(), WriteBuffers[CurrentFillBuffer].Num());
-		}
-
-		// Start filling the next buffer.
-		CurrentFillBuffer = !CurrentFillBuffer;
-		WriteBuffers[CurrentFillBuffer].SetNumUninitialized(0, EAllowShrinking::No);
-
-		return true;
-	};
-
 	// Attempt to create the file
 	ISpeedRecorder::FRecord ActivityRecord;
 	FileConstructorStat->OnBeforeAdminister();
 	ActivityRecord.CyclesStart = FStatsCollector::GetCycles();
 	TUniquePtr<FArchive> NewFile = FileSystem->CreateFileWriter(*NewFilename, bResumeExisting ? EWriteFlags::Append : EWriteFlags::None);
-	LastError = FPlatformMisc::GetLastError();
+	int32 CreateFilePlatformLastError = FPlatformMisc::GetLastError();
 	ActivityRecord.CyclesEnd = FStatsCollector::GetCycles();
 	ActivityRecord.Size = 0;
 	FileConstructorStat->OnAfterAdminister(ActivityRecord);
@@ -845,70 +2520,207 @@
 			FileConstructorStat->OnAfterAdminister(ActivityRecord);
 		}
 
-		// For each chunk, load it, and place it's data into the file
-		for (int32 ChunkPartIdx = StartChunkPart; ChunkPartIdx < FileManifest.ChunkParts.Num() && bSuccess && !bShouldAbort; ++ChunkPartIdx)
+		constexpr int32 BatchCount = UE_ARRAY_COUNT(WriteBuffers);
+		FBatchState BatchState[BatchCount];
+		for (int32 BatchIndex = 0; BatchIndex < BatchCount; BatchIndex++)
 		{
-			TRACE_CPUPROFILER_EVENT_SCOPE(ConstructFileFromChunks_Chunk);
+			BatchState[BatchIndex].BatchIndex = BatchIndex;
+		}
 
-			const FChunkPart& ChunkPart = FileManifest.ChunkParts[ChunkPartIdx];
+		FBatchState* WritingBatch = nullptr;
+	
+		// Fill read/write buffers in batches until we complete the file.
+		int32 CurrentFileChunkPartIndex = 0;
 
-			// If we can't fit in the buffer, flush. Conditional arranged to avoid overflow risk.
-			if (ChunkPart.Size > (WriteBufferSize - WriteBuffers[CurrentFillBuffer].Num()))
+		for (;;)
+		{
+			// Check state. We can continue to run this after a failure has been encountered in order
+			// to drain any async tasks so we know we can shut down safely.
 			{
-				if (!FlushToAsyncWriter(*NewFile, HashState))
+				TRACE_CPUPROFILER_EVENT_SCOPE(CFFC_StateCheck);
+
+				// Is our next write done?
+				if (WritingBatch && WritingBatch->bIsFinished.load(std::memory_order_acquire))
 				{
+					WritingBatch->bIsWriting = false;
+					WritingBatch = nullptr;
+				
+					if (NewFile->IsError())
+					{
+						bSuccess = false;
+
+						// SerializeError signals the error management code to check for low disk space
+						// since we currently have no way of routing out of disk errors from the UE low level
+						// writing code.
+						ConstructionError = EConstructionError::SerializeError;
+					}
+
+					// Can use Tell() because right now there's nothing writing.
+					FileConstructorStat->OnFileProgress(BuildFilename, NewFile->Tell());
+
+					// If that was the last write, then we are done.
+					if (CurrentFileChunkPartIndex == FileManifest.ChunkParts.Num())
+					{
+						int32 WriteBatchCount = 0;
+						for (uint32 i=0; i < BatchCount; i++)
+						{
+							WriteBatchCount += BatchState[i].bNeedsWrite ? 1 : 0;
+						}
+
+						if (WriteBatchCount == 0)
+						{
+							break;
+						}
+					}
+				}
+
+				// are any reads done?
+				for (;;)
+				{
+					// We have to retire reads in order because later reads could be waiting for data to
+					// get placed correctly by the earlier ones.
+					int32 EarliestReadingBatch = -1;
+					for (int32 i = 0; i < BatchCount; i++)
+					{
+						if (BatchState[i].bIsReading)
+						{
+							if (EarliestReadingBatch == -1 || BatchState[i].StartChunkFileIndex < BatchState[EarliestReadingBatch].StartChunkFileIndex)
+							{
+								EarliestReadingBatch = i;
+							}
+						}
+					}
+
+					if (EarliestReadingBatch == -1)
+					{
+						break; // Nothing reading.
+					}
+
+					if (!BatchState[EarliestReadingBatch].bIsFinished.load(std::memory_order_acquire))
+					{
+						break; // Read we care about isn't done.
+					}
+
+					CompleteReadBatch(FileManifest, BatchState[EarliestReadingBatch]);
+					BatchState[EarliestReadingBatch].bIsReading = false;
+					BatchState[EarliestReadingBatch].bNeedsWrite = true;
+
+					if (ConstructionError != EConstructionError::None ||
+						BatchState[EarliestReadingBatch].Error != EConstructionError::None)
+					{
+						ErrorContextGuid = BatchState[EarliestReadingBatch].ErrorContextGuid;
+						ConstructionError = BatchState[EarliestReadingBatch].Error;
+						bSuccess = false;
+					}
+
+					// Loop around to check if another is done.
+				}
+
+				if (bSuccess &&
+					bShouldAbort)
+				{
+					ConstructionError = EConstructionError::Aborted;
 					bSuccess = false;
-					InstallerAnalytics->RecordConstructionError(BuildFilename, INDEX_NONE, TEXT("Serialization Error"));
-					UE_LOG(LogBuildPatchServices, Error, TEXT("FBuildPatchFileConstructor: Failed %s due to serialization error"), *BuildFilename);
-					ConstructionError = EConstructionError::SerializeError;
+				}
+
+				if (!bIsPaused && bSuccess)
+				{
+					// Can we start another read?
+					// LONGTERM - if we are memory constrained and have an install source we can consider only dispatching
+					// single reads. This dramatically lowers our memory requirements as install sources often use a lot
+					// of small pieces of chunks - and we have to load the entire chunk into memory.
+					for (int32 i = 0; i < BatchCount; i++)
+					{
+						if (CurrentFileChunkPartIndex < FileManifest.ChunkParts.Num() &&
+							BatchState[i].BufferIsAvailableForRead())
+						{
+							BatchState[i].bNeedsWrite = true;
+							BatchState[i].bIsReading = true;
+							BatchState[i].bIsFinished.store(false, std::memory_order_release);
+							CurrentFileChunkPartIndex = StartReadBatch(FileManifest, CurrentFileChunkPartIndex, BatchState[i]);
+							if (CurrentFileChunkPartIndex == -1)
+							{
+								ConstructionError = BatchState[i].Error;
+								bSuccess = false;
+								break;
+							}
+						}
+					}
+
+					// Can we start a write?
+					if (bSuccess && WritingBatch == nullptr)
+					{
+						// We always have to issue writes in order and there can only be one.
+						int32 WriteBufferIndex = -1;
+						for (int32 i=0; i<BatchCount; i++)
+						{
+							if (BatchState[i].bNeedsWrite)
+							{
+								if (WriteBufferIndex == -1 ||
+									BatchState[WriteBufferIndex].StartChunkFileIndex > BatchState[i].StartChunkFileIndex)
+								{
+									WriteBufferIndex = i;
+								}
+							}
+						}
+
+						// However the next write might not be ready (fix possible of out of order read completion)
+						if (WriteBufferIndex >= 0 &&
+							!BatchState[WriteBufferIndex].bIsReading)
+						{
+							BatchState[WriteBufferIndex].bNeedsWrite = false;
+							BatchState[WriteBufferIndex].bIsWriting = true;
+							BatchState[WriteBufferIndex].bIsFinished.store(false, std::memory_order_relaxed);
+
+							WritingBatch = BatchState + WriteBufferIndex;
+							
+							// Launch the write and hash the buffer on this thread.
+							CountBytesProcessed(WriteBuffers[WriteBufferIndex].Num());
+
+							IConstructorChunkSource::FRequestProcessFn WriteFn = CreateWriteRequest(NewFile.Get(), BatchState[WriteBufferIndex]);
+							QueueGenericThreadTask(WriteThreadIndex, MoveTemp(WriteFn));
+					
+							{
+								TRACE_CPUPROFILER_EVENT_SCOPE(CFFC_Hash);
+								HashState.Update(WriteBuffers[WriteBufferIndex].GetData(), WriteBuffers[WriteBufferIndex].Num());
+							}
+						}
+					} // end if checking for write
+				} // end not paused
+			} // end state check
+
+			uint32 WaitTimeMs = TNumericLimits<uint32>::Max();
+			{
+				TRACE_CPUPROFILER_EVENT_SCOPE(CFFC_TickCloud);
+
+				// Cloud source may want to be woken up after a specific time.
+				CloudSource->Tick(!bIsPaused && bSuccess, WaitTimeMs);
+			}
+
+			if (!bSuccess)
+			{
+				int32 ReadingBatchCount = 0;
+				for (int32 i = 0; i < BatchCount; i++)
+				{
+					ReadingBatchCount += BatchState[i].bIsReading ? 1 : 0;
+				}
+
+				// We can only bail when all our async jobs have completed.
+				if (WritingBatch == nullptr && ReadingBatchCount == 0)
+				{
 					break;
 				}
 			}
 
-			bSuccess = AppendChunkData(ChunkPart, WriteBuffers[CurrentFillBuffer], ConstructionError);
+			{
+				TRACE_CPUPROFILER_EVENT_SCOPE(CFFC_Waiting);
+				// We have a bunch of stuff outstanding that will wake us up if something happens.
+				WakeUpDispatchThreadEvent->Wait(WaitTimeMs);
+			}
 
-			FileConstructorStat->OnFileProgress(BuildFilename, NewFile->Tell() + WriteBuffers[CurrentFillBuffer].Num());
-			if (bSuccess)
-			{
-				CountBytesProcessed(ChunkPart.Size);
-				// Wait while paused
-				FileConstructorHelpers::WaitWhilePaused(bIsPaused, bShouldAbort);
-			}
-			// Only report or log if this is the first error
-			else if (InstallerError->HasError() == false)
-			{
-				if (ConstructionError == EConstructionError::MissingChunk)
-				{
-					InstallerAnalytics->RecordConstructionError(BuildFilename, INDEX_NONE, TEXT("Missing Chunk"));
-					UE_LOG(LogBuildPatchServices, Error, TEXT("FBuildPatchFileConstructor: Failed %s due to missing chunk %s"), *BuildFilename, *ChunkPart.Guid.ToString());
-				}
-				else if (ConstructionError == EConstructionError::TrackingError)
-				{
-					InstallerAnalytics->RecordConstructionError(BuildFilename, INDEX_NONE, TEXT("Tracking Error"));
-					UE_LOG(LogBuildPatchServices, Error, TEXT("FBuildPatchFileConstructor: Failed %s due to untracked chunk %s"), *BuildFilename, *ChunkPart.Guid.ToString());
-				}
-			}
-		}
+		} // end looping over
 
-		if (WriteBuffers[CurrentFillBuffer].Num())
-		{
-			if (!FlushToAsyncWriter(*NewFile, HashState))
-			{
-				bSuccess = false;
-				InstallerAnalytics->RecordConstructionError(BuildFilename, INDEX_NONE, TEXT("Serialization Error"));
-				UE_LOG(LogBuildPatchServices, Error, TEXT("FBuildPatchFileConstructor: Failed %s due to serialization error"), *BuildFilename);
-				ConstructionError = EConstructionError::SerializeError;
-			}
-		}
+		bSuccess = bSuccess && !NewFile->IsError();
 
-		// Wait for the last write if there is one
-		if (bWriteJobRunning)
-		{
-			WriteJobCompleteEvent->Wait();
-			bWriteJobRunning = false;
-		}
-
-		bSuccess = !NewFile->IsError();
-
 		// Update this for disk space requirements tracking below on error
 		StartPosition = NewFile->Tell();
@@ -935,6 +2747,8 @@
 		ConstructionError = EConstructionError::CannotCreateFile;
 	}
 
+	bSuccess = bSuccess && BackingStore->CheckNoLocks(false);
+
 	// Check for error state
 	if (!bSuccess)
 	{
@@ -944,10 +2758,24 @@
 			// migrated such that it's when we fail to load an existing chunk (i.e. corruption)
 			// but instead that shows up as a missing chunk.
 
+			// space issue if it looks like that's the problem.			
 
 			uint64 TotalSize = 0;
 			uint64 FreeSize = 0;
+			uint64 LocalDiskSpaceRequired = TNumericLimits<uint64>::Max();
 			if (FPlatformMisc::GetDiskTotalAndFreeSpace(Configuration.InstallDirectory, TotalSize, FreeSize))
 			{
+				LocalDiskSpaceRequired = CalculateInProgressDiskSpaceRequired(FileManifest, StartPosition);
+
+				// If we are delete-during-install this gets more complicated because we'll be freeing up
+				// space as we add.
+				if (Configuration.bDeleteChunkDBFilesAfterUse)
+				{
+					LocalDiskSpaceRequired = CalculateDiskSpaceRequirementsWithDeleteDuringInstall(ConstructionStack);
+				}
+			}
+
+			if (LocalDiskSpaceRequired != TNumericLimits<uint64>::Max())
+			{
 				// We're responding to an actual failure, which would have happened because we literally weren't able
 				// to write our write butter. Because of transient stuff this might not be correct so we double our
@@ -956,21 +2784,8 @@
 				{
 					// We've already failed so it makes sense to reevaluate how much extra we need. 
 					// I'm not sure I like using the same error wording for initial and ongoing disk space failure, but whatevs
-					{
-						uint64 LocalDiskSpaceRequired = CalculateInProgressDiskSpaceRequired(FileManifest, StartPosition);
-
-						// If we are delete-during-install this gets more complicated because we'll be freeing up
-						// space as we add.
-						if (Configuration.bDeleteChunkDBFilesAfterUse)
-						{
-							LocalDiskSpaceRequired = CalculateDiskSpaceRequirementsWithDeleteDuringInstall(ConstructionStack);
-						}
-
-						AvailableDiskSpace.store(FreeSize, std::memory_order_release);
-						RequiredDiskSpace.store(LocalDiskSpaceRequired, std::memory_order_release);
-
-					}
-
+					AvailableDiskSpace.store(FreeSize, std::memory_order_release);
+					RequiredDiskSpace.store(LocalDiskSpaceRequired, std::memory_order_release);
 					ConstructionError = EConstructionError::OutOfDiskSpace;
 				}
 				else
@@ -986,7 +2801,6 @@
 			}
 		}
 
-		// \todo not exactly sure why this only reports on a file creation error?
 		const bool bReportAnalytic = InstallerError->HasError() == false;
 		switch (ConstructionError)
 		{
@@ -1006,27 +2820,58 @@
 			{
 				if (bReportAnalytic)
 				{
-					InstallerAnalytics->RecordConstructionError(BuildFilename, LastError, TEXT("Could Not Create File"));
-					UE_LOG(LogBuildPatchServices, Error, TEXT("FBuildPatchFileConstructor: Could not create %s"), *BuildFilename);
+					InstallerAnalytics->RecordConstructionError(BuildFilename, CreateFilePlatformLastError, TEXT("Could Not Create File"));
 				}
-				InstallerError->SetError(EBuildPatchInstallError::FileConstructionFail, ConstructionErrorCodes::FileCreateFail, LastError);
+				UE_LOG(LogBuildPatchServices, Error, TEXT("FBuildPatchFileConstructor: Could not create %s"), *BuildFilename);
+				InstallerError->SetError(EBuildPatchInstallError::FileConstructionFail, ConstructionErrorCodes::FileCreateFail, CreateFilePlatformLastError);
 				break;
 			}
 		case EConstructionError::MissingChunk:
 			{
+				if (bReportAnalytic)
+				{
+					InstallerAnalytics->RecordConstructionError(BuildFilename, INDEX_NONE, TEXT("Missing Chunk"));
+				}
+
+				UE_LOG(LogBuildPatchServices, Error, TEXT("FBuildPatchFileConstructor: Failed %s due to missing chunk %s (can be 0000 if unknown)"), *BuildFilename, *WriteToString<64>(ErrorContextGuid));
 				InstallerError->SetError(EBuildPatchInstallError::FileConstructionFail, ConstructionErrorCodes::MissingChunkData);
 				break;
 			}
 		case EConstructionError::SerializeError:
 			{
+				if (bReportAnalytic)
+				{
+					InstallerAnalytics->RecordConstructionError(BuildFilename, INDEX_NONE, TEXT("Serialization Error"));
+				}
+				UE_LOG(LogBuildPatchServices, Error, TEXT("FBuildPatchFileConstructor: Failed %s due to serialization error"), *BuildFilename);
 				InstallerError->SetError(EBuildPatchInstallError::FileConstructionFail, ConstructionErrorCodes::SerializationError);
 				break;
 			}
 		case EConstructionError::TrackingError:
 			{
+				if (bReportAnalytic)
+				{
+					InstallerAnalytics->RecordConstructionError(BuildFilename, INDEX_NONE, TEXT("Tracking Error"));
+				}
+				UE_LOG(LogBuildPatchServices, Error, TEXT("FBuildPatchFileConstructor: Failed %s due to untracked chunk %s (can be 0000 if unknown)"), *BuildFilename, *WriteToString<64>(ErrorContextGuid));
 				InstallerError->SetError(EBuildPatchInstallError::FileConstructionFail, ConstructionErrorCodes::TrackingError);
 				break;
 			}
+		case EConstructionError::InternalConsistencyError:
+			{
+				if (bReportAnalytic)
+				{
+					InstallerAnalytics->RecordConstructionError(BuildFilename, INDEX_NONE, TEXT("Internal Consistency Error"));
+				}
+				UE_LOG(LogBuildPatchServices, Error, TEXT("FBuildPatchFileConstructor: Failed %s due to internal consistency checking failure"), *BuildFilename);
+				InstallerError->SetError(EBuildPatchInstallError::FileConstructionFail, ConstructionErrorCodes::InternalConsistencyFailure);
+				break;
+			}
+			case EConstructionError::Aborted:
+			{
+				// We don't set errors on abort.
+				break;
+			}
 		}
 	}
 
@@ -1034,18 +2879,19 @@
 	if (bSuccess)
 	{
 		HashState.Final();
+
+		FSHAHash HashValue;
 		HashState.GetHash(HashValue.Hash);
 		bSuccess = HashValue == FileManifest.FileHash;
 		if (!bSuccess)
 		{
 			ConstructionError = EConstructionError::OutboundDataError;
-			// Only report or log if the first error
+			// Only report if the first error
 			if (InstallerError->HasError() == false)
 			{
 				InstallerAnalytics->RecordConstructionError(BuildFilename, INDEX_NONE, TEXT("Serialised Verify Fail"));
-				UE_LOG(LogBuildPatchServices, Error, TEXT("FBuildPatchFileConstructor: Verify failed after constructing %s"), *BuildFilename);
 			}
-			// Always set
+			UE_LOG(LogBuildPatchServices, Error, TEXT("FBuildPatchFileConstructor: Verify failed after constructing %s"), *BuildFilename);
 			InstallerError->SetError(EBuildPatchInstallError::FileConstructionFail, ConstructionErrorCodes::OutboundCorrupt);
 		}
 	}
@@ -1071,7 +2917,7 @@
 
 	if (bSuccess)
 	{
-		ChunkSource->ReportFileCompletion();
+		ChunkDbSource->ReportFileCompletion(ChunkReferenceTracker->GetRemainingChunkCount());
 	}
 
 	// Delete the staging file if unsuccessful by means of any failure that could leave the file in unknown state.
@@ -1079,58 +2925,33 @@
 	{
 		switch (ConstructionError)
 		{
+			// Errors we expect that we can conceptually resume:
+		case EConstructionError::OutOfDiskSpace:
+		case EConstructionError::MissingChunk:
+		case EConstructionError::Aborted:
+			{
+				break;
+			}
+
+			// Errors we expect there to be issues with the outbound file:
 		case EConstructionError::CannotCreateFile:
 		case EConstructionError::SerializeError:
 		case EConstructionError::TrackingError:
 		case EConstructionError::OutboundDataError:
-			if (!FileSystem->DeleteFile(*NewFilename))
+		case EConstructionError::InternalConsistencyError:
 			{
-				UE_LOG(LogBuildPatchServices, Warning, TEXT("FBuildPatchFileConstructor: Error deleting file: %s (Error Code %i)"), *NewFilename, FPlatformMisc::GetLastError());
+				if (!FileSystem->DeleteFile(*NewFilename))
+				{
+					UE_LOG(LogBuildPatchServices, Warning, TEXT("FBuildPatchFileConstructor: Error deleting file: %s (Error Code %i)"), *NewFilename, FPlatformMisc::GetLastError());
+				}
+				break;
 			}
-			break;
 		}
 	}
 
 	return bSuccess;
 }
 
-bool FBuildPatchFileConstructor::AppendChunkData(const FChunkPart& ChunkPart, TArray<uint8>& DestinationBuffer, EConstructionError& ConstructionError)
-{
-	ConstructionError = EConstructionError::None;
-	
-	FileConstructorStat->OnChunkGet(ChunkPart.Guid);
-	IChunkDataAccess* ChunkDataAccess = nullptr;
-	{
-		TRACE_CPUPROFILER_EVENT_SCOPE(GetChunkData);
-		ChunkDataAccess = ChunkSource->Get(ChunkPart.Guid);
-	}
-	if (ChunkDataAccess != nullptr)
-	{
-		uint8* Data;
-		ChunkDataAccess->GetDataLock(&Data, nullptr);
-
-		uint8* DataStart = &Data[ChunkPart.Offset];
-
-		{
-			TRACE_CPUPROFILER_EVENT_SCOPE(GetChunkData_Copy);
-			DestinationBuffer.Append(DataStart, ChunkPart.Size);
-		}
-
-		ChunkDataAccess->ReleaseDataLock();
-		const bool bPopReferenceOk = ChunkReferenceTracker->PopReference(ChunkPart.Guid);
-		if (!bPopReferenceOk)
-		{
-			ConstructionError = EConstructionError::TrackingError;
-		}
-	}
-	else
-	{
-		// We'd really like to know if this was because it's missing or because it failed in some way.
-		ConstructionError = EConstructionError::MissingChunk;
-	}
-	return ConstructionError == EConstructionError::None;
-}
-
 void FBuildPatchFileConstructor::DeleteDirectoryContents(const FString& RootDirectory)
 {
 	TArray<FString> SubDirNames;
diff --git a/Engine/Source/Runtime/Online/BuildPatchServices/Private/BuildPatchFileConstructor.h b/Engine/Source/Runtime/Online/BuildPatchServices/Private/BuildPatchFileConstructor.h
--- a/Engine/Source/Runtime/Online/BuildPatchServices/Private/BuildPatchFileConstructor.h
+++ b/Engine/Source/Runtime/Online/BuildPatchServices/Private/BuildPatchFileConstructor.h
@@ -3,23 +3,37 @@
 #pragma once
 
 #include "CoreMinimal.h"
+#include "Async/Mutex.h"
 #include "Misc/Guid.h"
 #include "HAL/Runnable.h"
 #include "BuildPatchProgress.h"
 #include "BuildPatchManifest.h"
+#include "Installer/ChunkSource.h"
 #include "Installer/Controllable.h"
 #include "Common/SpeedRecorder.h"
 #include "BuildPatchInstall.h"
 
+
 // Forward declarations
+struct FBatchState;
 class FBuildPatchAppManifest;
 enum class EConstructionError : uint8;
 
 class IBuildInstallerSharedContext;
+class FChunkBackingStore;
 
 namespace BuildPatchServices
 {
-	
+	enum EConstructorChunkLocation : uint8
+	{
+		Install,
+		ChunkDb,
+		Memory,
+		DiskOverflow,
+		Cloud,
+		Retired,
+		COUNT
+	};
 
 	struct FChunkPart;
 	class IFileSystem;
@@ -30,7 +44,10 @@ namespace BuildPatchServices
 	class IFileConstructorStat;
 	class IBuildManifestSet;
 	class IBuildInstallerThread;
-	class IChunkDbChunkSource;
+	class IConstructorChunkDbChunkSource;
+	class IConstructorInstallChunkSource;
+	class IConstructorCloudChunkSource;
+	
 
 	/**
 	 * A struct containing the configuration values for a file constructor.
@@ -55,6 +72,9 @@ namespace BuildPatchServices
 		// The install mode used for this installation.
 		EInstallMode InstallMode;
 
+		// The location where memory overflow will get written to.
+		FString BackingStoreDirectory;
+
 		IBuildInstallerSharedContext* SharedContext;
 
 		bool bDeleteChunkDBFilesAfterUse = false;
@@ -66,6 +86,7 @@ namespace BuildPatchServices
 	 */
 	class FBuildPatchFileConstructor : public IControllable
 	{
+		friend FChunkBackingStore;
 	public:
 
 		/**
@@ -79,9 +100,10 @@ namespace BuildPatchServices
 		 * @param FileConstructorStat       Pointer to the stat class for receiving updates.
 		 */
 		FBuildPatchFileConstructor(
-			FFileConstructorConfig Configuration, IFileSystem* FileSystem, IChunkSource* ChunkSource, 
-			IChunkDbChunkSource* ChunkDbChunkSource, IChunkReferenceTracker* ChunkReferenceTracker, IInstallerError* InstallerError, 
-			IInstallerAnalytics* InstallerAnalytics, IFileConstructorStat* FileConstructorStat);
+			FFileConstructorConfig Configuration, IFileSystem* FileSystem, 
+			IConstructorChunkDbChunkSource* ChunkDbChunkSource, IConstructorCloudChunkSource* CloudChunkSource, IConstructorInstallChunkSource* InstallChunkSource, 
+			IChunkReferenceTracker* ChunkReferenceTracker, IInstallerError* InstallerError, IInstallerAnalytics* InstallerAnalytics, 
+			IFileConstructorStat* FileConstructorStat, TMap<FGuid, EConstructorChunkLocation>&& ChunkLocations);
 
 		/**
 		 * Default Destructor, will delete the allocated Thread
@@ -91,17 +113,12 @@ namespace BuildPatchServices
 		void Run();
 
 		// IControllable interface begin.
-		virtual void SetPaused(bool bInIsPaused) override
-		{
-			bIsPaused = bInIsPaused;
-		}
-
-		virtual void Abort() override
-		{
-			bShouldAbort = true;
-		}
+		virtual void SetPaused(bool bInIsPaused) override;
+		virtual void Abort() override;
 		// IControllable interface end.
 
+		void WakeUpDispatch();
+
 		/**
 		 * Get the disk space that was required to perform the installation. This can change over time and indicates the required
 		 * space to _finish_ the installation from the current state. It is not initialized until after resume is processed and returns
@@ -128,6 +145,8 @@ namespace BuildPatchServices
 
 	private:
 
+		void SetChunkLocation(const FGuid& InGuid, EConstructorChunkLocation InNewLocation);
+
 		/**
 		 * Count additional bytes processed, and set new install progress value
 		 * @param ByteCount		Number of bytes to increment by
@@ -167,14 +186,6 @@ namespace BuildPatchServices
 		 */
 		bool ConstructFileFromChunks(const FString& BuildFilename, const FFileManifest& FileManifest, bool bResumeExisting);
 
-		/**
-		 * Adds the data from a chunk to the given buffer.
-		 * @param ChunkPart          The chunk part details.
-		 * @param ConstructionError  Will be set to the error type that ocurred or EConstructionError::None.
-		 * @return true if no errors were detected
-		 */
-		bool AppendChunkData(const FChunkPart& ChunkPart, TArray<uint8>& DestinationBuffer, EConstructionError& ConstructionError);
-
 		/**
 		 * Delete all contents of a directory
 		 * @param RootDirectory	 	Directory to make empty
@@ -194,10 +205,11 @@ namespace BuildPatchServices
 		// A flag marking whether we should be paused.
 		FThreadSafeBool bIsPaused;
 
-		// A flag marking whether we should abort operations and exit.
+		// A flag marking whether we should abort operations and exit. Always call Abort() to set this.
 		FThreadSafeBool bShouldAbort;
 
 		// A critical section to protect the flags and variables.
+		// \todo this is another thing where if we keep current index in the stack we don't need a mutex.
 		FCriticalSection ThreadLock;
 
 		// A stack of filenames for files that need to be constructed.
@@ -206,9 +218,13 @@ namespace BuildPatchServices
 		// Pointer to the file system.
 		IFileSystem* FileSystem;
 
-		// Pointer to chunk source.
-		IChunkSource* ChunkSource;
-		IChunkDbChunkSource* ChunkDbSource; // can be null if not using.
+		IConstructorChunkDbChunkSource* ChunkDbSource; // can be null if not using.
+		IConstructorInstallChunkSource* InstallSource;
+		IConstructorCloudChunkSource* CloudSource;
+
+		// We always want to know exactly where we think a chunk should be. If it's not there,
+		// we update this list to where it can be found (i.e. cloud)
+		TMap<FGuid, EConstructorChunkLocation> ChunkLocations;
 
 		// Pointer to the chunk reference tracker.
 		IChunkReferenceTracker* ChunkReferenceTracker;
@@ -237,25 +253,62 @@ namespace BuildPatchServices
 		// Event executed before deleting an old installation file.
 		FOnBeforeDeleteFile BeforeDeleteFileEvent;
 
-		//
-		// Async write management.
-		IBuildInstallerThread* WriteJobThread;
-		void WriteJobThreadRun();
+		// Track how much data we expect to have to download
+		uint64 DownloadRequirement = 0;
+		// When we get a failure and switch to the cloud source this will be on any thread. Rather than get fancy
+		// we just lock around it since it's tiny.
+		UE::FMutex DownloadRequirementLock;
+
+		TArray<FEvent*> ThreadWakeups;
+		TArray<TArray<IConstructorChunkSource::FRequestProcessFn>> ThreadJobPostings;
+		TArray<UE::FMutex> ThreadJobPostingLocks;
+		TArray<FEvent*> ThreadCompleteEvents;
+		TArray<IBuildInstallerThread*> Threads;
+
+		void QueueGenericThreadTask(int32 ThreadIndex, TUniqueFunction<void(bool)>&& Task);
+		void GenericThreadFn(int32 ThreadIndex);
+		int8 ThreadAssignments[EConstructorChunkLocation::COUNT];
+		int8 WriteThreadIndex = -1;
+
+		TUniquePtr<FChunkBackingStore> BackingStore;
 
-		// We ping pong between two buffers, filling/hashing one, and writing the other.
 		TArray<uint8> WriteBuffers[2];
-		int32 CurrentFillBuffer = 0;
-		uint32 WriteBufferSize = (4 << 20); // Default write buffer size 4MB.
+		uint32 WriteBufferSize = (16 << 20); // Default write buffer size.
+
+		// Fire this to wake up the main thread to process completed tasks.
+		FEvent* WakeUpDispatchThreadEvent = nullptr;
 
-		FEvent* WriteJobCompleteEvent = nullptr;
-		FEvent* WriteJobStartEvent = nullptr;
-		TArray<uint8>* WriteJobBufferToWrite = nullptr;
-		FArchive* WriteJobArchive = nullptr;
-		std::atomic_bool bWriteJobCompleted = false; // Only set to true if the Serialize() call was completed.
-		bool bWriteJobRunning = false; // Foreground thread only - have we dispatched a job?
+		IConstructorChunkSource::FRequestProcessFn CreateWriteRequest(FArchive* File, FBatchState& Batch);
 
 		// Where we are in the chunk consumption list after each file.
 		TArray<int32> FileCompletionPositions;
+
+		int32 StartReadBatch(const FFileManifest& FileManifest, int32 StartFileChunkPartIndex, FBatchState& Buffer);
+		void CompleteReadBatch(const FFileManifest& FileManifest, FBatchState& Buffer);
+		void RequestCompletedFn(const FGuid& Guid, bool bAborted, bool bFailedToRead, void* UserPtr);
+
+		std::atomic_int32_t PendingHarvestRequests = 0;
+		void ChunkHarvestCompletedFn(const FGuid& Guid, bool bAborted, bool bFailedToRead, void* UserPtr);
+
+	public:
+		struct FBackingStoreStats
+		{
+			uint64 DiskPeakUsageBytes=0;
+			uint64 MemoryPeakUsageBytes=0;
+			uint64 MemoryLimitBytes=0;
+			uint32 DiskLoadFailureCount=0;
+			uint32 DiskLostChunkCount=0;
+			uint32 DiskChunkLoadCount=0;
+		};
+
+		// This isn't safe to call during operations as values are changing on other threads.
+		FBackingStoreStats GetBackingStoreStats()
+		{
+			return BackingStoreStats;
+		}
+	private:
+
+		FBackingStoreStats BackingStoreStats;
 	};
 
 	/**
diff --git a/Engine/Source/Runtime/Online/BuildPatchServices/Private/BuildPatchInstaller.cpp b/Engine/Source/Runtime/Online/BuildPatchServices/Private/BuildPatchInstaller.cpp
--- a/Engine/Source/Runtime/Online/BuildPatchServices/Private/BuildPatchInstaller.cpp
+++ b/Engine/Source/Runtime/Online/BuildPatchServices/Private/BuildPatchInstaller.cpp
@@ -26,7 +26,6 @@
 #include "Common/FileSystem.h"
 #include "Common/ChunkDataSizeProvider.h"
 #include "Installer/InstallerError.h"
-#include "Installer/DiskChunkStore.h"
 #include "Installer/CloudChunkSource.h"
 #include "Installer/DownloadConnectionCount.h"
 #include "Installer/MemoryChunkStore.h"
@@ -35,7 +34,6 @@
 #include "Installer/DownloadService.h"
 #include "Installer/ChunkDbChunkSource.h"
 #include "Installer/InstallChunkSource.h"
-#include "Installer/ChainedChunkSource.h"
 #include "Installer/Verifier.h"
 #include "Installer/FileAttribution.h"
 #include "Installer/InstallerAnalytics.h"
@@ -45,7 +43,6 @@
 #include "Installer/MessagePump.h"
 #include "Installer/OptimisedDelta.h"
 #include "Installer/Statistics/MemoryChunkStoreStatistics.h"
-#include "Installer/Statistics/DiskChunkStoreStatistics.h"
 #include "Installer/Statistics/DownloadServiceStatistics.h"
 #include "Installer/Statistics/ChunkDbChunkSourceStatistics.h"
 #include "Installer/Statistics/InstallChunkSourceStatistics.h"
@@ -143,7 +140,7 @@ namespace InstallerHelpers
 		UE_LOG(LogBuildPatchServices, Log, TEXT("Build Stat: NumBadDownloads: %u"), BuildStats.NumBadDownloads);
 		UE_LOG(LogBuildPatchServices, Log, TEXT("Build Stat: NumAbortedDownloads: %u"), BuildStats.NumAbortedDownloads);
 		UE_LOG(LogBuildPatchServices, Log, TEXT("Build Stat: NumRecycleFailures: %u"), BuildStats.NumRecycleFailures);
-		UE_LOG(LogBuildPatchServices, Log, TEXT("Build Stat: NumChunksStoreBooted: %u"), BuildStats.NumChunksStoreBooted);
+		UE_LOG(LogBuildPatchServices, Log, TEXT("Build Stat: NumDriveStoreLostChunks: %u"), BuildStats.NumDriveStoreLostChunks);
 		UE_LOG(LogBuildPatchServices, Log, TEXT("Build Stat: NumDriveStoreChunkLoads: %u"), BuildStats.NumDriveStoreChunkLoads);
 		UE_LOG(LogBuildPatchServices, Log, TEXT("Build Stat: NumDriveStoreLoadFailures: %u"), BuildStats.NumDriveStoreLoadFailures);
 		UE_LOG(LogBuildPatchServices, Log, TEXT("Build Stat: NumChunkDbChunksFailed: %u"), BuildStats.NumChunkDbChunksFailed);
@@ -151,10 +148,10 @@ namespace InstallerHelpers
 		UE_LOG(LogBuildPatchServices, Log, TEXT("Build Stat: ActiveRequestCountPeak : %u"), BuildStats.ActiveRequestCountPeak);
 		UE_LOG(LogBuildPatchServices, Log, TEXT("Build Stat: AverageDownloadSpeed: %s bytes (%s, %s) /sec"), *FText::AsNumber(BuildStats.AverageDownloadSpeed).ToString(), *FText::AsMemory(BuildStats.AverageDownloadSpeed, EMemoryUnitStandard::SI).ToString(), *FText::AsMemory(BuildStats.AverageDownloadSpeed, EMemoryUnitStandard::IEC).ToString());
 		UE_LOG(LogBuildPatchServices, Log, TEXT("Build Stat: PeakDownloadSpeed: %s bytes (%s, %s) /sec"), *FText::AsNumber(BuildStats.PeakDownloadSpeed).ToString(), *FText::AsMemory(BuildStats.PeakDownloadSpeed, EMemoryUnitStandard::SI).ToString(), *FText::AsMemory(BuildStats.PeakDownloadSpeed, EMemoryUnitStandard::IEC).ToString());
-		UE_LOG(LogBuildPatchServices, Log, TEXT("Build Stat: TotalReadData: %llu"), BuildStats.TotalReadData);
+		UE_LOG(LogBuildPatchServices, Log, TEXT("Build Stat: TotalReadData: %s"), *FText::AsNumber(BuildStats.TotalReadData).ToString());
 		UE_LOG(LogBuildPatchServices, Log, TEXT("Build Stat: AverageDiskReadSpeed: %s bytes (%s, %s) /sec"), *FText::AsNumber(BuildStats.AverageDiskReadSpeed).ToString(), *FText::AsMemory(BuildStats.AverageDiskReadSpeed, EMemoryUnitStandard::SI).ToString(), *FText::AsMemory(BuildStats.AverageDiskReadSpeed, EMemoryUnitStandard::IEC).ToString());
 		UE_LOG(LogBuildPatchServices, Log, TEXT("Build Stat: PeakDiskReadSpeed: %s bytes (%s, %s) /sec"), *FText::AsNumber(BuildStats.PeakDiskReadSpeed).ToString(), *FText::AsMemory(BuildStats.PeakDiskReadSpeed, EMemoryUnitStandard::SI).ToString(), *FText::AsMemory(BuildStats.PeakDiskReadSpeed, EMemoryUnitStandard::IEC).ToString());
-		UE_LOG(LogBuildPatchServices, Log, TEXT("Build Stat: TotalWrittenData: %llu"), BuildStats.TotalWrittenData);
+		UE_LOG(LogBuildPatchServices, Log, TEXT("Build Stat: TotalWrittenData: %s"), *FText::AsNumber(BuildStats.TotalWrittenData).ToString());
 		UE_LOG(LogBuildPatchServices, Log, TEXT("Build Stat: AverageDiskWriteSpeed: %s bytes (%s, %s) /sec"), *FText::AsNumber(BuildStats.AverageDiskWriteSpeed).ToString(), *FText::AsMemory(BuildStats.AverageDiskWriteSpeed, EMemoryUnitStandard::SI).ToString(), *FText::AsMemory(BuildStats.AverageDiskWriteSpeed, EMemoryUnitStandard::IEC).ToString());
 		UE_LOG(LogBuildPatchServices, Log, TEXT("Build Stat: PeakDiskWriteSpeed: %s bytes (%s, %s) /sec"), *FText::AsNumber(BuildStats.PeakDiskWriteSpeed).ToString(), *FText::AsMemory(BuildStats.PeakDiskWriteSpeed, EMemoryUnitStandard::SI).ToString(), *FText::AsMemory(BuildStats.PeakDiskWriteSpeed, EMemoryUnitStandard::IEC).ToString());
 		UE_LOG(LogBuildPatchServices, Log, TEXT("Build Stat: NumFilesConstructed: %u"), BuildStats.NumFilesConstructed);
@@ -188,13 +185,11 @@ namespace InstallerHelpers
 		UE_LOG(LogBuildPatchServices, Log, TEXT("Build Stat: OkDownloadHealthTime: %s"), *FPlatformTime::PrettyTime(BuildStats.OkDownloadHealthTime));
 		UE_LOG(LogBuildPatchServices, Log, TEXT("Build Stat: PoorDownloadHealthTime: %s"), *FPlatformTime::PrettyTime(BuildStats.PoorDownloadHealthTime));
 		UE_LOG(LogBuildPatchServices, Log, TEXT("Build Stat: DisconnectedDownloadHealthTime: %s"), *FPlatformTime::PrettyTime(BuildStats.DisconnectedDownloadHealthTime));
-		UE_LOG(LogBuildPatchServices, Log, TEXT("Build Stat: AverageMemoryStoreUse: %.1f"), BuildStats.AverageMemoryStoreUse);
-		UE_LOG(LogBuildPatchServices, Log, TEXT("Build Stat: PeakMemoryStoreUse: %u"), BuildStats.PeakMemoryStoreUse);
-		UE_LOG(LogBuildPatchServices, Log, TEXT("Build Stat: AverageMemoryStoreRetained: %.1f"), BuildStats.AverageMemoryStoreRetained);
-		UE_LOG(LogBuildPatchServices, Log, TEXT("Build Stat: PeakMemoryStoreRetained: %u"), BuildStats.PeakMemoryStoreRetained);
-		UE_LOG(LogBuildPatchServices, Log, TEXT("Build Stat: MemoryStoreSize: %u"), BuildStats.MemoryStoreSize);
-		UE_LOG(LogBuildPatchServices, Log, TEXT("Build Stat: ProcessRequiredDiskSpace: %llu"), BuildStats.ProcessRequiredDiskSpace);
-		UE_LOG(LogBuildPatchServices, Log, TEXT("Build Stat: ProcessAvailableDiskSpace: %llu"), BuildStats.ProcessAvailableDiskSpace);
+		UE_LOG(LogBuildPatchServices, Log, TEXT("Build Stat: DriveStorePeakBytes: %u"), BuildStats.DriveStorePeakBytes);
+		UE_LOG(LogBuildPatchServices, Log, TEXT("Build Stat: MemoryStoreSizePeakBytes: %s"), *FText::AsNumber(BuildStats.MemoryStoreSizePeakBytes).ToString());
+		UE_LOG(LogBuildPatchServices, Log, TEXT("Build Stat: MemoryStoreSizeLimitBytes: %s"), *FText::AsNumber(BuildStats.MemoryStoreSizeLimitBytes).ToString());
+		UE_LOG(LogBuildPatchServices, Log, TEXT("Build Stat: ProcessRequiredDiskSpace: %s"), *FText::AsNumber(BuildStats.ProcessRequiredDiskSpace).ToString());
+		UE_LOG(LogBuildPatchServices, Log, TEXT("Build Stat: ProcessAvailableDiskSpace: %s"), *FText::AsNumber(BuildStats.ProcessAvailableDiskSpace).ToString());
 	}
 
 	const TCHAR* GetActionTypeLog(const FInstallerAction& InstallerAction)
@@ -412,12 +407,6 @@ namespace BuildPatchServices
 		}
 	};
 
-	enum class EMemoryStore : int32
-	{
-		Install = 0,
-		Cloud
-	};
-
 	static uint64 DetermineInstallMaxDiskSizeIfDeletingChunkDbs(const TArray<FString>& ChunkDbFiles, IBuildManifestSet* ManifestSet, IFileSystem* FileSystem, const FString& InstallDirectory, const TArray<FString>& CorruptFiles, bool bIsPrereqOnly);
 
 	/* FBuildPatchInstaller implementation
@@ -454,8 +443,6 @@ namespace BuildPatchServices
 		, Analytics(MoveTemp(InAnalytics))
 		, InstallerAnalytics(FInstallerAnalyticsFactory::Create(Analytics.Get()))
 		, FileOperationTracker(Configuration.bTrackFileOperations ? FFileOperationTrackerFactory::Create(FTSTicker::GetCoreTicker()) : FFileOperationTrackerFactory::CreateNull())
-		, MemoryChunkStoreStatistics(FMemoryChunkStoreStatisticsFactory::Create(FileOperationTracker.Get()))
-		, DiskChunkStoreStatistics(FDiskChunkStoreStatisticsFactory::Create(InstallerAnalytics.Get(), FileOperationTracker.Get()))
 		, DownloadSpeedRecorder(FSpeedRecorderFactory::Create())
 		, DiskReadSpeedRecorder(FSpeedRecorderFactory::Create())
 		, DiskWriteSpeedRecorder(FSpeedRecorderFactory::Create())
@@ -482,12 +469,14 @@ namespace BuildPatchServices
 		TArray<FBuildPatchAppManifestPtr> Manifests;
 		Manifests.Reserve(Configuration.InstallerActions.Num() * 2);
 
+		bool bHasCurrentManifest = false;
 		for (const FInstallerAction& InstallerAction : Configuration.InstallerActions)
 		{
 			FBuildPatchInstallerAction& BuildPatchInstallerAction = InstallerActions.Emplace_GetRef(InstallerAction);
 			// Make sure existing manifests are added to installation info.
 			if (BuildPatchInstallerAction.TryGetCurrentManifest())
 			{
+				bHasCurrentManifest = true;
 				InstallationInfo.Add(InstallDirectory / BuildPatchInstallerAction.GetInstallSubdirectory(), BuildPatchInstallerAction.GetSharedCurrentManifest());
 			}
 			// Cache chunk sizes too
@@ -507,7 +496,7 @@ namespace BuildPatchServices
 		{
 			Configuration.SharedContext = FBuildInstallerSharedContextFactory::Create(TEXT("BuildPatchInstaller"));
 			const bool bUseChunkDBs = !Configuration.ChunkDatabaseFiles.IsEmpty();
-			const uint32 NumExpectedThreads = Configuration.SharedContext->NumThreadsPerInstaller(bUseChunkDBs);
+			const uint32 NumExpectedThreads = Configuration.SharedContext->NumThreadsPerInstaller(bUseChunkDBs, bHasCurrentManifest);
 			Configuration.SharedContext->PreallocateThreads(NumExpectedThreads);
 		}
 	}
@@ -596,16 +585,6 @@ namespace BuildPatchServices
 		return VerifierStatistics.Get();
 	}
 
-	const IMemoryChunkStoreStatistics* FBuildPatchInstaller::GetMemoryChunkStoreStatistics() const
-	{
-		return MemoryChunkStoreStatistics.Get();
-	}
-
-	const IDiskChunkStoreStatistics* FBuildPatchInstaller::GetDiskChunkStoreStatistics() const
-	{
-		return DiskChunkStoreStatistics.Get();
-	}
-
 	const FBuildInstallerConfiguration& FBuildPatchInstaller::GetConfiguration() const
 	{
 		return Configuration;
@@ -804,6 +783,8 @@ namespace BuildPatchServices
 
 	uint32 FBuildPatchInstaller::Run()
 	{
+		TRACE_CPUPROFILER_EVENT_SCOPE(BuildPatchInstaller::Run);
+
 		// Make sure this function can never be parallelized
 		static FCriticalSection SingletonFunctionLockCS;
 		const bool bShouldLock = !Configuration.bAllowConcurrentExecution;
@@ -859,7 +840,6 @@ namespace BuildPatchServices
 				{
 					// Inform file operation tracker of the selected manifest.
 					FileOperationTracker->OnManifestSelection(ManifestSet.Get());
-					MemoryChunkStoreStatistics->SetMultipleReferencedChunk(InstallerHelpers::GetMultipleReferencedChunks(ManifestSet.Get()));
 
 					// Run the install
 					bool bInstallSuccess = RunInstallation(CorruptFiles);
@@ -1017,57 +997,21 @@ namespace BuildPatchServices
 		return false;
 	}
 
-	FChunkDbSourceConfig FBuildPatchInstaller::BuildChunkDbSourceConfig()
+	static FChunkDbSourceConfig BuildChunkDbSourceConfig(const FBuildInstallerConfiguration& InConfiguration)
 	{
-		FChunkDbSourceConfig ChunkDbSourceConfig(Configuration.ChunkDatabaseFiles);
-
-		ChunkDbSourceConfig.SharedContext = Configuration.SharedContext.Get();
-
-		// Load batch fetch config.
-		GConfig->GetInt(TEXT("Portal.BuildPatch"), TEXT("ChunkDbSourcePreFetchMinimum"), ChunkDbSourceConfig.PreFetchMinimum, GEngineIni);
-		GConfig->GetInt(TEXT("Portal.BuildPatch"), TEXT("ChunkDbSourcePreFetchMaximum"), ChunkDbSourceConfig.PreFetchMaximum, GEngineIni);
-		ChunkDbSourceConfig.PreFetchMinimum = FMath::Clamp<int32>(ChunkDbSourceConfig.PreFetchMinimum, 1, 1000);
-		ChunkDbSourceConfig.PreFetchMaximum = FMath::Clamp<int32>(ChunkDbSourceConfig.PreFetchMaximum, ChunkDbSourceConfig.PreFetchMinimum, 1000);
-
-		// Load reopen retry time.
-		GConfig->GetFloat(TEXT("Portal.BuildPatch"), TEXT("ChunkDbSourceChunkDbOpenRetryTime"), ChunkDbSourceConfig.ChunkDbOpenRetryTime, GEngineIni);
-		ChunkDbSourceConfig.ChunkDbOpenRetryTime = FMath::Clamp<float>(ChunkDbSourceConfig.ChunkDbOpenRetryTime, 0.5f, 60.0f);
-
-		ChunkDbSourceConfig.bDeleteChunkDBAfterUse = Configuration.bDeleteChunkDbFilesAfterUse;
-
+		FChunkDbSourceConfig ChunkDbSourceConfig(InConfiguration.ChunkDatabaseFiles);
+		ChunkDbSourceConfig.bDeleteChunkDBAfterUse = InConfiguration.bDeleteChunkDbFilesAfterUse;
 		return ChunkDbSourceConfig;
 	}
 
-	FInstallSourceConfig FBuildPatchInstaller::BuildInstallSourceConfig(TSet<FGuid> ChunkIgnoreSet)
-	{
-		FInstallSourceConfig InstallSourceConfig;
-		InstallSourceConfig.ChunkIgnoreSet = MoveTemp(ChunkIgnoreSet);
-
-		// Load batch fetch config
-		GConfig->GetInt(TEXT("Portal.BuildPatch"), TEXT("InstallSourceBatchFetchMinimum"), InstallSourceConfig.BatchFetchMinimum, GEngineIni);
-		GConfig->GetInt(TEXT("Portal.BuildPatch"), TEXT("InstallSourceBatchFetchMaximum"), InstallSourceConfig.BatchFetchMaximum, GEngineIni);
-		InstallSourceConfig.BatchFetchMinimum = FMath::Clamp<int32>(InstallSourceConfig.BatchFetchMinimum, 1, 1000);
-		InstallSourceConfig.BatchFetchMaximum = FMath::Clamp<int32>(InstallSourceConfig.BatchFetchMaximum, InstallSourceConfig.BatchFetchMinimum, 1000);
-
-		return InstallSourceConfig;
-	}
-
-	FCloudSourceConfig FBuildPatchInstaller::BuildCloudSourceConfig()
+	static FConstructorCloudChunkSourceConfig BuildCloudSourceConfig(const FBuildInstallerConfiguration& InConfiguration)
 	{
-		FCloudSourceConfig CloudSourceConfig(Configuration.CloudDirectories);
-
-		CloudSourceConfig.SharedContext = Configuration.SharedContext.Get();
+		FConstructorCloudChunkSourceConfig CloudSourceConfig(InConfiguration.CloudDirectories);
 
 		// Load max download retry count from engine config.
 		GConfig->GetInt(TEXT("Portal.BuildPatch"), TEXT("ChunkRetries"), CloudSourceConfig.MaxRetryCount, GEngineIni);
 		CloudSourceConfig.MaxRetryCount = FMath::Clamp<int32>(CloudSourceConfig.MaxRetryCount, -1, 1000);
 
-		// Load prefetch config
-		GConfig->GetInt(TEXT("Portal.BuildPatch"), TEXT("CloudSourcePreFetchMinimum"), CloudSourceConfig.PreFetchMinimum, GEngineIni);
-		GConfig->GetInt(TEXT("Portal.BuildPatch"), TEXT("CloudSourcePreFetchMaximum"), CloudSourceConfig.PreFetchMaximum, GEngineIni);
-		CloudSourceConfig.PreFetchMinimum = FMath::Clamp<int32>(CloudSourceConfig.PreFetchMinimum, 1, 1000);
-		CloudSourceConfig.PreFetchMaximum = FMath::Clamp<int32>(CloudSourceConfig.PreFetchMaximum, CloudSourceConfig.PreFetchMinimum, 1000);
-
 		// Load retry times from engine config.
 		TArray<FString> ConfigStrings;
 		GConfig->GetArray(TEXT("Portal.BuildPatch"), TEXT("RetryTimes"), ConfigStrings, GEngineIni);
@@ -1109,9 +1053,6 @@ namespace BuildPatchServices
 		GConfig->GetFloat(TEXT("Portal.BuildPatch"), TEXT("DisconnectedDelay"), CloudSourceConfig.DisconnectedDelay, GEngineIni);
 		CloudSourceConfig.DisconnectedDelay = FMath::Clamp<float>(CloudSourceConfig.DisconnectedDelay, 1.0f, 30.0f);
 
-		// We tell the cloud source to only start downloads once it receives the first get call.
-		CloudSourceConfig.bBeginDownloadsOnFirstGet = true;
-
 		return CloudSourceConfig;
 	}
 
@@ -1200,7 +1141,8 @@ namespace BuildPatchServices
 		{
 			FilesToConstruct = OutdatedFiles.Intersect(TaggedFiles);
 		}
-		FilesToConstruct.Sort(TLess<FString>());
+		// This is sorted coming in and needs to stay in that order to pass BPT test suite
+		//FilesToConstruct.Sort(TLess<FString>());
 	}
 
 	//
@@ -1274,7 +1216,7 @@ namespace BuildPatchServices
 		}
 
 		TArray<uint64> ChunkDbSizesAtPosition;
-		uint64 TotalChunkDbSize = IChunkDbChunkSource::GetChunkDbSizesAtIndexes(ChunkDbFiles, FileSystem, ReferenceChain, FileCompletionPositions, ChunkDbSizesAtPosition);
+		uint64 TotalChunkDbSize = IConstructorChunkDbChunkSource::GetChunkDbSizesAtIndexes(ChunkDbFiles, FileSystem, ReferenceChain, FileCompletionPositions, ChunkDbSizesAtPosition);
 
 		uint64 MaxDiskSize = FBuildPatchUtils::CalculateDiskSpaceRequirementsWithDeleteDuringInstall(
 			FilesToConstruct, 0, ManifestSet,
@@ -1372,64 +1314,6 @@ namespace BuildPatchServices
 			}
 		}
 
-		// Default chunk store size to tie in with the default prefetch maxes for source configs.
-		const int32 DefaultChunkDbMaxRead = FChunkDbSourceConfig({}).PreFetchMaximum;
-		const int32 DefaultInstallMaxRead = FInstallSourceConfig().BatchFetchMaximum;
-		const int32 DefaultCloudMaxRead = FCloudSourceConfig({}).PreFetchMaximum;
-		int32 ChunkStoreMemorySize = DefaultCloudMaxRead + DefaultChunkDbMaxRead + DefaultInstallMaxRead;
-		
-		// Add 50% because we dont evict chunks if they can be used again, so we end up dropping to disk unecessarily.
-		ChunkStoreMemorySize += ChunkStoreMemorySize / 2;
-
-		// Load overridden size from config.
-		if (!GConfig->GetInt(TEXT("Portal.BuildPatch"), TEXT("ChunkStoreMemorySize"), ChunkStoreMemorySize, GEngineIni))
-		{
-			// If we didn't get the new value, check for legacy ones.
-			int32 CloudChunkStoreMemorySize = DefaultCloudMaxRead + DefaultChunkDbMaxRead;
-			int32 InstallChunkStoreMemorySize = DefaultInstallMaxRead;
-			const bool bLoadedCloudSize = GConfig->GetInt(TEXT("Portal.BuildPatch"), TEXT("CloudChunkStoreMemorySize"), CloudChunkStoreMemorySize, GEngineIni);
-			const bool bLoadedInstallSize = GConfig->GetInt(TEXT("Portal.BuildPatch"), TEXT("InstallChunkStoreMemorySize"), InstallChunkStoreMemorySize, GEngineIni);
-			if (bLoadedCloudSize || bLoadedInstallSize)
-			{
-				ChunkStoreMemorySize = CloudChunkStoreMemorySize + InstallChunkStoreMemorySize;
-			}
-		}
-
-		// Clamp to sensible limits.
-		ChunkStoreMemorySize = FMath::Clamp<int32>(ChunkStoreMemorySize, 64, 2048);
-
-		// Check for new config - optionally fill memory up to a limit.
-		int32 ChunkStoreMemoryHeadRoomMB = 2000;
-		if (GConfig->GetInt(TEXT("Portal.BuildPatch"), TEXT("ChunkStoreMemoryHeadRoomMB"), ChunkStoreMemoryHeadRoomMB, GEngineIni) &&
-			ChunkStoreMemoryHeadRoomMB > 500) // We have to leave at least SOME room.
-		{
-			// Memory store works in chunks and we don't know how big the chunks are.
-			int32 ChunkSizeEach = 1 << 20; // default to 1MB
-			GConfig->GetInt(TEXT("Portal.BuildPatch"), TEXT("ChunkStoreMemoryChunkSizeBytes"), ChunkSizeEach, GEngineIni);
-
-			uint64 AvailableMem = FPlatformMemory::GetStats().AvailablePhysical;
-
-			uint64 RequestedHeadRoom = ((uint64)ChunkStoreMemoryHeadRoomMB << 20);
-
-			if (RequestedHeadRoom < AvailableMem)
-			{
-				uint64 MemoryStoreMem = AvailableMem - RequestedHeadRoom;
-				ChunkStoreMemorySize = MemoryStoreMem / ChunkSizeEach;
-
-				UE_LOG(LogBPSInstallerConfig, Display, TEXT("ChunkStoreMemoryHeadroom: Using %d as a result of %llu requested out of %llu available"), ChunkStoreMemorySize, RequestedHeadRoom, AvailableMem);
-			}
-			else
-			{
-				// Cap at available.
-				if (ChunkStoreMemorySize > (AvailableMem / ChunkSizeEach))
-				{
-					ChunkStoreMemorySize = AvailableMem / ChunkSizeEach;
-				}
-
-				UE_LOG(LogBPSInstallerConfig, Display, TEXT("ChunkStoreMemoryHeadroom: Requested more than available - Using %d as a result of %llu requested out of %llu available"), ChunkStoreMemorySize, RequestedHeadRoom, AvailableMem);
-			}
-		}
-
 		// Cache the last download requirement in case we are running a retry.
 		PreviousTotalDownloadRequired.Add(CloudChunkSourceStatistics->GetRequiredDownloadSize());
 		// Reset so that we don't double count data.
@@ -1446,84 +1330,74 @@ namespace BuildPatchServices
 				ManifestSet.Get(),
 				FilesToConstruct));
 			TSet<FGuid> ReferencedChunks = ChunkReferenceTracker->GetReferencedChunks();
-			TUniquePtr<IChunkEvictionPolicy> MemoryEvictionPolicy(FChunkEvictionPolicyFactory::Create(
-				ChunkReferenceTracker.Get()));
-			TUniquePtr<IDiskChunkStore> DiskOverflowStore;
-			bool bUseDiskOverflowStore = true;
-			GConfig->GetBool(TEXT("BuildPatchServices"), TEXT("bEnableDiskOverflowStore"), bUseDiskOverflowStore, GEngineIni);
-			UE_LOG(LogBPSInstallerConfig, Display, TEXT("DiskOverflowStore is: %s"), bUseDiskOverflowStore ? TEXT("Enabled") : TEXT("Disabled"));
-
-			if (bUseDiskOverflowStore)
-			{
-				FDiskChunkStoreConfig DiskChunkStoreConfig(DataStagingDir);
-				DiskChunkStoreConfig.SharedContext = Configuration.SharedContext.Get();
-				DiskOverflowStore.Reset(FDiskChunkStoreFactory::Create(
-					FileSystem.Get(),
-					ChunkDataSerialization.Get(),
-					DiskChunkStoreStatistics.Get(),
-					MoveTemp(DiskChunkStoreConfig)));
-			}
-
-			TUniquePtr<IMemoryChunkStore> CloudChunkStore(FMemoryChunkStoreFactory::Create(
-				ChunkStoreMemorySize,
-				MemoryEvictionPolicy.Get(),
-				DiskOverflowStore.Get(),
-				MemoryChunkStoreStatistics.Get(),
-				ChunkReferenceTracker.Get()));
-
+			
 			// Add a source for pulling from chunk "databases", basically tarballs of chunks.
-			TUniquePtr<IChunkDbChunkSource> ChunkDbChunkSource(FChunkDbChunkSourceFactory::Create(
-				BuildChunkDbSourceConfig(),
-				Platform.Get(),
+			TArray<FGuid> ChunkAccessOrderedList;
+			ChunkReferenceTracker->CopyOutOrderedUseList(ChunkAccessOrderedList);
+			TUniquePtr<IConstructorChunkDbChunkSource> ChunkDbChunkSource(IConstructorChunkDbChunkSource::CreateChunkDbSource(
+				BuildChunkDbSourceConfig(Configuration),
 				FileSystem.Get(),
-				CloudChunkStore.Get(),
-				ChunkReferenceTracker.Get(),
+				ChunkAccessOrderedList,
 				ChunkDataSerialization.Get(),
-				MessagePump.Get(),
-				InstallerError.Get(),
 				ChunkDbChunkSourceStatistics.Get()));
 
 			// Add a source for pulling from the existing directory (for patching). This uses the manifest for the
 			// currently deployed files to create a list of available chunks to read from.
-			TUniquePtr<IInstallChunkSource> InstallChunkSource(FInstallChunkSourceFactory::Create(
-				BuildInstallSourceConfig(ChunkDbChunkSource->GetAvailableChunks()),
+			TUniquePtr<IConstructorInstallChunkSource> InstallChunkSource(IConstructorInstallChunkSource::CreateInstallSource(
 				FileSystem.Get(),
-				CloudChunkStore.Get(),
-				ChunkReferenceTracker.Get(),
-				InstallerError.Get(),
 				InstallChunkSourceStatistics.Get(),
 				InstallationInfo,
-				ManifestSet.Get()));
+				ReferencedChunks));
 
 			// Sub out the chunks we have available on disk - anything we don't we have to download.
-			TSet<FGuid> InitialDownloadChunks = ReferencedChunks.Difference(InstallChunkSource->GetAvailableChunks()).Difference(ChunkDbChunkSource->GetAvailableChunks());
-			FileOperationTracker->OnDataStateUpdate(ReferencedChunks.Intersect(ChunkDbChunkSource->GetAvailableChunks()), EFileOperationState::PendingLocalChunkDbData);
-			FileOperationTracker->OnDataStateUpdate(ReferencedChunks.Intersect(InstallChunkSource->GetAvailableChunks()).Difference(ChunkDbChunkSource->GetAvailableChunks()), EFileOperationState::PendingLocalInstallData);
-			FileOperationTracker->OnDataStateUpdate(InitialDownloadChunks, EFileOperationState::PendingRemoteCloudData);
+			const TSet<FGuid>& ChunkDbChunksAvailable = ChunkDbChunkSource->GetAvailableChunks();
+			const TSet<FGuid>& InstallChunksAvailable = InstallChunkSource->GetAvailableChunks();
+			
+			TSet<FGuid> CloudChunks = ReferencedChunks.Difference(InstallChunksAvailable).Difference(ChunkDbChunksAvailable);
+			TSet<FGuid> ChunkDbChunks = ReferencedChunks.Intersect(ChunkDbChunksAvailable);
+
+			// Note ordering - if a chunk is available via install and chunkdb, we take it from chunkdb.
+			TSet<FGuid> InstallChunks = ReferencedChunks.Intersect(InstallChunksAvailable).Difference(ChunkDbChunksAvailable);
+
+			if (InstallChunks.Num() == 0)
+			{
+				InstallChunkSource.Reset();
+			}
+
+			// Set up our chunk location tracking.
+			int32 DownloadChunkCount = CloudChunks.Num();
+			TMap<FGuid, EConstructorChunkLocation> ChunkLocations;	
+			{
+				for (const FGuid& InstallChunk : InstallChunks)
+				{
+					ChunkLocations.Add(InstallChunk, EConstructorChunkLocation::Install);
+				}
+				for (const FGuid& ChunkDbChunk : ChunkDbChunks)
+				{
+					ChunkLocations.Add(ChunkDbChunk, EConstructorChunkLocation::ChunkDb);
+				}
+				for (const FGuid& CloudChunk : CloudChunks)
+				{
+					ChunkLocations.Add(CloudChunk, EConstructorChunkLocation::Cloud);
+				}
+			}
+
+			FileOperationTracker->OnDataStateUpdate(ChunkDbChunks, EFileOperationState::PendingLocalChunkDbData);
+			FileOperationTracker->OnDataStateUpdate(InstallChunks, EFileOperationState::PendingLocalInstallData);
+			FileOperationTracker->OnDataStateUpdate(CloudChunks, EFileOperationState::PendingRemoteCloudData);
 
 			// Add the source for downloading chunks we don't have on disk.
 			TUniquePtr<IDownloadConnectionCount> DownloadConnectionCount(FDownloadConnectionCountFactory::Create(BuildConnectionCountConfig(), DownloadServiceStatistics.Get()));
-			TUniquePtr<ICloudChunkSource> CloudChunkSource(FCloudChunkSourceFactory::Create(
-				BuildCloudSourceConfig(),
-				Platform.Get(),
-				CloudChunkStore.Get(),
-				DownloadService.Get(),
-				ChunkReferenceTracker.Get(),
-				ChunkDataSerialization.Get(),
-				MessagePump.Get(),
-				InstallerError.Get(),
-				DownloadConnectionCount.Get(),
-				CloudChunkSourceStatistics.Get(),
-				ManifestSet.Get(),
-				MoveTemp(InitialDownloadChunks)));
 
-			// Set up a fallback chain of sources.
-			TArray<IChunkSource*> ChunkSources;
-			ChunkSources.Add(ChunkDbChunkSource.Get());
-			ChunkSources.Add(InstallChunkSource.Get());
-			ChunkSources.Add(CloudChunkSource.Get());
-			TUniquePtr<IChainedChunkSource> ChainedChunkSource(FChainedChunkSourceFactory::Create(
-				ChunkSources));
+			TUniquePtr<IConstructorCloudChunkSource> CloudChunkSource(
+				IConstructorCloudChunkSource::CreateCloudSource(
+					BuildCloudSourceConfig(Configuration),
+					DownloadService.Get(),
+					ChunkDataSerialization.Get(),
+					DownloadConnectionCount.Get(),
+					CloudChunkSourceStatistics.Get(),
+					ManifestSet.Get())
+			);
 
 			// Set up the class that actually requests chunks and writes them to disk.
 			FFileConstructorConfig FCC;
@@ -1534,44 +1408,51 @@ namespace BuildPatchServices
 			FCC.MetaDirectory = MetaStagingDir;
 			FCC.SharedContext = Configuration.SharedContext.Get();
 			FCC.StagingDirectory = InstallStagingDir;
+			FCC.bDeleteChunkDBFilesAfterUse = Configuration.bDeleteChunkDbFilesAfterUse;
+			FCC.BackingStoreDirectory = DataStagingDir;
 
 			TUniquePtr<FBuildPatchFileConstructor> FileConstructor(
-				new FBuildPatchFileConstructor(FCC, FileSystem.Get(), ChainedChunkSource.Get(), ChunkDbChunkSource.Get(), ChunkReferenceTracker.Get(), InstallerError.Get(),InstallerAnalytics.Get(),FileConstructorStatistics.Get())
+				new FBuildPatchFileConstructor(
+					FCC, 
+					FileSystem.Get(), 
+					ChunkDbChunkSource.Get(), 
+					CloudChunkSource.Get(),
+					InstallChunkSource.Get(),
+					ChunkReferenceTracker.Get(), 
+					InstallerError.Get(),
+					InstallerAnalytics.Get(),
+					FileConstructorStatistics.Get(),
+					MoveTemp(ChunkLocations)
+				)
 			);
-			FDelegateHandle OnBeforeDeleteFileHandle = FileConstructor->OnBeforeDeleteFile().AddLambda([this, &InstallChunkSource](const FString& FilePath)
-			{
-				FString BuildRelativeFilename = FilePath;
-				BuildRelativeFilename.RemoveFromStart(Configuration.InstallDirectory);
-				BuildRelativeFilename.RemoveFromStart(TEXT("/"));
-				OldFilesRemovedBySystem.Add(BuildRelativeFilename);
-				InstallChunkSource->HarvestRemainingChunksFromFile(FilePath);
-			});
-			TFunction<void(const FGuid&)> LostChunkCallback = [&ChainedChunkSource](const FGuid& LostChunk)
-			{
-				ChainedChunkSource->AddRepeatRequirement(LostChunk);
-			};
-			if (DiskOverflowStore)
+
+			FDelegateHandle OnBeforeDeleteFileHandle;
+			if (InstallChunkSource)
 			{
-				DiskOverflowStore->SetLostChunkCallback(LostChunkCallback);
+				OnBeforeDeleteFileHandle = FileConstructor->OnBeforeDeleteFile().AddLambda([this, &InstallChunkSource](const FString& FilePath)
+				{
+					FString BuildRelativeFilename = FilePath;
+					BuildRelativeFilename.RemoveFromStart(Configuration.InstallDirectory);
+					BuildRelativeFilename.RemoveFromStart(TEXT("/"));
+					OldFilesRemovedBySystem.Add(BuildRelativeFilename);
+					InstallChunkSource->OnBeforeDeleteFile(FilePath);
+				});		
 			}
-			CloudChunkStore->SetLostChunkCallback(LostChunkCallback);
-
 
 			// Register controllables.
 			FScopedControllables ScopedControllables(&ThreadLock, Controllables, bIsPaused, bShouldAbort);
-			ScopedControllables.Register(InstallChunkSource.Get());
-			ScopedControllables.Register(CloudChunkSource.Get());
 			ScopedControllables.Register(FileConstructor.Get());
 
 			// Set chunk counter stats.
-			if(!bHasCorruptFiles)
+			if (!bHasCorruptFiles)
 			{
 				FScopeLock Lock(&ThreadLock);
 				BuildStats.NumChunksRequired = ReferencedChunks.Num();
-				BuildStats.ChunksQueuedForDownload = InitialDownloadChunks.Num();
-				BuildStats.ChunksLocallyAvailable = ReferencedChunks.Intersect(InstallChunkSource->GetAvailableChunks()).Num();
-				BuildStats.ChunksInChunkDbs = ReferencedChunks.Intersect(ChunkDbChunkSource->GetAvailableChunks()).Num();
-				UE_LOG(LogBPSInstallerConfig, Display, TEXT("Chunk Locations: Cloud %d Install %d ChunkDb %d"), BuildStats.ChunksQueuedForDownload, BuildStats.ChunksLocallyAvailable, BuildStats.ChunksInChunkDbs);
+				BuildStats.ChunksQueuedForDownload = CloudChunks.Num();
+				BuildStats.ChunksLocallyAvailable = InstallChunks.Num();
+				BuildStats.ChunksInChunkDbs = ChunkDbChunks.Num();
+				UE_LOG(LogBPSInstallerConfig, Display, TEXT("Chunk Locations: Cloud %d Install %d ChunkDb %d"), 
+					BuildStats.ChunksQueuedForDownload, BuildStats.ChunksLocallyAvailable, BuildStats.ChunksInChunkDbs);
 			}
 
 			// Setup some weightings for the progress tracking
@@ -1606,7 +1487,10 @@ namespace BuildPatchServices
 			ConstructTimer.Start();
 			FileConstructor->Run();
 			ConstructTimer.Stop();
-			FileConstructor->OnBeforeDeleteFile().Remove(OnBeforeDeleteFileHandle);
+			if (OnBeforeDeleteFileHandle.IsValid())
+			{
+				FileConstructor->OnBeforeDeleteFile().Remove(OnBeforeDeleteFileHandle);
+			}
 			UE_LOG(LogBuildPatchServices, Log, TEXT("File construction complete"));
 
 			// Set any final stats before system destruction.
@@ -1614,6 +1498,16 @@ namespace BuildPatchServices
 				FScopeLock Lock(&ThreadLock);
 				BuildStats.ProcessRequiredDiskSpace = FileConstructor->GetRequiredDiskSpace();
 				BuildStats.ProcessAvailableDiskSpace = FileConstructor->GetAvailableDiskSpace();
+
+				FBuildPatchFileConstructor::FBackingStoreStats BackingStoreStats = FileConstructor->GetBackingStoreStats();
+
+				BuildStats.DriveStorePeakBytes = BackingStoreStats.DiskPeakUsageBytes;
+				BuildStats.NumDriveStoreLoadFailures = BackingStoreStats.DiskLoadFailureCount;
+				BuildStats.NumDriveStoreLostChunks = BackingStoreStats.DiskLostChunkCount;
+				BuildStats.NumDriveStoreChunkLoads = BackingStoreStats.DiskChunkLoadCount;
+
+				BuildStats.MemoryStoreSizePeakBytes = BackingStoreStats.MemoryPeakUsageBytes;
+				BuildStats.MemoryStoreSizeLimitBytes = BackingStoreStats.MemoryLimitBytes;
 			}
 
 			// Let the verifier know which files we built.
@@ -1631,9 +1525,6 @@ namespace BuildPatchServices
 			BuildStats.NumChunksRecycled = InstallChunkSourceStatistics->GetNumSuccessfulChunkRecycles();
 			BuildStats.NumChunksReadFromChunkDbs = ChunkDbChunkSourceStatistics->GetNumSuccessfulLoads();
 			BuildStats.NumRecycleFailures = InstallChunkSourceStatistics->GetNumFailedChunkRecycles();
-			BuildStats.NumChunksStoreBooted = MemoryChunkStoreStatistics->GetNumBooted();
-			BuildStats.NumDriveStoreChunkLoads = DiskChunkStoreStatistics->GetNumSuccessfulLoads();
-			BuildStats.NumDriveStoreLoadFailures = DiskChunkStoreStatistics->GetNumFailedLoads();
 			BuildStats.NumChunkDbChunksFailed = ChunkDbChunkSourceStatistics->GetNumFailedLoads();
 			TArray<float> HealthTimers = CloudChunkSourceStatistics->GetDownloadHealthTimers();
 			BuildStats.ExcellentDownloadHealthTime = HealthTimers[(int32)EBuildPatchDownloadHealth::Excellent];
@@ -2281,11 +2172,6 @@ namespace BuildPatchServices
 			BuildStats.TotalWrittenData = FileConstructorStatistics->GetBytesConstructed();
 			BuildStats.NumFilesConstructed = FileConstructorStatistics->GetFilesConstructed();
 			BuildStats.TheoreticalDownloadTime = BuildStats.AverageDownloadSpeed > 0 ? BuildStats.TotalDownloadedData / BuildStats.AverageDownloadSpeed : 0;
-			BuildStats.AverageMemoryStoreUse = MemoryChunkStoreStatistics->GetAverageStoreUse();
-			BuildStats.PeakMemoryStoreUse = MemoryChunkStoreStatistics->GetPeakStoreUse();
-			BuildStats.AverageMemoryStoreRetained = MemoryChunkStoreStatistics->GetAverageStoreRetained();
-			BuildStats.PeakMemoryStoreRetained = MemoryChunkStoreStatistics->GetPeakStoreRetained();
-			BuildStats.MemoryStoreSize = MemoryChunkStoreStatistics->GetStoreSize();
 			InstallerHelpers::LogBuildStatInfo(BuildStats, SessionId);
 		}
 		// Call the complete delegate.
diff --git a/Engine/Source/Runtime/Online/BuildPatchServices/Private/BuildPatchInstaller.h b/Engine/Source/Runtime/Online/BuildPatchServices/Private/BuildPatchInstaller.h
--- a/Engine/Source/Runtime/Online/BuildPatchServices/Private/BuildPatchInstaller.h
+++ b/Engine/Source/Runtime/Online/BuildPatchServices/Private/BuildPatchInstaller.h
@@ -162,8 +162,6 @@ namespace BuildPatchServices
 
 		// Installer statistics tracking.
 		TUniquePtr<IFileOperationTracker> FileOperationTracker;
-		TUniquePtr<IMemoryChunkStoreStatistics> MemoryChunkStoreStatistics;
-		TUniquePtr<IDiskChunkStoreStatistics> DiskChunkStoreStatistics;
 		TUniquePtr<ISpeedRecorder> DownloadSpeedRecorder;
 		TUniquePtr<ISpeedRecorder> DiskReadSpeedRecorder;
 		TUniquePtr<ISpeedRecorder> DiskWriteSpeedRecorder;
@@ -321,15 +319,6 @@ namespace BuildPatchServices
 		 */
 		const IVerifierStatistics* GetVerifierStatistics() const;
 
-		/**
-		 * @return the memory chunk store statistics interface.
-		 */
-		const IMemoryChunkStoreStatistics* GetMemoryChunkStoreStatistics() const;
-
-		/**
-		 * @return the disk chunk store statistics interface.
-		 */
-		const IDiskChunkStoreStatistics* GetDiskChunkStoreStatistics() const;
 
 	private:
 
@@ -429,22 +418,6 @@ namespace BuildPatchServices
 		 */
 		bool RelocateFileWithRetries(const FString& ToFullFilename, const FString& FromFullFilename, uint32& RenameErrorCode, uint32& CopyErrorCode);
 
-		/**
-		 * Builds the chunkdb source configuration struct.
-		 */
-		FChunkDbSourceConfig BuildChunkDbSourceConfig();
-
-		/**
-		 * Builds the installation source configuration struct.
-		 * @param ChunkIgnoreSet    A set of chunks to initially not consider for loading, use for chunks that are knowingly available in higher priority sources.
-		 */
-		FInstallSourceConfig BuildInstallSourceConfig(TSet<FGuid> ChunkIgnoreSet);
-
-		/**
-		 * Builds the cloud source configuration struct.
-		 */
-		FCloudSourceConfig BuildCloudSourceConfig();
-
 		/**
 		 * Builds the cloud source configuration struct.
 		 */
diff --git a/Engine/Source/Runtime/Online/BuildPatchServices/Private/BuildPatchManifestSet.cpp b/Engine/Source/Runtime/Online/BuildPatchServices/Private/BuildPatchManifestSet.cpp
--- a/Engine/Source/Runtime/Online/BuildPatchServices/Private/BuildPatchManifestSet.cpp
+++ b/Engine/Source/Runtime/Online/BuildPatchServices/Private/BuildPatchManifestSet.cpp
@@ -47,6 +47,8 @@ namespace BuildPatchServices
 				{
 					for (const TPair<FGuid, const BuildPatchServices::FChunkInfo*>& ChunkInfoPair : Manifest.ChunkInfoLookup)
 					{
+						// Order is crucial here. The ManifestIterator above starts with the current manifest and then does the install manifest,
+						// so the data lookup prefers the new chunks over the deployed chunks.
 						DataLookup.Add(ChunkInfoPair.Key, FDataReference{ ChunkInfoPair.Value, &Manifest, &InstallerAction });
 					}
 				}
@@ -92,6 +94,16 @@ namespace BuildPatchServices
 			}
 		}
 
+		virtual FChunkInfo const* GetChunkInfo(const FGuid& DataGuid) const override
+		{
+			FDataReference const * const LookupResult = DataLookup.Find(DataGuid);
+			if (LookupResult)
+			{
+				return LookupResult->Get<0>();
+			}
+			return 0;
+		}
+
 		virtual void GetInstallResumeIds(TSet<FString>& ResumeIds, bool bIncludeLegacy) const
 		{
 			for (const FBuildPatchInstallerAction& InstallerAction : InstallerActions)
diff --git a/Engine/Source/Runtime/Online/BuildPatchServices/Private/Data/ChunkData.cpp b/Engine/Source/Runtime/Online/BuildPatchServices/Private/Data/ChunkData.cpp
--- a/Engine/Source/Runtime/Online/BuildPatchServices/Private/Data/ChunkData.cpp
+++ b/Engine/Source/Runtime/Online/BuildPatchServices/Private/Data/ChunkData.cpp
@@ -3,6 +3,7 @@
 #include "Data/ChunkData.h"
 #include "Serialization/MemoryReader.h"
 #include "Serialization/MemoryWriter.h"
+#include "Memory/SharedBuffer.h"
 #include "Misc/CommandLine.h"
 #include "Misc/Compression.h"
 #include "Common/FileSystem.h"
@@ -498,7 +499,97 @@ namespace BuildPatchServices
 		{
 		}
 
+		static bool CheckChunkHash(const FChunkHeader& Header, const void* Data, uint64 DataLen)
+		{
+			// If we have both hashes, only check Sha1 instead of both.
+			if (EnumHasAnyFlags(Header.HashType, EChunkHashFlags::Sha1))
+			{
+				FSHAHash ShaHashCheck;
+				TRACE_CPUPROFILER_EVENT_SCOPE(Chunk_ShaHash);
+				FSHA1::HashBuffer(Data, DataLen, ShaHashCheck.Hash);
+				return ShaHashCheck == Header.SHAHash;
+			}
+
+			if (EnumHasAnyFlags(Header.HashType, EChunkHashFlags::RollingPoly64))
+			{
+				return Header.RollingHash == FRollingHash::GetHashForDataSet((const uint8*)Data, DataLen);
+			}
+
+			return false;
+		}
+
 		// IChunkDataAccess interface begin.
+		virtual bool ValidateAndRead(FArchive& InArchive, FMutableMemoryView InDestinationBuffer, FChunkHeader& OutHeader, FUniqueBuffer& OutCompressedBuffer) const override
+		{
+			// Read and check the header.
+			InArchive << OutHeader;
+
+			if (!OutHeader.Guid.IsValid() || // Must have a valid guid
+				(!EnumHasAnyFlags(OutHeader.HashType, EChunkHashFlags::RollingPoly64) && !EnumHasAnyFlags(OutHeader.HashType, EChunkHashFlags::Sha1)) || // Must have a valid hash
+				((OutHeader.StoredAs & EChunkStorageFlags::Encrypted) != EChunkStorageFlags::None) || // Must not be encrypted
+				(InDestinationBuffer.GetSize() > OutHeader.DataSizeUncompressed)) // Unexpected buffer size.
+			{
+				return false;
+			}
+
+			// If the chunk isn't compressed and is the entire block, we can read directly in to the destination buffer.
+			// We seem to be getting chunkinfos that report window sizes that are not the full chunk size, so we need to be
+			// able to catch these subchunks and copy them out as needed.
+			if ((OutHeader.StoredAs & EChunkStorageFlags::Compressed) == EChunkStorageFlags::None &&
+				InDestinationBuffer.GetSize() == OutHeader.DataSizeUncompressed)
+			{
+				InArchive.Serialize(InDestinationBuffer.GetData(), InDestinationBuffer.GetSize());
+			}
+			else
+			{
+				// We need to allocate and read in to another buffer for decompression.
+				OutCompressedBuffer = FUniqueBuffer::Alloc(OutHeader.DataSizeCompressed);
+				InArchive.Serialize(OutCompressedBuffer.GetData(), OutCompressedBuffer.GetSize());
+			}
+
+			if (InArchive.IsError())
+			{
+				return false;
+			}
+			return true;
+		}
+
+
+		virtual bool DecompressValidatedRead(const FChunkHeader& InHeader, FMutableMemoryView InDestinationBuffer, const FUniqueBuffer& InCompressedBuffer) const override
+		{			
+			FMutableMemoryView UseDestinationBuffer = InDestinationBuffer;
+
+			FUniqueBuffer TempDestinationBuffer;
+			if (InHeader.DataSizeUncompressed != InDestinationBuffer.GetSize())
+			{
+				TempDestinationBuffer = FUniqueBuffer::Alloc(InHeader.DataSizeUncompressed);
+				UseDestinationBuffer = TempDestinationBuffer.GetView();
+			}
+
+			if ((InHeader.StoredAs & EChunkStorageFlags::Compressed) != EChunkStorageFlags::None)
+			{
+				if (!FCompression::UncompressMemory(
+					NAME_Zlib,
+					UseDestinationBuffer.GetData(),
+					UseDestinationBuffer.GetSize(),
+					InCompressedBuffer.GetData(),
+					InCompressedBuffer.GetSize()))
+				{
+					return false;
+				}
+			}
+
+			bool bHashMatches = CheckChunkHash(InHeader, UseDestinationBuffer.GetData(), UseDestinationBuffer.GetSize());
+			if (bHashMatches && InHeader.DataSizeUncompressed != InDestinationBuffer.GetSize())
+			{
+				// We routed through a temp buffer, copy out as necessary.
+				FMemory::Memcpy(InDestinationBuffer.GetData(), TempDestinationBuffer.GetData(), InDestinationBuffer.GetSize());
+			}
+
+			return bHashMatches;
+		}
+
+
 		virtual IChunkDataAccess* LoadFromFile(const FString& Filename, EChunkLoadResult& OutLoadResult) const override
 		{
 			IChunkDataAccess* ChunkData = nullptr;
diff --git a/Engine/Source/Runtime/Online/BuildPatchServices/Private/Data/ChunkData.h b/Engine/Source/Runtime/Online/BuildPatchServices/Private/Data/ChunkData.h
--- a/Engine/Source/Runtime/Online/BuildPatchServices/Private/Data/ChunkData.h
+++ b/Engine/Source/Runtime/Online/BuildPatchServices/Private/Data/ChunkData.h
@@ -1,6 +1,7 @@
 // Copyright Epic Games, Inc. All Rights Reserved.
 #pragma once
 
+#include "Memory/MemoryFwd.h"
 #include "Misc/SecureHash.h"
 #include "Misc/EnumClassFlags.h"
 #include "Misc/Guid.h"
@@ -409,6 +410,13 @@ namespace BuildPatchServices
 			return SaveToArchive(Archive, ChunkDataAccess);
 		}
 
+		// Read the header from the archive, make sure it's sane, and set up to decompress/hash in a separate call. This
+		// is to allow the IO to be on a separate thread from decompression/hashing.
+		virtual bool ValidateAndRead(FArchive& InArchive, FMutableMemoryView InDestinationBuffer, FChunkHeader& OutHeader, FUniqueBuffer& OutCompressedBuffer) const = 0;
+		
+		// Take the output of ValidateAndRead and process it in to an actual output chunk.
+		virtual bool DecompressValidatedRead(const FChunkHeader& InHeader, FMutableMemoryView InDestionationBuffer, const FUniqueBuffer& InCompressedBuffer) const = 0;
+
 		/**
 		 * Injects an SHA hash for the data into the structure of a serialized chunk.
 		 * @param Memory            The memory array containing the serialized chunk.
diff --git a/Engine/Source/Runtime/Online/BuildPatchServices/Private/Generation/ManifestBuildStreamer.cpp b/Engine/Source/Runtime/Online/BuildPatchServices/Private/Generation/ManifestBuildStreamer.cpp
--- a/Engine/Source/Runtime/Online/BuildPatchServices/Private/Generation/ManifestBuildStreamer.cpp
+++ b/Engine/Source/Runtime/Online/BuildPatchServices/Private/Generation/ManifestBuildStreamer.cpp
@@ -20,7 +20,6 @@
 #include "Installer/ChunkEvictionPolicy.h"
 #include "Installer/ChunkReferenceTracker.h"
 #include "Installer/CloudChunkSource.h"
-#include "Installer/DiskChunkStore.h"
 #include "Installer/DownloadService.h"
 #include "Installer/InstallerAnalytics.h"
 #include "Installer/InstallerError.h"
diff --git a/Engine/Source/Runtime/Online/BuildPatchServices/Private/IBuildManifestSet.h b/Engine/Source/Runtime/Online/BuildPatchServices/Private/IBuildManifestSet.h
--- a/Engine/Source/Runtime/Online/BuildPatchServices/Private/IBuildManifestSet.h
+++ b/Engine/Source/Runtime/Online/BuildPatchServices/Private/IBuildManifestSet.h
@@ -94,6 +94,7 @@ namespace BuildPatchServices
 		virtual void GetReferencedChunks(TSet<FGuid>& DataGuids) const = 0;
 		virtual uint64 GetDownloadSize(const FGuid& DataGuid) const = 0;
 		virtual uint64 GetDownloadSize(const TSet<FGuid>& DataGuids) const = 0;
+		virtual FChunkInfo const* GetChunkInfo(const FGuid& DataGuid) const = 0;
 		virtual bool GetChunkShaHash(const FGuid& ChunkGuid, FSHAHash& OutHash) const = 0;
 		virtual FString GetDataFilename(const FGuid& DataGuid) const = 0;
 
diff --git a/Engine/Source/Runtime/Online/BuildPatchServices/Private/Installer/BuildInstallStreamer.cpp b/Engine/Source/Runtime/Online/BuildPatchServices/Private/Installer/BuildInstallStreamer.cpp
--- a/Engine/Source/Runtime/Online/BuildPatchServices/Private/Installer/BuildInstallStreamer.cpp
+++ b/Engine/Source/Runtime/Online/BuildPatchServices/Private/Installer/BuildInstallStreamer.cpp
@@ -320,7 +320,6 @@ namespace BuildPatchServices
 				// Composition root.
 				FCloudSourceConfig CloudSourceConfig(Configuration.CloudDirectories);
 				CloudSourceConfig.bBeginDownloadsOnFirstGet = false;
-				CloudSourceConfig.bRunOwnThread = false;
 				TUniquePtr<IChunkReferenceTracker> ChunkReferenceTracker(FChunkReferenceTrackerFactory::Create(
 					ManifestSet.Get(),
 					VFCConfig.FilesToConstruct));
diff --git a/Engine/Source/Runtime/Online/BuildPatchServices/Private/Installer/BuildStatistics.cpp b/Engine/Source/Runtime/Online/BuildPatchServices/Private/Installer/BuildStatistics.cpp
--- a/Engine/Source/Runtime/Online/BuildPatchServices/Private/Installer/BuildStatistics.cpp
+++ b/Engine/Source/Runtime/Online/BuildPatchServices/Private/Installer/BuildStatistics.cpp
@@ -36,10 +36,6 @@ namespace BuildPatchServices
 		virtual const FBuildInstallerConfiguration& GetConfiguration() const override;
 		virtual int64 GetDownloadSize() const override;
 		virtual int64 GetBuildSize() const override;
-		virtual int32 GetMemoryChunkStoreSize() const override;
-		virtual int32 GetMemoryChunksInStore() const override;
-		virtual int32 GetMemoryChunksBooted() const override;
-		virtual int32 GetMemoryChunksRetained() const override;
 		virtual FString GetCurrentWorkingFileName() const override;
 		virtual float GetCurrentWorkingFileProgress() const override;
 		virtual FGuid GetCurrentWorkingData() const override;
@@ -81,26 +77,6 @@ namespace BuildPatchServices
 		return Installer->GetFileConstructorStatistics()->GetRequiredConstructSize();
 	}
 
-	int32 FBuildStatistics::GetMemoryChunkStoreSize() const
-	{
-		return Installer->GetMemoryChunkStoreStatistics()->GetStoreSize();
-	}
-
-	int32 FBuildStatistics::GetMemoryChunksInStore() const
-	{
-		return Installer->GetMemoryChunkStoreStatistics()->GetStoreUse();
-	}
-
-	int32 FBuildStatistics::GetMemoryChunksBooted() const
-	{
-		return Installer->GetMemoryChunkStoreStatistics()->GetNumBooted();
-	}
-
-	int32 FBuildStatistics::GetMemoryChunksRetained() const
-	{
-		return Installer->GetMemoryChunkStoreStatistics()->GetStoreRetained();
-	}
-
 	FString FBuildStatistics::GetCurrentWorkingFileName() const
 	{
 		return Installer->GetFileConstructorStatistics()->GetCurrentFile();
diff --git a/Engine/Source/Runtime/Online/BuildPatchServices/Private/Installer/ChainedChunkSource.cpp b/Engine/Source/Runtime/Online/BuildPatchServices/Private/Installer/ChainedChunkSource.cpp
deleted file mode 100644
diff --git a/Engine/Source/Runtime/Online/BuildPatchServices/Private/Installer/ChainedChunkSource.h b/Engine/Source/Runtime/Online/BuildPatchServices/Private/Installer/ChainedChunkSource.h
deleted file mode 100644
diff --git a/Engine/Source/Runtime/Online/BuildPatchServices/Private/Installer/ChunkDbChunkSource.cpp b/Engine/Source/Runtime/Online/BuildPatchServices/Private/Installer/ChunkDbChunkSource.cpp
--- a/Engine/Source/Runtime/Online/BuildPatchServices/Private/Installer/ChunkDbChunkSource.cpp
+++ b/Engine/Source/Runtime/Online/BuildPatchServices/Private/Installer/ChunkDbChunkSource.cpp
@@ -15,57 +15,49 @@
 #include "Installer/MessagePump.h"
 #include "Installer/InstallerError.h"
 #include "Installer/InstallerSharedContext.h"
+#include "Memory/SharedBuffer.h"
 #include "Misc/OutputDeviceRedirector.h"
+#include "Tasks/Task.h"
+#include "Serialization/ArrayReader.h"
+#include "Serialization/ArrayWriter.h"
+#include "Misc/Compression.h"
+#include "Compression/OodleDataCompressionUtil.h"
 
 #if PLATFORM_WINDOWS
 #include "Windows/WindowsHWrapper.h"
 namespace ChunkDbSourceHelpers
 {
-	int32 DisableOsIntervention()
+	bool DisableOsIntervention(uint32& OutPreviousOsIntervention)
 	{
-		return ::SetErrorMode(SEM_FAILCRITICALERRORS);
+		// This prevents dialogs from popping up if we try to read a chunkdb from a removable drive
+		// with no media. 
+
+		// We only call this during startup on a single thread so we can limit this to 
+		// just our thread.				
+		return ::SetThreadErrorMode(SEM_FAILCRITICALERRORS, (DWORD*)&OutPreviousOsIntervention);		
 	}
 
-	void ResetOsIntervention(int32 Previous)
+	void ResetOsIntervention(uint32 Previous)
 	{
-		::SetErrorMode(Previous);
+		::SetThreadErrorMode((DWORD)Previous, NULL);
 	}
 }
 #else
 namespace ChunkDbSourceHelpers
 {
-	int32 DisableOsIntervention()
+	bool DisableOsIntervention(uint32& OutPreviousOsIntervention)
 	{
-		return 0;
+		OutPreviousOsIntervention = 0;
+		return false;
 	}
 
-	void ResetOsIntervention(int32 Previous)
+	void ResetOsIntervention(uint32 Previous)
 	{
 	}
 }
 #endif
 
-DECLARE_LOG_CATEGORY_EXTERN(LogChunkDbChunkSource, Log, All);
-DEFINE_LOG_CATEGORY(LogChunkDbChunkSource);
-
-namespace ChunkDbSourceHelpers
-{
-	using namespace BuildPatchServices;
-	IChunkDbChunkSourceStat::ELoadResult FromSerializer(EChunkLoadResult LoadResult)
-	{
-		switch (LoadResult)
-		{
-			case EChunkLoadResult::Success:
-				return IChunkDbChunkSourceStat::ELoadResult::Success;
-			case EChunkLoadResult::MissingHashInfo:
-				return IChunkDbChunkSourceStat::ELoadResult::MissingHashInfo;
-			case EChunkLoadResult::HashCheckFailed:
-				return IChunkDbChunkSourceStat::ELoadResult::HashCheckFailed;
-			default:
-				return IChunkDbChunkSourceStat::ELoadResult::SerializationError;
-		}
-	}
-}
+DEFINE_LOG_CATEGORY_STATIC(LogChunkDbChunkSource, Log, All);
 
 namespace BuildPatchServices
 {
@@ -85,12 +77,11 @@ namespace BuildPatchServices
 		// If we're retired then any access is invalid and fatal as the file has been closed and could be deleted.
 		bool bIsRetired = false;
 
-		void Retire(IMessagePump* MessagePump, IFileSystem* FileSystemIfDeleting, bool bDelete)
+		void Retire(IFileSystem* FileSystemIfDeleting, bool bDelete)
 		{
 			bIsRetired = true;
 			FString ArchiveName = Archive->GetArchiveName();
 			Archive.Reset();
-			MessagePump->SendMessage(FChunkSourceEvent{ FChunkSourceEvent::EType::Retired, ArchiveName });
 
 			if (bDelete && FileSystemIfDeleting)
 			{
@@ -102,114 +93,68 @@ namespace BuildPatchServices
 		}
 	};
 
-	/**
-	 * Struct holding variables required for accessing a particular chunk.
-	 */
+	// Holds where to get the chunk data from. File + location in file.
 	struct FChunkAccessLookup
 	{
 		FChunkLocation* Location;
 		FChunkDbDataAccess* DbFile;
 	};
 
-	/**
-	 * Struct holding variables for tracking retry attempts on opening chunkdb files.
-	 */
-	struct FChunkDbRetryInfo
-	{
-		int32 Count;
-		double SecondsAtLastTry;
-	};
-
-	/**
-	 * The concrete implementation of IChunkDbChunkSource.
-	 */
-	class FChunkDbChunkSource
-		: public IChunkDbChunkSource
+	class FChunkDbChunkSource : public IConstructorChunkDbChunkSource
 	{
 	public:
-		FChunkDbChunkSource(FChunkDbSourceConfig Configuration, IPlatform* Platform, IFileSystem* FileSystem, IChunkStore* ChunkStore, IChunkReferenceTracker* ChunkReferenceTracker, IChunkDataSerialization* ChunkDataSerialization, IMessagePump* MessagePump, IInstallerError* InstallerError, IChunkDbChunkSourceStat* ChunkDbChunkSourceStat);
-		~FChunkDbChunkSource();
-
-		// IControllable interface begin.
-		virtual void SetPaused(bool bInIsPaused) override;
-		virtual void Abort() override;
-		// IControllable interface end.
-
-		// IChunkSource interface begin.
-		virtual IChunkDataAccess* Get(const FGuid& DataId) override;
-		virtual TSet<FGuid> AddRuntimeRequirements(TSet<FGuid> NewRequirements) override;
-		virtual bool AddRepeatRequirement(const FGuid& RepeatRequirement) override;
-		virtual void SetUnavailableChunksCallback(TFunction<void(TSet<FGuid>)> Callback) override;
-		virtual void ReportFileCompletion() override
+		FChunkDbChunkSource(FChunkDbSourceConfig Configuration, IFileSystem* FileSystem, const TArray<FGuid>& ChunkAccessOrderList, IChunkDataSerialization* ChunkDataSerialization, IChunkDbChunkSourceStat* ChunkDbChunkSourceStat);
+		~FChunkDbChunkSource()
+		{
+		}
+
+		virtual void ReportFileCompletion(int32 RemainingChunkCount) override
 		{
 			//
 			// Since we've completed a file we know we won't need to resume/retry it and can delete
 			// the source chunkdb that it used.
 			//
-			int32 RemainingChunkCount = ChunkReferenceTracker->GetRemainingChunkCount();
 			for (FChunkDbDataAccess& ChunkDbDataAccess : ChunkDbDataAccesses)
 			{
 				if (!ChunkDbDataAccess.bIsRetired &&
 					ChunkDbDataAccess.RetireAt > RemainingChunkCount)
 				{
-					ChunkDbDataAccess.Retire(MessagePump, FileSystem, Configuration.bDeleteChunkDBAfterUse);
+					ChunkDbDataAccess.Retire(FileSystem, Configuration.bDeleteChunkDBAfterUse);
 				}
 			}
 		}
-		// IChunkSource interface end.
+		virtual int32 GetChunkUnavailableAt(const FGuid& DataId) const override
+		{
+			// While technically the chunks retire as a result of delete-during-install, we only do this
+			// when they aren't needed any more, so we can set this to "never retires"
+			return TNumericLimits<int32>::Max();
+		}
+
+		virtual FRequestProcessFn CreateRequest(const FGuid& DataId, FMutableMemoryView DestinationBuffer, void* UserPtr, FChunkRequestCompleteDelegate CompleteFn) override;
+		virtual const TSet<FGuid>& GetAvailableChunks() const override 
+		{
+			return AvailableChunks; 
+		}
 
-		// IInstallChunkSource interface begin.
-		virtual const TSet<FGuid>& GetAvailableChunks() const override;
 		virtual uint64 GetChunkDbSizesAtIndexes(const TArray<int32>& FileCompletionIndexes, TArray<uint64>& OutChunkDbSizesAtCompletion) const override;
-		// IInstallChunkSource interface end.
 
 		static void LoadChunkDbFiles(
 			const TArray<FString>& ChunkDbFiles, IFileSystem* FileSystem, const TArray<FGuid>& ChunkAccessOrderedList,
 			TArray<FChunkDbDataAccess>& OutChunkDbDataAccesses, TMap<FGuid, FChunkAccessLookup>& OutChunkDbDataAccessLookup, TSet<FGuid>* OutOptionalAvailableStore);
 
 	private:
-		void ThreadRun();
-		void LoadChunk(const FGuid& DataId);
-		bool HasFailed(const FGuid& DataId);
 
-		typedef TFunction<bool(const FArchive&)> FArchiveSelector;
-		typedef TFunction<void(const FArchive&, bool)> FArchiveOpenResult;
-		void TryReopenChunkDbFiles(const FArchiveSelector& SelectPredicate, const FArchiveOpenResult& ResultCallback);
-
-	private:
 		// Configuration.
 		const FChunkDbSourceConfig Configuration;
 		// Dependencies.
-		IPlatform* Platform = nullptr;
 		IFileSystem* FileSystem = nullptr;
-		IChunkStore* ChunkStore = nullptr;
 		IChunkReferenceTracker* ChunkReferenceTracker = nullptr;
 		IChunkDataSerialization* ChunkDataSerialization = nullptr;
-		IMessagePump* MessagePump = nullptr;
-		IInstallerError* InstallerError = nullptr;
-		IChunkDbChunkSourceStat* ChunkDbChunkSourceStat;
-		// Worker thread lifetime management
-		TPromise<void> Promise;
-		TFuture<void> Future;
-		IBuildInstallerThread* Thread = nullptr;
-		// Control signals.
-		std::atomic<bool> bIsPaused = false;
-		std::atomic<bool> bShouldAbort = false;
-		std::atomic<bool> bStartedLoading = false;
-		// Handling of chunks we lose access to.
-		TFunction<void(TSet<FGuid>)> UnavailableChunksCallback;
-		TSet<FGuid> UnavailableChunks;
+		IChunkDbChunkSourceStat* ChunkDbChunkSourceStat = nullptr;
 		// Storage of our chunkdb and enumerated available chunks lookup.
 		TArray<FChunkDbDataAccess> ChunkDbDataAccesses;
 		TMap<FGuid, FChunkAccessLookup> ChunkDbDataAccessLookup;
-		TMap<FString, FChunkDbRetryInfo> ChunkDbReloadAttempts;
 		TSet<FGuid> AvailableChunks;
-		TSet<FGuid> PlacedInStore;
-		// Communication between the incoming Get calls, and our worker thread.
-		TQueue<FGuid, EQueueMode::Spsc> FailedToLoadMessages;
-		TSet<FGuid> FailedToLoad;
-		// Communication between the incoming AddRepeatRequirement calls, and our worker thread.
-		TQueue<FGuid, EQueueMode::Mpsc> RepeatRequirementMessages;
 
 		// Number of chunks to process in this manifest when we started.
 		int32 OriginalChunkCount = 0;
@@ -223,7 +168,8 @@ namespace BuildPatchServices
 
 		// Allow OS intervention only once.
 		bool bResetOsIntervention = false;
-		int32 PreviousOsIntervention = 0;
+		bool bHasPreviousOsIntervention = false;
+		uint32 PreviousOsIntervention = 0;
 		// Load each chunkdb's TOC to enumerate available chunks.
 		for (const FString& ChunkDbFilename : ChunkDbFiles)
 		{
@@ -254,11 +200,11 @@ namespace BuildPatchServices
 			else if (!bResetOsIntervention)
 			{
 				bResetOsIntervention = true;
-				PreviousOsIntervention = ChunkDbSourceHelpers::DisableOsIntervention();
+				bHasPreviousOsIntervention = ChunkDbSourceHelpers::DisableOsIntervention(PreviousOsIntervention);
 			}
 		}
 		// Reset OS intervention if we disabled it.
-		if (bResetOsIntervention)
+		if (bResetOsIntervention && bHasPreviousOsIntervention)
 		{
 			ChunkDbSourceHelpers::ResetOsIntervention(PreviousOsIntervention);
 		}
@@ -351,7 +297,7 @@ namespace BuildPatchServices
 		return GetChunkDbSizesAtIndexesInternal(ChunkDbDataAccesses, OriginalChunkCount, FileCompletionIndexes, OutChunkDbSizesAtCompletion);
 	}
 
-	uint64 IChunkDbChunkSource::GetChunkDbSizesAtIndexes(const TArray<FString>& ChunkDbFiles, IFileSystem* FileSystem, const TArray<FGuid>& ChunkAccessOrderedList, const TArray<int32>& FileCompletionIndexes, TArray<uint64>& OutChunkDbSizesAtCompletion)
+	uint64 IConstructorChunkDbChunkSource::GetChunkDbSizesAtIndexes(const TArray<FString>& ChunkDbFiles, IFileSystem* FileSystem, const TArray<FGuid>& ChunkAccessOrderedList, const TArray<int32>& FileCompletionIndexes, TArray<uint64>& OutChunkDbSizesAtCompletion)
 	{
 		TArray<FChunkDbDataAccess> ChunkFiles;
 		TMap<FGuid, FChunkAccessLookup> ChunkGuidLookup;
@@ -361,346 +307,104 @@ namespace BuildPatchServices
 		return GetChunkDbSizesAtIndexesInternal(ChunkFiles, ChunkAccessOrderedList.Num(), FileCompletionIndexes, OutChunkDbSizesAtCompletion);
 	}
 
-	FChunkDbChunkSource::FChunkDbChunkSource(FChunkDbSourceConfig InConfiguration, IPlatform* InPlatform, IFileSystem* InFileSystem, IChunkStore* InChunkStore, IChunkReferenceTracker* InChunkReferenceTracker, IChunkDataSerialization* InChunkDataSerialization, IMessagePump* InMessagePump, IInstallerError* InInstallerError, IChunkDbChunkSourceStat* InChunkDbChunkSourceStat)
+	FChunkDbChunkSource::FChunkDbChunkSource(FChunkDbSourceConfig InConfiguration, IFileSystem* InFileSystem, const TArray<FGuid>& InChunkAccessOrderList, 
+		IChunkDataSerialization* InChunkDataSerialization, IChunkDbChunkSourceStat* InChunkDbChunkSourceStat)
 		: Configuration(MoveTemp(InConfiguration))
-		, Platform(InPlatform)
 		, FileSystem(InFileSystem)
-		, ChunkStore(InChunkStore)
-		, ChunkReferenceTracker(InChunkReferenceTracker)
 		, ChunkDataSerialization(InChunkDataSerialization)
-		, MessagePump(InMessagePump)
-		, InstallerError(InInstallerError)
 		, ChunkDbChunkSourceStat(InChunkDbChunkSourceStat)
-		, bStartedLoading(!InConfiguration.bBeginLoadsOnFirstGet)
 	{
 		TRACE_CPUPROFILER_EVENT_SCOPE(FChunkDbChunkSource_ctor);
 
-		TArray<FGuid> OrderedGuidList;
-		InChunkReferenceTracker->CopyOutOrderedUseList(OrderedGuidList);
+		OriginalChunkCount = InChunkAccessOrderList.Num();
 
-		OriginalChunkCount = OrderedGuidList.Num();
-
-		LoadChunkDbFiles(Configuration.ChunkDbFiles, FileSystem, OrderedGuidList, ChunkDbDataAccesses, ChunkDbDataAccessLookup, &AvailableChunks);
+		LoadChunkDbFiles(Configuration.ChunkDbFiles, FileSystem, InChunkAccessOrderList, ChunkDbDataAccesses, ChunkDbDataAccessLookup, &AvailableChunks);
 
 		// Immediately retire any chunkdbs we don't need so they don't eat disk space during the first file.
-		FChunkDbChunkSource::ReportFileCompletion();
-
-		// Start threaded load worker.
-		if (ChunkDbDataAccesses.Num() > 0)
-		{
-			Future = Promise.GetFuture();
-			Thread = Configuration.SharedContext->CreateThread();
-			Thread->RunTask([this]() { ThreadRun(); });
-		}
+		FChunkDbChunkSource::ReportFileCompletion(OriginalChunkCount);
 	}
 
-	FChunkDbChunkSource::~FChunkDbChunkSource()
-	{
-		Abort();
 
-		if (!Thread)
+	IConstructorChunkSource::FRequestProcessFn FChunkDbChunkSource::CreateRequest(const FGuid& DataId, FMutableMemoryView DestinationBuffer, void* UserPtr, FChunkRequestCompleteDelegate CompleteFn)
+	{
+		FChunkAccessLookup* ChunkInfo = ChunkDbDataAccessLookup.Find(DataId);
+		if (!ChunkInfo)
 		{
-			// The thread was never started, so complete the promise here
-			Promise.SetValue();
+			CompleteFn.Execute(DataId, false, true, UserPtr);
+			return [](bool) {return;};
 		}
 
-		Future.Wait();
-		Configuration.SharedContext->ReleaseThread(Thread);
-		Thread = nullptr;
-	}
-
-	void FChunkDbChunkSource::SetPaused(bool bInIsPaused)
-	{
-		bIsPaused = bInIsPaused;
-	}
-
-	void FChunkDbChunkSource::Abort()
-	{
-		bShouldAbort = true;
-	}
-
-	IChunkDataAccess* FChunkDbChunkSource::Get(const FGuid& DataId)
-	{
-		TRACE_CPUPROFILER_EVENT_SCOPE(ChunkDb_Get);
-
-		// Get from our store
-		IChunkDataAccess* ChunkData = ChunkStore->Get(DataId);
-		if (ChunkData == nullptr)
+		return [ChunkInfo, DataId, DestinationBuffer, UserPtr, CompleteFn, ChunkDataSerialization = ChunkDataSerialization, ChunkDbChunkSourceStat = ChunkDbChunkSourceStat](bool bIsAborted)
 		{
-			bStartedLoading = true;
-			FChunkAccessLookup* ChunkInfo = ChunkDbDataAccessLookup.Find(DataId);
-			if (ChunkInfo && AvailableChunks.Contains(DataId))
+			TRACE_CPUPROFILER_EVENT_SCOPE(ChunkDbRead);
+			if (bIsAborted)
 			{
-				if (ChunkInfo->DbFile->bIsRetired)
-				{
-					GLog->Logf(TEXT("BAD!!!!!!!! Accessed chunk after chunkdb %s retired - invalid chunk reference order."), *ChunkInfo->DbFile->Archive->GetArchiveName());
-				}
-
-				// Wait for the chunk to be available.
-				while (!HasFailed(DataId) && (ChunkData = ChunkStore->Get(DataId)) == nullptr && !bShouldAbort)
-				{
-					TRACE_CPUPROFILER_EVENT_SCOPE(ChunkDb_GetChunkData_Sleep);
-					// 10ms was wayyy too long here. (def needs to be an event)
-					Platform->Sleep(0.001f);
-				}
-
-				// Dump out unavailable chunks on the incoming IO thread.
-				if (UnavailableChunksCallback && UnavailableChunks.Num() > 0)
-				{
-					UnavailableChunksCallback(MoveTemp(UnavailableChunks));
-				}
+				CompleteFn.Execute(DataId, true, false, UserPtr);
+				return;
 			}
-		}
-		return ChunkData;
-	}
 
-	TSet<FGuid> FChunkDbChunkSource::AddRuntimeRequirements(TSet<FGuid> NewRequirements)
-	{
-		// We can't actually get more than we are already getting, so just return what we don't already have in our list.
-		return NewRequirements.Difference(AvailableChunks);
-	}
+			FChunkLocation& ChunkLocation = *ChunkInfo->Location;
+			FChunkDbDataAccess& ChunkDbDataAccess = *ChunkInfo->DbFile;
+			FArchive* ChunkDbFile = ChunkDbDataAccess.Archive.Get();
+			if (ChunkDbFile->IsError())
+			{
+				CompleteFn.Execute(DataId, false, true, UserPtr);
+				return;
+			}
 
-	bool FChunkDbChunkSource::AddRepeatRequirement(const FGuid& RepeatRequirement)
-	{
-		if (AvailableChunks.Contains(RepeatRequirement))
-		{
-			RepeatRequirementMessages.Enqueue(RepeatRequirement);
-			return true;
-		}
-		return false;
-	}
+			IChunkDbChunkSourceStat::ELoadResult LoadResult = IChunkDbChunkSourceStat::ELoadResult::Success;
+			ISpeedRecorder::FRecord ActivityRecord;
+			ActivityRecord.CyclesStart = FStatsCollector::GetCycles();
+			ChunkDbChunkSourceStat->OnLoadStarted(DataId);
 
-	void FChunkDbChunkSource::SetUnavailableChunksCallback(TFunction<void(TSet<FGuid>)> Callback)
-	{
-		UnavailableChunksCallback = Callback;
-	}
+			// We'd love to read direct in to the destination if we don't have
+			// any compression. However we don't know if it's compressed until we read the header which
+			// is tiny and dependent - we don't know how big it is until we read part of it.
+			ChunkDbFile->Seek(ChunkLocation.ByteStart);
+			FChunkHeader Header;
 
-	const TSet<FGuid>& FChunkDbChunkSource::GetAvailableChunks() const
-	{
-		return AvailableChunks;
-	}
+			// If it's uncompressed, we can read direct to the destination.
+			FUniqueBuffer CompressedBuffer;
+			bool Result = ChunkDataSerialization->ValidateAndRead(*ChunkDbFile, DestinationBuffer, Header, CompressedBuffer);
 
-	void FChunkDbChunkSource::ThreadRun()
-	{
-		TRACE_CPUPROFILER_EVENT_SCOPE(ChunkDb_Thread);
+			ActivityRecord.Size = ChunkDbFile->Tell() - ChunkLocation.ByteStart;
 
-		while (!bShouldAbort)
-		{
-			bool bWorkPerformed = false;
-			if (bStartedLoading)
+			if (!Result)
 			{
-				// 'Forget' any repeat requirements.
-				FGuid RepeatRequirement;
-				while (RepeatRequirementMessages.Dequeue(RepeatRequirement))
-				{
-					PlacedInStore.Remove(RepeatRequirement);
-				}
-				// Select chunks that are contained in our chunk db files.
-				TFunction<bool(const FGuid&)> SelectPredicate = [this](const FGuid& ChunkId) { return AvailableChunks.Contains(ChunkId); };
-				// Grab all the chunks relevant to this source to fill the store.
-				int32 SearchLength = FMath::Max(ChunkStore->GetSize(), Configuration.PreFetchMinimum);
-				TArray<FGuid> BatchLoadChunks = ChunkReferenceTracker->SelectFromNextReferences(SearchLength, SelectPredicate);
-				// Remove already loaded chunks from our todo list.
-				// We only grab more chunks as they come into scope.
-				TFunction<bool(const FGuid&)> RemovePredicate = [this](const FGuid& ChunkId) { return PlacedInStore.Contains(ChunkId); };
-				BatchLoadChunks.RemoveAll(RemovePredicate);
-				// Clamp to configured max.
-				BatchLoadChunks.SetNum(FMath::Min(BatchLoadChunks.Num(), Configuration.PreFetchMaximum), EAllowShrinking::No);
-				// Load this batch.
-				ChunkDbChunkSourceStat->OnBatchStarted(BatchLoadChunks);
-				for (int32 ChunkIdx = 0; ChunkIdx < BatchLoadChunks.Num() && !bShouldAbort; ++ChunkIdx)
-				{
-					const FGuid& BatchLoadChunk = BatchLoadChunks[ChunkIdx];
-					LoadChunk(BatchLoadChunk);
-				}
+				LoadResult = IChunkDbChunkSourceStat::ELoadResult::SerializationError;
+				ActivityRecord.CyclesEnd = FStatsCollector::GetCycles();
+				ChunkDbChunkSourceStat->OnLoadComplete(DataId, LoadResult, ActivityRecord);
 
-				// Set whether we performed work.
-				if (BatchLoadChunks.Num() > 0 && !bShouldAbort)
-				{
-					bWorkPerformed = true;
-				}
+				// The header or chunk data was bad.				
+				CompleteFn.Execute(DataId, false, true, UserPtr);
+				return;
 			}
-			// If we had nothing to do, rest a little.
-			if(!bWorkPerformed)
-			{
-				// This sleep is hit during startup since we don't do work before we start loading. For fresh installs
-				// this means we usually ask for the first chunk immediately after this sleep starts. So I changed this from 100ms
-				// to 10ms. Ideally this would be waiting on an event...
-				Platform->Sleep(0.01f);
-			}
-		}
-
-		Promise.SetValue();
-	}
 
-	void FChunkDbChunkSource::LoadChunk(const FGuid& DataId)
-	{
-		bool bChunkGood = false;
-		if (ChunkDbDataAccessLookup.Contains(DataId))
-		{
-			TRACE_CPUPROFILER_EVENT_SCOPE(ChunkDb_LoadChunk);
-			ChunkDbChunkSourceStat->OnLoadStarted(DataId);
-			FChunkAccessLookup& ChunkAccessLookup = ChunkDbDataAccessLookup[DataId];
-			FChunkLocation& ChunkLocation = *ChunkAccessLookup.Location;
-			FChunkDbDataAccess& ChunkDbDataAccess = *ChunkAccessLookup.DbFile;
-			FArchive* ChunkDbFile = ChunkDbDataAccess.Archive.Get();
-			check(ChunkDbFile != nullptr); // We never set null.
-			bChunkGood = !ChunkDbFile->IsError();
-			if (!bChunkGood)
-			{
-				const double SecondsNow = FStatsCollector::GetSeconds();
-				// Time to retry?
-				FChunkDbRetryInfo& ChunkDbRetryInfo = ChunkDbReloadAttempts.FindOrAdd(ChunkDbFile->GetArchiveName());
-				const bool bIsFirstFail = ChunkDbRetryInfo.Count == 0;
-				if (bIsFirstFail)
-				{
-					MessagePump->SendMessage(FChunkSourceEvent{FChunkSourceEvent::EType::AccessLost, ChunkDbFile->GetArchiveName()});
-					// Also try reopening any chunkdbs files that have no error yet, in case they will also be lost.
-					// This gives us control over disabling OS intervention popups when we inevitably try to read from them later.
-					int32 PreviousOsIntervention = ChunkDbSourceHelpers::DisableOsIntervention();
-					TryReopenChunkDbFiles(
-						[this](const FArchive& Ar) -> bool
-						{
-							return Ar.IsError() == false;
-						},
-						[this](const FArchive& Ar, bool bSuccess) -> void
-						{
-							if (!bSuccess)
-							{
-								// Send a message about losing this chunkdb.
-								MessagePump->SendMessage(FChunkSourceEvent{FChunkSourceEvent::EType::AccessLost, Ar.GetArchiveName()});
-								ChunkDbReloadAttempts.FindOrAdd(Ar.GetArchiveName()).Count = 1;
-							}
-						});
-					// Reset OS intervention setting.
-					ChunkDbSourceHelpers::ResetOsIntervention(PreviousOsIntervention);
-				}
-				if (bIsFirstFail || (SecondsNow - ChunkDbRetryInfo.SecondsAtLastTry) >= Configuration.ChunkDbOpenRetryTime)
+			// We either need to hash the chunk for validation or decompress it into the destination buffer - don't
+			// block IO for this.
+			UE::Tasks::Launch(TEXT("ChunkDbDecompressionAndHash"), 
+				[DataId, CompleteFn, UserPtr, Header, DestinationBuffer, ChunkDataSerialization = ChunkDataSerialization, 
+				CompressedBuffer = MoveTemp(CompressedBuffer), ActivityRecord, ChunkDbChunkSourceStat = ChunkDbChunkSourceStat]()
 				{
-					UE_LOG(LogChunkDbChunkSource, Log, TEXT("Retrying ChunkDb archive which has errored %s"), *ChunkDbFile->GetArchiveName());
-					ChunkDbRetryInfo.SecondsAtLastTry = SecondsNow;
-					// Retry whilst disabling OS intervention. Any OS dialogs if they exist would have been popped during the read that flagged the error.
-					int32 PreviousOsIntervention = ChunkDbSourceHelpers::DisableOsIntervention();
-					TUniquePtr<FArchive> NewChunkDbFile(FileSystem->CreateFileReader(*ChunkDbFile->GetArchiveName()));
-					if (NewChunkDbFile.IsValid())
-					{
-						ChunkDbDataAccess.Archive = MoveTemp(NewChunkDbFile);
-						ChunkDbFile = ChunkDbDataAccess.Archive.Get();
-						bChunkGood = true;
-						ChunkDbRetryInfo.Count = 0;
-						MessagePump->SendMessage(FChunkSourceEvent{FChunkSourceEvent::EType::AccessRegained, ChunkDbFile->GetArchiveName()});
-						// Attempt to regain access to other failed chunkdb files too.
-						TryReopenChunkDbFiles(
-							[this](const FArchive& Ar) -> bool
-							{
-								return Ar.IsError() == true;
-							},
-							[this](const FArchive& Ar, bool bSuccess) -> void
-							{
-								if (bSuccess)
-								{
-									// Send a message about regaining this chunkdb.
-									MessagePump->SendMessage(FChunkSourceEvent{FChunkSourceEvent::EType::AccessRegained, Ar.GetArchiveName()});
-									ChunkDbReloadAttempts.FindOrAdd(Ar.GetArchiveName()).Count = 0;
-								}
-							});
-					}
-					else
-					{
-						++ChunkDbRetryInfo.Count;
-					}
-					// Reset OS intervention setting.
-					ChunkDbSourceHelpers::ResetOsIntervention(PreviousOsIntervention);
-				}
-			}
-			if (bChunkGood)
-			{
-				// Load the chunk.
-				bChunkGood = false;
-				const int64 TotalFileSize = ChunkDbFile->TotalSize();
-				int64 DataEndPoint = ChunkLocation.ByteStart + ChunkLocation.ByteSize;
-				if (TotalFileSize >= DataEndPoint)
-				{
-					ISpeedRecorder::FRecord ActivityRecord;
-					ActivityRecord.CyclesStart = FStatsCollector::GetCycles();
-					if (ChunkDbFile->Tell() != ChunkLocation.ByteStart)
-					{
-						ChunkDbFile->Seek(ChunkLocation.ByteStart);
-					}
-					EChunkLoadResult LoadResult;
-					TUniquePtr<IChunkDataAccess> ChunkDataAccess(ChunkDataSerialization->LoadFromArchive(*ChunkDbFile, LoadResult));
-					ActivityRecord.CyclesEnd = FStatsCollector::GetCycles();
-					const uint64 End = ChunkDbFile->Tell();
-					ActivityRecord.Size = End - ChunkLocation.ByteStart;
-					bChunkGood = LoadResult == EChunkLoadResult::Success && ChunkDataAccess.IsValid() && (End == DataEndPoint);
-					ChunkDbChunkSourceStat->OnLoadComplete(DataId, ChunkDbSourceHelpers::FromSerializer(LoadResult), ActivityRecord);
-					if (bChunkGood)
-					{
-						// Add it to our cache.
-						PlacedInStore.Add(DataId);
-						ChunkStore->Put(DataId, MoveTemp(ChunkDataAccess));
-					}
-				}
-				else
-				{
-					ChunkDbChunkSourceStat->OnLoadComplete(DataId, IChunkDbChunkSourceStat::ELoadResult::LocationOutOfBounds, ISpeedRecorder::FRecord());
-				}
-			}
-		}
-		if (!bChunkGood)
-		{
-			FailedToLoadMessages.Enqueue(DataId);
-		}
-	}
+					ISpeedRecorder::FRecord OurActivityRecord = ActivityRecord; // copy as the capture is const
+					bool bDecompressSucceeded = ChunkDataSerialization->DecompressValidatedRead(Header, DestinationBuffer, CompressedBuffer);
 
-	bool FChunkDbChunkSource::HasFailed(const FGuid& DataId)
-	{
-		FGuid Temp;
-		while (FailedToLoadMessages.Dequeue(Temp))
-		{
-			FailedToLoad.Add(Temp);
-			UnavailableChunks.Add(Temp);
-		}
-		return FailedToLoad.Contains(DataId);
-	}
+					OurActivityRecord.CyclesEnd = FStatsCollector::GetCycles();
 
-	void FChunkDbChunkSource::TryReopenChunkDbFiles(const FArchiveSelector& SelectPredicate, const FArchiveOpenResult& ResultCallback)
-	{
-		for (FChunkDbDataAccess& ChunkDbToCheck : ChunkDbDataAccesses)
-		{
-			TUniquePtr<FArchive>& ArchiveToCheck = ChunkDbToCheck.Archive;
-			if (ArchiveToCheck.IsValid() && SelectPredicate(*ArchiveToCheck.Get()))
-			{
-				TUniquePtr<FArchive> CanOpenFile(FileSystem->CreateFileReader(*ArchiveToCheck->GetArchiveName()));
-				if (CanOpenFile.IsValid())
-				{
-					// Always set in the case that reopening fixes the currently undetected problem.
-					ArchiveToCheck.Reset(CanOpenFile.Release());
-					ResultCallback(*ArchiveToCheck.Get(), true);
-				}
-				else
-				{
-					// Make sure error is set on the archive so that we know to be retrying.
-					ArchiveToCheck->SetError();
-					ResultCallback(*ArchiveToCheck.Get(), false);
+					ChunkDbChunkSourceStat->OnLoadComplete(DataId, 
+						bDecompressSucceeded ? IChunkDbChunkSourceStat::ELoadResult::Success : IChunkDbChunkSourceStat::ELoadResult::CorruptedData,
+						OurActivityRecord);
+
+					CompleteFn.Execute(DataId, false, !bDecompressSucceeded, UserPtr);
 				}
-			}
-		}
+			);
+		};
 	}
 
-	IChunkDbChunkSource* FChunkDbChunkSourceFactory::Create(FChunkDbSourceConfig Configuration, IPlatform* Platform, IFileSystem* FileSystem, IChunkStore* ChunkStore, IChunkReferenceTracker* ChunkReferenceTracker, IChunkDataSerialization* ChunkDataSerialization, IMessagePump* MessagePump, IInstallerError* InstallerError, IChunkDbChunkSourceStat* ChunkDbChunkSourceStat)
+	IConstructorChunkDbChunkSource* IConstructorChunkDbChunkSource::CreateChunkDbSource(FChunkDbSourceConfig&& Configuration, IFileSystem* FileSystem, 
+		const TArray<FGuid>& ChunkAccessOrderList, IChunkDataSerialization* ChunkDataSerialization, IChunkDbChunkSourceStat* ChunkDbChunkSourceStat)
 	{
-		UE_LOG(LogChunkDbChunkSource, Verbose, TEXT("FChunkDbChunkSourceFactory::Create for %d db files"), Configuration.ChunkDbFiles.Num());
-
-		check(Platform != nullptr);
-		check(FileSystem != nullptr);
-		check(ChunkStore != nullptr);
-		check(ChunkReferenceTracker != nullptr);
-		check(ChunkDataSerialization != nullptr);
-		check(MessagePump != nullptr);
-		check(InstallerError != nullptr);
-		check(ChunkDbChunkSourceStat != nullptr);
-		return new FChunkDbChunkSource(MoveTemp(Configuration), Platform, FileSystem, ChunkStore, ChunkReferenceTracker, ChunkDataSerialization, MessagePump, InstallerError, ChunkDbChunkSourceStat);
+		return new FChunkDbChunkSource(MoveTemp(Configuration), FileSystem, ChunkAccessOrderList, ChunkDataSerialization, ChunkDbChunkSourceStat);
 	}
 
 	const TCHAR* ToString(const IChunkDbChunkSourceStat::ELoadResult& LoadResult)
@@ -709,12 +413,8 @@ namespace BuildPatchServices
 		{
 			case IChunkDbChunkSourceStat::ELoadResult::Success:
 				return TEXT("Success");
-			case IChunkDbChunkSourceStat::ELoadResult::MissingHashInfo:
-				return TEXT("MissingHashInfo");
-			case IChunkDbChunkSourceStat::ELoadResult::HashCheckFailed:
-				return TEXT("HashCheckFailed");
-			case IChunkDbChunkSourceStat::ELoadResult::LocationOutOfBounds:
-				return TEXT("LocationOutOfBounds");
+			case IChunkDbChunkSourceStat::ELoadResult::CorruptedData:
+				return TEXT("CorruptedData");
 			case IChunkDbChunkSourceStat::ELoadResult::SerializationError:
 				return TEXT("SerializationError");
 			default:
diff --git a/Engine/Source/Runtime/Online/BuildPatchServices/Private/Installer/ChunkDbChunkSource.h b/Engine/Source/Runtime/Online/BuildPatchServices/Private/Installer/ChunkDbChunkSource.h
--- a/Engine/Source/Runtime/Online/BuildPatchServices/Private/Installer/ChunkDbChunkSource.h
+++ b/Engine/Source/Runtime/Online/BuildPatchServices/Private/Installer/ChunkDbChunkSource.h
@@ -10,24 +10,40 @@ class IBuildInstallerSharedContext;
 
 namespace BuildPatchServices
 {
-	class IPlatform;
 	class IFileSystem;
-	class IChunkStore;
 	class IChunkReferenceTracker;
 	class IChunkDataSerialization;
-	class IMessagePump;
-	class IInstallerError;
 	class IChunkDbChunkSourceStat;
+	class IBuildManifestSet;
+
+	
+	/**
+	* A struct containing the configuration values for a chunkdb chunk source.
+	*/
+	struct FChunkDbSourceConfig
+	{
+		// An array of chunkdb full file paths.
+		TArray<FString> ChunkDbFiles;
+
+		// If true, once we complete a file we delete all the chunkdbs used to create it.
+		bool bDeleteChunkDBAfterUse = false;
+
+		/**
+		* Constructor which sets usual defaults, and takes params for values that cannot use a default.
+		* @param InChunkDbFiles    The chunkdb filename array.
+		*/
+		FChunkDbSourceConfig(const TArray<FString>& InChunkDbFiles)
+		: ChunkDbFiles(InChunkDbFiles)
+		{}
+	};
 
 	/**
 	 * The interface for a chunkdb chunk source, which provides access to chunk data retrieved from chunkdb files.
 	 */
-	class IChunkDbChunkSource
-		: public IChunkSource
-		, public IControllable
+	class IConstructorChunkDbChunkSource : public IConstructorChunkSource
 	{
 	public:
-		virtual ~IChunkDbChunkSource() {}
+		virtual ~IConstructorChunkDbChunkSource() {}
 
 		/**
 		 * Get the set of chunks available in the chunkdbs which were provided to the source.
@@ -41,62 +57,13 @@ namespace BuildPatchServices
 
 		// As above, except use the remaining open chunkdbs for progressive disk space checking.		
 		virtual uint64 GetChunkDbSizesAtIndexes(const TArray<int32>& FileCompletionIndexes, TArray<uint64>& OutChunkDbSizesAtCompletion) const = 0;
-	};
-
-	/**
-	 * A struct containing the configuration values for a chunkdb chunk source.
-	 */
-	struct FChunkDbSourceConfig
-	{
-		// An array of chunkdb full file paths.
-		TArray<FString> ChunkDbFiles;
-		// The minimum number of chunks to load ahead of what is required, depending on store slack.
-		int32 PreFetchMinimum = 10;
-		// The maximum number of chunks to load ahead of what is required, depending on store slack.
-		int32 PreFetchMaximum = 40;
-		// The time in seconds to wait until trying to open a chunkdb file again after we lost the file handle (e.g. due to device eject or network error).
-		float ChunkDbOpenRetryTime = 5.0f;
-		// If true, the loading will not begin until the first Get request is made. It is fairly fundamental to stop loading of chunks until resume
-		// data is processed, but can be special case disabled.
-		bool bBeginLoadsOnFirstGet = true;
-		// The context for allocating shared resources.
-		IBuildInstallerSharedContext* SharedContext = nullptr;
 
-		// If true, once we complete a file we delete all the chunkdbs used to create it.
-		bool bDeleteChunkDBAfterUse = false;
+		virtual void ReportFileCompletion(int32 RemainingChunkCount) =0;
 
-		/**
-		 * Constructor which sets usual defaults, and takes params for values that cannot use a default.
-		 * @param InChunkDbFiles    The chunkdb filename array.
-		 */
-		FChunkDbSourceConfig(TArray<FString> InChunkDbFiles)
-			: ChunkDbFiles(MoveTemp(InChunkDbFiles))
-		{}
+		static IConstructorChunkDbChunkSource* CreateChunkDbSource(FChunkDbSourceConfig&& Configuration, IFileSystem* FileSystem, const TArray<FGuid>& ChunkAccessOrderList, 
+			IChunkDataSerialization* ChunkDataSerialization, IChunkDbChunkSourceStat* ChunkDbChunkSourceStat);
 	};
-
-	/**
-	 * A factory for creating an IChunkDbChunkSource instance.
-	 */
-	class FChunkDbChunkSourceFactory
-	{
-	public:
-		/**
-		 * This implementation will read chunks from provided chunkdbs if they are available. During initialization the provided chunkdb files
-		 * are enumerated to find each available chunk.
-		 * @param Configuration             The configuration struct for this source.
-		 * @param Platform                  The platform access interface.
-		 * @param FileSystem                The service used to open files.
-		 * @param ChunkStore                The chunk store where loaded chunks will be put.
-		 * @param ChunkReferenceTracker     The reference tracker for the installation, used to decide which chunks to fetch when loading.
-		 * @param ChunkDataSerialization    Chunk data serialization implementation.
-		 * @param MessagePump               The message pump to receive messages about source events.
-		 * @param InstallerError            Error tracker where fatal errors will be reported.
-		 * @param ChunkDbChunkSourceStat    The class to receive statistics information.
-		 * @return the new IChunkDbChunkSource instance created.
-		 */
-		static IChunkDbChunkSource* Create(FChunkDbSourceConfig Configuration, IPlatform* Platform, IFileSystem* FileSystem, IChunkStore* ChunkStore, IChunkReferenceTracker* ChunkReferenceTracker, IChunkDataSerialization* ChunkDataSerialization, IMessagePump* MessagePump, IInstallerError* InstallerError, IChunkDbChunkSourceStat* ChunkDbChunkSourceStat);
-	};
-
+	
 	/**
 	 * This interface defines the statistics class required by the chunkdb chunk source. It should be implemented in order to collect
 	 * desired information which is being broadcast by the system.
@@ -112,16 +79,20 @@ namespace BuildPatchServices
 			Success = 0,
 
 			// The hash information was missing.
-			MissingHashInfo,
+			MissingHashInfo_DEPRECATED,
 
 			// The expected data hash for the chunk did not match.
-			HashCheckFailed,
+			HashCheckFailed_DEPRECATED,
 
 			// The chunkdb header specified an invalid chunk location offset or size.
-			LocationOutOfBounds,
-
-			// An unexpected error during serialization.
-			SerializationError
+			LocationOutOfBounds_DEPRECATED,
+
+			// An unexpected error during serialization. This includes header validation
+			// checks like whether the hash is present.
+			SerializationError,
+			
+			// Either the hash didn't match or the decompression call failed.
+			CorruptedData
 		};
 
 	public:
@@ -131,7 +102,8 @@ namespace BuildPatchServices
 		 * Called when a batch of chunks are going to be loaded.
 		 * @param ChunkIds  The ids of each chunk.
 		 */
-		virtual void OnBatchStarted(const TArray<FGuid>& ChunkIds) = 0;
+		UE_DEPRECATED(5.6, "No longer batch loaded")
+		virtual void OnBatchStarted(const TArray<FGuid>& ChunkIds) {};
 
 		/**
 		 * Called each time a chunk load begins.
diff --git a/Engine/Source/Runtime/Online/BuildPatchServices/Private/Installer/ChunkReferenceTracker.cpp b/Engine/Source/Runtime/Online/BuildPatchServices/Private/Installer/ChunkReferenceTracker.cpp
--- a/Engine/Source/Runtime/Online/BuildPatchServices/Private/Installer/ChunkReferenceTracker.cpp
+++ b/Engine/Source/Runtime/Online/BuildPatchServices/Private/Installer/ChunkReferenceTracker.cpp
@@ -35,6 +35,12 @@ namespace BuildPatchServices
 		virtual TArray<FGuid> SelectFromNextReferences(int32 Count, const TFunction<bool(const FGuid&)>& SelectPredicate) const override;
 		virtual bool PopReference(const FGuid& ChunkId) override;
 		virtual int32 GetRemainingChunkCount() const override;
+		virtual int32 GetNextUsageForChunk(const FGuid& ChunkId, int32& OutLastUsageIndex) const override;
+		virtual int32 GetCurrentUsageIndex() const override
+		{
+			return CurrentPosition.load(std::memory_order_relaxed);
+		}
+
 		// IChunkReferenceTracker interface end.
 
 	private:
@@ -87,9 +93,8 @@ namespace BuildPatchServices
 
 	FChunkReferenceTracker::~FChunkReferenceTracker()
 	{
-	}
-
 
+	}
 	
 	TSet<FGuid> FChunkReferenceTracker::GetReferencedChunks() const
 	{
@@ -183,6 +188,28 @@ namespace BuildPatchServices
 		
 		return NextReferences;
 	}
+	
+	int32 FChunkReferenceTracker::GetNextUsageForChunk(const FGuid& ChunkId, int32& OutLastUsageIndex) const
+	{
+		int32 LocalCurrentPosition = CurrentPosition.load(std::memory_order_acquire);
+
+		int32 GuidPosition = Algo::LowerBound(GuidUsagePositions, TPair<FGuid,int32> { ChunkId, LocalCurrentPosition });
+		if (GuidPosition >= GuidUsagePositions.Num() ||
+			GuidUsagePositions[GuidPosition].Key != ChunkId)
+		{
+			return -1;
+		}
+		int32 NextUsage = GuidUsagePositions[GuidPosition].Value;
+
+		while (GuidPosition + 1 < GuidUsagePositions.Num() &&
+			GuidUsagePositions[GuidPosition + 1].Key == ChunkId)
+		{
+			GuidPosition++;
+		}
+
+		OutLastUsageIndex = GuidUsagePositions[GuidPosition].Value;
+		return NextUsage;
+	}
 
 	int32 FChunkReferenceTracker::GetRemainingChunkCount() const
 	{
diff --git a/Engine/Source/Runtime/Online/BuildPatchServices/Private/Installer/ChunkReferenceTracker.h b/Engine/Source/Runtime/Online/BuildPatchServices/Private/Installer/ChunkReferenceTracker.h
--- a/Engine/Source/Runtime/Online/BuildPatchServices/Private/Installer/ChunkReferenceTracker.h
+++ b/Engine/Source/Runtime/Online/BuildPatchServices/Private/Installer/ChunkReferenceTracker.h
@@ -77,6 +77,12 @@ namespace BuildPatchServices
 		// Return the ordered list we are consuming.
 		// copies the _remaining_ items!
 		virtual void CopyOutOrderedUseList(TArray<FGuid>& OutUseList) const = 0;
+
+		// Returns the next usage index in _advancing order_ for the chunkid, or -1 if it won't
+		// be used again.
+		virtual int32 GetNextUsageForChunk(const FGuid& ChunkId, int32& OutLastUsageIndex) const = 0;
+
+		virtual int32 GetCurrentUsageIndex() const = 0;
 	};
 
 	/**
diff --git a/Engine/Source/Runtime/Online/BuildPatchServices/Private/Installer/ChunkSource.h b/Engine/Source/Runtime/Online/BuildPatchServices/Private/Installer/ChunkSource.h
--- a/Engine/Source/Runtime/Online/BuildPatchServices/Private/Installer/ChunkSource.h
+++ b/Engine/Source/Runtime/Online/BuildPatchServices/Private/Installer/ChunkSource.h
@@ -3,12 +3,30 @@
 
 #include "Templates/Function.h"
 #include "Containers/Set.h"
+#include "Delegates/Delegate.h"
+#include "Memory/MemoryFwd.h"
 #include "Misc/Guid.h"
 
 namespace BuildPatchServices
 {
 	class IChunkDataAccess;
 
+	// This class is for providing chunks to a threaded file constructor that manages its own backing store.
+	class IConstructorChunkSource
+	{
+	public:
+		virtual ~IConstructorChunkSource() {}
+
+		DECLARE_DELEGATE_FourParams(FChunkRequestCompleteDelegate, const FGuid& DataId, bool bAborted, bool bFailedToRead, void* UserPtr);
+
+		typedef TUniqueFunction<void(bool bIsAborted)> FRequestProcessFn;
+
+		virtual FRequestProcessFn CreateRequest(const FGuid& DataId, FMutableMemoryView DestinationBuffer, void* UserPtr, FChunkRequestCompleteDelegate CompleteFn) = 0;
+
+		// Return the reference index at which we can no longer get this data id from this source.
+		virtual int32 GetChunkUnavailableAt(const FGuid& DataId) const = 0;
+	};
+
 	/**
 	 * An interface providing basic access to retrieving chunk data.
 	 */
@@ -45,9 +63,5 @@ namespace BuildPatchServices
 		 * @param Callback          The function to call with the set of chunks no longer still available.
 		 */
 		virtual void SetUnavailableChunksCallback(TFunction<void(TSet<FGuid>)> Callback) = 0;
-
-		// Inform the chunk source(s) that a file has been completed so we can clean up since we know we
-		// won't be going backwards for retries/resumes.
-		virtual void ReportFileCompletion() {}
 	};
 }
diff --git a/Engine/Source/Runtime/Online/BuildPatchServices/Private/Installer/CloudChunkSource.cpp b/Engine/Source/Runtime/Online/BuildPatchServices/Private/Installer/CloudChunkSource.cpp
--- a/Engine/Source/Runtime/Online/BuildPatchServices/Private/Installer/CloudChunkSource.cpp
+++ b/Engine/Source/Runtime/Online/BuildPatchServices/Private/Installer/CloudChunkSource.cpp
@@ -20,6 +20,8 @@
 #include "Interfaces/IBuildInstallerSharedContext.h"
 #include "BuildPatchUtil.h"
 #include "Installer/Statistics/DownloadServiceStatistics.h"
+#include "Memory/SharedBuffer.h"
+#include "Tasks/Task.h"
 
 DECLARE_LOG_CATEGORY_EXTERN(LogCloudChunkSource, Log, All);
 DEFINE_LOG_CATEGORY(LogCloudChunkSource);
@@ -89,7 +91,40 @@ namespace BuildPatchServices
 		ImmediateValueCount += 1.0L;
 		TotalValueCount += 1.0L;
 	}
+
+	static float GetRetryDelay(const TArray<float>& RetryDelayTimes, int32 RetryNum)
+	{
+		const int32 RetryTimeIndex = FMath::Clamp<int32>(RetryNum - 1, 0, RetryDelayTimes.Num() - 1);
+		return RetryDelayTimes[RetryTimeIndex];
+	}
 	
+	static EBuildPatchDownloadHealth GetDownloadHealth(bool bIsDisconnected, const TArray<float>& HealthPercentages, float ChunkSuccessRate)
+	{
+		EBuildPatchDownloadHealth DownloadHealth;
+		if (bIsDisconnected)
+		{
+			DownloadHealth = EBuildPatchDownloadHealth::Disconnected;
+		}
+		else if (ChunkSuccessRate >= HealthPercentages[(int32)EBuildPatchDownloadHealth::Excellent])
+		{
+			DownloadHealth = EBuildPatchDownloadHealth::Excellent;
+		}
+		else if (ChunkSuccessRate >= HealthPercentages[(int32)EBuildPatchDownloadHealth::Good])
+		{
+			DownloadHealth = EBuildPatchDownloadHealth::Good;
+		}
+		else if (ChunkSuccessRate >= HealthPercentages[(int32)EBuildPatchDownloadHealth::OK])
+		{
+			DownloadHealth = EBuildPatchDownloadHealth::OK;
+		}
+		else
+		{
+			DownloadHealth = EBuildPatchDownloadHealth::Poor;
+		}
+		return DownloadHealth;
+	}
+
+
 	/**
 	 * The concrete implementation of ICloudChunkSource
 	 */
@@ -156,8 +191,6 @@ namespace BuildPatchServices
 	private:
 		void EnsureAquiring(const FGuid& DataId);
 
-		float GetRetryDelay(int32 RetryNum);
-		EBuildPatchDownloadHealth GetDownloadHealth(bool bIsDisconnected, float ChunkSuccessRate);
 		FGuid GetNextTask(const TMap<FGuid, FTaskInfo>& TaskInfos, const TMap<int32, FGuid>& InFlightDownloads, const TSet<FGuid>& TotalRequiredChunks, const TSet<FGuid>& PriorityRequests, const TSet<FGuid>& FailedDownloads, const TSet<FGuid>& Stored, TArray<FGuid>& DownloadQueue, EBuildPatchDownloadHealth DownloadHealth);
 		void OnDownloadProgress(int32 RequestId, uint64 BytesSoFar);
 		void OnDownloadComplete(int32 RequestId, const FDownloadRef& Download);
@@ -338,38 +371,6 @@ namespace BuildPatchServices
 		RequestedDownloads.Add(DataId);
 	}
 
-	float FCloudChunkSource::GetRetryDelay(int32 RetryNum)
-	{
-		const int32 RetryTimeIndex = FMath::Clamp<int32>(RetryNum - 1, 0, Configuration.RetryDelayTimes.Num() - 1);
-		return Configuration.RetryDelayTimes[RetryTimeIndex];
-	}
-
-	EBuildPatchDownloadHealth FCloudChunkSource::GetDownloadHealth(bool bIsDisconnected, float ChunkSuccessRate)
-	{
-		EBuildPatchDownloadHealth DownloadHealth;
-		if (bIsDisconnected)
-		{
-			DownloadHealth = EBuildPatchDownloadHealth::Disconnected;
-		}
-		else if (ChunkSuccessRate >= Configuration.HealthPercentages[(int32)EBuildPatchDownloadHealth::Excellent])
-		{
-			DownloadHealth = EBuildPatchDownloadHealth::Excellent;
-		}
-		else if (ChunkSuccessRate >= Configuration.HealthPercentages[(int32)EBuildPatchDownloadHealth::Good])
-		{
-			DownloadHealth = EBuildPatchDownloadHealth::Good;
-		}
-		else if (ChunkSuccessRate >= Configuration.HealthPercentages[(int32)EBuildPatchDownloadHealth::OK])
-		{
-			DownloadHealth = EBuildPatchDownloadHealth::OK;
-		}
-		else
-		{
-			DownloadHealth = EBuildPatchDownloadHealth::Poor;
-		}
-		return DownloadHealth;
-	}
-
 	FGuid FCloudChunkSource::GetNextTask(const TMap<FGuid, FTaskInfo>& TaskInfos, const TMap<int32, FGuid>& InFlightDownloads, const TSet<FGuid>& TotalRequiredChunks, const TSet<FGuid>& PriorityRequests, const TSet<FGuid>& FailedDownloads, const TSet<FGuid>& Stored, TArray<FGuid>& DownloadQueue, EBuildPatchDownloadHealth DownloadHealth)
 	{
 		// Check for aborting.
@@ -391,7 +392,7 @@ namespace BuildPatchServices
 		{
 			const FTaskInfo& FailedDownload = TaskInfos[*FailedIt];
 			const double SecondsSinceFailure = SecondsNow - FailedDownload.SecondsAtFail;
-			if (SecondsSinceFailure >= GetRetryDelay(FailedDownload.RetryNum))
+			if (SecondsSinceFailure >= GetRetryDelay(Configuration.RetryDelayTimes, FailedDownload.RetryNum))
 			{
 				ChunkToRetry = *FailedIt;
 			}
@@ -686,7 +687,7 @@ namespace BuildPatchServices
 			const double SecondsSinceData = FStatsCollector::CyclesToSeconds(FStatsCollector::GetCycles() - CyclesAtLastData);
 			const bool bReportAsDisconnected = (bAllDownloadsRetrying && SecondsSinceData > Configuration.DisconnectedDelay);
 			const float SuccessRate = ChunkSuccessRate.GetOverall();
-			EBuildPatchDownloadHealth OverallDownloadHealth = GetDownloadHealth(bReportAsDisconnected, SuccessRate);
+			EBuildPatchDownloadHealth OverallDownloadHealth = GetDownloadHealth(bReportAsDisconnected, Configuration.HealthPercentages, SuccessRate);
 			if (TrackedDownloadHealth != OverallDownloadHealth)
 			{
 				TrackedDownloadHealth = OverallDownloadHealth;
@@ -697,7 +698,7 @@ namespace BuildPatchServices
 				CloudChunkSourceStat->OnSuccessRateUpdated(SuccessRate);
 			}
 			const float ImmediateSuccessRate = ChunkSuccessRate.GetImmediate();
-			EBuildPatchDownloadHealth ImmediateDownloadHealth = GetDownloadHealth(bReportAsDisconnected, ImmediateSuccessRate);
+			EBuildPatchDownloadHealth ImmediateDownloadHealth = GetDownloadHealth(bReportAsDisconnected, Configuration.HealthPercentages, ImmediateSuccessRate);
 			// Kick off new downloads.
 			if (bDownloadsStarted)
 			{
@@ -801,6 +802,7 @@ namespace BuildPatchServices
 		CompletedDownloads.Add(RequestId, Download);
 	}
 
+	
 	ICloudChunkSource* FCloudChunkSourceFactory::Create(FCloudSourceConfig Configuration, IPlatform* Platform, IChunkStore* ChunkStore, IDownloadService* DownloadService, IChunkReferenceTracker* ChunkReferenceTracker, IChunkDataSerialization* ChunkDataSerialization, IMessagePump* MessagePump, IInstallerError* InstallerError, IDownloadConnectionCount* ConnectionCount, ICloudChunkSourceStat* CloudChunkSourceStat, IBuildManifestSet* ManifestSet, TSet<FGuid> InitialDownloadSet)
 	{
 		UE_LOG(LogCloudChunkSource, Verbose, TEXT("FCloudChunkSourceFactory::Create for %d roots"), Configuration.CloudRoots.Num());
@@ -816,4 +818,456 @@ namespace BuildPatchServices
 		check(CloudChunkSourceStat != nullptr);
 		return new FCloudChunkSource(MoveTemp(Configuration), Platform, ChunkStore, DownloadService, ChunkReferenceTracker, ChunkDataSerialization, MessagePump, InstallerError, ConnectionCount, CloudChunkSourceStat, ManifestSet, MoveTemp(InitialDownloadSet));
 	}
-}
+
+	//-------------------------------------------------------------------------
+	//
+	// ConstructorCloudChunkSource below here.
+	//
+	//-------------------------------------------------------------------------
+	class FConstructorCloudChunkSource : public IConstructorCloudChunkSource
+	{
+	private:
+		/**
+		* This is a wrapper class for binding thread safe shared ptr delegates for the download service, without having to enforce that
+		* this service should be made using TShared* reference controllers.
+		*/
+		class FDownloadDelegates
+		{
+		public:
+			FDownloadDelegates(FConstructorCloudChunkSource& InCloudChunkSource)
+			: CloudSource(InCloudChunkSource)
+			{
+			}
+
+		public:
+			void OnDownloadProgress(int32 RequestId, uint64 BytesSoFar)
+			{
+				CloudSource.OnDownloadProgress(RequestId, BytesSoFar);
+			}
+			void OnDownloadComplete(int32 RequestId, const FDownloadRef& Download)
+			{
+				CloudSource.OnDownloadComplete(RequestId, Download);
+			}
+
+		private:
+			FConstructorCloudChunkSource& CloudSource;
+		};
+		
+		struct FCloudRead
+		{
+			FGuid DataId;
+			FMutableMemoryView Destination;
+			FChunkRequestCompleteDelegate CompleteFn;
+			void* UserPtr;
+
+			bool bQueuedForDownload = false;
+			int32 CloudDirUsed = 0;
+			FString UrlUsed;
+			int32 RetryNum = 0;
+			int32 ExpectedSize = 0;
+			int32 RequestId = 0;
+			double SecondsAtRequested = 0;
+
+			// The time when we want to launch this retry, if RetryNum != 0
+			double RetryTime = 0;
+		};
+
+	public:
+		FConstructorCloudChunkSource(FConstructorCloudChunkSourceConfig&& InConfiguration, IDownloadService* InDownloadService, IChunkDataSerialization* InChunkDataSerialization,
+			ICloudChunkSourceStat* InCloudChunkSourceStat, IBuildManifestSet* InManifestSet, IDownloadConnectionCount* InDownloadCount)
+			: Configuration(MoveTemp(InConfiguration))
+			, DownloadDelegates(MakeShareable(new FDownloadDelegates(*this)))
+			, OnDownloadProgressDelegate(FDownloadProgressDelegate::CreateThreadSafeSP(DownloadDelegates, &FDownloadDelegates::OnDownloadProgress))
+			, OnDownloadCompleteDelegate(FDownloadCompleteDelegate::CreateThreadSafeSP(DownloadDelegates, &FDownloadDelegates::OnDownloadComplete))
+			, DownloadService(InDownloadService)
+			, ChunkDataSerialization(InChunkDataSerialization)
+			, CloudChunkSourceStat(InCloudChunkSourceStat)
+			, ManifestSet(InManifestSet)
+			, DownloadCount(InDownloadCount)
+		{
+			CloudDirFailureCount.SetNumZeroed(Configuration.CloudRoots.Num());
+		}
+
+		virtual ~FConstructorCloudChunkSource()
+		{
+			// Release any unqueued reads.
+			for (FCloudRead* Read : ReadQueue)
+			{
+				delete Read;
+			}
+
+			// Abort any downloads.
+			for (TPair<int32, FCloudRead*> Pair : InFlightDownloads)
+			{
+				FCloudRead* Read = Pair.Value;
+				DownloadService->RequestAbandon(Read->RequestId);
+				delete Read;
+			}
+		}
+		
+		virtual FRequestProcessFn CreateRequest(const FGuid& DataId, FMutableMemoryView DestinationBuffer, void* UserPtr, FChunkRequestCompleteDelegate CompleteFn) override
+		{
+			// We don't have a request that can be serviced directly because of how we have multiple in flight,
+			// so we return an empty request function and let the tick do the work.
+			FCloudRead* Read = new FCloudRead();
+			Read->DataId = DataId;
+			Read->Destination = DestinationBuffer;
+			Read->CompleteFn = CompleteFn;
+			Read->UserPtr = UserPtr;
+			ReadQueue.Add(Read);
+
+			// We need to make sure that we get ticked to start the request as we might not be called from
+			// the dispatch thread.
+			WakeupMainThreadFn();
+
+			return [](bool){return;};
+		}
+	
+		virtual int32 GetChunkUnavailableAt(const FGuid& DataId) const override
+		{
+			// We can always redownload but it's never kept local so it's immediately unavailable
+			return 0;
+		}
+
+		virtual void Tick(bool bStartNewDownloads, uint32& OutTimeToNextTickMs) override;
+		virtual void SetWakeupFunction(TUniqueFunction<void()>&& InWakeupMainThreadFn) override
+		{
+			WakeupMainThreadFn = MoveTemp(InWakeupMainThreadFn);
+		}
+
+		virtual void PostRequiredByteCount(uint64 InDownloadExpected) override
+		{
+			CloudChunkSourceStat->OnRequiredDataUpdated(InDownloadExpected);
+		}
+
+	private:	
+		
+		void OnDownloadProgress(int32 RequestId, uint64 BytesSoFar)
+		{
+			CyclesAtLastData.store(FStatsCollector::GetCycles(), std::memory_order_relaxed);
+		}
+
+		void OnDownloadComplete(int32 RequestId, const FDownloadRef& Download)
+		{
+			FScopeLock ScopeLock(&CompletedDownloadsCS);
+			CompletedDownloads.Add(RequestId, Download);
+		}
+
+		FConstructorCloudChunkSourceConfig Configuration;
+		TSharedRef<FDownloadDelegates, ESPMode::ThreadSafe> DownloadDelegates;
+		FDownloadProgressDelegate OnDownloadProgressDelegate;
+		FDownloadCompleteDelegate OnDownloadCompleteDelegate;
+
+		// Handles to various systems/data that we require to be valid across execution duration.
+		IDownloadService* DownloadService;
+		IChunkDataSerialization* ChunkDataSerialization;
+		ICloudChunkSourceStat* CloudChunkSourceStat;
+		IBuildManifestSet* ManifestSet;
+		IDownloadConnectionCount* DownloadCount;
+
+		// Called when we get a completed download and we want to get ticked again.
+		TUniqueFunction<void()> WakeupMainThreadFn;
+
+		// Tracking health and connection state
+		std::atomic<int64> CyclesAtLastData = 0;
+
+		// Communication from download thread to processing thread.
+		FCriticalSection CompletedDownloadsCS;
+		TMap<int32, FDownloadRef> CompletedDownloads;
+
+		// If we start getting failures on our downloads, we track which ones
+		// fail and avoid them until everything goes bad. Initially we just hit the
+		// first directory.
+		int32 CurrentBestCloudDir = 0;
+		TArray<int32> CloudDirFailureCount;
+
+		// Read pointers are in one of these until they are either destroyed or passed
+		// to the validation task.
+		TArray<FCloudRead*> ReadQueue;
+		TMap<int32, FCloudRead*> InFlightDownloads;
+
+		// \todo does this thing have any recency bias? Surely we only care about the last few
+		// seconds of success...
+		FChunkSuccessRate ChunkSuccessTracker;
+
+		// Change recognition
+		EBuildPatchDownloadHealth LastSeenHealth = EBuildPatchDownloadHealth::NUM_Values;
+		float LastSeenSuccessRate = -1.0f;
+
+		double StartTime = -1;
+		uint64 TotalBytes = 0;
+	};
+
+	//
+	// Since we need to have a lot of requests in flight at the same time _and_ we need to
+	// synchronize across them to manage how many are in flight at any given moment, we have
+	// to have a touchpoint on the main thread to dispatch and manage everything. This thread
+	// should do very little work! It should just be brokering things, not accomplishing anything.
+	//
+	void FConstructorCloudChunkSource::Tick(bool bStartNewDownloads, uint32& OutTimeToNextTickMs)
+	{
+		int32 StartingDownloadCount = InFlightDownloads.Num();
+			
+		//UE_LOG(LogTemp, Warning, TEXT("Cloud: Active: %d, Pending %d"), InFlightDownloads.Num(), ReadQueue.Num());
+
+		// Process completed downloads.
+		TMap<int32, FDownloadRef> FrameCompletedDownloads;
+		{
+			CompletedDownloadsCS.Lock();
+			FrameCompletedDownloads = MoveTemp(CompletedDownloads);
+			CompletedDownloadsCS.Unlock();
+		}
+
+		{
+			TRACE_CPUPROFILER_EVENT_SCOPE(CDT_ProcessCompleted);
+			for (const TPair<int32, FDownloadRef>& FrameCompletedDownload : FrameCompletedDownloads)
+			{
+				TRACE_CPUPROFILER_EVENT_SCOPE(CDT_ProcessOne);
+				const int32& RequestId = FrameCompletedDownload.Key;
+				const FDownloadRef& Download = FrameCompletedDownload.Value;
+
+				FCloudRead* Read = InFlightDownloads[RequestId];
+				InFlightDownloads.Remove(RequestId);
+
+				Read->bQueuedForDownload = false;
+
+				//UE_LOG(LogTemp, Warning, TEXT("Downloaded chunk: %s"), *WriteToString<40>(Read->DataId));
+
+				bool bDownloadSuccess = Download->ResponseSuccessful();
+				if (bDownloadSuccess)
+				{
+					const TArray<uint8>& DownloadedData = Download->GetData();
+
+					TotalBytes += DownloadedData.Num();
+
+					double MegaBytesPerSecond = (TotalBytes / (FStatsCollector::GetSeconds() - StartTime)) / (1024*1024);
+					//UE_LOG(LogTemp, Warning, TEXT("Rate: %.2f MB/s"), MegaBytesPerSecond);
+
+					CloudChunkSourceStat->OnReceivedDataUpdated(TotalBytes);
+
+					FMemoryReader ReaderThunk(DownloadedData);
+
+					// If the chunk is uncompressed it can be directly copied to the output here, otherwise it'll get decompressed
+					// during the verification task;
+					FChunkHeader Header;
+					FUniqueBuffer CompressedBuffer;
+					bool Result = ChunkDataSerialization->ValidateAndRead(ReaderThunk, Read->Destination, Header, CompressedBuffer);
+
+					if (!Result)
+					{
+						// The header or chunk data was bad.				
+						Read->CompleteFn.Execute(Read->DataId, false, true, Read->UserPtr);
+						delete Read;
+					}
+					else
+					{
+						// Older chunks might not have a sha hash internally, but the manifest is new and might have it,
+						// so we can inject it.
+						if (!EnumHasAnyFlags(Header.HashType, EChunkHashFlags::Sha1))
+						{
+							if (ManifestSet->GetChunkShaHash(Read->DataId, Header.SHAHash))
+							{
+								EnumAddFlags(Header.HashType, EChunkHashFlags::Sha1);
+							}
+						}
+					
+						// We either need to hash the chunk for validation or decompress it in to the destination buffer - don't
+						// block IO for this.
+						UE::Tasks::Launch(TEXT("CloudDecompressionAndHash"), [Read, Header, CloudChunkSourceStat = CloudChunkSourceStat, ChunkDataSerialization = ChunkDataSerialization, CompressedBuffer = MoveTemp(CompressedBuffer)]()
+							{
+								bool bDecompressSucceeded = ChunkDataSerialization->DecompressValidatedRead(Header, Read->Destination, CompressedBuffer);
+
+								if (bDecompressSucceeded)
+								{
+									CloudChunkSourceStat->OnDownloadSuccess(Read->DataId);
+								}
+								else
+								{
+									// \todo this wants to know what the actual internal error was for some reason. idk why.
+									CloudChunkSourceStat->OnDownloadCorrupt(Read->DataId, Read->UrlUsed, EChunkLoadResult::HashCheckFailed);
+								}
+
+								Read->CompleteFn.Execute(Read->DataId, false, !bDecompressSucceeded, Read->UserPtr);
+								delete Read;
+							}
+						);
+					}
+
+					ChunkSuccessTracker.AddSuccess();
+				}
+				else
+				{
+					CloudChunkSourceStat->OnDownloadFailed(Read->DataId, Read->UrlUsed);
+					UE_LOG(LogCloudChunkSource, Error, TEXT("FAILED: %s"), *Read->UrlUsed);
+					
+					// Mark this CDN as failed.
+					{
+						CloudDirFailureCount[Read->CloudDirUsed]++;
+						if (CloudDirFailureCount[Read->CloudDirUsed] > (100 << 20))
+						{
+							// Cap to prevent wrap. I think this is technically impossible due to the time it would take to 
+							// get here but...
+							CloudDirFailureCount[Read->CloudDirUsed] = 100 << 20;
+						}
+
+						// Find who has failed the least, be sure to take equivalents in the initial specified order.
+						// We expect this to be like 5 entries.
+						int32 MinFailCount = CloudDirFailureCount[0];
+						int32 MinAtIndex = 0;
+						for (int32 CloudDirSeek = 1; CloudDirSeek < CloudDirFailureCount.Num(); CloudDirSeek++)
+						{
+							if (CloudDirFailureCount[CloudDirSeek] < MinFailCount)
+							{
+								MinFailCount = CloudDirFailureCount[CloudDirSeek];
+								MinAtIndex = CloudDirSeek;
+							}
+						}
+
+						CurrentBestCloudDir = MinAtIndex;
+						UE_LOG(LogCloudChunkSource, Warning, TEXT("CDN %s failed download, updating CDN selection to: %s"), *Configuration.CloudRoots[Read->CloudDirUsed], *Configuration.CloudRoots[MinAtIndex]);
+					}
+					ChunkSuccessTracker.AddFail();
+
+					// Update retry.
+					Read->RetryNum++;
+					if (Configuration.MaxRetryCount >= 0 && Read->RetryNum >= Configuration.MaxRetryCount)
+					{
+						// Fail the request.
+						Read->CompleteFn.Execute(Read->DataId, false, true, Read->UserPtr);
+						delete Read;
+					}
+					else
+					{
+						// Set retry time and put the task back in the queue.
+						Read->RetryTime = FStatsCollector::GetSeconds() + GetRetryDelay(Configuration.RetryDelayTimes, Read->RetryNum);
+						ReadQueue.Add(Read);
+					}
+				}
+			} // end each completed
+		} // end process completed			
+
+		// Update connection status and health.
+		EBuildPatchDownloadHealth DownloadHealth;
+		{
+			bool bAllDownloadsRetrying = ReadQueue.Num() > 0;
+			for (FCloudRead* Read : ReadQueue)
+			{
+				if (!Read->RetryNum)
+				{
+					bAllDownloadsRetrying = false;
+				}
+			}
+
+			const double SecondsSinceData = FStatsCollector::CyclesToSeconds(FStatsCollector::GetCycles() - CyclesAtLastData.load(std::memory_order_relaxed));
+			const bool bReportAsDisconnected = (bAllDownloadsRetrying && SecondsSinceData > Configuration.DisconnectedDelay);
+
+			const float SuccessRate = ChunkSuccessTracker.GetOverall();
+			EBuildPatchDownloadHealth OverallDownloadHealth = GetDownloadHealth(bReportAsDisconnected, Configuration.HealthPercentages, SuccessRate);
+			if (LastSeenHealth != OverallDownloadHealth)
+			{
+				LastSeenHealth = OverallDownloadHealth;
+				CloudChunkSourceStat->OnDownloadHealthUpdated(LastSeenHealth);
+			}
+			if (LastSeenSuccessRate != SuccessRate)
+			{
+				LastSeenSuccessRate = SuccessRate;
+				CloudChunkSourceStat->OnSuccessRateUpdated(SuccessRate);
+			}
+
+			DownloadHealth = GetDownloadHealth(bReportAsDisconnected, Configuration.HealthPercentages, ChunkSuccessTracker.GetImmediate());
+		}
+
+		// Start new downloads, if we aren't gated or paused
+		while (bStartNewDownloads)
+		{
+			const uint32 MaxDownloads = FMath::Max(1U, DownloadCount->GetAdjustedCount(InFlightDownloads.Num(), DownloadHealth));
+			if ((uint32)InFlightDownloads.Num() >= MaxDownloads ||
+				ReadQueue.Num() == 0)
+			{
+				break;
+			}
+
+			FCloudRead* Read = ReadQueue[0];
+			ReadQueue.RemoveAt(0);
+
+			Read->CloudDirUsed = CurrentBestCloudDir;
+			Read->bQueuedForDownload = true;
+			Read->UrlUsed = Configuration.CloudRoots[CurrentBestCloudDir] / ManifestSet->GetDataFilename(Read->DataId);
+			Read->ExpectedSize = ManifestSet->GetDownloadSize(Read->DataId);
+			Read->SecondsAtRequested = FStatsCollector::GetSeconds();
+			Read->RequestId = DownloadService->RequestFile(Read->UrlUsed, OnDownloadCompleteDelegate, OnDownloadProgressDelegate);
+
+			InFlightDownloads.Add(Read->RequestId, Read);
+
+			CloudChunkSourceStat->OnDownloadRequested(Read->DataId);
+
+			if (StartTime == -1)
+			{
+				StartTime = Read->SecondsAtRequested;
+			}
+		}
+
+		if (StartingDownloadCount != InFlightDownloads.Num())
+		{
+			//UE_LOG(LogTemp, Warning, TEXT("count: %d"), InFlightDownloads.Num());
+			CloudChunkSourceStat->OnActiveRequestCountUpdated(InFlightDownloads.Num());
+		}
+
+		
+		// Determine how long we should wait. Normally we want to just wait until we get triggered,
+		// either because we completed a download or because a new request came in. However if we have 
+		// a retry that we want to requeue in a bit we just need to check back in.
+		uint32 WaitTimeMs = TNumericLimits<uint32>::Max(); // default infinite
+		if (ReadQueue.Num())
+		{
+			// This should never be that long. We should only have what has been requested
+			// for the current block the constructor is writing.
+			//
+			// Note that we only check retries - for reads that have never queued, the reason they
+			// aren't queued is because we ran out of download slots and we will check those
+			// when we get woken up by a completed download.
+			double ClosestRetry = -1;
+			for (FCloudRead* Read : ReadQueue)
+			{					
+				if (Read->RetryNum && 
+					(ClosestRetry == -1 || Read->RetryTime < ClosestRetry))
+				{
+					ClosestRetry = Read->RetryTime;
+				}
+			}
+
+			if (ClosestRetry != -1)
+			{
+				double TimeToClosestRetry = ClosestRetry - FStatsCollector::GetSeconds();
+				if (TimeToClosestRetry <= 0)
+				{
+					WaitTimeMs = 0;
+				}
+				else
+				{
+					WaitTimeMs = (uint32)(TimeToClosestRetry * 1000);
+				}
+			}
+		}
+		if (InFlightDownloads.Num())
+		{
+			// If we have in flight downloads we want to wake up every so often to
+			// check for disconnections. Every second should do the trick.
+			if (WaitTimeMs == TNumericLimits<uint32>::Max() ||
+				WaitTimeMs > 1000)
+			{
+				WaitTimeMs = 1000;
+			}
+		}
+
+		OutTimeToNextTickMs = WaitTimeMs;		
+	}
+
+
+	IConstructorCloudChunkSource* IConstructorCloudChunkSource::CreateCloudSource(FConstructorCloudChunkSourceConfig&& Configuration, IDownloadService* InDownloadService, 
+		IChunkDataSerialization* InChunkDataSerialization, IDownloadConnectionCount* InDownloadCount, ICloudChunkSourceStat* InCloudChunkSourceStat, 
+		IBuildManifestSet* InManifestSet)
+	{
+		return new FConstructorCloudChunkSource(MoveTemp(Configuration), InDownloadService, InChunkDataSerialization, InCloudChunkSourceStat, InManifestSet, InDownloadCount);
+	}
+}
\ No newline at end of file
diff --git a/Engine/Source/Runtime/Online/BuildPatchServices/Private/Installer/CloudChunkSource.h b/Engine/Source/Runtime/Online/BuildPatchServices/Private/Installer/CloudChunkSource.h
--- a/Engine/Source/Runtime/Online/BuildPatchServices/Private/Installer/CloudChunkSource.h
+++ b/Engine/Source/Runtime/Online/BuildPatchServices/Private/Installer/CloudChunkSource.h
@@ -23,6 +23,53 @@ namespace BuildPatchServices
 	class IDownloadServiceStatistics;
 	class IBuildManifestSet;
 
+	struct FConstructorCloudChunkSourceConfig
+	{
+		// An array of cloud root paths, supporting HTTP(s) and file access. HTTP(s) roots must begin with the protocol.
+		TArray<FString> CloudRoots;
+		// The maximum number of times that a single chunk should retry, before registering a fatal error.
+		// Infinite can be specified as < 0.
+		int32 MaxRetryCount = 6;
+		// Array of times in seconds, representing the time between each retry upon failure. The last entry will be used
+		// indefinitely once it is reached.
+		TArray<float> RetryDelayTimes;
+		// Configures what success rate values apply to which EBuildPatchDownloadHealth value, there should be
+		// EBuildPatchDownloadHealth::NUM_Values entries in this array.
+		TArray<float> HealthPercentages;
+		// When all requests are failing, how many seconds before a success until we determine the state as disconnected.
+		float DisconnectedDelay = 5.0f;
+
+		FConstructorCloudChunkSourceConfig(const TArray<FString>& InCloudRoots)
+		: CloudRoots(InCloudRoots)
+		{
+			const float RetryFloats[] = {0.5f, 1.0f, 1.0f, 3.0f, 3.0f, 10.0f, 10.0f, 20.0f, 20.0f, 30.0f};
+			RetryDelayTimes.Empty(UE_ARRAY_COUNT(RetryFloats));
+			RetryDelayTimes.Append(RetryFloats, UE_ARRAY_COUNT(RetryFloats));
+			const float HealthFloats[] = {0.0f, 0.0f, 0.9f, 0.99f, 1.0f};
+			check((int32)EBuildPatchDownloadHealth::NUM_Values == UE_ARRAY_COUNT(HealthFloats));
+			HealthPercentages.Empty(UE_ARRAY_COUNT(HealthFloats));
+			HealthPercentages.Append(HealthFloats, UE_ARRAY_COUNT(HealthFloats));
+		}
+	};
+
+	class IConstructorCloudChunkSource : public IConstructorChunkSource
+	{
+	public:
+		virtual ~IConstructorCloudChunkSource() {}
+
+		virtual void Tick(bool bStartNewDownloads, uint32& OutTimeToNextTickMs) = 0;
+
+		// When downloads are complete we need a way to alert the main thread.
+		virtual void SetWakeupFunction(TUniqueFunction<void()>&& WakeupMainThreadFn) = 0;
+
+		// Pass in total download requirements for the install (including already completed downloads)
+		virtual void PostRequiredByteCount(uint64 InDownloadExpected) = 0;
+
+		static IConstructorCloudChunkSource* CreateCloudSource(FConstructorCloudChunkSourceConfig&& Configuration, IDownloadService* InDownloadService, 
+			IChunkDataSerialization* InChunkDataSerialization, IDownloadConnectionCount* InDownloadCount, ICloudChunkSourceStat* InCloudChunkSourceStat, 
+			IBuildManifestSet* InManifestSet);
+	};
+
 	/**
 	 * The interface for a cloud chunk source, which provides access to chunk data retrieved from provided cloud roots.
 	 */
@@ -78,8 +125,8 @@ namespace BuildPatchServices
 		 * Constructor which sets usual defaults, and takes params for values that cannot use a default.
 		 * @param InCloudRoots  The cloud roots array
 		 */
-		FCloudSourceConfig(TArray<FString> InCloudRoots)
-			: CloudRoots(MoveTemp(InCloudRoots))
+		FCloudSourceConfig(const TArray<FString>& InCloudRoots)
+			: CloudRoots(InCloudRoots)
 		{
 			const float RetryFloats[] = {0.5f, 1.0f, 1.0f, 3.0f, 3.0f, 10.0f, 10.0f, 20.0f, 20.0f, 30.0f};
 			RetryDelayTimes.Empty(UE_ARRAY_COUNT(RetryFloats));
diff --git a/Engine/Source/Runtime/Online/BuildPatchServices/Private/Installer/DiskChunkStore.cpp b/Engine/Source/Runtime/Online/BuildPatchServices/Private/Installer/DiskChunkStore.cpp
deleted file mode 100644
diff --git a/Engine/Source/Runtime/Online/BuildPatchServices/Private/Installer/DiskChunkStore.h b/Engine/Source/Runtime/Online/BuildPatchServices/Private/Installer/DiskChunkStore.h
deleted file mode 100644
diff --git a/Engine/Source/Runtime/Online/BuildPatchServices/Private/Installer/InstallChunkSource.cpp b/Engine/Source/Runtime/Online/BuildPatchServices/Private/Installer/InstallChunkSource.cpp
--- a/Engine/Source/Runtime/Online/BuildPatchServices/Private/Installer/InstallChunkSource.cpp
+++ b/Engine/Source/Runtime/Online/BuildPatchServices/Private/Installer/InstallChunkSource.cpp
@@ -3,198 +3,129 @@
 #include "Installer/InstallChunkSource.h"
 
 #include "Algo/Transform.h"
-#include "Containers/Queue.h"
+#include "BuildPatchHash.h"
+#include "Common/FileSystem.h"
+#include "Common/StatsCollector.h"
 #include "Core/BlockStructure.h"
-#include "HAL/ThreadSafeBool.h"
 #include "Misc/Paths.h"
 #include "Serialization/MemoryReader.h"
+#include "Tasks/Task.h"
 
-#include "Common/FileSystem.h"
-#include "Common/StatsCollector.h"
-#include "Installer/InstallerError.h"
-#include "Installer/ChunkStore.h"
-#include "Installer/ChunkReferenceTracker.h"
-#include "BuildPatchHash.h"
-#include "IBuildManifestSet.h"
-
-DECLARE_LOG_CATEGORY_CLASS(LogInstallChunkSource, Warning, All);
+DEFINE_LOG_CATEGORY_STATIC(LogInstallChunkSource, Log, All);
 
 namespace BuildPatchServices
 {
-	class FInstallChunkSource : public IInstallChunkSource
+	static FSHAHash GetShaHashForDataSet(const void* ChunkData, const uint32 ChunkSize)
+	{
+		FSHAHash ShaHashCheck;
+		FSHA1::HashBuffer(ChunkData, ChunkSize, ShaHashCheck.Hash);
+		return ShaHashCheck;
+	}
+	
+	class FInstallChunkSource : public IConstructorInstallChunkSource
 	{
 	public:
-		FInstallChunkSource(FInstallSourceConfig InConfiguration, IFileSystem* FileSystem, IChunkStore* InChunkStore, IChunkReferenceTracker* InChunkReferenceTracker, IInstallerError* InInstallerError, IInstallChunkSourceStat* InInstallChunkSourceStat, const TMultiMap<FString, FBuildPatchAppManifestRef>& InInstallationSources, IBuildManifestSet* ManifestSet);
+		FInstallChunkSource(IFileSystem* FileSystem, IInstallChunkSourceStat* InInstallChunkSourceStat, const TMultiMap<FString, FBuildPatchAppManifestRef>& InInstallationSources, 
+			const TSet<FGuid>& ChunksThatWillBeNeeded);
 		~FInstallChunkSource();
+		
+		virtual FRequestProcessFn CreateRequest(const FGuid& DataId, FMutableMemoryView DestinationBuffer, void* UserPtr, FChunkRequestCompleteDelegate CompleteFn) override;
+		virtual const TSet<FGuid>& GetAvailableChunks() const override
+		{
+			return AvailableInBuilds;
+		}
+		virtual void OnBeforeDeleteFile(const FString& FilePath) override
+		{
+			// Make sure we close our handle before the deletion occurs.
+			OpenedFileHandles.Remove(FilePath);
+		}
+		virtual int32 GetChunkUnavailableAt(const FGuid& DataId) const override;
+		virtual void SetFileRetirementPositions(TMap<FString, int32>&& InFileRetirementPositions) override
+		{
+			FileRetirementPositions = MoveTemp(InFileRetirementPositions);
+		}
 
-		// IControllable interface begin.
-		virtual void SetPaused(bool bInIsPaused) override;
-		virtual void Abort() override;
-		// IControllable interface end.
-
-		// IChunkSource interface begin.
-		virtual IChunkDataAccess* Get(const FGuid& DataId) override;
-		virtual TSet<FGuid> AddRuntimeRequirements(TSet<FGuid> NewRequirements) override;
-		virtual bool AddRepeatRequirement(const FGuid& RepeatRequirement) override;
-		virtual void SetUnavailableChunksCallback(TFunction<void(TSet<FGuid>)> Callback) override;
-		// IChunkSource interface end.
-
-		// IInstallChunkSource interface begin.
-		virtual const TSet<FGuid>& GetAvailableChunks() const override;
-		virtual void HarvestRemainingChunksFromFile(const FString& BuildFile) override;
-		// IInstallChunkSource interface end.
+		virtual void GetChunksForFile(const FString& FilePath, TSet<FGuid>& OutChunks) const override;
 
 	private:
 		void FindChunkLocation(const FGuid& DataId, const FString** FoundInstallDirectory, const FBuildPatchAppManifest** FoundInstallManifest) const;
-		bool LoadFromBuild(const FGuid& DataId);
-		FSHAHash GetShaHashForDataSet(const uint8* ChunkData, const uint32 ChunkSize);
 
 	private:
-		// Configuration.
-		const FInstallSourceConfig Configuration;
-		// Dependencies.
+	
 		IFileSystem* FileSystem;
-		IChunkStore* ChunkStore;
-		IChunkReferenceTracker* ChunkReferenceTracker;
-		IInstallerError* InstallerError;
 		IInstallChunkSourceStat* InstallChunkSourceStat;
-		// Control signals.
-		FThreadSafeBool bIsPaused;
-		FThreadSafeBool bShouldAbort;
-		// Handling of chunks we lose access to.
-		TFunction<void(TSet<FGuid>)> UnavailableChunksCallback;
-		TSet<FGuid> UnavailableChunks;
-		// Keep a separate copy of failed chunks internally.
-		TSet<FGuid> FailedChunks;
+
 		// Storage of enumerated chunks.
 		TSet<FGuid> AvailableInBuilds;
 		TArray<TPair<FString, FBuildPatchAppManifestRef>> InstallationSources;
-		TSet<FGuid> PlacedInStore;
-		// Track additional requests.
-		TSet<FGuid> RuntimeRequests;
-		// Communication and storage of incoming repeat requirements.
-		TQueue<FGuid, EQueueMode::Mpsc> RepeatRequirementMessages;
 
 		TMap<FString, TUniquePtr<FArchive>> OpenedFileHandles;
+
+		// The index (ChunkReferenceTracker->GetCurrentUsageIndex) at which our files will get deleted due to destructive install to make room for
+		// the new file. 
+		TMap<FString, int32> FileRetirementPositions;
+
+		struct FChunkSourceDetails
+		{
+			const FFileManifest* FileManifest;
+			const FChunkPart* ChunkPart;
+			uint64 FileOffset;
+		};
+
+		TMap<FGuid, TArray<FChunkSourceDetails>> ChunkSources;
 	};
 
-	FInstallChunkSource::FInstallChunkSource(FInstallSourceConfig InConfiguration, IFileSystem* InFileSystem, IChunkStore* InChunkStore, IChunkReferenceTracker* InChunkReferenceTracker, IInstallerError* InInstallerError, IInstallChunkSourceStat* InInstallChunkSourceStat, const TMultiMap<FString, FBuildPatchAppManifestRef>& InInstallationSources, IBuildManifestSet* ManifestSet)
-		: Configuration(MoveTemp(InConfiguration))
-		, FileSystem(InFileSystem)
-		, ChunkStore(InChunkStore)
-		, ChunkReferenceTracker(InChunkReferenceTracker)
-		, InstallerError(InInstallerError)
+	FInstallChunkSource::FInstallChunkSource(IFileSystem* InFileSystem, IInstallChunkSourceStat* InInstallChunkSourceStat, 
+		const TMultiMap<FString, FBuildPatchAppManifestRef>& InInstallationSources, const TSet<FGuid>& ChunksThatWillBeNeeded)
+		: FileSystem(InFileSystem)
 		, InstallChunkSourceStat(InInstallChunkSourceStat)
-		, bIsPaused(false)
-		, bShouldAbort(false)
-		, UnavailableChunksCallback(nullptr)
 	{
-		// Cache faster lookup information.
-		TSet<FGuid> ReferencedChunks;
-		ManifestSet->GetReferencedChunks(ReferencedChunks);
+		TRACE_CPUPROFILER_EVENT_SCOPE(InstallChunkSource_ctor);
+
+		// Cache faster lookup information.		
 		for (const TPair<FString, FBuildPatchAppManifestRef>& Pair : InInstallationSources)
 		{
-			if (Pair.Value->EnumerateProducibleChunks(Pair.Key, ReferencedChunks, AvailableInBuilds) > 0)
+			if (Pair.Value->EnumerateProducibleChunks(Pair.Key, ChunksThatWillBeNeeded, AvailableInBuilds) > 0)
 			{
 				InstallationSources.Add(Pair);
-
-				TSet<FString> FileSourceThisSource;
-				Pair.Value->GetFileList(FileSourceThisSource);
 			}
 		}
 		UE_LOG(LogInstallChunkSource, Log, TEXT("Useful Sources:%d. Available Chunks:%d."), InstallationSources.Num(), AvailableInBuilds.Num());
-	}
-
-	FInstallChunkSource::~FInstallChunkSource()
-	{
-	}
 
-	void FInstallChunkSource::SetPaused(bool bInIsPaused)
-	{
-		bIsPaused = bInIsPaused;
-	}
+		// Cache what file we get everything from so we aren't linearly looking over this every request
+		for (const TPair<FString, FBuildPatchAppManifestRef>& InstallationSource : InstallationSources)
+		{
+			FBuildPatchAppManifestRef Manifest = InstallationSource.Value;
 
-	void FInstallChunkSource::Abort()
-	{
-		bShouldAbort = true;
-	}
+			TSet<FString> FilesInManifest;
+			Manifest->GetFileList(FilesInManifest);
 
-	IChunkDataAccess* FInstallChunkSource::Get(const FGuid& DataId)
-	{
-		TRACE_CPUPROFILER_EVENT_SCOPE(Install_Get);
-		// Get from our store
-		IChunkDataAccess* ChunkData = ChunkStore->Get(DataId);
-		if (ChunkData == nullptr)
-		{
-			// If available, load next batch into store.
-			if (AvailableInBuilds.Contains(DataId))
+			for (const FString& FileName : FilesInManifest)
 			{
-				// 'Forget' any repeat requirements.
-				FGuid RepeatRequirement;
-				while (RepeatRequirementMessages.Dequeue(RepeatRequirement))
-				{
-					PlacedInStore.Remove(RepeatRequirement);
-				}
-                
-
-                
-                // There's no reason to prefetch anything - we are called when we need this exact chunk, so just load it directly.
-                // Anything further just fills up memory for no reason and doesn't actually save time. If we need it during harvest we'll get it.
-                
-                TArray<FGuid> BatchLoadChunks;
-				BatchLoadChunks.AddUnique(DataId);
-				// Call to stat.
-				InstallChunkSourceStat->OnBatchStarted(BatchLoadChunks);
-				// Load this batch
-				for (int32 ChunkIdx = 0; ChunkIdx < BatchLoadChunks.Num() && !bShouldAbort; ++ChunkIdx)
-				{
-					const FGuid& BatchLoadChunk = BatchLoadChunks[ChunkIdx];
-					LoadFromBuild(BatchLoadChunk);
-				}
-				// Get from store again.
-				ChunkData = ChunkStore->Get(DataId);
-				// Dump out unavailable chunks on the incoming IO thread.
-				if (UnavailableChunksCallback && UnavailableChunks.Num() > 0)
+				const FFileManifest* FileManifest = Manifest->GetFileManifest(FileName);
+
+				uint64 FileOffset = 0;
+				for (const FChunkPart& ChunkPart : FileManifest->ChunkParts)
 				{
-					UnavailableChunksCallback(MoveTemp(UnavailableChunks));
+					if (AvailableInBuilds.Contains(ChunkPart.Guid))
+					{
+						TArray<FChunkSourceDetails>& ChunkSource = ChunkSources.FindOrAdd(ChunkPart.Guid);
+						ChunkSource.Add({FileManifest, &ChunkPart, FileOffset});
+					}
+
+					FileOffset += ChunkPart.Size;
 				}
 			}
 		}
-		return ChunkData;
-	}
-
-	TSet<FGuid> FInstallChunkSource::AddRuntimeRequirements(TSet<FGuid> NewRequirements)
-	{
-		TSet<FGuid> Unhandled = NewRequirements.Difference(AvailableInBuilds);
-		TSet<FGuid> Accepted = NewRequirements.Intersect(AvailableInBuilds);
-		RuntimeRequests.Append(Accepted);
-		InstallChunkSourceStat->OnAcceptedNewRequirements(Accepted);
-		return Unhandled;
-	}
-
-	bool FInstallChunkSource::AddRepeatRequirement(const FGuid& RepeatRequirement)
-	{
-		if (AvailableInBuilds.Contains(RepeatRequirement))
-		{
-			RepeatRequirementMessages.Enqueue(RepeatRequirement);
-			return true;
-		}
-		return false;
 	}
 
-	void FInstallChunkSource::SetUnavailableChunksCallback(TFunction<void(TSet<FGuid>)> Callback)
-	{
-		UnavailableChunksCallback = Callback;
-	}
-
-	const TSet<FGuid>& FInstallChunkSource::GetAvailableChunks() const
+	FInstallChunkSource::~FInstallChunkSource()
 	{
-		return AvailableInBuilds;
+		
 	}
 
-	void FInstallChunkSource::HarvestRemainingChunksFromFile(const FString& FilePath)
+	void FInstallChunkSource::GetChunksForFile(const FString& FilePath, TSet<FGuid>& OutChunks) const
 	{
-		TRACE_CPUPROFILER_EVENT_SCOPE(Install_Harvest);
 		const FFileManifest* FileManifest = nullptr;
 		for (const TPair<FString, FBuildPatchAppManifestRef>& Pair : InstallationSources)
 		{
@@ -208,234 +139,239 @@ namespace BuildPatchServices
 		}
 		if (FileManifest != nullptr)
 		{
-			// Collect all chunks in this file.
-			TSet<FGuid> FileManifestChunks;
-			Algo::Transform(FileManifest->ChunkParts, FileManifestChunks, &FChunkPart::Guid);
-			FileManifestChunks = FileManifestChunks.Intersect(AvailableInBuilds);
-			// Select all chunks still required from this file.
-			TFunction<bool(const FGuid&)> SelectPredicate = [this, &FileManifestChunks](const FGuid& ChunkId)
-			{
-				return !PlacedInStore.Contains(ChunkId) && (FileManifestChunks.Contains(ChunkId) && (!Configuration.ChunkIgnoreSet.Contains(ChunkId) || RuntimeRequests.Contains(ChunkId)));
-			};
-			TArray<FGuid> BatchLoadChunks = ChunkReferenceTracker->GetNextReferences(TNumericLimits<int32>::Max(), SelectPredicate);
-			if (BatchLoadChunks.Num() > 0)
-			{
-				// Call to stat.
-				InstallChunkSourceStat->OnBatchStarted(BatchLoadChunks);
-				// Load the batch.
-				for (int32 ChunkIdx = 0; ChunkIdx < BatchLoadChunks.Num() && !bShouldAbort; ++ChunkIdx)
-				{
-					const FGuid& BatchLoadChunk = BatchLoadChunks[ChunkIdx];
-					LoadFromBuild(BatchLoadChunk);
-				}
-			}
+			Algo::Transform(FileManifest->ChunkParts, OutChunks, &FChunkPart::Guid);
 		}
-
-		// Make sure we close our handle before the deletion occurs.
-		OpenedFileHandles.Remove(FilePath);
 	}
 
-	void FInstallChunkSource::FindChunkLocation(const FGuid& DataId, const FString** FoundInstallDirectory, const FBuildPatchAppManifest** FoundInstallManifest) const
+	IConstructorChunkSource::FRequestProcessFn FInstallChunkSource::CreateRequest(const FGuid& DataId, FMutableMemoryView DestinationBuffer, void* UserPtr, FChunkRequestCompleteDelegate CompleteFn)
 	{
-		uint64 ChunkHash;
-		*FoundInstallDirectory = nullptr;
-		*FoundInstallManifest = nullptr;
-		for (const TPair<FString, FBuildPatchAppManifestRef>& Pair : InstallationSources)
-		{
-			// GetChunkHash can be used as a check for whether this manifest references this chunk.
-			if (Pair.Value->GetChunkHash(DataId, ChunkHash))
-			{
-				*FoundInstallDirectory = &Pair.Key;
-				*FoundInstallManifest = &Pair.Value.Get();
-				return;
-			}
-		}
-	}
-
-	bool FInstallChunkSource::LoadFromBuild(const FGuid& DataId)
-	{
-		TRACE_CPUPROFILER_EVENT_SCOPE(Install_LoadFromBuild);
-		// Find the location of this chunk.
 		const FString* FoundInstallDirectory;
 		const FBuildPatchAppManifest* FoundInstallManifest;
 		FindChunkLocation(DataId, &FoundInstallDirectory, &FoundInstallManifest);
 		if (FoundInstallDirectory == nullptr || FoundInstallManifest == nullptr)
 		{
-			return false;
+			CompleteFn.Execute(DataId, false, true, UserPtr);
+			return [](bool) {return;};
 		}
 
-		// Attempt construction of the chunk from the parts.
-		ISpeedRecorder::FRecord LoadRecord;
-		LoadRecord.Size = 0;
-		InstallChunkSourceStat->OnLoadStarted(DataId);
+		// Must be run sequentially! No threaded protection provided.
+		return [this, FoundInstallManifest, FoundInstallDirectory = *FoundInstallDirectory, DataId, DestinationBuffer, UserPtr, CompleteFn](bool bIsAborted)
+		{
+			ISpeedRecorder::FRecord ActivityRecord;
+			ActivityRecord.CyclesStart = FStatsCollector::GetCycles();
+			InstallChunkSourceStat->OnLoadStarted(DataId);
 
+			TRACE_CPUPROFILER_EVENT_SCOPE(InstallRead);
+			if (bIsAborted)
+			{
+				ActivityRecord.CyclesEnd = ActivityRecord.CyclesStart;
+				InstallChunkSourceStat->OnLoadComplete(DataId, IInstallChunkSourceStat::ELoadResult::Aborted, ActivityRecord);
 
-		// We must have a hash for this chunk or else we cant verify it
-		EChunkHashFlags HashType = EChunkHashFlags::None;
-		uint64 ChunkHash = 0;
-		FSHAHash ChunkShaHash;
-		if (FoundInstallManifest->GetChunkShaHash(DataId, ChunkShaHash))
-		{
-			HashType |= EChunkHashFlags::Sha1;
-		}
-		if (FoundInstallManifest->GetChunkHash(DataId, ChunkHash))
-		{
-			HashType |= EChunkHashFlags::RollingPoly64;
-		}
-		LoadRecord.CyclesStart = FStatsCollector::GetCycles();
-		IInstallChunkSourceStat::ELoadResult LoadResult = HashType == EChunkHashFlags::None ? IInstallChunkSourceStat::ELoadResult::MissingHashInfo : IInstallChunkSourceStat::ELoadResult::Success;
-		if (LoadResult == IInstallChunkSourceStat::ELoadResult::Success)
-		{
-			// Get the list of data pieces we need to load.
-			TArray<FFileChunkPart> FileChunkParts = FoundInstallManifest->GetFilePartsForChunk(DataId);
+				CompleteFn.Execute(DataId, true, false, UserPtr);
+				return;
+			}
+
+			const TArray<FChunkSourceDetails>* ChunkSource = ChunkSources.Find(DataId);
 			const FChunkInfo* ChunkInfoPtr = FoundInstallManifest->GetChunkInfo(DataId);
-			LoadResult = ChunkInfoPtr == nullptr || FileChunkParts.Num() <= 0 ? IInstallChunkSourceStat::ELoadResult::MissingPartInfo : IInstallChunkSourceStat::ELoadResult::Success;
-			if (LoadResult == IInstallChunkSourceStat::ELoadResult::Success)
+			if (!ChunkInfoPtr || !ChunkSource)
 			{
-				FBlockStructure ChunkBlocks;
-				const int32 InitialDataSize = 1024*1024;
-				TArray<uint8> TempArray;
-				uint8* TempChunkConstruction = nullptr;
-				uint32 LoadedChunkSize = 0;
-				bool bLoadedWholeChunk = false;
-				for (int32 FileChunkPartsIdx = 0; FileChunkPartsIdx < FileChunkParts.Num() && !bLoadedWholeChunk && !bShouldAbort; ++FileChunkPartsIdx)
-				{
-					const FFileChunkPart& FileChunkPart = FileChunkParts[FileChunkPartsIdx];
-					LoadedChunkSize = FMath::Max<uint32>(LoadedChunkSize, FileChunkPart.ChunkPart.Offset + FileChunkPart.ChunkPart.Size);
-					TempArray.SetNumZeroed(FMath::Max<uint32>(LoadedChunkSize, InitialDataSize));
-					TempChunkConstruction = TempArray.GetData();
-					FString FullFilename = *FoundInstallDirectory / FileChunkPart.Filename;
+				ActivityRecord.CyclesEnd = FStatsCollector::GetCycles();
+				InstallChunkSourceStat->OnLoadComplete(DataId, IInstallChunkSourceStat::ELoadResult::MissingPartInfo, ActivityRecord);
 
-					int64 FileSize = 0;
+				CompleteFn.Execute(DataId, false, true, UserPtr);
+				return;
+			}
 
-					TUniquePtr<FArchive>* FileArchive = OpenedFileHandles.Find(FullFilename);
-					if (FileArchive == nullptr)
-					{
-						TRACE_CPUPROFILER_EVENT_SCOPE(Install_OpenSource);
-						TUniquePtr<FArchive> NewReader = FileSystem->CreateFileReader(*FullFilename);
-						if (NewReader.IsValid())
-						{
-							OpenedFileHandles.Add(FullFilename, MoveTemp(NewReader));
-							FileArchive = OpenedFileHandles.Find(FullFilename);
-						}
-					}
+			FBlockStructure ChunkBlocks;
 
-					if (FileArchive == nullptr)
-					{
-						LoadResult = IInstallChunkSourceStat::ELoadResult::OpenFileFail;
-					}
-					else
-					{
-						FileSize = (*FileArchive)->TotalSize();
-					}
+			IInstallChunkSourceStat::ELoadResult Result = IInstallChunkSourceStat::ELoadResult::Success;
+	
+			bool bLoadedWholeChunk = false;
+			for (int32 FileChunkPartsIdx = 0; FileChunkPartsIdx < ChunkSource->Num(); ++FileChunkPartsIdx)
+			{
+				if (bLoadedWholeChunk)
+				{
+					// The manifest gave us more chunk parts than we needed to generate the full chunk. This shouldn't happen,
+					// and so conceptually is an error, but since we have all the data we can technically proceed.
+					// This seems to happen with some regularity - need to understand why.
+					//UE_LOG(LogInstallChunkSource, Display, TEXT("Chunk %s had more chunk sources than necessary to re-assemble"), *WriteToString<40>(DataId));
+					break;
+				}
+				const FChunkSourceDetails& FileChunkPart = ChunkSource->operator[](FileChunkPartsIdx);
 
-					// Grab the section of the file.
-					if (LoadResult == IInstallChunkSourceStat::ELoadResult::Success)
-					{
-						// Make sure we don't attempt to read off the end of the file
-						const int64 LastRequiredByte = FileChunkPart.FileOffset + FileChunkPart.ChunkPart.Size;
-						LoadResult = FileSize < LastRequiredByte ? IInstallChunkSourceStat::ELoadResult::IncorrectFileSize : IInstallChunkSourceStat::ELoadResult::Success;
-						if (LoadResult == IInstallChunkSourceStat::ELoadResult::Success)
-						{
-							TRACE_CPUPROFILER_EVENT_SCOPE(Install_Serialize);
-							(*FileArchive)->Seek(FileChunkPart.FileOffset);
-							(*FileArchive)->Serialize(TempChunkConstruction + FileChunkPart.ChunkPart.Offset, FileChunkPart.ChunkPart.Size);
-							ChunkBlocks.Add(FileChunkPart.ChunkPart.Offset, FileChunkPart.ChunkPart.Size);
-							LoadRecord.Size += FileChunkPart.ChunkPart.Size;
-						}
-					}
-					bLoadedWholeChunk = ChunkBlocks.GetHead() && ChunkBlocks.GetHead() == ChunkBlocks.GetTail() && ChunkBlocks.GetHead()->GetSize() == ChunkInfoPtr->WindowSize;
-					// Wait while paused
-					while (bIsPaused && !bShouldAbort)
-					{
-						FPlatformProcess::Sleep(0.5f);
-					}
+				// Validate the chunk can load into the destination.
+				uint32 ChunkEndLocation = FileChunkPart.ChunkPart->Offset + FileChunkPart.ChunkPart->Size;
+				if (ChunkEndLocation > DestinationBuffer.GetSize())
+				{
+					// The chunk metadata tried to assemble larger than the chunk itself - error
+					UE_LOG(LogInstallChunkSource, Error, TEXT("Chunk %s assembled larger than the actual chunk size (chunk wanted end %u vs buffer size %llu"), *WriteToString<40>(DataId), ChunkEndLocation, DestinationBuffer.GetSize());
+					Result = IInstallChunkSourceStat::ELoadResult::InvalidChunkParts;
+					break;
 				}
-				if (bLoadedWholeChunk)
+
+				FString FullFilename = FoundInstallDirectory / FileChunkPart.FileManifest->Filename;
+
+				TUniquePtr<FArchive>* FileArchive = OpenedFileHandles.Find(FullFilename);
+				if (FileArchive == nullptr)
 				{
-					LoadResult = IInstallChunkSourceStat::ELoadResult::Success;
+					TRACE_CPUPROFILER_EVENT_SCOPE(Install_OpenSource);
+					TUniquePtr<FArchive> NewReader = FileSystem->CreateFileReader(*FullFilename);
+					if (NewReader.IsValid())
+					{
+						OpenedFileHandles.Add(FullFilename, MoveTemp(NewReader));
+						FileArchive = OpenedFileHandles.Find(FullFilename);
+					}
 				}
 
-				// Mark not success if we gave up due to being aborted.
-				if (bShouldAbort)
+				if (FileArchive == nullptr)
 				{
-					LoadResult = IInstallChunkSourceStat::ELoadResult::Aborted;
+					Result = IInstallChunkSourceStat::ELoadResult::OpenFileFail;
+					break;
 				}
 
-				// Check chunk hash
-				if (LoadResult == IInstallChunkSourceStat::ELoadResult::Success)
 				{
-					TRACE_CPUPROFILER_EVENT_SCOPE(Install_Hash);
-					const bool bUseSha = (HashType & EChunkHashFlags::Sha1) != EChunkHashFlags::None;
-					const bool bHashCheckOk = bUseSha ? GetShaHashForDataSet(TempChunkConstruction, LoadedChunkSize) == ChunkShaHash : FRollingHash::GetHashForDataSet(TempChunkConstruction, LoadedChunkSize) == ChunkHash;
-					if (bHashCheckOk == false)
+					TRACE_CPUPROFILER_EVENT_SCOPE(Install_Serialize);
+					(*FileArchive)->Seek(FileChunkPart.FileOffset);
+					(*FileArchive)->Serialize((uint8*)DestinationBuffer.GetData() + FileChunkPart.ChunkPart->Offset, FileChunkPart.ChunkPart->Size);
+
+					ActivityRecord.Size += FileChunkPart.ChunkPart->Size;
+
+					FBlockStructure NewChunk;
+					NewChunk.Add(FileChunkPart.ChunkPart->Offset, FileChunkPart.ChunkPart->Size);
+					if (NewChunk.Intersect(ChunkBlocks).GetHead() != nullptr)
 					{
-						LoadResult = IInstallChunkSourceStat::ELoadResult::HashCheckFailed;
+						// This used to be allowed but in advance of multi threaded reading we want to make sure this
+						// doesn't happen anymore (already shouldn't be...)
+						UE_LOG(LogInstallChunkSource, Error, TEXT("Chunk %s had overlapping chunk parts"), *WriteToString<40>(DataId));
+						Result = IInstallChunkSourceStat::ELoadResult::InvalidChunkParts;
+						break;
 					}
+					ChunkBlocks.Add(FileChunkPart.ChunkPart->Offset, FileChunkPart.ChunkPart->Size);
 				}
 
-				// Save the chunk to the store if all went well.
-				if (LoadResult == IInstallChunkSourceStat::ELoadResult::Success)
+				// The expectation is that we get the full chunk only once we've assembled all of the parts
+				// provided by the manifest, so the last iteration should set this to true. If it isn't the last
+				// iteration, then we'll hit the faux-error case at the top of the loop.
+				bLoadedWholeChunk = ChunkBlocks.GetHead() && ChunkBlocks.GetHead() == ChunkBlocks.GetTail() && ChunkBlocks.GetHead()->GetSize() == ChunkInfoPtr->WindowSize;
+			}
+			
+			if (!bLoadedWholeChunk)
+			{
+				if (Result == IInstallChunkSourceStat::ELoadResult::Success)
 				{
-					// Create the ChunkFile data structure
-					IChunkDataAccess* NewChunkFile = FChunkDataAccessFactory::Create(LoadedChunkSize);
-
-					// Lock data
-					FChunkHeader* ChunkHeader;
-					uint8* ChunkData;
-					NewChunkFile->GetDataLock(&ChunkData, &ChunkHeader);
-
-					// Copy the data
-					FMemoryReader MemReader(TempArray);
-					MemReader.Serialize(ChunkData, LoadedChunkSize);
-
-					// Setup the header
-					ChunkHeader->Guid = DataId;
-					ChunkHeader->StoredAs = EChunkStorageFlags::None;
-					ChunkHeader->DataSizeCompressed = LoadedChunkSize;
-					ChunkHeader->DataSizeUncompressed = LoadedChunkSize;
-					ChunkHeader->HashType = HashType;
-					ChunkHeader->RollingHash = ChunkHash;
-					FMemory::Memcpy(ChunkHeader->SHAHash.Hash, ChunkShaHash.Hash, FSHA1::DigestSize);
-
-					// Release data
-					NewChunkFile->ReleaseDataLock();
-
-					// Add it to our cache.
-					PlacedInStore.Add(DataId);
-					ChunkStore->Put(DataId, TUniquePtr<IChunkDataAccess>(NewChunkFile));
+					// If we failed without hitting a different case, then we just didn't have enough parts.
+					Result = IInstallChunkSourceStat::ELoadResult::InvalidChunkParts;
 				}
+
+				ActivityRecord.CyclesEnd = FStatsCollector::GetCycles();
+				InstallChunkSourceStat->OnLoadComplete(DataId, IInstallChunkSourceStat::ELoadResult::MissingPartInfo, ActivityRecord);
+
+				CompleteFn.Execute(DataId, false, !bLoadedWholeChunk, UserPtr);
+				return;
 			}
+
+			// Check chunk hash. 
+			UE::Tasks::Launch(TEXT("Install_Hash"), [DataId, UserPtr, DestinationBuffer, FoundInstallManifest, CompleteFn, ActivityRecord, InstallChunkSourceStat=InstallChunkSourceStat]()
+				{
+					// This is here because the compiler is complaining about ActivityRecord being a const object
+					// even though it's passed by value.
+					ISpeedRecorder::FRecord OurActivityRecord = ActivityRecord;
+
+					IInstallChunkSourceStat::ELoadResult Result = IInstallChunkSourceStat::ELoadResult::Success;
+					{
+						TRACE_CPUPROFILER_EVENT_SCOPE(InstallHash);
+						FSHAHash ChunkShaHash;
+						uint64 ChunkRollingHash = 0;
+
+						if (FoundInstallManifest->GetChunkShaHash(DataId, ChunkShaHash))
+						{
+							if (GetShaHashForDataSet(DestinationBuffer.GetData(), DestinationBuffer.GetSize()) != ChunkShaHash)
+							{
+								Result = IInstallChunkSourceStat::ELoadResult::HashCheckFailed;
+							}
+						}
+						else if (FoundInstallManifest->GetChunkHash(DataId, ChunkRollingHash))
+						{
+							if (FRollingHash::GetHashForDataSet((const uint8*)DestinationBuffer.GetData(), DestinationBuffer.GetSize()) != ChunkRollingHash)
+							{
+								Result = IInstallChunkSourceStat::ELoadResult::HashCheckFailed;
+							}
+						}
+						else
+						{
+							Result = IInstallChunkSourceStat::ELoadResult::MissingHashInfo;
+						}
+					}
+
+					OurActivityRecord.CyclesEnd = FStatsCollector::GetCycles();
+					InstallChunkSourceStat->OnLoadComplete(
+						DataId, 
+						Result, 
+						OurActivityRecord);					
+
+					CompleteFn.Execute(DataId, false, Result != IInstallChunkSourceStat::ELoadResult::Success, UserPtr);
+				}
+			);
+		};
+	}
+	
+	int32 FInstallChunkSource::GetChunkUnavailableAt(const FGuid& DataId) const
+	{
+		if (FileRetirementPositions.Num() == 0)
+		{
+			// If we aren't destructive then it's always available.
+			return TNumericLimits<int32>::Max();
 		}
-		LoadRecord.CyclesEnd = FStatsCollector::GetCycles();
-		InstallChunkSourceStat->OnLoadComplete(DataId, LoadResult, LoadRecord);
-		if (LoadResult == IInstallChunkSourceStat::ELoadResult::Success)
+
+		const FString* FoundInstallDirectory;
+		const FBuildPatchAppManifest* FoundInstallManifest;
+		FindChunkLocation(DataId, &FoundInstallDirectory, &FoundInstallManifest);
+		if (FoundInstallDirectory == nullptr || FoundInstallManifest == nullptr)
 		{
-			return true;
+			return TNumericLimits<int32>::Max();
 		}
-		else
+
+		// This chunk is no longer available as soon as the first file containing a part is complete (if destructive install)
+		int32 ChunkUnavailableAt = TNumericLimits<int32>::Max();
+
+		const TArray<FChunkSourceDetails>* ChunkSource = ChunkSources.Find(DataId);
+		if (ChunkSource)
 		{
-			UnavailableChunks.Add(DataId);
-			FailedChunks.Add(DataId);
-			return false;
+			for (const FChunkSourceDetails& Part : *ChunkSource)
+			{
+				const int32* FirstIndexAfterFile = FileRetirementPositions.Find(Part.FileManifest->Filename);
+				if (FirstIndexAfterFile &&
+					*FirstIndexAfterFile < ChunkUnavailableAt)
+				{
+					ChunkUnavailableAt = *FirstIndexAfterFile;
+				}
+			}
 		}
-	}
 
-	FSHAHash FInstallChunkSource::GetShaHashForDataSet(const uint8* ChunkData, const uint32 ChunkSize)
-	{
-		FSHAHash ShaHashCheck;
-		FSHA1::HashBuffer(ChunkData, ChunkSize, ShaHashCheck.Hash);
-		return ShaHashCheck;
+		return ChunkUnavailableAt;
 	}
 	
-	IInstallChunkSource* FInstallChunkSourceFactory::Create(FInstallSourceConfig Configuration, IFileSystem* FileSystem, IChunkStore* ChunkStore, IChunkReferenceTracker* ChunkReferenceTracker, IInstallerError* InstallerError, IInstallChunkSourceStat* InstallChunkSourceStat, const TMultiMap<FString, FBuildPatchAppManifestRef>& InstallationSources, IBuildManifestSet* ManifestSet)
+	void FInstallChunkSource::FindChunkLocation(const FGuid& DataId, const FString** FoundInstallDirectory, const FBuildPatchAppManifest** FoundInstallManifest) const
+	{
+		uint64 ChunkHash;
+		*FoundInstallDirectory = nullptr;
+		*FoundInstallManifest = nullptr;
+		for (const TPair<FString, FBuildPatchAppManifestRef>& Pair : InstallationSources)
+		{
+			// GetChunkHash can be used as a check for whether this manifest references this chunk.
+			if (Pair.Value->GetChunkHash(DataId, ChunkHash))
+			{
+				*FoundInstallDirectory = &Pair.Key;
+				*FoundInstallManifest = &Pair.Value.Get();
+				return;
+			}
+		}
+	}
+
+	IConstructorInstallChunkSource* IConstructorInstallChunkSource::CreateInstallSource(IFileSystem* FileSystem, IInstallChunkSourceStat* InstallChunkSourceStat, 
+		const TMultiMap<FString, FBuildPatchAppManifestRef>& InstallationSources, const TSet<FGuid>& ChunksThatWillBeNeeded)
 	{
-		check(FileSystem != nullptr);
-		check(ChunkStore != nullptr);
-		check(ChunkReferenceTracker != nullptr);
-		check(InstallerError != nullptr);
-		check(InstallChunkSourceStat != nullptr);
-		return new FInstallChunkSource(MoveTemp(Configuration), FileSystem, ChunkStore, ChunkReferenceTracker, InstallerError, InstallChunkSourceStat, InstallationSources, ManifestSet);
+		return new FInstallChunkSource(FileSystem, InstallChunkSourceStat, InstallationSources, ChunksThatWillBeNeeded);
 	}
 
 	const TCHAR* ToString(const IInstallChunkSourceStat::ELoadResult& LoadResult)
@@ -450,12 +386,12 @@ namespace BuildPatchServices
 				return TEXT("MissingPartInfo");
 			case IInstallChunkSourceStat::ELoadResult::OpenFileFail:
 				return TEXT("OpenFileFail");
-			case IInstallChunkSourceStat::ELoadResult::IncorrectFileSize:
-				return TEXT("IncorrectFileSize");
 			case IInstallChunkSourceStat::ELoadResult::HashCheckFailed:
 				return TEXT("HashCheckFailed");
 			case IInstallChunkSourceStat::ELoadResult::Aborted:
 				return TEXT("Aborted");
+			case IInstallChunkSourceStat::ELoadResult::InvalidChunkParts:
+				return TEXT("InvalidChunkParts");
 			default:
 				return TEXT("Unknown");
 		}
diff --git a/Engine/Source/Runtime/Online/BuildPatchServices/Private/Installer/InstallChunkSource.h b/Engine/Source/Runtime/Online/BuildPatchServices/Private/Installer/InstallChunkSource.h
--- a/Engine/Source/Runtime/Online/BuildPatchServices/Private/Installer/InstallChunkSource.h
+++ b/Engine/Source/Runtime/Online/BuildPatchServices/Private/Installer/InstallChunkSource.h
@@ -6,6 +6,7 @@
 #include "Common/SpeedRecorder.h"
 #include "BuildPatchManifest.h"
 
+class IBuildInstallerSharedContext;
 
 namespace BuildPatchServices
 {
@@ -16,74 +17,38 @@ namespace BuildPatchServices
 	class IFileSystem;
 	class IBuildManifestSet;
 
+
 	/**
-	 * The interface for an installation chunk source, which provides access to chunk data retrieved from known local installations.
-	 */
-	class IInstallChunkSource
-		: public IChunkSource
-		, public IControllable
+	* The interface for an installation chunk source, which provides access to chunk data retrieved from known local installations.
+	*/
+	class IConstructorInstallChunkSource : public IConstructorChunkSource
 	{
 	public:
-		virtual ~IInstallChunkSource() {}
+		virtual ~IConstructorInstallChunkSource() {}
 
-		/**
-		 * Get the set of chunks available locally which are relevant to the installation being performed.
-		 * @return the set of chunks that should be available locally.
-		 */
+		// Return the chunks this source can provide.
 		virtual const TSet<FGuid>& GetAvailableChunks() const = 0;
 
-		/**
-		 * Given an installation file path, harvest all remaining required chunks from the file immediately.
-		 * NB: The filepath must match a file contained in one of the installation sources provided in order to load chunks.
-		 * @param FilePath  The full filepath to the file to harvest.
-		 */
-		virtual void HarvestRemainingChunksFromFile(const FString& FilePath) = 0;
-	};
-
-	/**
-	 * A struct containing the configuration values for an install chunk source.
-	 */
-	struct FInstallSourceConfig
-	{
-		// A set of chunks to not retrieve unless specifically asked for.
-		TSet<FGuid> ChunkIgnoreSet;
-		// The minimum number of chunks to load at a time when one is requested, depending on store slack.
-		int32 BatchFetchMinimum;
-		// The maximum number of chunks to load at a time when one is requested, depending on store slack.
-		int32 BatchFetchMaximum;
-
-		/**
-		 * Constructor which sets usual defaults.
-		 */
-		FInstallSourceConfig()
-			: ChunkIgnoreSet()
-			, BatchFetchMinimum(10)
-			, BatchFetchMaximum(40)
-		{
-		}
-	};
-
-	/**
-	 * A factory for creating an IInstallChunkSource instance.
-	 */
-	class FInstallChunkSourceFactory
-	{
-	public:
-		/**
-		 * This implementation will read chunks from provided local installations if they are available.
-		 * During initialization the local installations are enumerated to find each available chunk and expected local files
-		 * are checked and skipped if missing or incorrect size.
-		 * @param Configuration             The configuration struct for this cloud source.
-		 * @param FileSystem                The service used to open files on disk.
-		 * @param ChunkStore                The chunk store where loaded chunks will be put.
-		 * @param ChunkReferenceTracker     The reference tracker for the installation, used to decide which chunks to fetch when loading.
-		 * @param InstallerError            Error tracker where fatal errors will be reported.
-		 * @param InstallChunkSourceStat    The class to receive statistics and event information.
-		 * @param InstallationSources       The map of known local installation directories and the associated manifest.
-		 * @param InstallManifest           The manifest that chunks are required for.
-		 * @return the new IInstallChunkSource instance created.
-		 */
-		static IInstallChunkSource* Create(FInstallSourceConfig Configuration, IFileSystem* FileSystem, IChunkStore* ChunkStore, IChunkReferenceTracker* ChunkReferenceTracker, IInstallerError* InstallerError, IInstallChunkSourceStat* InstallChunkSourceStat, const TMultiMap<FString, FBuildPatchAppManifestRef>& InstallationSources, IBuildManifestSet* ManifestSet);
+		// Notification that a file is about to be deleted so we need to clear out any references.
+		virtual void OnBeforeDeleteFile(const FString& FilePath) = 0;
+
+		// Return the list of chunks for a file so that the constructor can harvest them before the file
+		// is deleted
+		virtual void GetChunksForFile(const FString& FilePath, TSet<FGuid>& OutChunks) const = 0;
+
+		// Tell the install source when files will be going away so we can report when chunks are unavailable.
+		// This should be the ChunkReferenceTracker->GetCurrentUsageIndex at which the files will get deleted.
+		virtual void SetFileRetirementPositions(TMap<FString, int32>&& FileRetirementPositions) = 0;
+
+		// InstallationSources is the install manifest for each installed app that we can pull from. This
+		// is expected to be a single entry.
+		// LONGTERM question -- how does this work during delta generation? It is expected that all apps
+		// are installed and chunks are pulled from everything?
+		static IConstructorInstallChunkSource* CreateInstallSource(
+			IFileSystem* FileSystem, 
+			IInstallChunkSourceStat* InstallChunkSourceStat, 
+			const TMultiMap<FString, FBuildPatchAppManifestRef>& InstallationSources,
+			const TSet<FGuid>& ChunksThatWillBeNeeded);
 	};
 
 	/**
@@ -110,13 +75,16 @@ namespace BuildPatchServices
 			OpenFileFail,
 
 			// The expected source file size was not matched.
-			IncorrectFileSize,
+			IncorrectFileSize_DEPRECATED,
 
 			// The expected data hash for the chunk did not match.
 			HashCheckFailed,
 
 			// The process has been aborted.
-			Aborted
+			Aborted,
+
+			// Invalid assemble structure (i.e. overlapping chunk parts)
+			InvalidChunkParts
 		};
 
 	public:
@@ -126,7 +94,8 @@ namespace BuildPatchServices
 		 * Called when a batch of chunks are going to be loaded.
 		 * @param ChunkIds  The ids of each chunk.
 		 */
-		virtual void OnBatchStarted(const TArray<FGuid>& ChunkIds) = 0;
+		UE_DEPRECATED(5.6, "No longer batch loaded")
+		virtual void OnBatchStarted(const TArray<FGuid>& ChunkIds) {};
 
 		/**
 		 * Called each time a chunk load begins.
@@ -146,7 +115,8 @@ namespace BuildPatchServices
 		 * Called when a batch of chunks are added and accepted via IChunkSource::AddRuntimeRequirements.
 		 * @param ChunkIds  The ids of each chunk.
 		 */
-		virtual void OnAcceptedNewRequirements(const TSet<FGuid>& ChunkIds) = 0;
+		UE_DEPRECATED(5.6, "No longer batch loaded")
+		virtual void OnAcceptedNewRequirements(const TSet<FGuid>& ChunkIds) {};
 	};
 	
 	/**
diff --git a/Engine/Source/Runtime/Online/BuildPatchServices/Private/Installer/InstallerError.h b/Engine/Source/Runtime/Online/BuildPatchServices/Private/Installer/InstallerError.h
--- a/Engine/Source/Runtime/Online/BuildPatchServices/Private/Installer/InstallerError.h
+++ b/Engine/Source/Runtime/Online/BuildPatchServices/Private/Installer/InstallerError.h
@@ -62,6 +62,7 @@ namespace BuildPatchServices
 		static const TCHAR* OutboundCorrupt = TEXT("05");
 		static const TCHAR* SerializationError = TEXT("06");
 		static const TCHAR* TrackingError = TEXT("07");
+		static const TCHAR* InternalConsistencyFailure = TEXT("08");
 	}
 
 	/**
diff --git a/Engine/Source/Runtime/Online/BuildPatchServices/Private/Installer/InstallerSharedContext.cpp b/Engine/Source/Runtime/Online/BuildPatchServices/Private/Installer/InstallerSharedContext.cpp
--- a/Engine/Source/Runtime/Online/BuildPatchServices/Private/Installer/InstallerSharedContext.cpp
+++ b/Engine/Source/Runtime/Online/BuildPatchServices/Private/Installer/InstallerSharedContext.cpp
@@ -120,19 +120,18 @@ namespace BuildPatchServices
 		bWarnOnCreateThread = true;
 	}
 
-	uint32 FBuildInstallerSharedContext::NumThreadsPerInstaller(bool bUseChunkDBs) const
+	uint32 FBuildInstallerSharedContext::NumThreadsPerInstaller(bool bUseChunkDBs, bool bHasInstallSource) const
 	{
 		const uint32 NumInstallerMainThreads = 1;
-		const uint32 NumCloudChunkSourceThreads = 1;
+		const uint32 NumWriteThreads = 1;
 
-		bool bUseDiskOverflowStore = true;
-		GConfig->GetBool(TEXT("BuildPatchServices"), TEXT("bEnableDiskOverflowStore"), bUseDiskOverflowStore, GEngineIni);
-		const uint32 NumDiskChunkStoreThreads = bUseDiskOverflowStore ? 1 : 0;
 		const uint32 NumChunkDBThreads = bUseChunkDBs ? 1 : 0;
+		const uint32 NumInstallSourceThreads = bHasInstallSource ? 1 : 0;
+
 		const uint32 NumExpectedThreads =
 			NumInstallerMainThreads + 
-			NumCloudChunkSourceThreads + 
-			NumDiskChunkStoreThreads + 
+			NumWriteThreads +
+			NumInstallSourceThreads + 
 			NumChunkDBThreads;
 
 		return NumExpectedThreads;
diff --git a/Engine/Source/Runtime/Online/BuildPatchServices/Private/Installer/InstallerSharedContext.h b/Engine/Source/Runtime/Online/BuildPatchServices/Private/Installer/InstallerSharedContext.h
--- a/Engine/Source/Runtime/Online/BuildPatchServices/Private/Installer/InstallerSharedContext.h
+++ b/Engine/Source/Runtime/Online/BuildPatchServices/Private/Installer/InstallerSharedContext.h
@@ -68,7 +68,7 @@ namespace BuildPatchServices
 		virtual IBuildInstallerThread* CreateThread() override;
 		virtual void ReleaseThread(IBuildInstallerThread* Thread) override;
 		virtual void PreallocateThreads(uint32 NumThreads) override;
-		virtual uint32 NumThreadsPerInstaller(bool bUseChunkDBs) const override;
+		virtual uint32 NumThreadsPerInstaller(bool bUseChunkDBs, bool bHasInstallSource) const override;
 		// IBuildInstallerSharedContext interface end
 
 	private:
diff --git a/Engine/Source/Runtime/Online/BuildPatchServices/Private/Installer/Statistics/ChunkDbChunkSourceStatistics.cpp b/Engine/Source/Runtime/Online/BuildPatchServices/Private/Installer/Statistics/ChunkDbChunkSourceStatistics.cpp
--- a/Engine/Source/Runtime/Online/BuildPatchServices/Private/Installer/Statistics/ChunkDbChunkSourceStatistics.cpp
+++ b/Engine/Source/Runtime/Online/BuildPatchServices/Private/Installer/Statistics/ChunkDbChunkSourceStatistics.cpp
@@ -17,7 +17,6 @@ namespace BuildPatchServices
 
 	public:
 		// IChunkDbChunkSourceStat interface begin.
-		virtual void OnBatchStarted(const TArray<FGuid>& ChunkIds) override;
 		virtual void OnLoadStarted(const FGuid& ChunkId) override;
 		virtual void OnLoadComplete(const FGuid& ChunkId, ELoadResult Result, const ISpeedRecorder::FRecord& Record) override;
 		// IChunkDbChunkSourceStat interface end.
@@ -46,11 +45,6 @@ namespace BuildPatchServices
 	{
 	}
 
-	void FChunkDbChunkSourceStatistics::OnBatchStarted(const TArray<FGuid>& ChunkIds)
-	{
-		FileOperationTracker->OnDataStateUpdate(ChunkIds, EFileOperationState::RetrievingLocalChunkDbData);
-	}
-
 	void FChunkDbChunkSourceStatistics::OnLoadStarted(const FGuid& ChunkId)
 	{
 		FileOperationTracker->OnDataStateUpdate(ChunkId, EFileOperationState::RetrievingLocalChunkDbData);
diff --git a/Engine/Source/Runtime/Online/BuildPatchServices/Private/Installer/Statistics/DiskChunkStoreStatistics.cpp b/Engine/Source/Runtime/Online/BuildPatchServices/Private/Installer/Statistics/DiskChunkStoreStatistics.cpp
deleted file mode 100644
diff --git a/Engine/Source/Runtime/Online/BuildPatchServices/Private/Installer/Statistics/DiskChunkStoreStatistics.h b/Engine/Source/Runtime/Online/BuildPatchServices/Private/Installer/Statistics/DiskChunkStoreStatistics.h
deleted file mode 100644
diff --git a/Engine/Source/Runtime/Online/BuildPatchServices/Private/Installer/Statistics/InstallChunkSourceStatistics.cpp b/Engine/Source/Runtime/Online/BuildPatchServices/Private/Installer/Statistics/InstallChunkSourceStatistics.cpp
--- a/Engine/Source/Runtime/Online/BuildPatchServices/Private/Installer/Statistics/InstallChunkSourceStatistics.cpp
+++ b/Engine/Source/Runtime/Online/BuildPatchServices/Private/Installer/Statistics/InstallChunkSourceStatistics.cpp
@@ -18,10 +18,8 @@ namespace BuildPatchServices
 		~FInstallChunkSourceStatistics();
 
 		// IInstallChunkSourceStat interface begin.
-		virtual void OnBatchStarted(const TArray<FGuid>& ChunkIds) override;
 		virtual void OnLoadStarted(const FGuid& ChunkId) override;
 		virtual void OnLoadComplete(const FGuid& ChunkId, const ELoadResult& Result, const ISpeedRecorder::FRecord& Record) override;
-		virtual void OnAcceptedNewRequirements(const TSet<FGuid>& ChunkIds) override;
 		// IInstallChunkSourceStat interface end.
 
 		// IInstallChunkSourceStatistics interface begin.
@@ -57,12 +55,7 @@ namespace BuildPatchServices
 	FInstallChunkSourceStatistics::~FInstallChunkSourceStatistics()
 	{
 	}
-
-	void FInstallChunkSourceStatistics::OnBatchStarted(const TArray<FGuid>& ChunkIds)
-	{
-		FileOperationTracker->OnDataStateUpdate(ChunkIds, EFileOperationState::RetrievingLocalInstallData);
-	}
-
+	
 	void FInstallChunkSourceStatistics::OnLoadStarted(const FGuid& ChunkId)
 	{
 		bIsReading = true;
@@ -87,11 +80,6 @@ namespace BuildPatchServices
 		LastReadCycles = Record.CyclesEnd;
 	}
 
-	void FInstallChunkSourceStatistics::OnAcceptedNewRequirements(const TSet<FGuid>& ChunkIds)
-	{
-		FileOperationTracker->OnDataStateUpdate(ChunkIds, EFileOperationState::PendingLocalInstallData);
-	}
-
 	uint64 FInstallChunkSourceStatistics::GetBytesRead() const
 	{
 		return TotalBytesRead.GetValue();
diff --git a/Engine/Source/Runtime/Online/BuildPatchServices/Private/Tests/Fake/ChunkReferenceTracker.fake.h b/Engine/Source/Runtime/Online/BuildPatchServices/Private/Tests/Fake/ChunkReferenceTracker.fake.h
--- a/Engine/Source/Runtime/Online/BuildPatchServices/Private/Tests/Fake/ChunkReferenceTracker.fake.h
+++ b/Engine/Source/Runtime/Online/BuildPatchServices/Private/Tests/Fake/ChunkReferenceTracker.fake.h
@@ -63,6 +63,7 @@ namespace BuildPatchServices
 		{
 			if (NextReferences.Num() && NextReferences[0] == ChunkId)
 			{
+				PoppedCount++;
 				NextReferences.RemoveAt(0);
 				return true;
 			}
@@ -75,7 +76,25 @@ namespace BuildPatchServices
 			OutUseList = NextReferences;
 		}
 
+		int32 GetNextUsageForChunk(const FGuid& ChunkId, int32& OutLastUsageIndex) const
+		{
+			for (int32 i=0; i < NextReferences.Num(); i++)
+			{
+				if (NextReferences[i] == ChunkId)
+				{
+					return i + PoppedCount;
+				}
+			}
+			return -1;
+		}
+
+		virtual int32 GetCurrentUsageIndex() const
+		{
+			return PoppedCount;
+		}
+
 	public:
+		int32 PoppedCount = 0;
 		TSet<FGuid> ReferencedChunks;
 		TMap<FGuid, int32> ReferenceCounts;
 		TArray<FGuid> NextReferences;
diff --git a/Engine/Source/Runtime/Online/BuildPatchServices/Private/Tests/Mock/ChunkDataSerialization.mock.h b/Engine/Source/Runtime/Online/BuildPatchServices/Private/Tests/Mock/ChunkDataSerialization.mock.h
--- a/Engine/Source/Runtime/Online/BuildPatchServices/Private/Tests/Mock/ChunkDataSerialization.mock.h
+++ b/Engine/Source/Runtime/Online/BuildPatchServices/Private/Tests/Mock/ChunkDataSerialization.mock.h
@@ -4,6 +4,7 @@
 #include "Data/ChunkData.h"
 #include "Tests/TestHelpers.h"
 #include "Common/StatsCollector.h"
+#include "Memory/MemoryView.h"
 
 #if WITH_DEV_AUTOMATION_TESTS
 
@@ -92,6 +93,16 @@ namespace BuildPatchServices
 			RxInjectShaToChunkData.Emplace(Memory, ShaHashData);
 		}
 
+		virtual bool ValidateAndRead(FArchive& InArchive, FMutableMemoryView InDestinationBuffer, FChunkHeader& OutHeader, FUniqueBuffer& OutCompressedBuffer) const
+		{
+			return false;
+		}
+		virtual bool DecompressValidatedRead(const FChunkHeader& InHeader, FMutableMemoryView InDestionationBuffer, const FUniqueBuffer& InCompressedBuffer) const
+		{
+			return false;
+		}
+
+
 	public:
 		mutable TArray<FRxLoadFromFile> RxLoadFromFile;
 		mutable TArray<FRxSaveToFile> RxSaveToFile;
diff --git a/Engine/Source/Runtime/Online/BuildPatchServices/Private/Tests/Mock/InstallChunkSourceStat.mock.h b/Engine/Source/Runtime/Online/BuildPatchServices/Private/Tests/Mock/InstallChunkSourceStat.mock.h
--- a/Engine/Source/Runtime/Online/BuildPatchServices/Private/Tests/Mock/InstallChunkSourceStat.mock.h
+++ b/Engine/Source/Runtime/Online/BuildPatchServices/Private/Tests/Mock/InstallChunkSourceStat.mock.h
@@ -12,20 +12,10 @@ namespace BuildPatchServices
 		: public IInstallChunkSourceStat
 	{
 	public:
-		typedef TTuple<double, TArray<FGuid>> FBatchStarted;
 		typedef TTuple<double, FGuid> FLoadStarted;
 		typedef TTuple<double, FGuid, ELoadResult, ISpeedRecorder::FRecord> FLoadComplete;
 
 	public:
-		virtual void OnBatchStarted(const TArray<FGuid>& ChunkIds) override
-		{
-			if (OnBatchStartedFunc != nullptr)
-			{
-				OnBatchStartedFunc(ChunkIds);
-			}
-			RxBatchStarted.Emplace(FStatsCollector::GetSeconds(), ChunkIds);
-		}
-
 		virtual void OnLoadStarted(const FGuid& ChunkId) override
 		{
 			if (OnLoadStartedFunc != nullptr)
@@ -44,15 +34,9 @@ namespace BuildPatchServices
 			RxLoadComplete.Emplace(FStatsCollector::GetSeconds(), ChunkId, Result, Record);
 		}
 
-		virtual void OnAcceptedNewRequirements(const TSet<FGuid>& ChunkIds) override
-		{
-		}
-
 	public:
-		TArray<FBatchStarted> RxBatchStarted;
 		TArray<FLoadStarted> RxLoadStarted;
 		TArray<FLoadComplete> RxLoadComplete;
-		TFunction<void(const TArray<FGuid>&)> OnBatchStartedFunc;
 		TFunction<void(const FGuid&)> OnLoadStartedFunc;
 		TFunction<void(const FGuid&, const ELoadResult&, const ISpeedRecorder::FRecord&)> OnLoadCompleteFunc;
 	};
diff --git a/Engine/Source/Runtime/Online/BuildPatchServices/Private/Tests/Unit/ChainedChunkSource.spec.cpp b/Engine/Source/Runtime/Online/BuildPatchServices/Private/Tests/Unit/ChainedChunkSource.spec.cpp
deleted file mode 100644
diff --git a/Engine/Source/Runtime/Online/BuildPatchServices/Private/Tests/Unit/CloudChunkSource.spec.cpp b/Engine/Source/Runtime/Online/BuildPatchServices/Private/Tests/Unit/CloudChunkSource.spec.cpp
--- a/Engine/Source/Runtime/Online/BuildPatchServices/Private/Tests/Unit/CloudChunkSource.spec.cpp
+++ b/Engine/Source/Runtime/Online/BuildPatchServices/Private/Tests/Unit/CloudChunkSource.spec.cpp
@@ -91,7 +91,6 @@ void FCloudChunkSourceSpec::Define()
 		MockPlatform.Reset(new FMockPlatform());
 		InstallerSharedContext.Reset(new FBuildInstallerSharedContext(TEXT("CloudChunkSourceSpec")));
 		CloudSourceConfig.Reset(new FCloudSourceConfig({ "http://download.mydomain.com/clouddata" }));
-		CloudSourceConfig->SharedContext = InstallerSharedContext.Get();
 		FakeChunkStore.Reset(new FFakeChunkStore());
 		FakeDownloadService.Reset(new FFakeDownloadService());
 		MockChunkReferenceTracker.Reset(new FFakeChunkReferenceTracker());
diff --git a/Engine/Source/Runtime/Online/BuildPatchServices/Private/Tests/Unit/DiskChunkStore.spec.cpp b/Engine/Source/Runtime/Online/BuildPatchServices/Private/Tests/Unit/DiskChunkStore.spec.cpp
deleted file mode 100644
diff --git a/Engine/Source/Runtime/Online/BuildPatchServices/Private/Tests/Unit/InstallChunkSource.spec.cpp b/Engine/Source/Runtime/Online/BuildPatchServices/Private/Tests/Unit/InstallChunkSource.spec.cpp
--- a/Engine/Source/Runtime/Online/BuildPatchServices/Private/Tests/Unit/InstallChunkSource.spec.cpp
+++ b/Engine/Source/Runtime/Online/BuildPatchServices/Private/Tests/Unit/InstallChunkSource.spec.cpp
@@ -22,29 +22,23 @@
 BEGIN_DEFINE_SPEC(FInstallChunkSourceSpec, "BuildPatchServices.Unit", EAutomationTestFlags::ProductFilter | EAutomationTestFlags_ApplicationContextMask)
 const uint32 TestChunkSize = 128 * 1024;
 // Unit.
-TUniquePtr<BuildPatchServices::IInstallChunkSource> InstallChunkSource;
+TUniquePtr<BuildPatchServices::IConstructorInstallChunkSource> InstallChunkSource;
 // Mock.
 TUniquePtr<BuildPatchServices::FFakeFileSystem> FakeFileSystem;
-TUniquePtr<BuildPatchServices::FFakeChunkStore> FakeChunkStore;
-TUniquePtr<BuildPatchServices::FFakeChunkReferenceTracker> MockChunkReferenceTracker;
-TUniquePtr<BuildPatchServices::FFakeInstallerError> MockInstallerError;
 TUniquePtr<BuildPatchServices::FMockInstallChunkSourceStat> MockInstallChunkSourceStat;
 BuildPatchServices::FMockManifestPtr MockManifest;
 // Data.
-BuildPatchServices::FInstallSourceConfig Configuration;
 TMultiMap<FString, FBuildPatchAppManifestRef> InstallationSources;
 TSet<FGuid> SomeAvailableChunks;
 FGuid SomeChunk;
-float PauseTime;
-bool bHasPaused;
 TUniquePtr<BuildPatchServices::IBuildManifestSet> ManifestSet;
+TArray<uint8> ReadDestination;
 // Test helpers.
 void MakeUnit();
 void InventUsableChunkData();
 void SomeChunkAvailable();
 void SomeChunkUnavailable();
 void SomeChunkCorrupted();
-TFuture<void> PauseFor(float Seconds);
 END_DEFINE_SPEC(FInstallChunkSourceSpec)
 
 void FInstallChunkSourceSpec::Define()
@@ -54,16 +48,11 @@ void FInstallChunkSourceSpec::Define()
 	// Data setup.
 	FRollingHashConst::Init();
 	SomeChunk = FGuid::NewGuid();
-	Configuration.BatchFetchMinimum = 10;
-	Configuration.BatchFetchMaximum = 10;
 
 	// Specs.
 	BeforeEach([this]()
 	{
 		FakeFileSystem.Reset(new FFakeFileSystem());
-		FakeChunkStore.Reset(new FFakeChunkStore());
-		MockChunkReferenceTracker.Reset(new FFakeChunkReferenceTracker());
-		MockInstallerError.Reset(new FFakeInstallerError());
 		MockInstallChunkSourceStat.Reset(new FMockInstallChunkSourceStat());
 		MockManifest = MakeShareable(new FMockManifest());
 		ManifestSet.Reset(FBuildManifestSetFactory::Create({ BuildPatchServices::FInstallerAction::MakeInstall(MockManifest.ToSharedRef()) }));
@@ -116,22 +105,19 @@ void FInstallChunkSourceSpec::Define()
 
 				Describe("when some chunk is not in the store", [this]()
 				{
-					It("should return nullptr.", [this]()
+					It("should fail.", [this]()
 					{
-						TEST_NULL(InstallChunkSource->Get(SomeChunk));
-					});
-				});
+						bool bFailed = false;
+						TUniqueFunction<void(bool)> Request = InstallChunkSource->CreateRequest(SomeChunk, {}, 0, IConstructorChunkSource::FChunkRequestCompleteDelegate::CreateLambda(
+						[&bFailed](const FGuid& DataId, bool bAborted, bool bFailedToRead, void* UserPtr)
+						{
+							bFailed = bFailedToRead;
+						}
+						));
 
-				Describe("when some chunk is in the store", [this]()
-				{
-					BeforeEach([this]()
-					{
-						FakeChunkStore->Store.Emplace(SomeChunk, new FFakeChunkDataAccess());
-					});
+						Request(false);
 
-					It("should return some chunk.", [this]()
-					{
-						TEST_NOT_NULL(InstallChunkSource->Get(SomeChunk));
+						TEST_TRUE(bFailed);
 					});
 				});
 			});
@@ -149,33 +135,25 @@ void FInstallChunkSourceSpec::Define()
 				{
 					It("should return some chunk loading from disk.", [this]()
 					{
-						TEST_NOT_NULL(InstallChunkSource->Get(SomeChunk));
+						bool bFailed = false;
+						TUniqueFunction<void(bool)> Request = InstallChunkSource->CreateRequest(SomeChunk, 
+							FMutableMemoryView(ReadDestination.GetData(), ReadDestination.Num()), 
+							0, 
+							IConstructorChunkSource::FChunkRequestCompleteDelegate::CreateLambda(
+								[&bFailed](const FGuid& DataId, bool bAborted, bool bFailedToRead, void* UserPtr)
+								{
+									bFailed = bFailedToRead;
+								}
+							));
+
+						Request(false);
+
+						TEST_FALSE(bFailed);
 						TEST_EQUAL(MockInstallChunkSourceStat->RxLoadStarted.Num(), 1);
 						TEST_EQUAL(MockInstallChunkSourceStat->RxLoadComplete.Num(), 1);
 						TEST_TRUE(FakeFileSystem->RxCreateFileReader.Num() > 0);
 					});
-
-					Describe("when upcoming chunk references are known", [this]()
-					{
-						BeforeEach([this]()
-						{
-							MockChunkReferenceTracker->NextReferences = SomeAvailableChunks.Array();
-						});
-
-						It("should also load upcoming chunks according to provided configuration.", [this]()
-						{
-							InstallChunkSource->Get(SomeChunk);
-							TEST_TRUE(FakeChunkStore->Store.Contains(SomeChunk));
-
-							// disable this one because we no longer prefetch in the install thread, so only the one chunk
-							// we "get" should exist.
-//							for (int32 NextReferenceIdx = 0; NextReferenceIdx < Configuration.BatchFetchMaximum; ++NextReferenceIdx)
-//							{
-//								TEST_TRUE(FakeChunkStore->Store.Contains(MockChunkReferenceTracker->NextReferences[NextReferenceIdx]));
-//							}
-						});
-					});
-
+					
 					Describe("when some chunk hashes are not known", [this]()
 					{
 						BeforeEach([this]()
@@ -189,7 +167,20 @@ void FInstallChunkSourceSpec::Define()
 
 						It("should not have attempted to load some chunk from disk.", [this]()
 						{
-							InstallChunkSource->Get(SomeChunk);
+							bool bFailed = false;
+							TUniqueFunction<void(bool)> Request = InstallChunkSource->CreateRequest(SomeChunk, 
+								FMutableMemoryView(ReadDestination.GetData(), ReadDestination.Num()), 
+								0, 
+								IConstructorChunkSource::FChunkRequestCompleteDelegate::CreateLambda(
+									[&bFailed](const FGuid& DataId, bool bAborted, bool bFailedToRead, void* UserPtr)
+									{
+										bFailed = bFailedToRead;
+									}
+							));
+
+							Request(false);
+
+							TEST_TRUE(bFailed);
 							TEST_EQUAL(FakeFileSystem->RxCreateFileReader.Num(), 0);
 						});
 					});
@@ -210,7 +201,20 @@ void FInstallChunkSourceSpec::Define()
 
 						It("should still succeed to load some chunk from disk.", [this]()
 						{
-							TEST_NOT_NULL(InstallChunkSource->Get(SomeChunk));
+							bool bFailed = false;
+							TUniqueFunction<void(bool)> Request = InstallChunkSource->CreateRequest(SomeChunk, 
+								FMutableMemoryView(ReadDestination.GetData(), ReadDestination.Num()), 
+								0, 
+								IConstructorChunkSource::FChunkRequestCompleteDelegate::CreateLambda(
+									[&bFailed](const FGuid& DataId, bool bAborted, bool bFailedToRead, void* UserPtr)
+									{
+										bFailed = bFailedToRead;
+									}
+							));
+
+							Request(false);
+
+							TEST_FALSE(bFailed);
 							TEST_EQUAL(MockInstallChunkSourceStat->RxLoadStarted.Num(), 1);
 							TEST_EQUAL(MockInstallChunkSourceStat->RxLoadComplete.Num(), 1);
 							TEST_TRUE(FakeFileSystem->RxCreateFileReader.Num() > 0);
@@ -225,7 +229,21 @@ void FInstallChunkSourceSpec::Define()
 
 							It("should fail to load some chunk from disk, reporting IInstallChunkSourceStat::ELoadResult::HashCheckFailed.", [this]()
 							{
-								TEST_NULL(InstallChunkSource->Get(SomeChunk));
+								bool bFailed = false;
+								TUniqueFunction<void(bool)> Request = InstallChunkSource->CreateRequest(SomeChunk, 
+									FMutableMemoryView(ReadDestination.GetData(), ReadDestination.Num()), 
+									0, 
+									IConstructorChunkSource::FChunkRequestCompleteDelegate::CreateLambda(
+										[&bFailed](const FGuid& DataId, bool bAborted, bool bFailedToRead, void* UserPtr)
+										{
+											bFailed = bFailedToRead;
+										}
+								));
+
+								Request(false);
+
+								TEST_TRUE(bFailed);
+
 								TEST_EQUAL(MockInstallChunkSourceStat->RxLoadComplete.Num(), 1);
 								if (MockInstallChunkSourceStat->RxLoadComplete.Num() == 1)
 								{
@@ -244,7 +262,20 @@ void FInstallChunkSourceSpec::Define()
 
 						It("should fail to load some chunk from disk, reporting IInstallChunkSourceStat::ELoadResult::HashCheckFailed.", [this]()
 						{
-							TEST_NULL(InstallChunkSource->Get(SomeChunk));
+							bool bFailed = false;
+							TUniqueFunction<void(bool)> Request = InstallChunkSource->CreateRequest(SomeChunk, 
+								FMutableMemoryView(ReadDestination.GetData(), ReadDestination.Num()), 
+								0, 
+								IConstructorChunkSource::FChunkRequestCompleteDelegate::CreateLambda(
+									[&bFailed](const FGuid& DataId, bool bAborted, bool bFailedToRead, void* UserPtr)
+									{
+										bFailed = bFailedToRead;
+									}
+							));
+
+							Request(false);
+
+							TEST_TRUE(bFailed);
 							TEST_EQUAL(MockInstallChunkSourceStat->RxLoadComplete.Num(), 1);
 							if (MockInstallChunkSourceStat->RxLoadComplete.Num() == 1)
 							{
@@ -253,108 +284,14 @@ void FInstallChunkSourceSpec::Define()
 						});
 					});
 				});
-
-				Describe("when some chunk is in the store", [this]()
-				{
-					BeforeEach([this]()
-					{
-						FakeChunkStore->Store.Emplace(SomeChunk, new FFakeChunkDataAccess());
-					});
-
-					It("should return some chunk without loading.", [this]()
-					{
-						TEST_NOT_NULL(InstallChunkSource->Get(SomeChunk));
-						TEST_EQUAL(MockInstallChunkSourceStat->RxLoadStarted.Num(), 0);
-						TEST_EQUAL(FakeFileSystem->RxCreateFileReader.Num(), 0);
-					});
-				});
-			});
-		});
-
-		Describe("SetPaused", [this]()
-		{
-			BeforeEach([this]()
-			{
-				InventUsableChunkData();
-				SomeChunkAvailable();
-				MakeUnit();
-				MockChunkReferenceTracker->NextReferences = SomeAvailableChunks.Array();
-				bHasPaused = false;
-				PauseTime = 0.1f;
-				MockInstallChunkSourceStat->OnLoadCompleteFunc = [this](const FGuid& ChunkId, const IInstallChunkSourceStat::ELoadResult& Result, const ISpeedRecorder::FRecord& Record)
-				{
-					if (!bHasPaused)
-					{
-						bHasPaused = true;
-						PauseFor(PauseTime);
-					}
-				};
-			});
-
-			It("should delay the chunk load process.", [this]()
-			{
-				// Pausing happens in the LoadFromBuild and we only load 1 per Get, so get two and be sure
-				// to check complete time.
-				int32 GetCounter = 2;
-				for (FGuid& AvailableChunk : SomeAvailableChunks)
-				{
-					InstallChunkSource->Get(AvailableChunk);
-					GetCounter--;
-					if (GetCounter == 0)
-					{
-						break;
-					}
-				}
-
-				double LongestDelay = 0.0f;
-				for (int32 Idx = 1; Idx < MockInstallChunkSourceStat->RxLoadComplete.Num(); ++Idx)
-				{
-					double ThisDelay = MockInstallChunkSourceStat->RxLoadComplete[Idx].Get<0>() - MockInstallChunkSourceStat->RxLoadComplete[Idx - 1].Get<0>();
-					if (ThisDelay > LongestDelay)
-					{
-						LongestDelay = ThisDelay;
-					}
-				}
-				TEST_TRUE(LongestDelay >= PauseTime);
 			});
 		});
 
-		Describe("Abort", [this]()
-		{
-			BeforeEach([this]()
-			{
-				InventUsableChunkData();
-				SomeChunkAvailable();
-				MakeUnit();
-				MockChunkReferenceTracker->NextReferences = SomeAvailableChunks.Array();
-				MockInstallChunkSourceStat->OnLoadStartedFunc = [this](const FGuid& ChunkId)
-				{
-					if (MockInstallChunkSourceStat->RxLoadStarted.Num() > 1)
-					{
-						InstallChunkSource->Abort();
-					}
-				};
-			});
-
-			It("should abort loading of chunks, reporting IInstallChunkSourceStat::ELoadResult::Aborted.", [this]()
-			{
-				InstallChunkSource->Get(SomeChunk);
-				TSet<FGuid> StartedChunks;
-				for (const FMockInstallChunkSourceStat::FLoadStarted& Call : MockInstallChunkSourceStat->RxLoadStarted)
-				{
-					StartedChunks.Add(Call.Get<1>());
-				}
-				TEST_TRUE(StartedChunks.Num() < Configuration.BatchFetchMinimum);
-			});
-		});
 	});
 
 	AfterEach([this]()
 	{
 		FakeFileSystem.Reset();
-		FakeChunkStore.Reset();
-		MockChunkReferenceTracker.Reset();
-		MockInstallerError.Reset();
 		MockInstallChunkSourceStat.Reset();
 		ManifestSet.Reset();
 		MockManifest.Reset();
@@ -366,19 +303,19 @@ void FInstallChunkSourceSpec::Define()
 void FInstallChunkSourceSpec::MakeUnit()
 {
 	using namespace BuildPatchServices;
-	InstallChunkSource.Reset(FInstallChunkSourceFactory::Create(
-		Configuration,
+	TSet<FGuid> ChunksThatWillBeNeeded;
+	ManifestSet->GetReferencedChunks(ChunksThatWillBeNeeded);
+	InstallChunkSource.Reset(IConstructorInstallChunkSource::CreateInstallSource(
 		FakeFileSystem.Get(),
-		FakeChunkStore.Get(),
-		MockChunkReferenceTracker.Get(),
-		MockInstallerError.Get(),
 		MockInstallChunkSourceStat.Get(),
 		InstallationSources,
-		ManifestSet.Get()));
+		ChunksThatWillBeNeeded));
 }
 
 void FInstallChunkSourceSpec::InventUsableChunkData()
 {
+	ReadDestination.SetNumUninitialized(TestChunkSize);
+
 	using namespace BuildPatchServices;
 	for (int32 Count = 0; Count < 100; ++Count)
 	{
@@ -526,19 +463,5 @@ void FInstallChunkSourceSpec::SomeChunkCorrupted()
 	}
 }
 
-TFuture<void> FInstallChunkSourceSpec::PauseFor(float Seconds)
-{
-	double PauseAt = BuildPatchServices::FStatsCollector::GetSeconds();
-	InstallChunkSource->SetPaused(true);
-	TFunction<void()> Task = [this, PauseAt, Seconds]()
-	{
-		while ((BuildPatchServices::FStatsCollector::GetSeconds() - PauseAt) < Seconds)
-		{
-			FPlatformProcess::Sleep(0.01f);
-		}
-		InstallChunkSource->SetPaused(false);
-	};
-	return Async(EAsyncExecution::Thread, MoveTemp(Task));
-}
 
 #endif //WITH_DEV_AUTOMATION_TESTS
diff --git a/Engine/Source/Runtime/Online/BuildPatchServices/Public/BuildPatchMessage.h b/Engine/Source/Runtime/Online/BuildPatchServices/Public/BuildPatchMessage.h
--- a/Engine/Source/Runtime/Online/BuildPatchServices/Public/BuildPatchMessage.h
+++ b/Engine/Source/Runtime/Online/BuildPatchServices/Public/BuildPatchMessage.h
@@ -18,8 +18,6 @@ namespace BuildPatchServices
 			AccessLost = 0,
 			// Access has been regained after being lost.
 			AccessRegained,
-			// The chunk source has been retired and is not expected to be accessed again (i.e. can be deleted).
-			Retired
 		};
 
 		// The type of event that occurred.
diff --git a/Engine/Source/Runtime/Online/BuildPatchServices/Public/Interfaces/IBuildInstaller.h b/Engine/Source/Runtime/Online/BuildPatchServices/Public/Interfaces/IBuildInstaller.h
--- a/Engine/Source/Runtime/Online/BuildPatchServices/Public/Interfaces/IBuildInstaller.h
+++ b/Engine/Source/Runtime/Online/BuildPatchServices/Public/Interfaces/IBuildInstaller.h
@@ -122,7 +122,6 @@ struct FBuildInstallStats
 		, NumBadDownloads(0)
 		, NumAbortedDownloads(0)
 		, NumRecycleFailures(0)
-		, NumChunksStoreBooted(0)
 		, NumDriveStoreChunkLoads(0)
 		, NumDriveStoreLoadFailures(0)
 		, NumChunkDbChunksFailed(0)
@@ -163,11 +162,6 @@ struct FBuildInstallStats
 		, OkDownloadHealthTime(0.0f)
 		, PoorDownloadHealthTime(0.0f)
 		, DisconnectedDownloadHealthTime(0.0f)
-		, AverageMemoryStoreUse(0.0f)
-		, PeakMemoryStoreUse(0)
-		, AverageMemoryStoreRetained(0.0f)
-		, PeakMemoryStoreRetained(0)
-		, MemoryStoreSize(0)
 		, ProcessRequiredDiskSpace(0)
 		, ProcessAvailableDiskSpace(0)
 	{}
@@ -200,12 +194,14 @@ struct FBuildInstallStats
 	uint32 NumAbortedDownloads;
 	// The number of chunks that failed to be recycled from existing build.
 	uint32 NumRecycleFailures;
-	// The number of chunks that had to be booted from the memory stores.
-	uint32 NumChunksStoreBooted;
 	// The number of chunks that had to be loaded from the drive store.
 	uint32 NumDriveStoreChunkLoads;
 	// The number of chunks that failed to load from the drive store.
 	uint32 NumDriveStoreLoadFailures;
+	// The max disk space used by the drive store
+	uint32 DriveStorePeakBytes = 0;
+	// The number of chunks that were lost due to drive space limits
+	uint32 NumDriveStoreLostChunks = 0;
 	// The number of chunks that were not successfully loaded from provided chunkdbs.
 	uint32 NumChunkDbChunksFailed;
 	// The total number of bytes downloaded.
@@ -287,16 +283,10 @@ struct FBuildInstallStats
 	float PoorDownloadHealthTime;
 	// The amount of time that was spent with Disconnected download health.
 	float DisconnectedDownloadHealthTime;
-	// AverageMemoryStoreUse
-	float AverageMemoryStoreUse;
-	// PeakMemoryStoreUse
-	uint32 PeakMemoryStoreUse;
-	// AverageMemoryStoreRetained
-	float AverageMemoryStoreRetained;
-	// PeakMemoryStoreRetained
-	uint32 PeakMemoryStoreRetained;
-	// MemoryStoreSize
-	uint32 MemoryStoreSize;
+	// Max memory usage for the memory store.
+	uint64 MemoryStoreSizePeakBytes = 0; 
+	// Memory limit for the memory backing store
+	uint64 MemoryStoreSizeLimitBytes = 0;
 	// The total number of bytes required to start the installation.
 	uint64 ProcessRequiredDiskSpace;
 	// The total number of bytes available at the time of checking ProcessRequiredDiskSpace.
diff --git a/Engine/Source/Runtime/Online/BuildPatchServices/Public/Interfaces/IBuildInstallerSharedContext.h b/Engine/Source/Runtime/Online/BuildPatchServices/Public/Interfaces/IBuildInstallerSharedContext.h
--- a/Engine/Source/Runtime/Online/BuildPatchServices/Public/Interfaces/IBuildInstallerSharedContext.h
+++ b/Engine/Source/Runtime/Online/BuildPatchServices/Public/Interfaces/IBuildInstallerSharedContext.h
@@ -57,7 +57,7 @@ public:
 	/**
 	 * @return The number of threads required per installer
 	 */
-	virtual uint32 NumThreadsPerInstaller(bool bUseChunkDBs) const = 0;
+	virtual uint32 NumThreadsPerInstaller(bool bUseChunkDBs, bool bHasInstallSource) const = 0;
 };
 
 using IBuildInstallerSharedContextRef = TSharedRef<IBuildInstallerSharedContext>;
diff --git a/Engine/Source/Runtime/Online/BuildPatchServices/Public/Interfaces/IBuildStatistics.h b/Engine/Source/Runtime/Online/BuildPatchServices/Public/Interfaces/IBuildStatistics.h
--- a/Engine/Source/Runtime/Online/BuildPatchServices/Public/Interfaces/IBuildStatistics.h
+++ b/Engine/Source/Runtime/Online/BuildPatchServices/Public/Interfaces/IBuildStatistics.h
@@ -128,22 +128,26 @@ namespace BuildPatchServices
 		/**
 		 * @return the size in chunks of the memory store.
 		 */
-		virtual int32 GetMemoryChunkStoreSize() const = 0;
+		UE_DEPRECATED(5.6, "Appears to be unused, not replaced")
+		virtual int32 GetMemoryChunkStoreSize() const { return 0; }
 
 		/**
 		 * @return the number of chunks currently in the memory store.
 		 */
-		virtual int32 GetMemoryChunksInStore() const = 0;
+		UE_DEPRECATED(5.6, "Appears to be unused, not replaced")
+		virtual int32 GetMemoryChunksInStore() const { return 0; }
 
 		/**
 		 * @return the number of chunks that have been booted from the memory store.
 		 */
-		virtual int32 GetMemoryChunksBooted() const = 0;
+		UE_DEPRECATED(5.6, "Appears to be unused, not replaced")
+		virtual int32 GetMemoryChunksBooted() const { return 0; }
 
 		/**
 		 * @return the number of chunks currently in the memory store which are held due to multiple referencing.
 		 */
-		virtual int32 GetMemoryChunksRetained() const = 0;
+		UE_DEPRECATED(5.6, "Appears to be unused, not replaced")
+		virtual int32 GetMemoryChunksRetained() const { return 0; }
 
 		/**
 		 * @return the filename of the file currently being worked on.
@@ -219,23 +223,6 @@ namespace BuildPatchServices
 		 * @return the verify error counts per error type experienced during this installation.
 		 */
 		virtual TMap<EVerifyError, int32> GetVerifyErrorCounts() const = 0;
-
-		UE_DEPRECATED(4.21, "Please use GetMemoryChunkStoreSize instead.")
-		virtual int32 GetInstallMemoryChunkStoreSize() const { return GetMemoryChunkStoreSize(); }
-		UE_DEPRECATED(4.21, "Please use GetMemoryChunksInStore instead.")
-		virtual int32 GetInstallMemoryChunksInStore() const { return GetMemoryChunksInStore(); }
-		UE_DEPRECATED(4.21, "Please use GetMemoryChunksBooted instead.")
-		virtual int32 GetInstallMemoryChunksBooted() const { return GetMemoryChunksBooted(); }
-		UE_DEPRECATED(4.21, "Please use GetMemoryChunksRetained instead.")
-		virtual int32 GetInstallMemoryChunksRetained() const { return GetMemoryChunksRetained(); }
-		UE_DEPRECATED(4.21, "Please use GetMemoryChunkStoreSize instead.")
-		virtual int32 GetCloudMemoryChunkStoreSize() const { return GetMemoryChunkStoreSize(); }
-		UE_DEPRECATED(4.21, "Please use GetMemoryChunksInStore instead.")
-		virtual int32 GetCloudMemoryChunksInStore() const { return GetMemoryChunksInStore(); }
-		UE_DEPRECATED(4.21, "Please use GetMemoryChunksBooted instead.")
-		virtual int32 GetCloudMemoryChunksBooted() const { return GetMemoryChunksBooted(); }
-		UE_DEPRECATED(4.21, "Please use GetMemoryChunksRetained instead.")
-		virtual int32 GetCloudMemoryChunksRetained() const { return GetMemoryChunksRetained(); }
 	};
 
 	typedef TSharedPtr<IBuildStatistics> IBuildStatisticsPtr;
