Static analysis fixes for file constructor rework.

#rnx

diff --git a/Engine/Source/Runtime/Online/BuildPatchServices/Private/BuildPatchFileConstructor.cpp b/Engine/Source/Runtime/Online/BuildPatchServices/Private/BuildPatchFileConstructor.cpp
--- a/Engine/Source/Runtime/Online/BuildPatchServices/Private/BuildPatchFileConstructor.cpp
+++ b/Engine/Source/Runtime/Online/BuildPatchServices/Private/BuildPatchFileConstructor.cpp
@@ -299,6 +299,7 @@ class FChunkBackingStore : IConstructorChunkSource
 	uint64 CurrentMemoryLoad = 0;
 	uint64 PeakMemoryLoad = 0;
 
+	static constexpr uint64 ChunkStoreMemoryLimitDisabledSentinel = TNumericLimits<uint64>::Max();
 	uint64 MaxMemoryBytes = 1 << 30; // 1GB.
 
 
@@ -498,7 +499,7 @@ class FChunkBackingStore : IConstructorChunkSource
 			}
 
 			// Need to expand the backing store.
-			if (UseMaxDiskSpaceBytes && ((uint64)(BackingStoreEntryCount + EntriesRequired) << BitsPerEntry) > UseMaxDiskSpaceBytes)
+			if (UseMaxDiskSpaceBytes && (((uint64)BackingStoreEntryCount + EntriesRequired) << BitsPerEntry) > UseMaxDiskSpaceBytes)
 			{
 				// We can't expand - this fails to page out and gets evicted from the backing store.
 				Stats.DiskLostChunkCount++;
@@ -780,11 +781,22 @@ public:
 			}
 		}
 
+		
+
 		// check current values. We expect this to override anything from above.
 		int32 ChunkStoreMemoryLimitMB = 0;
 		if (GConfig->GetInt(TEXT("Portal.BuildPatch"), TEXT("ChunkStoreMemorySizeMB"), ChunkStoreMemoryLimitMB, GEngineIni))
 		{
-			ChunkStoreMemoryLimit = (uint64)ChunkStoreMemoryLimitMB << 20;
+			// To be consistent with other limitations, negative disables it
+			// 0 is OK - we require locked data to be in memory so we'll go over the limit but it'll be a minimum.
+			if (ChunkStoreMemoryLimitMB < 0)
+			{
+				ChunkStoreMemoryLimit = ChunkStoreMemoryLimitDisabledSentinel;
+			}
+			else
+			{
+				ChunkStoreMemoryLimit = (uint64)ChunkStoreMemoryLimitMB << 20;
+			}
 		}
 		
 
@@ -805,12 +817,7 @@ public:
 			}
 		}		
 
-		// To be consistent with other limitations, negative disables it
-		// 0 is OK - we require locked data to be in memory so we'll go over the limit but it'll be a minimum.
-		if (ChunkStoreMemoryLimit < 0)
-		{
-			ChunkStoreMemoryLimit = TNumericLimits<uint64>::Max();
-		}
+
 
 
 		uint64 AvailableMem = FPlatformMemory::GetStats().AvailablePhysical;
@@ -833,7 +840,7 @@ public:
 			}
 		}
 
-		if (ChunkStoreMemoryLimit == -1)
+		if (ChunkStoreMemoryLimit == ChunkStoreMemoryLimitDisabledSentinel)
 		{
 			UE_LOG(LogBuildPatchServices, Display, TEXT("ChunkStoreMemoryLimits are disabled"));
 		}
@@ -960,70 +967,72 @@ public:
 
 		int32 CurrentUsageIndex = ParentConstructor->ChunkReferenceTracker->GetCurrentUsageIndex();
 
-		while (CurrentMemoryLoad + InChunkSize > MaxMemoryBytes)
+		if (MaxMemoryBytes != ChunkStoreMemoryLimitDisabledSentinel)
 		{
-			// Gotta dump stuff to disk. If we fail to dump to disk we mark the chunk
-			// as only available via the cloud source.
-
-			// Evict the one that is the longest until we use it. Chunks are almost always the same size
-			// so we expect this to run once.
-			int32 FarthestNextUsage = -1;
-			FGuid EvictGuid;
-			for (TPair<FGuid, FStoredChunk>& Pair : StoredChunks)
+			while (CurrentMemoryLoad + InChunkSize > MaxMemoryBytes)
 			{
-				FStoredChunk& Chunk = Pair.Value;
-
-				if (Chunk.LockCount || // Actively in use - don't evict.
-					!Chunk.ChunkData.Num()) // Not in memory - can't evict.
-				{					
-					continue;
-				}
+				// Gotta dump stuff to disk. If we fail to dump to disk we mark the chunk
+				// as only available via the cloud source.
 
-				if (Chunk.NextUsageIndex < CurrentUsageIndex)
+				// Evict the one that is the longest until we use it. Chunks are almost always the same size
+				// so we expect this to run once.
+				int32 FarthestNextUsage = -1;
+				FGuid EvictGuid;
+				for (TPair<FGuid, FStoredChunk>& Pair : StoredChunks)
 				{
-					Chunk.NextUsageIndex = ParentConstructor->ChunkReferenceTracker->GetNextUsageForChunk(Pair.Key, Chunk.LastUsageIndex);
+					FStoredChunk& Chunk = Pair.Value;
+
+					if (Chunk.LockCount || // Actively in use - don't evict.
+						!Chunk.ChunkData.Num()) // Not in memory - can't evict.
+					{					
+						continue;
+					}
+
+					if (Chunk.NextUsageIndex < CurrentUsageIndex)
+					{
+						Chunk.NextUsageIndex = ParentConstructor->ChunkReferenceTracker->GetNextUsageForChunk(Pair.Key, Chunk.LastUsageIndex);
+					}
+
+					if (Chunk.NextUsageIndex > FarthestNextUsage)
+					{
+						FarthestNextUsage = Pair.Value.NextUsageIndex;
+						EvictGuid = Pair.Key;
+					}
 				}
 
-				if (Chunk.NextUsageIndex > FarthestNextUsage)
+				if (FarthestNextUsage == -1)
 				{
-					FarthestNextUsage = Pair.Value.NextUsageIndex;
-					EvictGuid = Pair.Key;
+					// This means we can't reserve and also keep our memory requirements. 
+					//
+					// We currently have a minimum memory requirement for construction: Partial or reused chunks
+					// are routed through the backing store for holding. Partial because we need a full allocation to decompress
+					// the chunk and reuse so we don't have to re-read from disk.
+					//
+					// For unpatched installs this is minimal as most chunks are able to write directly to
+					// the destination buffer. For optimized delta manifests this is not the case as BPT assembles
+					// chunks from all over the place - I've seen 80+ chunk references to assemble 16MB of data,
+					// resulting in total buffer usage of 80MB + 16MB > 100MB of memory use per batch. With two
+					// batches in flight this pushes 200MB total buffer allocation.
+					//
+					// This is not likely an issue in real life as 200MB isn't much, but it does mean if we have a 
+					// low memory usage limit we can hit this legitimately.
+					//
+					// So...we let the reservation continue in violation of our memory constraints and hope we don't OOM.
+					//
+					// Future work could be:
+					// 1. Limit batch creation based on inflight chunk in addition to write buffer size
+					// 2. Only allow 1 batch in flight when using an optimized delta + low memory constraints.
+					break;
 				}
-			}
-
-			if (FarthestNextUsage == -1)
-			{
-				// This means we can't reserve and also keep our memory requirements. 
-				//
-				// We currently have a minimum memory requirement for construction: Partial or reused chunks
-				// are routed through the backing store for holding. Partial because we need a full allocation to decompress
-				// the chunk and reuse so we don't have to re-read from disk.
-				//
-				// For unpatched installs this is minimal as most chunks are able to write directly to
-				// the destination buffer. For optimized delta manifests this is not the case as BPT assembles
-				// chunks from all over the place - I've seen 80+ chunk references to assemble 16MB of data,
-				// resulting in total buffer usage of 80MB + 16MB > 100MB of memory use per batch. With two
-				// batches in flight this pushes 200MB total buffer allocation.
-				//
-				// This is not likely an issue in real life as 200MB isn't much, but it does mean if we have a 
-				// low memory usage limit we can hit this legitimately.
-				//
-				// So...we let the reservation continue in violation of our memory constraints and hope we don't OOM.
-				//
-				// Future work could be:
-				// 1. Limit batch creation based on inflight chunk in addition to write buffer size
-				// 2. Only allow 1 batch in flight when using an optimized delta + low memory constraints.
-				break;
-			}
-			else
-			{
-				if (!PageOut(EvictGuid))
+				else
 				{
-					return FMutableMemoryView();
+					if (!PageOut(EvictGuid))
+					{
+						return FMutableMemoryView();
+					}
 				}
-			}
-		}
-
+			} // end while over limits
+		} // end if memory limits enabled
 		FStoredChunk& StoredChunk = StoredChunks.FindOrAdd(InGuid);
 
 		//
