Removed Online code marked for deprecation in 5.4

diff --git a/Engine/Plugins/Online/VoiceChat/VoiceChat/Source/Public/VoiceChat.h b/Engine/Plugins/Online/VoiceChat/VoiceChat/Source/Public/VoiceChat.h
--- a/Engine/Plugins/Online/VoiceChat/VoiceChat/Source/Public/VoiceChat.h
+++ b/Engine/Plugins/Online/VoiceChat/VoiceChat/Source/Public/VoiceChat.h
@@ -455,25 +455,13 @@ public:
 	 */
 	virtual FOnVoiceChatCallStatsUpdatedDelegate& OnVoiceChatCallStatsUpdated() = 0;
 
-	/**
-	 * Set the 3d position of the player
-	 *
-	 * @param ChannelName
-	 * @param SpeakerPosition
-	 * @param ListenerPosition
-	 * @param ListenerForwardDirection
-	 * @param ListenerUpDirection
-	 */
-	UE_DEPRECATED(5.4, "Please replace with new Set3DPosition override") // TODO also remove the default impl of Set3DPosition below when removing this.
-	virtual void Set3DPosition(const FString& ChannelName, const FVector& SpeakerPosition, const FVector& ListenerPosition, const FVector& ListenerForwardDirection, const FVector& ListenerUpDirection) {};
-
 	/**
 	 * Set the 3d position of the player
 	 *
 	 * @param ChannelName
 	 * @param Position
 	 */
-	virtual void Set3DPosition(const FString& ChannelName, const FVector& Position) {};
+	virtual void Set3DPosition(const FString& ChannelName, const FVector& Position) = 0;
 
 	/**
 	 * Get an array of channels the user is in
diff --git a/Engine/Source/Runtime/Analytics/AnalyticsET/Private/IAnalyticsProviderET.cpp b/Engine/Source/Runtime/Analytics/AnalyticsET/Private/IAnalyticsProviderET.cpp
--- a/Engine/Source/Runtime/Analytics/AnalyticsET/Private/IAnalyticsProviderET.cpp
+++ b/Engine/Source/Runtime/Analytics/AnalyticsET/Private/IAnalyticsProviderET.cpp
@@ -331,10 +331,6 @@ bool FAnalyticsProviderET::Tick(float DeltaSeconds)
 	LLM_SCOPE_BYNAME(TEXT("Analytics"));
 	QUICK_SCOPE_CYCLE_COUNTER(STAT_FAnalyticsProviderET_Tick);
 
-PRAGMA_DISABLE_DEPRECATION_WARNINGS
-	HttpRetryManager->Update();
-PRAGMA_ENABLE_DEPRECATION_WARNINGS
-
 	// hold a lock the entire time here because we're making several calls to the event cache that we need to be consistent when we decide to flush.
 	// With more care, we can likely avoid holding this lock the entire time.
 	FAnalyticsProviderETEventCache::Lock EventCacheLock(EventCache);
diff --git a/Engine/Source/Runtime/Online/HTTP/Private/HttpRetrySystem.cpp b/Engine/Source/Runtime/Online/HTTP/Private/HttpRetrySystem.cpp
--- a/Engine/Source/Runtime/Online/HTTP/Private/HttpRetrySystem.cpp
+++ b/Engine/Source/Runtime/Online/HTTP/Private/HttpRetrySystem.cpp
@@ -602,250 +602,6 @@ float FHttpRetrySystem::FManager::GetLockoutPeriodSeconds(const FHttpRetryReques
 
 static FRandomStream TempRandomStream(4435261);
 
-bool FHttpRetrySystem::FManager::Update(uint32* FileCount, uint32* FailingCount, uint32* FailedCount, uint32* CompletedCount)
-{
-	if (CVarHttpRetrySystemNonGameThreadSupportEnabled.GetValueOnAnyThread())
-	{
-		return true;
-	}
-
-	QUICK_SCOPE_CYCLE_COUNTER(STAT_FHttpRetrySystem_FManager_Update);
-	LLM_SCOPE_BYTAG(HTTP);
-
-	FScopeLock ScopeLock(&RequestListLock);
-
-	bool bIsGreen = true;
-
-	if (FileCount != nullptr)
-	{
-		*FileCount = RequestList.Num();
-	}
-
-	const double NowAbsoluteSeconds = FPlatformTime::Seconds();
-
-	// Basic algorithm
-	// for each managed item
-	//    if the item hasn't timed out
-	//       if the item's retry state is NotStarted
-	//          if the item's request's state is not NotStarted
-	//             move the item's retry state to Processing
-	//          endif
-	//       endif
-	//       if the item's retry state is Processing
-	//          if the item's request's state is Failed
-	//             flag return code to false
-	//             if the item can be retried
-	//                increment FailingCount if applicable
-	//                retry the item's request
-	//                increment the item's retry count
-	//             else
-	//                increment FailedCount if applicable
-	//                set the item's retry state to FailedRetry
-	//             endif
-	//          else if the item's request's state is Succeeded
-	//          endif
-	//       endif
-	//    else
-	//       flag return code to false
-	//       set the item's retry state to FailedTimeout
-	//       increment FailedCount if applicable
-	//    endif
-	//    if the item's retry state is FailedRetry
-	//       do stuff
-	//    endif
-	//    if the item's retry state is FailedTimeout
-	//       do stuff
-	//    endif
-	//    if the item's retry state is Succeeded
-	//       do stuff
-	//    endif
-	// endfor
-
-	int32 RequestIndex = 0;
-	while (RequestIndex < RequestList.Num())
-	{
-		QUICK_SCOPE_CYCLE_COUNTER(STAT_FHttpRetrySystem_FManager_Update_RequestListItem);
-
-		FHttpRetryRequestEntry* HttpRetryRequestEntry = &RequestList[RequestIndex];
-		TSharedRef<FHttpRetrySystem::FRequest, ESPMode::ThreadSafe> HttpRetryRequest = HttpRetryRequestEntry->Request;
-		// Delegates fired in this loop can resize the array if new requests are added, invalidating HttpRetryRequestEntry and HttpRetryRequest.  Call this when the array may have been modified.
-		auto ResetCurrentIterationVariables = [this, &HttpRetryRequestEntry, &HttpRetryRequest, RequestIndex]()
-		{
-			HttpRetryRequestEntry = &RequestList[RequestIndex];
-		};
-
-		const EHttpRequestStatus::Type RequestStatus = HttpRetryRequest->GetStatus();
-
-		if (HttpRetryRequestEntry->bShouldCancel)
-		{
-			UE_LOG(LogHttp, Warning, TEXT("Request cancelled on %s"), *(HttpRetryRequest->GetURL()));
-			HttpRetryRequest->RetryStatus = FHttpRetrySystem::FRequest::EStatus::Cancelled;
-		}
-		else
-		{
-			if (!HasTimedOut(*HttpRetryRequestEntry, NowAbsoluteSeconds))
-			{
-				if (HttpRetryRequest->RetryStatus == FHttpRetrySystem::FRequest::EStatus::NotStarted)
-				{
-					if (RequestStatus != EHttpRequestStatus::NotStarted)
-					{
-						HttpRetryRequest->RetryStatus = FHttpRetrySystem::FRequest::EStatus::Processing;
-					}
-				}
-
-				if (HttpRetryRequest->RetryStatus == FHttpRetrySystem::FRequest::EStatus::Processing)
-				{
-					bool forceFail = false;
-
-					// Code to simulate request failure
-					if (RequestStatus == EHttpRequestStatus::Succeeded && RandomFailureRate.IsSet())
-					{
-						float random = TempRandomStream.GetFraction();
-						if (random < RandomFailureRate.GetValue())
-						{
-							forceFail = true;
-						}
-					}
-
-					// If we failed to connect, try the next domain in the list
-					if (HttpRetryRequest->GetStatus() == EHttpRequestStatus::Failed && HttpRetryRequest->GetFailureReason() == EHttpFailureReason::ConnectionError)
-					{
-						if (HttpRetryRequest->RetryDomains.IsValid())
-						{
-							HttpRetryRequest->MoveToNextRetryDomain();
-						}
-					}
-					// Save these for failure case retry checks if we hit a completion state
-					bool bShouldRetry = false;
-					bool bCanRetry = false;
-					if (EHttpRequestStatus::IsFinished(RequestStatus))
-					{
-						bShouldRetry = ShouldRetry(*HttpRetryRequestEntry);
-						bCanRetry = CanRetry(*HttpRetryRequestEntry);
-					}
-
-					if (RequestStatus == EHttpRequestStatus::Failed || forceFail || (bShouldRetry && bCanRetry))
-					{
-						bIsGreen = false;
-
-						if (forceFail || (bShouldRetry && bCanRetry))
-						{
-							float LockoutPeriod = GetLockoutPeriodSeconds(*HttpRetryRequestEntry);
-
-							if (LockoutPeriod > 0.0f)
-							{
-								UE_LOG(LogHttp, Warning, TEXT("Lockout of %fs on %s"), LockoutPeriod, *(HttpRetryRequest->GetURL()));
-							}
-
-							HttpRetryRequestEntry->LockoutEndTimeAbsoluteSeconds = NowAbsoluteSeconds + LockoutPeriod;
-							HttpRetryRequest->RetryStatus = FHttpRetrySystem::FRequest::EStatus::ProcessingLockout;
-							
-							QUICK_SCOPE_CYCLE_COUNTER(STAT_FHttpRetrySystem_FManager_Update_OnRequestWillRetry);
-							HttpRetryRequest->OnRequestWillRetry().ExecuteIfBound(HttpRetryRequest, HttpRetryRequest->GetResponse(), LockoutPeriod);
-							ResetCurrentIterationVariables();
-						}
-						else
-						{
-							UE_LOG(LogHttp, Warning, TEXT("Retry exhausted on %s"), *(HttpRetryRequest->GetURL()));
-							if (FailedCount != nullptr)
-							{
-								++(*FailedCount);
-							}
-							HttpRetryRequest->RetryStatus = FHttpRetrySystem::FRequest::EStatus::FailedRetry;
-						}
-					}
-					else if (RequestStatus == EHttpRequestStatus::Succeeded)
-					{
-						if (HttpRetryRequestEntry->CurrentRetryCount > 0)
-						{
-							UE_LOG(LogHttp, Warning, TEXT("Success on %s"), *(HttpRetryRequest->GetURL()));
-						}
-
-						if (CompletedCount != nullptr)
-						{
-							++(*CompletedCount);
-						}
-
-						HttpRetryRequest->RetryStatus = FHttpRetrySystem::FRequest::EStatus::Succeeded;
-					}
-				}
-
-				if (HttpRetryRequest->RetryStatus == FHttpRetrySystem::FRequest::EStatus::ProcessingLockout)
-				{
-					if (NowAbsoluteSeconds >= HttpRetryRequestEntry->LockoutEndTimeAbsoluteSeconds)
-					{
-						RetryHttpRequest(*HttpRetryRequestEntry);
-						ResetCurrentIterationVariables();
-					}
-
-					if (FailingCount != nullptr)
-					{
-						++(*FailingCount);
-					}
-				}
-			}
-			else
-			{
-				UE_LOG(LogHttp, Warning, TEXT("Timeout on retry %d: %s"), HttpRetryRequestEntry->CurrentRetryCount + 1, *(HttpRetryRequest->GetURL()));
-				bIsGreen = false;
-				HttpRetryRequest->RetryStatus = FHttpRetrySystem::FRequest::EStatus::FailedTimeout;
-				if (FailedCount != nullptr)
-				{
-					++(*FailedCount);
-				}
-			}
-		}
-
-		bool bWasCompleted = false;
-		bool bWasSuccessful = false;
-
-        if (HttpRetryRequest->RetryStatus == FHttpRetrySystem::FRequest::EStatus::Cancelled ||
-            HttpRetryRequest->RetryStatus == FHttpRetrySystem::FRequest::EStatus::FailedRetry ||
-            HttpRetryRequest->RetryStatus == FHttpRetrySystem::FRequest::EStatus::FailedTimeout ||
-            HttpRetryRequest->RetryStatus == FHttpRetrySystem::FRequest::EStatus::Succeeded)
-		{
-			bWasCompleted = true;
-            bWasSuccessful = HttpRetryRequest->RetryStatus == FHttpRetrySystem::FRequest::EStatus::Succeeded;
-		}
-
-		if (bWasCompleted)
-		{
-			if (bWasSuccessful)
-			{
-				HttpRetryRequest->BroadcastResponseHeadersReceived();
-				ResetCurrentIterationVariables();
-			}
-			
-			QUICK_SCOPE_CYCLE_COUNTER(STAT_FHttpRetrySystem_FManager_Update_OnProcessRequestComplete);
-			HttpRetryRequest->OnProcessRequestComplete().ExecuteIfBound(HttpRetryRequest, HttpRetryRequest->GetResponse(), bWasSuccessful);
-			ResetCurrentIterationVariables();
-		}
-
-        if (bWasSuccessful)
-        {
-            if (CompletedCount != nullptr)
-            {
-                ++(*CompletedCount);
-            }
-        }
-
-		if (bWasCompleted)
-		{
-			if (RequestList[RequestIndex].CurrentRetryCount > 0)
-			{
-				FHttpLogVerbosityTracker::Get().DecrementRetriedRequests();
-			}
-			RequestList.RemoveAtSwap(RequestIndex);
-		}
-		else
-		{
-			++RequestIndex;
-		}
-	}
-
-	return bIsGreen;
-}
-
 FHttpRetrySystem::FManager::FHttpRetryRequestEntry::FHttpRetryRequestEntry(TSharedRef<FHttpRetrySystem::FRequest, ESPMode::ThreadSafe>& InRequest)
 	: bShouldCancel(false)
 	, CurrentRetryCount(0)
@@ -916,7 +672,6 @@ void FHttpRetrySystem::FManager::BlockUntilFlushed(float InTimeoutSec)
 {
 	const float SleepInterval = 0.016;
 	float TimeElapsed = 0.0f;
-	uint32 FileCount, FailingCount, FailedCount, CompleteCount;
 
 	while (TimeElapsed < InTimeoutSec)
 	{
@@ -930,10 +685,6 @@ void FHttpRetrySystem::FManager::BlockUntilFlushed(float InTimeoutSec)
 
 		FHttpModule::Get().GetHttpManager().Tick(SleepInterval);
 
-PRAGMA_DISABLE_DEPRECATION_WARNINGS
-		Update(&FileCount, &FailingCount, &FailedCount, &CompleteCount);
-PRAGMA_ENABLE_DEPRECATION_WARNINGS
-
 		FPlatformProcess::Sleep(SleepInterval);
 		TimeElapsed += SleepInterval;
 	}
diff --git a/Engine/Source/Runtime/Online/HTTP/Public/HttpModule.h b/Engine/Source/Runtime/Online/HTTP/Public/HttpModule.h
--- a/Engine/Source/Runtime/Online/HTTP/Public/HttpModule.h
+++ b/Engine/Source/Runtime/Online/HTTP/Public/HttpModule.h
@@ -81,15 +81,6 @@ public:
 		return *HttpManager;
 	}
 
-	/**
-	 * @return timeout in seconds for the entire http request to complete
-	 */
-	UE_DEPRECATED(5.4, "GetHttpTimeout has been deprecated, use GetHttpActivityTimeout or GetHttpTotalTimeout instead")
-	inline float GetHttpTimeout() const
-	{
-		return HttpTotalTimeout;
-	}
-
 	/**
 	 * @return total timeout in seconds for the entire http request to complete
 	 */
@@ -98,15 +89,6 @@ public:
 		return HttpTotalTimeout;
 	}
 
-	/**
-	 * Sets timeout in seconds for the entire http request to complete
-	 */
-	UE_DEPRECATED(5.4, "SetHttpTimeout has been deprecated, config it through HttpActivityTimeout or HttpTotalTimeout instead")
-	inline void SetHttpTimeout(float TimeOutInSec)
-	{
-		HttpTotalTimeout = TimeOutInSec;
-	}
-
 	/**
 	 * @return timeout in seconds to establish the connection
 	 */
@@ -115,24 +97,6 @@ public:
 		return HttpConnectionTimeout;
 	}
 
-	/**
-	 * @return timeout in seconds to receive a response on the connection 
-	 */
-	UE_DEPRECATED(5.4, "GetHttpReceiveTimeout has been deprecated, Use GetHttpActivityTimeout instead to decide if there is still any ongoing activity.")
-	inline float GetHttpReceiveTimeout() const
-	{
-		return HttpReceiveTimeout;
-	}
-
-	/**
-	 * @return timeout in seconds to send a request on the connection
-	 */
-	UE_DEPRECATED(5.4, "GetHttpSendTimeout has been deprecated, The legacy behavior doesn't make use of HttpSendTimeout at all, and there is no such support on some platforms. Use GetHttpActivityTimeout instead to decide if there is still any ongoing activity.")
-	inline float GetHttpSendTimeout() const
-	{
-		return HttpSendTimeout;
-	}
-
 	/**
 	 * @return timeout in seconds to check there is any ongoing activity on the established connection
 	 */
diff --git a/Engine/Source/Runtime/Online/HTTP/Public/HttpRetrySystem.h b/Engine/Source/Runtime/Online/HTTP/Public/HttpRetrySystem.h
--- a/Engine/Source/Runtime/Online/HTTP/Public/HttpRetrySystem.h
+++ b/Engine/Source/Runtime/Online/HTTP/Public/HttpRetrySystem.h
@@ -204,19 +204,6 @@ public:
 
 	HTTP_API virtual ~FManager();
 
-	/**
-	 * Updates the entries in the list of retry requests. Optional parameters are for future connection health assessment
-	 *
-	 * @param FileCount       optional parameter that will be filled with the total files updated
-	 * @param FailingCount    optional parameter that will be filled with the total files that have are in a retrying state
-	 * @param FailedCount     optional parameter that will be filled with the total files that have failed
-	 * @param CompletedCount  optional parameter that will be filled with the total files that have completed
-	 *
-	 * @return                true if there are no failures or retries
-	 */
-	UE_DEPRECATED(5.4, "HttpRetrySystem::Update has been deprecated, all logic will be processed in http thread instead")
-	HTTP_API bool Update(uint32* FileCount = NULL, uint32* FailingCount = NULL, uint32* FailedCount = NULL, uint32* CompletedCount = NULL);
-
 	void SetRandomFailureRate(float Value) { RandomFailureRate = FRandomFailureRateSetting(Value); }
 	void SetDefaultRetryLimit(uint32 Value) { RetryLimitCountDefault = FRetryLimitCountSetting(Value); }
 
diff --git a/Engine/Source/Runtime/Online/HTTP/Public/Interfaces/IHttpBase.h b/Engine/Source/Runtime/Online/HTTP/Public/Interfaces/IHttpBase.h
--- a/Engine/Source/Runtime/Online/HTTP/Public/Interfaces/IHttpBase.h
+++ b/Engine/Source/Runtime/Online/HTTP/Public/Interfaces/IHttpBase.h
@@ -17,8 +17,6 @@ namespace EHttpRequestStatus
 		Processing,
 		/** Finished but failed */
 		Failed,
-		/** Failed because it was unable to connect (safe to retry) */
-		Failed_ConnectionError UE_DEPRECATED(5.4, "Failed_ConnectionError has been deprecated, use Failed + EHttpFailureReason::ConnectionError instead"),
 		/** Finished and was successful */
 		Succeeded
 	};
@@ -26,7 +24,6 @@ namespace EHttpRequestStatus
 	/** @return the stringified version of the enum passed in */
 	inline const TCHAR* ToString(EHttpRequestStatus::Type EnumVal)
 	{
-PRAGMA_DISABLE_DEPRECATION_WARNINGS
 		switch (EnumVal)
 		{
 			case NotStarted:
@@ -41,17 +38,12 @@ PRAGMA_DISABLE_DEPRECATION_WARNINGS
 			{
 				return TEXT("Failed");
 			}
-			case Failed_ConnectionError:
-			{
-				return TEXT("ConnectionError");
-			}
 			case Succeeded:
 			{
 				return TEXT("Succeeded");
 			}
 		}
 		return TEXT("");
-PRAGMA_ENABLE_DEPRECATION_WARNINGS
 	}
 
 	inline bool IsFinished(const EHttpRequestStatus::Type Value)
diff --git a/Engine/Source/Runtime/Online/ImageDownload/Public/WebImage.h b/Engine/Source/Runtime/Online/ImageDownload/Public/WebImage.h
--- a/Engine/Source/Runtime/Online/ImageDownload/Public/WebImage.h
+++ b/Engine/Source/Runtime/Online/ImageDownload/Public/WebImage.h
@@ -55,10 +55,6 @@ public:
 	/** Set the brush that is currently being returned (this will be overridden when any async download completes) */
 	FORCEINLINE FWebImage& SetStandInBrush(TAttribute<const FSlateBrush*> StandInBrushIn) { StandInBrush = StandInBrushIn; DownloadedBrush.Reset(); return *this; }
 
-	/** Sets the RGB format used to read the downloaded image data (ERGBFormat::RGBA by default) */
-	UE_DEPRECATED(5.4, "ERGBFormat format is detected from downloaded image")
-	void SetRGBFormat(const ERGBFormat InRGBFormat) {}
-
 	/** Begin downloading an image. This will automatically set the current brush to the downloaded image when it completes (if successful) */
 	IMAGEDOWNLOAD_API bool BeginDownload(const FString& InUrl, const TOptional<FString>& StandInETag = TOptional<FString>(), const FOnImageDownloaded& DownloadCallback = FOnImageDownloaded());
 
