Fixes to BPT file constructor rewrite:
1. Resume wasn't resuming mid file correctly
2. Disabled disk chunk store crashes when a pageout is commanded

Additionally, extra logging has been added under Verbose and VeryVerbose to aid any future issues.

#rnx

diff --git a/Engine/Source/Runtime/Online/BuildPatchServices/Private/BuildPatchFileConstructor.cpp b/Engine/Source/Runtime/Online/BuildPatchServices/Private/BuildPatchFileConstructor.cpp
--- a/Engine/Source/Runtime/Online/BuildPatchServices/Private/BuildPatchFileConstructor.cpp
+++ b/Engine/Source/Runtime/Online/BuildPatchServices/Private/BuildPatchFileConstructor.cpp
@@ -446,6 +446,8 @@ class FChunkBackingStore : IConstructorChunkSource
 			return true;
 		}
 
+		UE_LOG(LogBuildPatchServices, VeryVerbose, TEXT("Paging out: %s"), *WriteToString<40>(InGuid));
+
 		uint32 EntriesRequired = Align(InChunk.ChunkSize, BitsPerEntry) >> BitsPerEntry;
 
 		int32 SpanIndex = 0;
@@ -499,7 +501,8 @@ class FChunkBackingStore : IConstructorChunkSource
 			}
 
 			// Need to expand the backing store.
-			if (UseMaxDiskSpaceBytes && (((uint64)BackingStoreEntryCount + EntriesRequired) << BitsPerEntry) > UseMaxDiskSpaceBytes)
+			if (BackingStoreFileHandle == nullptr ||
+				(UseMaxDiskSpaceBytes && (((uint64)BackingStoreEntryCount + EntriesRequired) << BitsPerEntry) > UseMaxDiskSpaceBytes))
 			{
 				// We can't expand - this fails to page out and gets evicted from the backing store.
 				Stats.DiskLostChunkCount++;
@@ -1377,6 +1380,9 @@ void FBuildPatchFileConstructor::SetPaused(bool bInIsPaused)
 void FBuildPatchFileConstructor::Abort()
 {
 	bool bAlreadyAborted = bShouldAbort.AtomicSet(true);
+
+	UE_LOG(LogBuildPatchServices, Verbose, TEXT("Issuing abort (previously aborted: %d)"), bAlreadyAborted)
+
 	if (bAlreadyAborted)
 	{
 		return;
@@ -1413,6 +1419,8 @@ void FBuildPatchFileConstructor::SetChunkLocation(const FGuid& InGuid, EConstruc
 		{
 			uint64 ChunkSize = Configuration.ManifestSet->GetDownloadSize(InGuid);
 
+			UE_LOG(LogBuildPatchServices, VeryVerbose, TEXT("Migrating chunk to cloud: %s, %llu bytes"), *WriteToString<40>(InGuid), ChunkSize);
+
 			DownloadRequirementLock.Lock();
 			DownloadRequirement += ChunkSize;
 			CloudSource->PostRequiredByteCount(DownloadRequirement);
@@ -1596,6 +1604,8 @@ void FBuildPatchFileConstructor::Run()
 					bool bHarvestSuccess = true;
 					if (InstallSource)
 					{
+						UE_LOG(LogBuildPatchServices, Verbose, TEXT("Harvesting source: %s"), *FileToDelete);
+
 						TRACE_CPUPROFILER_EVENT_SCOPE(Harvest);
 						// We need to grab any chunks from install sources that are no longer available.
 						// Anything that's already been loaded is already placed into the memory store appropriately,
@@ -1730,7 +1740,7 @@ void FBuildPatchFileConstructor::Run()
 								{
 									break;
 								}
-							}						
+							}
 						} // end if we need to harvest
 
 					} // end if install source exists.
@@ -2112,7 +2122,7 @@ void FBuildPatchFileConstructor::CompleteReadBatch(const FFileManifest& FileMani
 		Batch.Error = EConstructionError::MissingChunk;
 	}
 
-	//UE_LOG(LogTemp, Warning, TEXT("Completing batch: %d"), Batch.BatchId);
+	UE_LOG(LogBuildPatchServices, Verbose, TEXT("Completing ReadBatch: %d"), Batch.BatchId);
 
 	// We have to copy the memory source chunks after the reads are done
 	// because if we have two buffer's reads queued, the first one could be
@@ -2204,7 +2214,6 @@ int32 FBuildPatchFileConstructor::StartReadBatch(const FFileManifest& FileManife
 	Batch.BatchId = UniqueBatchId.fetch_add(1);
 	Batch.StartChunkFileIndex = StartFileChunkPartIndex;
 
-	//UE_LOG(LogTemp, Warning, TEXT("Starting batch: %d"), Batch.BatchId);
 
 	int32 EndChunkIdx = StartFileChunkPartIndex;
 	{
@@ -2237,6 +2246,8 @@ int32 FBuildPatchFileConstructor::StartReadBatch(const FFileManifest& FileManife
 	Batch.ChunkCount = EndChunkIdx - StartFileChunkPartIndex;
 	Batch.PendingRequestCount = Batch.Requests.Num();
 	Batch.FailedRequestCount = 0;
+
+	UE_LOG(LogBuildPatchServices, Verbose, TEXT("Starting ReadBatch: %d, Chunks=%d"), Batch.BatchId, Batch.ChunkCount);
 	
 	//
 	// IMPORTANT!!!
@@ -2366,6 +2377,7 @@ void FBuildPatchFileConstructor::WakeUpDispatch()
 bool FBuildPatchFileConstructor::ConstructFileFromChunks(const FString& BuildFilename, const FFileManifest& FileManifest, bool bResumeExisting)
 {
 	TRACE_CPUPROFILER_EVENT_SCOPE(ConstructFileFromChunks);
+	UE_LOG(LogBuildPatchServices, Verbose, TEXT("Starting File: %s"), *BuildFilename);
 
 	bool bSuccess = true;
 
@@ -2484,6 +2496,8 @@ bool FBuildPatchFileConstructor::ConstructFileFromChunks(const FString& BuildFil
 			}
 		}
 
+		UE_LOG(LogBuildPatchServices, Verbose, TEXT("Initial Disk Sizes: Required: %s Available: %s"), *FormatNumber(LocalDiskSpaceRequired), *FormatNumber(LocalDiskSpaceAvailable));
+
 		AvailableDiskSpace.store(LocalDiskSpaceAvailable, std::memory_order_release);
 		RequiredDiskSpace.store(LocalDiskSpaceRequired, std::memory_order_release);	
 
@@ -2539,7 +2553,7 @@ bool FBuildPatchFileConstructor::ConstructFileFromChunks(const FString& BuildFil
 		FBatchState* WritingBatch = nullptr;
 	
 		// Fill read/write buffers in batches until we complete the file.
-		int32 CurrentFileChunkPartIndex = 0;
+		int32 CurrentFileChunkPartIndex = StartChunkPart;
 
 		for (;;)
 		{
@@ -2686,6 +2700,8 @@ bool FBuildPatchFileConstructor::ConstructFileFromChunks(const FString& BuildFil
 							// Launch the write and hash the buffer on this thread.
 							CountBytesProcessed(WriteBuffers[WriteBufferIndex].Num());
 
+							UE_LOG(LogBuildPatchServices, Verbose, TEXT("Writing Batch: %d"), BatchState[WriteBufferIndex].BatchId);
+
 							IConstructorChunkSource::FRequestProcessFn WriteFn = CreateWriteRequest(NewFile.Get(), BatchState[WriteBufferIndex]);
 							QueueGenericThreadTask(WriteThreadIndex, MoveTemp(WriteFn));
 					
@@ -2706,14 +2722,14 @@ bool FBuildPatchFileConstructor::ConstructFileFromChunks(const FString& BuildFil
 				CloudSource->Tick(!bIsPaused && bSuccess, WaitTimeMs);
 			}
 
-			if (!bSuccess)
+			int32 ReadingBatchCount = 0;
+			for (int32 i = 0; i < BatchCount; i++)
 			{
-				int32 ReadingBatchCount = 0;
-				for (int32 i = 0; i < BatchCount; i++)
-				{
-					ReadingBatchCount += BatchState[i].bIsReading ? 1 : 0;
-				}
+				ReadingBatchCount += BatchState[i].bIsReading ? 1 : 0;
+			}
 
+			if (!bSuccess)
+			{
 				// We can only bail when all our async jobs have completed.
 				if (WritingBatch == nullptr && ReadingBatchCount == 0)
 				{
@@ -2722,6 +2738,11 @@ bool FBuildPatchFileConstructor::ConstructFileFromChunks(const FString& BuildFil
 			}
 
 			{
+				if (WaitTimeMs == TNumericLimits<uint32>::Max())
+				{
+					UE_LOG(LogBuildPatchServices, VeryVerbose, TEXT("Constructor waiting infinite. Reading batches: %d, Writing Batch: %d"), ReadingBatchCount, WritingBatch ? 1 : 0);
+				}
+
 				TRACE_CPUPROFILER_EVENT_SCOPE(CFFC_Waiting);
 				// We have a bunch of stuff outstanding that will wake us up if something happens.
 				WakeUpDispatchThreadEvent->Wait(WaitTimeMs);
diff --git a/Engine/Source/Runtime/Online/BuildPatchServices/Private/Installer/CloudChunkSource.cpp b/Engine/Source/Runtime/Online/BuildPatchServices/Private/Installer/CloudChunkSource.cpp
--- a/Engine/Source/Runtime/Online/BuildPatchServices/Private/Installer/CloudChunkSource.cpp
+++ b/Engine/Source/Runtime/Online/BuildPatchServices/Private/Installer/CloudChunkSource.cpp
@@ -1008,7 +1008,7 @@ namespace BuildPatchServices
 	{
 		int32 StartingDownloadCount = InFlightDownloads.Num();
 			
-		//UE_LOG(LogTemp, Warning, TEXT("Cloud: Active: %d, Pending %d"), InFlightDownloads.Num(), ReadQueue.Num());
+		UE_LOG(LogCloudChunkSource, Verbose, TEXT("Cloud: Active: %d, Pending %d"), InFlightDownloads.Num(), ReadQueue.Num());
 
 		// Process completed downloads.
 		TMap<int32, FDownloadRef> FrameCompletedDownloads;
@@ -1031,7 +1031,7 @@ namespace BuildPatchServices
 
 				Read->bQueuedForDownload = false;
 
-				//UE_LOG(LogTemp, Warning, TEXT("Downloaded chunk: %s"), *WriteToString<40>(Read->DataId));
+				UE_LOG(LogCloudChunkSource, VeryVerbose, TEXT("Downloaded chunk: %s"), *WriteToString<40>(Read->DataId));
 
 				bool bDownloadSuccess = Download->ResponseSuccessful();
 				if (bDownloadSuccess)
