Make FAppleHttpRequest::SetContentAsStreamedFile log the actual error when stat fails

diff --git a/Engine/Source/Runtime/Online/HTTP/Private/Apple/AppleHttp.cpp b/Engine/Source/Runtime/Online/HTTP/Private/Apple/AppleHttp.cpp
--- a/Engine/Source/Runtime/Online/HTTP/Private/Apple/AppleHttp.cpp
+++ b/Engine/Source/Runtime/Online/HTTP/Private/Apple/AppleHttp.cpp
@@ -691,7 +691,7 @@ bool FAppleHttpRequest::SetContentAsStreamedFile(const FString& Filename)
 	}
 	else
 	{
-		UE_LOG(LogHttp, Warning, TEXT("FAppleHttpRequest::SetContentAsStreamedFile failed to get file size"));
+		UE_LOG(LogHttp, Warning, TEXT("FAppleHttpRequest::SetContentAsStreamedFile failed to get file size errno: %d: %s"), errno, UTF8_TO_TCHAR(strerror(errno)));
 		StreamedContentSource.Emplace<FNoStreamSource>();
 		ContentBytesLength = 0;
 		return false;
