Adding platform last error to create file error log

#rnx

diff --git a/Engine/Source/Runtime/Online/BuildPatchServices/Private/BuildPatchFileConstructor.cpp b/Engine/Source/Runtime/Online/BuildPatchServices/Private/BuildPatchFileConstructor.cpp
--- a/Engine/Source/Runtime/Online/BuildPatchServices/Private/BuildPatchFileConstructor.cpp
+++ b/Engine/Source/Runtime/Online/BuildPatchServices/Private/BuildPatchFileConstructor.cpp
@@ -2852,7 +2852,7 @@ bool FBuildPatchFileConstructor::ConstructFileFromChunks(const FString& BuildFil
 				{
 					InstallerAnalytics->RecordConstructionError(BuildFilename, CreateFilePlatformLastError, TEXT("Could Not Create File"));
 				}
-				UE_LOG(LogBuildPatchServices, Error, TEXT("FBuildPatchFileConstructor: Could not create %s"), *BuildFilename);
+				UE_LOG(LogBuildPatchServices, Error, TEXT("FBuildPatchFileConstructor: Could not create %s (LastError=%d)"), *BuildFilename, CreateFilePlatformLastError);
 				InstallerError->SetError(EBuildPatchInstallError::FileConstructionFail, ConstructionErrorCodes::FileCreateFail, CreateFilePlatformLastError);
 				break;
 			}
