Move AddRequest to internal in HttpManager.

diff --git a/Engine/Source/Runtime/Online/HTTP/Public/HttpManager.h b/Engine/Source/Runtime/Online/HTTP/Public/HttpManager.h
--- a/Engine/Source/Runtime/Online/HTTP/Public/HttpManager.h
+++ b/Engine/Source/Runtime/Online/HTTP/Public/HttpManager.h
@@ -155,15 +155,6 @@ public:
 	 */
 	HTTP_API void Shutdown();
 
-	/**
-	 * Adds an Http request instance to the manager for tracking/ticking
-	 * Manager should always have a list of requests currently being processed
-	 *
-	 * @param Request - the request object to add
-	 */
-	UE_DEPRECATED(5.4, "AddRequest has been deprecated, use AddThreadedRequest instead")
-	HTTP_API void AddRequest(const FHttpRequestRef& Request);
-
 	/**
 	 * Set a delegate to be triggered when an http request added to http manager. 
 	 * NOTE: The delegate can be triggered from different threads, depends on which
@@ -348,6 +339,8 @@ protected:
 	HTTP_API void UpdateUrlPatternsToLogResponse(IConsoleVariable* CVar);
 	HTTP_API void UpdateUrlPatternsToMockFailure(IConsoleVariable* CVar);
 
+	HTTP_API void AddRequest(const FHttpRequestRef& Request);
+
 protected:
 	/** List of Http requests that are actively being processed */
 	TArray<FHttpRequestRef> Requests;
