Update the install progress as the reads come in instead of when the writes are sent so that we get a bit more granularity on our progress bar.

#rnx

diff --git a/Engine/Source/Runtime/Online/BuildPatchServices/Private/BuildPatchFileConstructor.cpp b/Engine/Source/Runtime/Online/BuildPatchServices/Private/BuildPatchFileConstructor.cpp
--- a/Engine/Source/Runtime/Online/BuildPatchServices/Private/BuildPatchFileConstructor.cpp
+++ b/Engine/Source/Runtime/Online/BuildPatchServices/Private/BuildPatchFileConstructor.cpp
@@ -2042,18 +2042,33 @@ void FBuildPatchFileConstructor::RequestCompletedFn(const FGuid& Guid, bool bAbo
 		{
 			Request->Batch->FailedRequestCount.fetch_add(1, std::memory_order_relaxed);
 		}
-		else if (!bAborted && Request->bReadIntoMemoryStore)
+		else if (!bAborted)
 		{
-			// If the read went to memory, we need to copy splats. Otherwise it was a single
-			// direct read so do nothing.
-			TRACE_CPUPROFILER_EVENT_SCOPE(CFFC_Splats);
-			for (const FRequestInfo::FRequestSplat& Splat : Request->Splats)
+			uint64 TotalToDestinationBuffer = 0;
+
+			if (Request->bReadIntoMemoryStore)
+			{
+				// If the read went to memory, we need to copy splats. Otherwise it was a single
+				// direct read so do nothing.
+				TRACE_CPUPROFILER_EVENT_SCOPE(CFFC_Splats);
+				for (const FRequestInfo::FRequestSplat& Splat : Request->Splats)
+				{
+					FMemory::Memcpy(
+						(uint8*)Request->DestinationBuffer.GetData() + Splat.DestinationOffset,
+						(uint8*)Request->ReadBuffer.GetData() + Splat.OffsetInChunk,
+						Splat.BytesToCopy);
+					TotalToDestinationBuffer += Splat.BytesToCopy;
+				}
+			}
+			else
 			{
-				FMemory::Memcpy(
-					(uint8*)Request->DestinationBuffer.GetData() + Splat.DestinationOffset,
-					(uint8*)Request->ReadBuffer.GetData() + Splat.OffsetInChunk,
-					Splat.BytesToCopy);
+				// Direct means 1 splat
+				TotalToDestinationBuffer += Request->Splats[0].BytesToCopy;
 			}
+
+			FileProgressLock.Lock();
+			FileProgress += TotalToDestinationBuffer;
+			FileProgressLock.Unlock();
 		}
 		
 		if (Request->Batch->PendingRequestCount.fetch_add(-1, std::memory_order_relaxed) == 1)
@@ -2144,13 +2159,19 @@ void FBuildPatchFileConstructor::CompleteReadBatch(const FFileManifest& FileMani
 
 				if (!bFailedToGetChunk)
 				{
+					uint64 TotalToDestinationBuffer = 0;
 					for (const FRequestInfo::FRequestSplat& Splat : Request.Splats)
 					{
 						FMemory::Memcpy(
 							(uint8*)Request.DestinationBuffer.GetData() + Splat.DestinationOffset,
 							(uint8*)ChunkData.GetData() + Splat.OffsetInChunk,
 							Splat.BytesToCopy);
+						TotalToDestinationBuffer += Splat.BytesToCopy;
 					}
+
+					FileProgressLock.Lock();
+					FileProgress += TotalToDestinationBuffer;
+					FileProgressLock.Unlock();
 				}
 				else
 				{
@@ -2532,6 +2553,8 @@ bool FBuildPatchFileConstructor::ConstructFileFromChunks(const FString& BuildFil
 	bSuccess = NewFile != nullptr;
 	if (bSuccess)
 	{
+		FileProgress = StartPosition;
+
 		// Seek to file write position
 		if (NewFile->Tell() != StartPosition)
 		{
@@ -2555,6 +2578,8 @@ bool FBuildPatchFileConstructor::ConstructFileFromChunks(const FString& BuildFil
 		// Fill read/write buffers in batches until we complete the file.
 		int32 CurrentFileChunkPartIndex = StartChunkPart;
 
+		uint64 LastSeenFileProgress = FileProgress; // no need to lock - nothing async yet.
+
 		for (;;)
 		{
 			// Check state. We can continue to run this after a failure has been encountered in order
@@ -2577,10 +2602,7 @@ bool FBuildPatchFileConstructor::ConstructFileFromChunks(const FString& BuildFil
 						// writing code.
 						ConstructionError = EConstructionError::SerializeError;
 					}
-
-					// Can use Tell() because right now there's nothing writing.
-					FileConstructorStat->OnFileProgress(BuildFilename, NewFile->Tell());
-
+					
 					// If that was the last write, then we are done.
 					if (CurrentFileChunkPartIndex == FileManifest.ChunkParts.Num())
 					{
@@ -2698,8 +2720,6 @@ bool FBuildPatchFileConstructor::ConstructFileFromChunks(const FString& BuildFil
 							WritingBatch = BatchState + WriteBufferIndex;
 							
 							// Launch the write and hash the buffer on this thread.
-							CountBytesProcessed(WriteBuffers[WriteBufferIndex].Num());
-
 							UE_LOG(LogBuildPatchServices, Verbose, TEXT("Writing Batch: %d"), BatchState[WriteBufferIndex].BatchId);
 
 							IConstructorChunkSource::FRequestProcessFn WriteFn = CreateWriteRequest(NewFile.Get(), BatchState[WriteBufferIndex]);
@@ -2714,6 +2734,30 @@ bool FBuildPatchFileConstructor::ConstructFileFromChunks(const FString& BuildFil
 				} // end not paused
 			} // end state check
 
+			// If the file progress changed since we last saw it, post the update.
+			// Note that we want this to update reasonably often but we're about to wait
+			// potentially until all reads complete - however the only time things actually
+			// take a long time wall-clock wise is when we are downloading, and the cloud
+			// source will then prevent us from sleeping too long, so we actually catch these
+			// updates.
+			// We do this from here to ensure we always increase rather than risk multi thread
+			// races.
+			{
+				uint64 CurrentFileProgress;
+				FileProgressLock.Lock();
+				CurrentFileProgress = FileProgress;
+				FileProgressLock.Unlock();
+
+				if (CurrentFileProgress != LastSeenFileProgress)
+				{
+					// this updates the overall install progress.
+					CountBytesProcessed(CurrentFileProgress - LastSeenFileProgress);
+					LastSeenFileProgress = CurrentFileProgress;
+					FileConstructorStat->OnFileProgress(BuildFilename, CurrentFileProgress);
+				}
+			}
+
+
 			uint32 WaitTimeMs = TNumericLimits<uint32>::Max();
 			{
 				TRACE_CPUPROFILER_EVENT_SCOPE(CFFC_TickCloud);
diff --git a/Engine/Source/Runtime/Online/BuildPatchServices/Private/BuildPatchFileConstructor.h b/Engine/Source/Runtime/Online/BuildPatchServices/Private/BuildPatchFileConstructor.h
--- a/Engine/Source/Runtime/Online/BuildPatchServices/Private/BuildPatchFileConstructor.h
+++ b/Engine/Source/Runtime/Online/BuildPatchServices/Private/BuildPatchFileConstructor.h
@@ -259,6 +259,11 @@ namespace BuildPatchServices
 		// we just lock around it since it's tiny.
 		UE::FMutex DownloadRequirementLock;
 
+		// Instead of reporting file progress on write (which jumps by the write buffer), advance as
+		// we fill the write buffer.
+		uint64 FileProgress = 0;
+		UE::FMutex FileProgressLock;
+
 		TArray<FEvent*> ThreadWakeups;
 		TArray<TArray<IConstructorChunkSource::FRequestProcessFn>> ThreadJobPostings;
 		TArray<UE::FMutex> ThreadJobPostingLocks;
