Add more checks for duplicated scheme protocol in http request.

diff --git a/Engine/Source/Programs/WebTests/Private/TestHttp.cpp b/Engine/Source/Programs/WebTests/Private/TestHttp.cpp
--- a/Engine/Source/Programs/WebTests/Private/TestHttp.cpp
+++ b/Engine/Source/Programs/WebTests/Private/TestHttp.cpp
@@ -15,6 +15,7 @@
 #include "HttpServerModule.h"
 #include "Misc/CommandLine.h"
 #include "TestHarness.h"
+#include "TestMacros/Assertions.h"
 #include "Serialization/JsonSerializerMacros.h"
 #include "catch2/generators/catch_generators.hpp"
 
@@ -221,6 +222,15 @@ TEST_CASE_METHOD(FHttpModuleTestFixture, "Shutdown http module without issue whe
 	HttpModule->GetHttpManager().Tick(0.0f);
 }
 
+TEST_CASE_METHOD(FHttpModuleTestFixture, "Duplicated scheme protocol will trigger check", HTTP_TAG)
+{
+	TSharedRef<IHttpRequest> HttpRequest = CreateRequest();
+	CHECK(HttpRequest->GetVerb() == TEXT("GET"));
+
+	HttpRequest->SetURL(TEXT("https://https://www.epicgames.com"));
+	REQUIRE_CHECK_MSG("Duplicated protocol", HttpRequest->ProcessRequest());
+}
+
 TEST_CASE_METHOD(FHttpModuleTestFixture, "Shutdown http module without issue when there are ongoing streaming http requests with timeout.", HTTP_TAG)
 {
 	if (!bRunHeavyTests)
diff --git a/Engine/Source/Runtime/Online/HTTP/Private/HttpManager.cpp b/Engine/Source/Runtime/Online/HTTP/Private/HttpManager.cpp
--- a/Engine/Source/Runtime/Online/HTTP/Private/HttpManager.cpp
+++ b/Engine/Source/Runtime/Online/HTTP/Private/HttpManager.cpp
@@ -237,8 +237,31 @@ FString FHttpManager::CreateCorrelationId() const
 	return CorrelationIdMethod();
 }
 
+#if !UE_BUILD_SHIPPING
+namespace HttpManagerInternal
+{
+
+static void CheckNoDuplicatedProtocol(FStringView UrlStringView)
+{
+	const FStringView Delimiter = TEXT("://");
+	const int32 FoundIndex = UE::String::FindFirst(UrlStringView, Delimiter);
+	if (FoundIndex != INDEX_NONE)
+	{
+		FStringView Domain = UrlStringView.RightChop(FoundIndex + Delimiter.Len());
+		checkf(!Domain.StartsWith(UrlStringView.Left(FoundIndex + Delimiter.Len())), TEXT("Duplicated protocol in url: %.*s"), UrlStringView.Len(), UrlStringView.GetData());
+		UE_CLOG(Domain.Contains(Delimiter), LogHttp, Warning, TEXT("Duplicated delimiter '://' found in url %.*s, is that intended?"), UrlStringView.Len(), UrlStringView.GetData());
+	}
+}
+
+}
+#endif // !UE_BUILD_SHIPPING
+
 bool FHttpManager::IsDomainAllowed(const FString& Url) const
 {
+#if !UE_BUILD_SHIPPING
+	HttpManagerInternal::CheckNoDuplicatedProtocol(Url);
+#endif // !UE_BUILD_SHIPPING
+
 	if (!URLRequestFilter.IsEmpty())
 	{
 		return URLRequestFilter.IsRequestAllowed(Url);
