Fix for integer Overflow and Truncation while loading fron Chunk File
Changed the FileSize type from an int64 a uint64
Made sure sum of HeaderSize and DataSizeCompressed do not overflow by casting to uint64 before adding
Added check that the resulting FileSize can still fit in uint32 expected by FChunkDataAccessFactory::Create

diff --git a/Engine/Source/Runtime/Online/BuildPatchServices/Private/Data/ChunkData.cpp b/Engine/Source/Runtime/Online/BuildPatchServices/Private/Data/ChunkData.cpp
--- a/Engine/Source/Runtime/Online/BuildPatchServices/Private/Data/ChunkData.cpp
+++ b/Engine/Source/Runtime/Online/BuildPatchServices/Private/Data/ChunkData.cpp
@@ -713,117 +713,125 @@ namespace BuildPatchServices
 			Reader << HeaderCheck;
 
 			// Get file size.
-			const int64 FileSize = HeaderCheck.HeaderSize + HeaderCheck.DataSizeCompressed;
+			const uint64 FileSize = (uint64)HeaderCheck.HeaderSize + (uint64)HeaderCheck.DataSizeCompressed;
 
-			if (HeaderCheck.Guid.IsValid())
+			// Check that Filesize can fit into a uint32 for FChunkDataAccessFactory::Create()
+			if (IntFitsIn<uint32>(FileSize))
 			{
-				if (HeaderCheck.HashType != EChunkHashFlags::None)
+				if (HeaderCheck.Guid.IsValid())
 				{
-					if ((HeaderCheck.HeaderSize + HeaderCheck.DataSizeCompressed) <= AvailableSize)
+					if (HeaderCheck.HashType != EChunkHashFlags::None)
 					{
-						if ((HeaderCheck.StoredAs & EChunkStorageFlags::Encrypted) == EChunkStorageFlags::None)
+						if ((HeaderCheck.HeaderSize + HeaderCheck.DataSizeCompressed) <= AvailableSize)
 						{
-							// Create the data.
-							ChunkData = FChunkDataAccessFactory::Create(FileSize);
+							if ((HeaderCheck.StoredAs & EChunkStorageFlags::Encrypted) == EChunkStorageFlags::None)
+							{
+								// Create the data.
+								ChunkData = FChunkDataAccessFactory::Create(FileSize);
 
-							// Lock data.
-							FChunkHeader* Header;
-							uint8* Data;
-							ChunkData->GetDataLock(&Data, &Header);
-							*Header = HeaderCheck;
+								// Lock data.
+								FChunkHeader* Header;
+								uint8* Data;
+								ChunkData->GetDataLock(&Data, &Header);
+								*Header = HeaderCheck;
 
-							// Read the data.
-							{
-								TRACE_CPUPROFILER_EVENT_SCOPE(Chunk_Serialize);
-								Reader.Serialize(Data, Header->DataSizeCompressed);
-							}
-							if (Reader.IsError() == false)
-							{
-								OutLoadResult = EChunkLoadResult::Success;
-								// Decompress.
-								if ((Header->StoredAs & EChunkStorageFlags::Compressed) != EChunkStorageFlags::None)
+								// Read the data.
 								{
-									// Create a new data instance.
-									IChunkDataAccess* NewChunkData = FChunkDataAccessFactory::Create(Header->DataSizeUncompressed);
-									// Lock data.
-									FChunkHeader* NewHeader;
-									uint8* NewData;
-									NewChunkData->GetDataLock(&NewData, &NewHeader);
-									// Uncompress the memory.
-									bool bSuccess = FCompression::UncompressMemory(
-										NAME_Zlib,
-										NewData,
-										Header->DataSizeUncompressed,
-										Data,
-										Header->DataSizeCompressed);
-									// If successful, switch over to new data.
-									if (bSuccess)
-									{
-										FMemory::Memcpy(*NewHeader, *Header);
-										NewHeader->StoredAs = EChunkStorageFlags::None;
-										NewHeader->DataSizeCompressed = Header->DataSizeUncompressed;
-										ChunkData->ReleaseDataLock();
-										delete ChunkData;
-										ChunkData = NewChunkData;
-										Header = NewHeader;
-										Data = NewData;
-									}
-									// Otherwise delete new data.
-									else
+									TRACE_CPUPROFILER_EVENT_SCOPE(Chunk_Serialize);
+									Reader.Serialize(Data, Header->DataSizeCompressed);
+								}
+								if (Reader.IsError() == false)
+								{
+									OutLoadResult = EChunkLoadResult::Success;
+									// Decompress.
+									if ((Header->StoredAs & EChunkStorageFlags::Compressed) != EChunkStorageFlags::None)
 									{
-										OutLoadResult = EChunkLoadResult::DecompressFailure;
-										NewChunkData->ReleaseDataLock();
-										delete NewChunkData;
-										NewChunkData = nullptr;
+										// Create a new data instance.
+										IChunkDataAccess* NewChunkData = FChunkDataAccessFactory::Create(Header->DataSizeUncompressed);
+										// Lock data.
+										FChunkHeader* NewHeader;
+										uint8* NewData;
+										NewChunkData->GetDataLock(&NewData, &NewHeader);
+										// Uncompress the memory.
+										bool bSuccess = FCompression::UncompressMemory(
+											NAME_Zlib,
+											NewData,
+											Header->DataSizeUncompressed,
+											Data,
+											Header->DataSizeCompressed);
+										// If successful, switch over to new data.
+										if (bSuccess)
+										{
+											FMemory::Memcpy(*NewHeader, *Header);
+											NewHeader->StoredAs = EChunkStorageFlags::None;
+											NewHeader->DataSizeCompressed = Header->DataSizeUncompressed;
+											ChunkData->ReleaseDataLock();
+											delete ChunkData;
+											ChunkData = NewChunkData;
+											Header = NewHeader;
+											Data = NewData;
+										}
+										// Otherwise delete new data.
+										else
+										{
+											OutLoadResult = EChunkLoadResult::DecompressFailure;
+											NewChunkData->ReleaseDataLock();
+											delete NewChunkData;
+											NewChunkData = nullptr;
+										}
 									}
-								}
-								// Verify.
+									// Verify.
 
-								// If we have both hashes, only check Sha1 instead of both.
-								bool bHasRollingHash = EnumHasAnyFlags(Header->HashType, EChunkHashFlags::RollingPoly64);
-								bool bHasSha1Hash = EnumHasAnyFlags(Header->HashType, EChunkHashFlags::Sha1);
+									// If we have both hashes, only check Sha1 instead of both.
+									bool bHasRollingHash = EnumHasAnyFlags(Header->HashType, EChunkHashFlags::RollingPoly64);
+									bool bHasSha1Hash = EnumHasAnyFlags(Header->HashType, EChunkHashFlags::Sha1);
 
-								if (OutLoadResult == EChunkLoadResult::Success && bHasRollingHash && !bHasSha1Hash) // Don't bother verifying the rolling hash if we have a sha1
-								{
-									if (Header->DataSizeCompressed != Header->DataSizeUncompressed || Header->RollingHash != FRollingHash::GetHashForDataSet(Data, Header->DataSizeUncompressed))
+									if (OutLoadResult == EChunkLoadResult::Success && bHasRollingHash && !bHasSha1Hash) // Don't bother verifying the rolling hash if we have a sha1
 									{
-										OutLoadResult = EChunkLoadResult::HashCheckFailed;
+										if (Header->DataSizeCompressed != Header->DataSizeUncompressed || Header->RollingHash != FRollingHash::GetHashForDataSet(Data, Header->DataSizeUncompressed))
+										{
+											OutLoadResult = EChunkLoadResult::HashCheckFailed;
+										}
 									}
-								}
-								FSHAHash ShaHashCheck;
-								if (OutLoadResult == EChunkLoadResult::Success && bHasSha1Hash)
-								{
-									TRACE_CPUPROFILER_EVENT_SCOPE(Chunk_ShaHash);
-									FSHA1::HashBuffer(Data, Header->DataSizeUncompressed, ShaHashCheck.Hash);
-									if (!(ShaHashCheck == Header->SHAHash))
+									FSHAHash ShaHashCheck;
+									if (OutLoadResult == EChunkLoadResult::Success && bHasSha1Hash)
 									{
-										OutLoadResult = EChunkLoadResult::HashCheckFailed;
+										TRACE_CPUPROFILER_EVENT_SCOPE(Chunk_ShaHash);
+										FSHA1::HashBuffer(Data, Header->DataSizeUncompressed, ShaHashCheck.Hash);
+										if (!(ShaHashCheck == Header->SHAHash))
+										{
+											OutLoadResult = EChunkLoadResult::HashCheckFailed;
+										}
 									}
 								}
+								else
+								{
+									OutLoadResult = EChunkLoadResult::SerializationError;
+								}
 							}
 							else
 							{
-								OutLoadResult = EChunkLoadResult::SerializationError;
+								OutLoadResult = EChunkLoadResult::UnsupportedStorage;
 							}
 						}
 						else
 						{
-							OutLoadResult = EChunkLoadResult::UnsupportedStorage;
+							OutLoadResult = EChunkLoadResult::IncorrectFileSize;
 						}
 					}
 					else
 					{
-						OutLoadResult = EChunkLoadResult::IncorrectFileSize;
+						OutLoadResult = EChunkLoadResult::MissingHashInfo;
 					}
 				}
 				else
 				{
-					OutLoadResult = EChunkLoadResult::MissingHashInfo;
+					OutLoadResult = EChunkLoadResult::CorruptHeader;
 				}
 			}
 			else
 			{
-				OutLoadResult = EChunkLoadResult::CorruptHeader;
+				OutLoadResult = EChunkLoadResult::FileSizeTooBig;
 			}
 
 			if (ChunkData != nullptr)
diff --git a/Engine/Source/Runtime/Online/BuildPatchServices/Private/Data/ChunkData.h b/Engine/Source/Runtime/Online/BuildPatchServices/Private/Data/ChunkData.h
--- a/Engine/Source/Runtime/Online/BuildPatchServices/Private/Data/ChunkData.h
+++ b/Engine/Source/Runtime/Online/BuildPatchServices/Private/Data/ChunkData.h
@@ -78,6 +78,9 @@ namespace BuildPatchServices
 		// The expected data hash in the header did not match the hash of the data.
 		HashCheckFailed,
 
+		// The size of the file is too big to handle 
+		FileSizeTooBig,
+
 		// The operation was aborted.
 		Aborted
 	};
