Add safe check of index before removing request entry from retry manager when complete http request.

diff --git a/Engine/Source/Runtime/Online/HTTP/Private/HttpRetrySystem.cpp b/Engine/Source/Runtime/Online/HTTP/Private/HttpRetrySystem.cpp
--- a/Engine/Source/Runtime/Online/HTTP/Private/HttpRetrySystem.cpp
+++ b/Engine/Source/Runtime/Online/HTTP/Private/HttpRetrySystem.cpp
@@ -226,60 +226,61 @@ void FHttpRetrySystem::FRequest::HttpOnProcessRequestComplete(FHttpRequestPtr In
 		FScopeLock ScopeLock(&RetryManagerPtr->RequestListLock);
 
 		uint32 EntryIndex = RetryManagerPtr->RequestList.IndexOfByPredicate([this](const FManager::FHttpRetryRequestEntry& Entry) { return Entry.Request == AsShared(); });
-		check(EntryIndex != INDEX_NONE);
-		FManager::FHttpRetryRequestEntry& HttpRetryRequestEntry = RetryManagerPtr->RequestList[EntryIndex];
-
-
-		if (RetryStatus == FHttpRetrySystem::FRequest::EStatus::Cancelled)
-		{
-			// Do nothing here
-		}
-		else if (GetStatus() == EHttpRequestStatus::Failed)
+		if (ensure(EntryIndex != INDEX_NONE))
 		{
-			if (GetFailureReason() == EHttpFailureReason::ConnectionError && RetryDomains.IsValid())
+			FManager::FHttpRetryRequestEntry& HttpRetryRequestEntry = RetryManagerPtr->RequestList[EntryIndex];
+
+			if (RetryStatus == FHttpRetrySystem::FRequest::EStatus::Cancelled)
 			{
-				MoveToNextRetryDomain();
+				// Do nothing here
 			}
-
-			if (GetFailureReason() == EHttpFailureReason::TimedOut)
+			else if (GetStatus() == EHttpRequestStatus::Failed)
 			{
-				RetryStatus = FHttpRetrySystem::FRequest::EStatus::FailedTimeout;
+				if (GetFailureReason() == EHttpFailureReason::ConnectionError && RetryDomains.IsValid())
+				{
+					MoveToNextRetryDomain();
+				}
+
+				if (GetFailureReason() == EHttpFailureReason::TimedOut)
+				{
+					RetryStatus = FHttpRetrySystem::FRequest::EStatus::FailedTimeout;
+				}
+				else
+				{
+					RetryStatus = FHttpRetrySystem::FRequest::EStatus::FailedRetry;
+				}
 			}
 			else
 			{
-				RetryStatus = FHttpRetrySystem::FRequest::EStatus::FailedRetry;
+				RetryStatus = FHttpRetrySystem::FRequest::EStatus::Succeeded;
 			}
-		}
-		else
-		{
-			RetryStatus = FHttpRetrySystem::FRequest::EStatus::Succeeded;
-		}
 
-		if (RetryStatus != FHttpRetrySystem::FRequest::EStatus::Cancelled && 
-			RetryStatus != FHttpRetrySystem::FRequest::EStatus::FailedTimeout && 
-			RetryManagerPtr->ShouldRetry(HttpRetryRequestEntry) && 
-			RetryManagerPtr->CanRetry(HttpRetryRequestEntry))
-		{
-			const double NowAbsoluteSeconds = FPlatformTime::Seconds();
-			float LockoutPeriod = RetryManagerPtr->GetLockoutPeriodSeconds(HttpRetryRequestEntry);
+			if (RetryStatus != FHttpRetrySystem::FRequest::EStatus::Cancelled && 
+				RetryStatus != FHttpRetrySystem::FRequest::EStatus::FailedTimeout && 
+				RetryManagerPtr->ShouldRetry(HttpRetryRequestEntry) && 
+				RetryManagerPtr->CanRetry(HttpRetryRequestEntry))
+			{
+				const double NowAbsoluteSeconds = FPlatformTime::Seconds();
+				float LockoutPeriod = RetryManagerPtr->GetLockoutPeriodSeconds(HttpRetryRequestEntry);
 
-			RetryStatus = FHttpRetrySystem::FRequest::EStatus::ProcessingLockout;
+				RetryStatus = FHttpRetrySystem::FRequest::EStatus::ProcessingLockout;
 
-			{
-				QUICK_SCOPE_CYCLE_COUNTER(STAT_FHttpRetrySystem_FManager_Update_OnRequestWillRetry);
-				OnRequestWillRetry().ExecuteIfBound(HttpRetryRequestEntry.Request, GetResponse(), LockoutPeriod);
+				{
+					QUICK_SCOPE_CYCLE_COUNTER(STAT_FHttpRetrySystem_FManager_Update_OnRequestWillRetry);
+					OnRequestWillRetry().ExecuteIfBound(HttpRetryRequestEntry.Request, GetResponse(), LockoutPeriod);
+				}
+
+				RetryManagerPtr->RetryHttpRequestWithDelay(HttpRetryRequestEntry, LockoutPeriod, bSucceeded);
+				return;
 			}
 
-			RetryManagerPtr->RetryHttpRequestWithDelay(HttpRetryRequestEntry, LockoutPeriod, bSucceeded);
-			return;
-		}
+			if (HttpRetryRequestEntry.CurrentRetryCount > 0)
+			{
+				FHttpRetrySystem::FManager::FHttpLogVerbosityTracker::Get().DecrementRetriedRequests();
+			}
 
-		if (HttpRetryRequestEntry.CurrentRetryCount > 0)
-		{
-			FHttpRetrySystem::FManager::FHttpLogVerbosityTracker::Get().DecrementRetriedRequests();
+			RetryManagerPtr->RequestList.RemoveAtSwap(EntryIndex);
 		}
-
-		RetryManagerPtr->RequestList.RemoveAtSwap(EntryIndex);
 	}
 
 	FHttpResponsePtr ResultResponse = HttpResponse;
