Add domain of default url in retry domains if there is any retry domain and when it doesn't contain default domain.
Make CanRetry and ShouldRetry methods const in RetryManager.

diff --git a/Engine/Source/Programs/WebTests/Private/TestHttp.cpp b/Engine/Source/Programs/WebTests/Private/TestHttp.cpp
--- a/Engine/Source/Programs/WebTests/Private/TestHttp.cpp
+++ b/Engine/Source/Programs/WebTests/Private/TestHttp.cpp
@@ -223,6 +223,12 @@ TEST_CASE_METHOD(FHttpModuleTestFixture, "Shutdown http module without issue whe
 	HttpModule->GetHttpManager().Tick(0.0f);
 }
 
+TEST_CASE("Get domain of url can work well", HTTP_TAG)
+{
+	FString DefaultUrlDomain = FPlatformHttp::GetUrlDomain(TEXT("https://www.epicgames.com/the/relative/path/"));
+	CHECK(DefaultUrlDomain == TEXT("www.epicgames.com"));
+}
+
 TEST_CASE_METHOD(FHttpModuleTestFixture, "Duplicated scheme protocol will trigger warning", HTTP_TAG)
 {
 	// TODO: find a way to catch expected warning log instead
@@ -2201,6 +2207,82 @@ TEST_CASE_METHOD(FWaitUntilCompleteHttpFixture, "Retry immediately without lock
 	HttpRequest->ProcessRequest();
 }
 
+TEST_CASE_METHOD(FWaitUntilCompleteHttpFixture, "Retry will add default domain to alt domains", HTTP_TAG)
+{
+	if (!bRetryEnabled)
+	{
+		return;
+	}
+
+	DisableWarningsInThisTest();
+
+	CVarHttpUrlPatternsToMockFailure->Set(TEXT("fakedomain->0"));
+
+	HttpModule->HttpConnectionTimeout = 0.1f;
+
+	TQueue<FString>* ExpectedDomainsOrderInRetry = new TQueue<FString>();
+	TArray<FString> AltDomains;
+
+	FString DefautUrl = TEXT("https://www.fakedomain-default.com/webtests/httptests");
+	ExpectedDomainsOrderInRetry->Enqueue(TEXT("www.fakedomain-a.com"));
+	ExpectedDomainsOrderInRetry->Enqueue(TEXT("www.fakedomain-b.com"));
+	ExpectedDomainsOrderInRetry->Enqueue(TEXT("www.fakedomain-default.com"));
+	ExpectedDomainsOrderInRetry->Enqueue(TEXT("www.fakedomain-a.com"));
+
+	SECTION("When default domain is not listed in AltDomains")
+	{
+		AltDomains.Add(TEXT("www.fakedomain-a.com"));
+		AltDomains.Add(TEXT("www.fakedomain-b.com"));
+	}
+	SECTION("When default domain configured as first in AltDomains, retry system will keep the order in the array intead of adding it")
+	{
+		AltDomains.Add(TEXT("www.fakedomain-default.com"));
+		AltDomains.Add(TEXT("www.fakedomain-a.com"));
+		AltDomains.Add(TEXT("www.fakedomain-b.com"));
+	}
+	SECTION("When default domain configured in the middle of AltDomains, retry system will move it to first in the array")
+	{
+		AltDomains.Add(TEXT("www.fakedomain-a.com"));
+		AltDomains.Add(TEXT("www.fakedomain-default.com"));
+		AltDomains.Add(TEXT("www.fakedomain-b.com"));
+	}
+	SECTION("When default domain configured in the end of AltDomains, retry system will move it to first in the array")
+	{
+		AltDomains.Add(TEXT("www.fakedomain-a.com"));
+		AltDomains.Add(TEXT("www.fakedomain-b.com"));
+		AltDomains.Add(TEXT("www.fakedomain-default.com"));
+	}
+
+	FHttpRetrySystem::FRetryDomainsPtr RetryDomains = MakeShared<FHttpRetrySystem::FRetryDomains>(MoveTemp(AltDomains));
+	TSharedRef<IHttpRequest> HttpRequest = HttpRetryManager->CreateRequest(
+		4/*InRetryLimitCountOverride*/,
+		FHttpRetrySystem::FRetryTimeoutRelativeSecondsSetting()/*InRetryTimeoutRelativeSecondsOverride unused*/,
+		{ EHttpResponseCodes::TooManyRequests }/*InRetryResponseCodes*/,
+		FHttpRetrySystem::FRetryVerbs(),
+		RetryDomains
+	);
+
+	HttpRequest->SetURL(DefautUrl);
+	ExpectingExtraCallbacks = 4;
+	HttpRequest->OnRequestWillRetry().BindLambda([this, ExpectedDomainsOrderInRetry](FHttpRequestPtr HttpRequest, FHttpResponsePtr /*Response*/, float LockoutPeriod) {
+		--ExpectingExtraCallbacks;
+		CHECK(LockoutPeriod == 0);
+		FString ExpectedDomain;
+		ExpectedDomainsOrderInRetry->Dequeue(ExpectedDomain);
+		FString ActualDomain = FPlatformHttp::GetUrlDomain(HttpRequest->GetURL());
+		CHECK(ExpectedDomain == ActualDomain);
+	});
+	HttpRequest->OnProcessRequestComplete().BindLambda([ExpectedDomainsOrderInRetry](FHttpRequestPtr HttpRequest, FHttpResponsePtr HttpResponse, bool bSucceeded) {
+		CHECK(ExpectedDomainsOrderInRetry->IsEmpty());
+		delete ExpectedDomainsOrderInRetry;
+		CVarHttpUrlPatternsToMockFailure->Set(TEXT(""));
+
+		CHECK(!bSucceeded);
+		REQUIRE(HttpResponse == nullptr);
+	});
+	HttpRequest->ProcessRequest();
+}
+
 #if UE_HTTP_CONNECTION_TIMEOUT_SUPPORT_RETRY
 TEST_CASE_METHOD(FWaitUntilCompleteHttpFixture, "Optionally retry limit can be set differently for connection error", HTTP_TAG)
 {
@@ -2614,4 +2696,4 @@ TEST_CASE_METHOD(FWaitUntilCompleteHttpFixture, "Http Methods over Unix Domain S
 	HttpRequest->ProcessRequest();
 }
 
-#endif //UE_HTTP_SUPPORT_UNIX_SOCKET
\ No newline at end of file
+#endif //UE_HTTP_SUPPORT_UNIX_SOCKET
diff --git a/Engine/Source/Runtime/Online/HTTP/Private/HttpRetrySystem.cpp b/Engine/Source/Runtime/Online/HTTP/Private/HttpRetrySystem.cpp
--- a/Engine/Source/Runtime/Online/HTTP/Private/HttpRetrySystem.cpp
+++ b/Engine/Source/Runtime/Online/HTTP/Private/HttpRetrySystem.cpp
@@ -21,6 +21,13 @@ TAutoConsoleVariable<bool> CVarHttpRetrySystemNonGameThreadSupportEnabled(
 	TEXT("Enable retry system non-game thread support")
 );
 
+TAutoConsoleVariable<bool> CVarHttpAddDomainOfDefaultUrlToRetryDomainsEnabled(
+	TEXT("http.OnlySetUrlFromAltDomainInRetry"),
+	true,
+	TEXT("When enabled, add domain in default url into retry domains"),
+	ECVF_SaveForNextBoot
+);
+
 LLM_DEFINE_TAG(HTTP);
 
 namespace FHttpRetrySystem
@@ -106,8 +113,8 @@ FHttpRetrySystem::FRequest::FRequest(
 	, RetryManager(InManager)
 	, RetryExponentialBackoffCurve(InExponentialBackoffCurve)
 {
-    // if the InRetryTimeoutRelativeSecondsOverride override is being used the value cannot be negative
-    check(!(InRetryTimeoutRelativeSecondsOverride.IsSet()) || (InRetryTimeoutRelativeSecondsOverride.GetValue() >= 0.0));
+	// if the InRetryTimeoutRelativeSecondsOverride override is being used the value cannot be negative
+	check(!(InRetryTimeoutRelativeSecondsOverride.IsSet()) || (InRetryTimeoutRelativeSecondsOverride.GetValue() >= 0.0));
 
 	if (RetryDomains.IsValid())
 	{
@@ -144,6 +151,26 @@ bool FHttpRetrySystem::FRequest::ProcessRequest()
 	TSharedRef<FRequest> RetryRequest = StaticCastSharedRef<FRequest>(AsShared());
 
 	OriginalUrl = HttpRequest->GetURL();
+
+	if (CVarHttpAddDomainOfDefaultUrlToRetryDomainsEnabled.GetValueOnAnyThread())
+	{
+		if (RetryDomains.IsValid() && !RetryDomains->Domains.IsEmpty())
+		{
+			FString OriginalUrlDomainAndPort = FPlatformHttp::GetUrlDomainAndPort(OriginalUrl);
+			int32 Index = RetryDomains->Domains.Find(OriginalUrlDomainAndPort);
+			if (Index == INDEX_NONE)
+			{
+				RetryDomains->Domains.Insert(MoveTemp(OriginalUrlDomainAndPort), 0);
+			}
+			else if (Index > 0)
+			{
+				RetryDomains->Domains.RemoveAt(Index);
+				RetryDomains->Domains.Insert(MoveTemp(OriginalUrlDomainAndPort), 0);
+			}
+		}
+	}
+
+	// The ActiveIndex inside FRetryDomains could have been increased before, so to skip the domains failed to connect
 	if (RetryDomains.IsValid())
 	{
 		SetUrlFromRetryDomains();
@@ -367,7 +394,7 @@ TSharedRef<FHttpRetrySystem::FRequest, ESPMode::ThreadSafe> FHttpRetrySystem::FM
 		));
 }
 
-bool FHttpRetrySystem::FManager::ShouldRetry(const FHttpRetryRequestEntry& HttpRetryRequestEntry)
+bool FHttpRetrySystem::FManager::ShouldRetry(const FHttpRetryRequestEntry& HttpRetryRequestEntry) const
 {
 	FHttpResponsePtr Response = HttpRetryRequestEntry.Request->GetResponse();
 	if (Response)
@@ -405,24 +432,24 @@ bool FHttpRetrySystem::FManager::ShouldRetry(const FHttpRetryRequestEntry& HttpR
 	return DefaultRetryVerbs.Contains(Verb);
 }
 
-bool FHttpRetrySystem::FManager::RetryLimitForConnectionErrorIsSet(const FHttpRetryRequestEntry& HttpRetryRequestEntry)
+bool FHttpRetrySystem::FManager::RetryLimitForConnectionErrorIsSet(const FHttpRetryRequestEntry& HttpRetryRequestEntry) const
 {
 	return HttpRetryRequestEntry.Request->RetryLimitCountForConnectionErrorOverride.IsSet() || RetryLimitCountForConnectionErrorDefault.IsSet();
 }
 
-bool FHttpRetrySystem::FManager::CanRetryForConnectionError(const FHttpRetryRequestEntry& HttpRetryRequestEntry)
+bool FHttpRetrySystem::FManager::CanRetryForConnectionError(const FHttpRetryRequestEntry& HttpRetryRequestEntry) const
 {
 	uint32 RetryLimitForConnectionError = HttpRetryRequestEntry.Request->RetryLimitCountForConnectionErrorOverride.Get(RetryLimitCountForConnectionErrorDefault.Get(0));
 	return HttpRetryRequestEntry.CurrentRetryCountForConnectionError < RetryLimitForConnectionError;
 }
 
-bool FHttpRetrySystem::FManager::CanRetryInGeneral(const FHttpRetryRequestEntry& HttpRetryRequestEntry)
+bool FHttpRetrySystem::FManager::CanRetryInGeneral(const FHttpRetryRequestEntry& HttpRetryRequestEntry) const
 {
 	uint32 RetryLimit = HttpRetryRequestEntry.Request->RetryLimitCountOverride.Get(RetryLimitCountDefault.Get(0));
 	return HttpRetryRequestEntry.CurrentRetryCount < RetryLimit;
 }
 
-bool FHttpRetrySystem::FManager::CanRetry(const FHttpRetryRequestEntry& HttpRetryRequestEntry)
+bool FHttpRetrySystem::FManager::CanRetry(const FHttpRetryRequestEntry& HttpRetryRequestEntry) const
 {
 	if (HttpRetryRequestEntry.Request->GetFailureReason() == EHttpFailureReason::ConnectionError && RetryLimitForConnectionErrorIsSet(HttpRetryRequestEntry))
 	{
diff --git a/Engine/Source/Runtime/Online/HTTP/Public/HttpRetrySystem.h b/Engine/Source/Runtime/Online/HTTP/Public/HttpRetrySystem.h
--- a/Engine/Source/Runtime/Online/HTTP/Public/HttpRetrySystem.h
+++ b/Engine/Source/Runtime/Online/HTTP/Public/HttpRetrySystem.h
@@ -51,7 +51,7 @@ namespace FHttpRetrySystem
 		{}
 
 		/** The domains to use */
-		const TArray<FString> Domains;
+		TArray<FString> Domains;
 		/**
 		 * Index into Domains to attempt
 		 * Domains are cycled through on some errors, and when we succeed on one domain, we remain on that domain until that domain results in an error
@@ -257,10 +257,10 @@ protected:
 
 	// @return true if there is a no formal response to the request
 	// @TODO return true if a variety of 5xx errors are the result of a formal response
-	bool ShouldRetry(const FHttpRetryRequestEntry& HttpRetryRequestEntry);
+	bool ShouldRetry(const FHttpRetryRequestEntry& HttpRetryRequestEntry) const;
 
 	// @return true if retry chances have not been exhausted
-	bool CanRetry(const FHttpRetryRequestEntry& HttpRetryRequestEntry);
+	bool CanRetry(const FHttpRetryRequestEntry& HttpRetryRequestEntry) const;
 
 	// @return true if the retry request has timed out
 	bool HasTimedOut(const FHttpRetryRequestEntry& HttpRetryRequestEntry, const double NowAbsoluteSeconds);
@@ -284,9 +284,9 @@ protected:
 	// @return number of seconds to lockout for
 	float GetLockoutPeriodSeconds(const FHttpRetryRequestEntry& HttpRetryRequestEntry);
 
-	bool RetryLimitForConnectionErrorIsSet(const FHttpRetryRequestEntry& HttpRetryRequestEntry);
-	bool CanRetryForConnectionError(const FHttpRetryRequestEntry& HttpRetryRequestEntry);
-	bool CanRetryInGeneral(const FHttpRetryRequestEntry& HttpRetryRequestEntry);
+	bool RetryLimitForConnectionErrorIsSet(const FHttpRetryRequestEntry& HttpRetryRequestEntry) const;
+	bool CanRetryForConnectionError(const FHttpRetryRequestEntry& HttpRetryRequestEntry) const;
+	bool CanRetryInGeneral(const FHttpRetryRequestEntry& HttpRetryRequestEntry) const;
 
 	// Default configuration for the retry system
 	FRandomFailureRateSetting RandomFailureRate;
