Fix for slow BPT cloud source downloads - thread wakeup function got lost somewhere along the way.

#rnx

diff --git a/Engine/Source/Runtime/Online/BuildPatchServices/Private/Installer/CloudChunkSource.cpp b/Engine/Source/Runtime/Online/BuildPatchServices/Private/Installer/CloudChunkSource.cpp
--- a/Engine/Source/Runtime/Online/BuildPatchServices/Private/Installer/CloudChunkSource.cpp
+++ b/Engine/Source/Runtime/Online/BuildPatchServices/Private/Installer/CloudChunkSource.cpp
@@ -949,8 +949,11 @@ namespace BuildPatchServices
 
 		void OnDownloadComplete(int32 RequestId, const FDownloadRef& Download)
 		{
-			FScopeLock ScopeLock(&CompletedDownloadsCS);
-			CompletedDownloads.Add(RequestId, Download);
+			{
+				FScopeLock ScopeLock(&CompletedDownloadsCS);
+				CompletedDownloads.Add(RequestId, Download);
+			}
+			WakeupMainThreadFn();
 		}
 
 		FConstructorCloudChunkSourceConfig Configuration;
