Log the failed url when http request failed, if request is not aborted nor log disabled.

diff --git a/Engine/Source/Programs/WebTests/Private/TestHttp.cpp b/Engine/Source/Programs/WebTests/Private/TestHttp.cpp
--- a/Engine/Source/Programs/WebTests/Private/TestHttp.cpp
+++ b/Engine/Source/Programs/WebTests/Private/TestHttp.cpp
@@ -425,6 +425,8 @@ TEST_CASE_METHOD(FWaitUntilCompleteHttpFixture, "Can process https request", HTT
 
 TEST_CASE_METHOD(FWaitUntilCompleteHttpFixture, "Can mock connection error through CVar", HTTP_TAG)
 {
+	DisableWarningsInThisTest();
+
 	CVarHttpUrlPatternsToMockFailure->Set(TEXT("epicgames.com->0 unrealengine.com->503"));
 
 	float ExpectedTimeoutDuration = 2.0f;
diff --git a/Engine/Source/Runtime/Online/HTTP/Private/GenericPlatform/HttpRequestCommon.cpp b/Engine/Source/Runtime/Online/HTTP/Private/GenericPlatform/HttpRequestCommon.cpp
--- a/Engine/Source/Runtime/Online/HTTP/Private/GenericPlatform/HttpRequestCommon.cpp
+++ b/Engine/Source/Runtime/Online/HTTP/Private/GenericPlatform/HttpRequestCommon.cpp
@@ -290,9 +290,27 @@ void FHttpRequestCommon::HandleRequestFailed(TSharedPtr<IHttpResponse> InRespons
 
 	SetStatus(EHttpRequestStatus::Failed);
 
+	LogFailure();
+
 	OnProcessRequestComplete().ExecuteIfBound(SharedThis(this), InResponse, false);
 }
 
+#define UE_HTTP_LOG_AS_WARNING_IF(Condition, Format, ...) \
+	if (Condition) \
+	{ \
+		UE_LOG(LogHttp, Warning, Format, ##__VA_ARGS__); \
+	} \
+	else \
+	{ \
+		UE_LOG(LogHttp, Verbose, Format, ##__VA_ARGS__); \
+	}
+
+void FHttpRequestCommon::LogFailure() const
+{
+	const bool bAborted = (bCanceled || bTimedOut || bActivityTimedOut);
+	UE_HTTP_LOG_AS_WARNING_IF(!bAborted && !FHttpModule::Get().GetHttpManager().ShouldDisableFailedLog(GetURL()), TEXT("%p %s %s completed with reason '%s' after %.2fs"), this, *GetVerb(), *GetURL(), LexToString(GetFailureReason()), ElapsedTime);
+}
+
 void FHttpRequestCommon::SetStatus(EHttpRequestStatus::Type InCompletionStatus)
 {
 	CompletionStatus = InCompletionStatus;
@@ -706,27 +724,17 @@ void FHttpRequestCommon::CloseRequestPayloadDefaultImpl()
 	}
 }
 
-#define UE_HTTP_LOG_RESPONSE_PRIVATE(Condition, Format, ...) \
-	if (Condition) \
-	{ \
-		UE_LOG(LogHttp, Warning, Format, ##__VA_ARGS__); \
-	} \
-	else \
-	{ \
-		UE_LOG(LogHttp, Verbose, Format, ##__VA_ARGS__); \
-	}
-
 void FHttpRequestCommon::LogResponse(const TSharedPtr<IHttpResponse>& InResponse)
 {
 	bool bShouldLogResponse = FHttpModule::Get().GetHttpManager().ShouldLogResponse(GetURL());
-	UE_HTTP_LOG_RESPONSE_PRIVATE(bShouldLogResponse, TEXT("%p %s %s completed with code %d after %.2fs. Content length: %ld"), this, *GetVerb(), *GetURL(), InResponse->GetResponseCode(), ElapsedTime, InResponse->GetContentLength());
+	UE_HTTP_LOG_AS_WARNING_IF(bShouldLogResponse, TEXT("%p %s %s completed with code %d after %.2fs. Content length: %ld"), this, *GetVerb(), *GetURL(), InResponse->GetResponseCode(), ElapsedTime, InResponse->GetContentLength());
 
 	TArray<FString> AllHeaders = InResponse->GetAllHeaders();
 	for (const FString& HeaderStr : AllHeaders)
 	{
 		if (!HeaderStr.StartsWith(TEXT("Authorization")) && !HeaderStr.StartsWith(TEXT("Set-Cookie")))
 		{
-			UE_HTTP_LOG_RESPONSE_PRIVATE(bShouldLogResponse, TEXT("%p Response Header %s"), this, *HeaderStr);
+			UE_HTTP_LOG_AS_WARNING_IF(bShouldLogResponse, TEXT("%p Response Header %s"), this, *HeaderStr);
 		}
 	}
 
diff --git a/Engine/Source/Runtime/Online/HTTP/Private/HttpManager.cpp b/Engine/Source/Runtime/Online/HTTP/Private/HttpManager.cpp
--- a/Engine/Source/Runtime/Online/HTTP/Private/HttpManager.cpp
+++ b/Engine/Source/Runtime/Online/HTTP/Private/HttpManager.cpp
@@ -33,6 +33,12 @@ TAutoConsoleVariable<FString> CVarHttpUrlPatternsToLogResponse(
 	ECVF_SaveForNextBoot
 );
 
+TAutoConsoleVariable<FString> CVarHttpUrlPatternsToDisableFailedLog(
+	TEXT("http.UrlPatternsToDisableFailedLog"),
+	TEXT(""),
+	TEXT("List of url patterns to disable failed log: \"epicgames.com unrealengine.com ...\"")
+);
+
 TAutoConsoleVariable<FString> CVarHttpUrlPatternsToMockFailure(
 	TEXT("http.UrlPatternsToMockFailure"),
 	TEXT(""),
@@ -114,6 +120,9 @@ void FHttpManager::Initialize()
 	UpdateUrlPatternsToLogResponse(CVarHttpUrlPatternsToLogResponse.AsVariable());
 	CVarHttpUrlPatternsToLogResponse.AsVariable()->OnChangedDelegate().AddRaw(this, &FHttpManager::UpdateUrlPatternsToLogResponse);
 
+	UpdateUrlPatternsToDisableFailedLog(CVarHttpUrlPatternsToDisableFailedLog.AsVariable());
+	CVarHttpUrlPatternsToDisableFailedLog.AsVariable()->OnChangedDelegate().AddRaw(this, &FHttpManager::UpdateUrlPatternsToDisableFailedLog);
+
 	UpdateUrlPatternsToMockFailure(CVarHttpUrlPatternsToMockFailure.AsVariable());
 	CVarHttpUrlPatternsToMockFailure.AsVariable()->OnChangedDelegate().AddRaw(this, &FHttpManager::UpdateUrlPatternsToMockFailure);
 }
@@ -691,7 +700,7 @@ void FHttpManager::UpdateUrlPatternsToLogResponse(IConsoleVariable* CVar)
 	UrlPatternsToLogResponseStr.ParseIntoArray(UrlPatternsToLogResponse, TEXT(" "));
 }
 
-bool FHttpManager::ShouldLogResponse(FStringView Url)
+bool FHttpManager::ShouldLogResponse(FStringView Url) const
 {
 	const FScopeLock CacheLock(&UrlPatternsToLogResponseCriticalSection);
 	for (const FString& UrlPatternToLogResponse : UrlPatternsToLogResponse)
@@ -705,6 +714,27 @@ bool FHttpManager::ShouldLogResponse(FStringView Url)
 	return false;
 }
 
+void FHttpManager::UpdateUrlPatternsToDisableFailedLog(IConsoleVariable* CVar)
+{
+	const FScopeLock CacheLock(&UrlPatternsToDisableFailedLogCriticalSection);
+	const FString UrlPatternsToDisableFailedLogStr = CVar->AsVariable()->GetString();
+	UrlPatternsToDisableFailedLogStr.ParseIntoArray(UrlPatternsToDisableFailedLog, TEXT(" "));
+}
+
+bool FHttpManager::ShouldDisableFailedLog(FStringView Url) const
+{
+	const FScopeLock CacheLock(&UrlPatternsToDisableFailedLogCriticalSection);
+	for (const FString& UrlPatternToDisableFailedLog : UrlPatternsToDisableFailedLog)
+	{
+		if (Url.Contains(UrlPatternToDisableFailedLog))
+		{
+			return true;
+		}
+	}
+
+	return false;
+}
+
 void FHttpManager::UpdateUrlPatternsToMockFailure(IConsoleVariable* CVar)
 {
 	const FScopeLock CacheLock(&UrlPatternsToMockFailureCriticalSection);
@@ -725,7 +755,7 @@ void FHttpManager::UpdateUrlPatternsToMockFailure(IConsoleVariable* CVar)
 	}
 }
 
-TOptional<int32> FHttpManager::GetMockFailure(FStringView Url)
+TOptional<int32> FHttpManager::GetMockFailure(FStringView Url) const
 {
 	TOptional<int32> Result;
 	const FScopeLock CacheLock(&UrlPatternsToMockFailureCriticalSection);
diff --git a/Engine/Source/Runtime/Online/HTTP/Public/GenericPlatform/HttpRequestCommon.h b/Engine/Source/Runtime/Online/HTTP/Public/GenericPlatform/HttpRequestCommon.h
--- a/Engine/Source/Runtime/Online/HTTP/Public/GenericPlatform/HttpRequestCommon.h
+++ b/Engine/Source/Runtime/Online/HTTP/Public/GenericPlatform/HttpRequestCommon.h
@@ -107,6 +107,7 @@ protected:
 	HTTP_API bool OpenRequestPayloadDefaultImpl();
 	HTTP_API void CloseRequestPayloadDefaultImpl();
 
+	HTTP_API void LogFailure() const;
 	HTTP_API void LogResponse(const TSharedPtr<IHttpResponse>& InResponse);
 	HTTP_API bool TriggerMockFailure();
 
diff --git a/Engine/Source/Runtime/Online/HTTP/Public/HttpManager.h b/Engine/Source/Runtime/Online/HTTP/Public/HttpManager.h
--- a/Engine/Source/Runtime/Online/HTTP/Public/HttpManager.h
+++ b/Engine/Source/Runtime/Online/HTTP/Public/HttpManager.h
@@ -337,6 +337,7 @@ protected:
 	HTTP_API bool HasAnyBoundDelegate() const;
 
 	HTTP_API void UpdateUrlPatternsToLogResponse(IConsoleVariable* CVar);
+	HTTP_API void UpdateUrlPatternsToDisableFailedLog(IConsoleVariable* CVar);
 	HTTP_API void UpdateUrlPatternsToMockFailure(IConsoleVariable* CVar);
 
 	HTTP_API void AddRequest(const FHttpRequestRef& Request);
@@ -410,10 +411,13 @@ protected:
 	bool bUseEventLoop = true;
 
 	TArray<FString> UrlPatternsToLogResponse;
-	FCriticalSection UrlPatternsToLogResponseCriticalSection;
+	mutable FCriticalSection UrlPatternsToLogResponseCriticalSection;
+
+	TArray<FString> UrlPatternsToDisableFailedLog;
+	mutable FCriticalSection UrlPatternsToDisableFailedLogCriticalSection;
 
 	TMap<FString, int32> UrlPatternsToMockFailure;
-	FCriticalSection UrlPatternsToMockFailureCriticalSection;
+	mutable FCriticalSection UrlPatternsToMockFailureCriticalSection;
 
 PACKAGE_SCOPE:
 
@@ -450,7 +454,8 @@ PACKAGE_SCOPE:
 	/** Record platform specific stats */
 	HTTP_API void RecordPlatformStats(const FHttpStatsPlatform& PlatformStats);
 
-	HTTP_API bool ShouldLogResponse(FStringView Url);
+	HTTP_API bool ShouldLogResponse(FStringView Url) const;
+	HTTP_API bool ShouldDisableFailedLog(FStringView Url) const;
 
-	HTTP_API TOptional<int32> GetMockFailure(FStringView Url);
+	HTTP_API TOptional<int32> GetMockFailure(FStringView Url) const;
 };
