Move the SetURL/GetURL into common impl;
Change the GetURL of IHttpBase to return const reference.

diff --git a/Engine/Source/Runtime/Online/BuildPatchServices/Private/Tests/Mock/HttpRequest.mock.h b/Engine/Source/Runtime/Online/BuildPatchServices/Private/Tests/Mock/HttpRequest.mock.h
--- a/Engine/Source/Runtime/Online/BuildPatchServices/Private/Tests/Mock/HttpRequest.mock.h
+++ b/Engine/Source/Runtime/Online/BuildPatchServices/Private/Tests/Mock/HttpRequest.mock.h
@@ -16,10 +16,10 @@ namespace BuildPatchServices
 		typedef TTuple<FString> FRxSetURL;
 
 	public:
-		virtual FString GetURL() const override
+		virtual const FString& GetURL() const override
 		{
 			MOCK_FUNC_NOT_IMPLEMENTED("FMockHttpRequest::GetURL");
-			return FString();
+			return URL;
 		}
 
 		virtual FHttpRequestWillRetryDelegate& OnRequestWillRetry() override
@@ -82,9 +82,9 @@ namespace BuildPatchServices
 			RxSetVerb.Emplace(Verb);
 		}
 
-		virtual void SetURL(const FString& URL) override
+		virtual void SetURL(const FString& InURL) override
 		{
-			RxSetURL.Emplace(URL);
+			RxSetURL.Emplace(InURL);
 		}
 
 		virtual void SetOption(const FName Option, const FString& OptionValue) override
@@ -255,6 +255,7 @@ namespace BuildPatchServices
 		int32 RxProcessRequest;
 		int32 RxCancelRequest;
 		FString EffectiveURL;
+		FString URL;
 	};
 }
 
diff --git a/Engine/Source/Runtime/Online/BuildPatchServices/Private/Tests/Mock/HttpResponse.mock.h b/Engine/Source/Runtime/Online/BuildPatchServices/Private/Tests/Mock/HttpResponse.mock.h
--- a/Engine/Source/Runtime/Online/BuildPatchServices/Private/Tests/Mock/HttpResponse.mock.h
+++ b/Engine/Source/Runtime/Online/BuildPatchServices/Private/Tests/Mock/HttpResponse.mock.h
@@ -28,10 +28,10 @@ namespace BuildPatchServices
 			return FString();
 		}
 
-		virtual FString GetURL() const override
+		virtual const FString& GetURL() const override
 		{
 			MOCK_FUNC_NOT_IMPLEMENTED("FMockHttpRequest::GetURL");
-			return FString();
+			return URL;
 		}
 
 		virtual FString GetURLParameter(const FString& ParameterName) const override
@@ -89,6 +89,9 @@ namespace BuildPatchServices
 			static FString None;
 			return None;
 		}
+
+	private:
+		FString URL;
 	};
 }
 
diff --git a/Engine/Source/Runtime/Online/HTTP/Private/Apple/AppleHttp.cpp b/Engine/Source/Runtime/Online/HTTP/Private/Apple/AppleHttp.cpp
--- a/Engine/Source/Runtime/Online/HTTP/Private/Apple/AppleHttp.cpp
+++ b/Engine/Source/Runtime/Online/HTTP/Private/Apple/AppleHttp.cpp
@@ -517,28 +517,6 @@ FAppleHttpRequest::~FAppleHttpRequest()
     [Session release];
 }
 
-FString FAppleHttpRequest::GetURL() const
-{
-	SCOPED_AUTORELEASE_POOL;
-	NSURL* URL = Request.URL;
-	if (URL != nullptr)
-	{
-		FString ConvertedURL(URL.absoluteString);
-		return ConvertedURL;
-	}
-	else
-	{
-		return FString();
-	}
-}
-
-void FAppleHttpRequest::SetURL(const FString& URL)
-{
-	SCOPED_AUTORELEASE_POOL;
-	UE_LOG(LogHttp, Verbose, TEXT("FAppleHttpRequest::SetURL() - %s"), *URL);
-	Request.URL = [NSURL URLWithString: URL.GetNSString()];
-}
-
 FString FAppleHttpRequest::GetHeader(const FString& HeaderName) const
 {
 	SCOPED_AUTORELEASE_POOL;
@@ -763,6 +741,8 @@ bool FAppleHttpRequest::SetupRequest()
 {
 	SCOPED_AUTORELEASE_POOL;
 
+	Request.URL = [NSURL URLWithString: URL.GetNSString()];
+
 	// set the content-length and user-agent (it is possible that the OS ignores this value)
 	if(GetContentLength() > 0)
 	{
diff --git a/Engine/Source/Runtime/Online/HTTP/Private/Apple/AppleHttp.h b/Engine/Source/Runtime/Online/HTTP/Private/Apple/AppleHttp.h
--- a/Engine/Source/Runtime/Online/HTTP/Private/Apple/AppleHttp.h
+++ b/Engine/Source/Runtime/Online/HTTP/Private/Apple/AppleHttp.h
@@ -27,7 +27,6 @@ public:
 	friend class FAppleHttpResponse;
 
 	//~ Begin IHttpBase Interface
-	virtual FString GetURL() const override;
 	virtual FString GetHeader(const FString& HeaderName) const override;
 	virtual TArray<FString> GetAllHeaders() const override;	
 	virtual FString GetContentType() const override;
@@ -38,7 +37,6 @@ public:
 	//~ Begin IHttpRequest Interface
 	virtual FString GetVerb() const override;
 	virtual void SetVerb(const FString& Verb) override;
-	virtual void SetURL(const FString& URL) override;
 	virtual void SetContent(const TArray<uint8>& ContentPayload) override;
 	virtual void SetContent(TArray<uint8>&& ContentPayload) override;
 	virtual void SetContentAsString(const FString& ContentString) override;
diff --git a/Engine/Source/Runtime/Online/HTTP/Private/Curl/CurlHttp.cpp b/Engine/Source/Runtime/Online/HTTP/Private/Curl/CurlHttp.cpp
--- a/Engine/Source/Runtime/Online/HTTP/Private/Curl/CurlHttp.cpp
+++ b/Engine/Source/Runtime/Online/HTTP/Private/Curl/CurlHttp.cpp
@@ -54,7 +54,10 @@ static int SslCertVerify(int PreverifyOk, X509_STORE_CTX* Context)
 		FCurlHttpRequest* Request = static_cast<FCurlHttpRequest*>(SSL_CTX_get_app_data(SslContext));
 		check(Request);
 
+// FHttpRequestCommon deprecation
+PRAGMA_DISABLE_DEPRECATION_WARNINGS
 		const FString Domain = FPlatformHttp::GetUrlDomain(Request->GetURL());
+PRAGMA_ENABLE_DEPRECATION_WARNINGS
 
 		if (!FSslModule::Get().GetCertificateManager().VerifySslCertificates(Context, Domain))
 		{
@@ -229,11 +232,6 @@ FCurlHttpRequest::~FCurlHttpRequest()
 	}
 }
 
-FString FCurlHttpRequest::GetURL() const
-{
-	return URL;
-}
-
 FString FCurlHttpRequest::GetHeader(const FString& HeaderName) const
 {
 	FString Result;
@@ -298,18 +296,6 @@ void FCurlHttpRequest::SetVerb(const FString& InVerb)
 	Verb = InVerb.ToUpper();
 }
 
-void FCurlHttpRequest::SetURL(const FString& InURL)
-{
-	if (CompletionStatus == EHttpRequestStatus::Processing)
-	{
-		UE_LOG(LogHttp, Warning, TEXT("FCurlHttpRequest::SetURL() - attempted to set url on a request that is inflight"));
-		return;
-	}
-
-	check(EasyHandle);
-	URL = InURL;
-}
-
 void FCurlHttpRequest::SetOption(const FName Option, const FString& OptionValue)
 {
 	if (CompletionStatus == EHttpRequestStatus::Processing)
@@ -877,7 +863,7 @@ bool FCurlHttpRequest::SetupRequestHttpThread()
 	{
 		QUICK_SCOPE_CYCLE_COUNTER(STAT_FCurlHttpRequest_SetupRequest_EASY_SETOPT);
 
-		curl_easy_setopt(EasyHandle, CURLOPT_URL, TCHAR_TO_ANSI(*URL));
+		curl_easy_setopt(EasyHandle, CURLOPT_URL, TCHAR_TO_ANSI(*GetURL()));
 
 #if UE_HTTP_SUPPORT_UNIX_SOCKET
 		FString UnixSocketPath = GetOption(HttpRequestOptions::UnixSocketPath);
@@ -1037,7 +1023,7 @@ bool FCurlHttpRequest::SetupRequestHttpThread()
 	curl_easy_setopt(EasyHandle, CURLOPT_QUICK_EXIT, 1L);
 #endif
 
-	UE_LOG(LogHttp, Log, TEXT("%p: Starting %s request to URL='%s'"), this, *Verb, *URL);
+	UE_LOG(LogHttp, Log, TEXT("%p: Starting %s request to URL='%s'"), this, *Verb, *GetURL());
 
 	return true;
 }
@@ -1246,7 +1232,7 @@ void FCurlHttpRequest::FinishRequest()
 				}
 			}
 
-			if (Response->GetResponseCode() <= 0 && URL.StartsWith(TEXT("Http"), ESearchCase::IgnoreCase))
+			if (Response->GetResponseCode() <= 0 && GetURL().StartsWith(TEXT("Http"), ESearchCase::IgnoreCase))
 			{
 				UE_LOG(LogHttp, Warning, TEXT("%p: invalid HTTP response code received. URL: %s, HTTP code: %d, content length: %llu, actual payload size: %llu"),
 					this, *GetURL(), Response->GetResponseCode(), Response->ContentLength, TotalBytesRead.load());
diff --git a/Engine/Source/Runtime/Online/HTTP/Private/Curl/CurlHttp.h b/Engine/Source/Runtime/Online/HTTP/Private/Curl/CurlHttp.h
--- a/Engine/Source/Runtime/Online/HTTP/Private/Curl/CurlHttp.h
+++ b/Engine/Source/Runtime/Online/HTTP/Private/Curl/CurlHttp.h
@@ -145,7 +145,6 @@ public:
 	friend class FCurlHttpResponse;
 
 	//~ Begin IHttpBase Interface
-	virtual FString GetURL() const override;
 	virtual FString GetHeader(const FString& HeaderName) const override;
 	virtual TArray<FString> GetAllHeaders() const override;
 	virtual FString GetContentType() const override;
@@ -156,7 +155,6 @@ public:
 	//~ Begin IHttpRequest Interface
 	virtual FString GetVerb() const override;
 	virtual void SetVerb(const FString& InVerb) override;
-	virtual void SetURL(const FString& InURL) override;
 	virtual void SetOption(const FName Option, const FString& OptionValue) override;
 	virtual void SetContent(const TArray<uint8>& ContentPayload) override;
 	virtual void SetContent(TArray<uint8>&& ContentPayload) override;
@@ -370,8 +368,6 @@ private:
 	CURL *			EasyHandle;	
 	/** List of custom headers to be passed to CURL */
 	curl_slist *	HeaderList;
-	/** Cached URL */
-	FString			URL;
 	/** Cached verb */
 	FString			Verb;
 	/** Set to true when request has been completed */
diff --git a/Engine/Source/Runtime/Online/HTTP/Private/GenericPlatform/GenericPlatformHttp.cpp b/Engine/Source/Runtime/Online/HTTP/Private/GenericPlatform/GenericPlatformHttp.cpp
--- a/Engine/Source/Runtime/Online/HTTP/Private/GenericPlatform/GenericPlatformHttp.cpp
+++ b/Engine/Source/Runtime/Online/HTTP/Private/GenericPlatform/GenericPlatformHttp.cpp
@@ -95,7 +95,7 @@ class FGenericPlatformHttpRequest : public FHttpRequestImpl
 {
 public:
 	// IHttpBase
-	virtual FString GetURL() const override { return TEXT(""); }
+	virtual const FString& GetURL() const override { return URL; }
 	virtual FString GetURLParameter(const FString& ParameterName) const override { return TEXT(""); }
 	virtual FString GetHeader(const FString& HeaderName) const override { return TEXT(""); }
 	virtual TArray<FString> GetAllHeaders() const override { return TArray<FString>(); }
@@ -106,7 +106,7 @@ public:
 	// IHttpRequest
 	virtual FString GetVerb() const override { return TEXT(""); }
 	virtual void SetVerb(const FString& Verb) override {}
-	virtual void SetURL(const FString& URL) override {}
+	virtual void SetURL(const FString& InURL) override { URL = InURL; }
 	virtual FString GetOption(const FName Option) const override { return TEXT(""); }
 	virtual void SetOption(const FName Option, const FString& OptionValue) override {}
 	virtual void SetContent(const TArray<uint8>& ContentPayload) override {}
@@ -135,6 +135,7 @@ public:
 	virtual void ProcessRequestUntilComplete() override {}
 
 private:
+	FString URL;
 	FString EffectiveURL;
 };
 
diff --git a/Engine/Source/Runtime/Online/HTTP/Private/GenericPlatform/HttpRequestCommon.cpp b/Engine/Source/Runtime/Online/HTTP/Private/GenericPlatform/HttpRequestCommon.cpp
--- a/Engine/Source/Runtime/Online/HTTP/Private/GenericPlatform/HttpRequestCommon.cpp
+++ b/Engine/Source/Runtime/Online/HTTP/Private/GenericPlatform/HttpRequestCommon.cpp
@@ -654,6 +654,22 @@ float FHttpRequestCommon::GetTimeStartedWaitingInQueue() const
 	return TimeStartedWaitingInQueue;
 }
 
+void FHttpRequestCommon::SetURL(const FString& InURL)
+{
+	if (CompletionStatus == EHttpRequestStatus::Processing)
+	{
+		UE_LOG(LogHttp, Warning, TEXT("FHttpRequestCommon::SetURL() - attempted to set url on a request that is inflight"));
+		return;
+	}
+
+	URL = InURL;
+}
+
+const FString& FHttpRequestCommon::GetURL() const
+{
+	return URL;
+}
+
 bool FHttpRequestCommon::PassReceivedDataToStream(void* Ptr, int64 Length)
 {
 	const FScopeLock StreamLock(&ResponseBodyReceiveStreamCriticalSection);
diff --git a/Engine/Source/Runtime/Online/HTTP/Private/GenericPlatform/HttpResponseCommon.cpp b/Engine/Source/Runtime/Online/HTTP/Private/GenericPlatform/HttpResponseCommon.cpp
--- a/Engine/Source/Runtime/Online/HTTP/Private/GenericPlatform/HttpResponseCommon.cpp
+++ b/Engine/Source/Runtime/Online/HTTP/Private/GenericPlatform/HttpResponseCommon.cpp
@@ -23,7 +23,7 @@ FString FHttpResponseCommon::GetURLParameter(const FString& ParameterName) const
 	return ReturnValue;
 }
 
-FString FHttpResponseCommon::GetURL() const
+const FString& FHttpResponseCommon::GetURL() const
 {
 	return URL;
 }
diff --git a/Engine/Source/Runtime/Online/HTTP/Private/HttpRequestAdapter.cpp b/Engine/Source/Runtime/Online/HTTP/Private/HttpRequestAdapter.cpp
--- a/Engine/Source/Runtime/Online/HTTP/Private/HttpRequestAdapter.cpp
+++ b/Engine/Source/Runtime/Online/HTTP/Private/HttpRequestAdapter.cpp
@@ -9,7 +9,7 @@ FHttpRequestAdapterBase::FHttpRequestAdapterBase(const TSharedRef<IHttpRequest>&
 
 FHttpRequestAdapterBase::~FHttpRequestAdapterBase() = default;
 
-FString FHttpRequestAdapterBase::GetURL() const 
+const FString& FHttpRequestAdapterBase::GetURL() const 
 { 
 	return HttpRequest->GetURL(); 
 }
diff --git a/Engine/Source/Runtime/Online/HTTP/Private/NullHttp.cpp b/Engine/Source/Runtime/Online/HTTP/Private/NullHttp.cpp
--- a/Engine/Source/Runtime/Online/HTTP/Private/NullHttp.cpp
+++ b/Engine/Source/Runtime/Online/HTTP/Private/NullHttp.cpp
@@ -7,7 +7,7 @@
 
 // FNullHttpRequest
 
-FString FNullHttpRequest::GetURL() const
+const FString& FNullHttpRequest::GetURL() const
 {
 	return Url;
 }
@@ -253,9 +253,9 @@ void FNullHttpRequest::ProcessRequestUntilComplete()
 
 // FNullHttpResponse
 
-FString FNullHttpResponse::GetURL() const
+const FString& FNullHttpResponse::GetURL() const
 {
-	return FString();
+	return URL;
 }
 
 FString FNullHttpResponse::GetURLParameter(const FString& ParameterName) const
diff --git a/Engine/Source/Runtime/Online/HTTP/Private/NullHttp.h b/Engine/Source/Runtime/Online/HTTP/Private/NullHttp.h
--- a/Engine/Source/Runtime/Online/HTTP/Private/NullHttp.h
+++ b/Engine/Source/Runtime/Online/HTTP/Private/NullHttp.h
@@ -14,7 +14,7 @@ class FNullHttpRequest : public FHttpRequestImpl
 public:
 
 	// IHttpBase
-	virtual FString GetURL() const override;
+	virtual const FString& GetURL() const override;
 	virtual FString GetURLParameter(const FString& ParameterName) const override;
 	virtual FString GetHeader(const FString& HeaderName) const override;
 	virtual TArray<FString> GetAllHeaders() const override;	
@@ -80,7 +80,7 @@ private:
 class FNullHttpResponse : public IHttpResponse
 {
 	// IHttpBase 
-	virtual FString GetURL() const override;
+	virtual const FString& GetURL() const override;
 	virtual FString GetURLParameter(const FString& ParameterName) const override;
 	virtual FString GetHeader(const FString& HeaderName) const override;
 	virtual TArray<FString> GetAllHeaders() const override;	
@@ -96,4 +96,5 @@ class FNullHttpResponse : public IHttpResponse
 
 private:
 	TArray<uint8> Payload;
+	FString URL;
 };
diff --git a/Engine/Source/Runtime/Online/HTTP/Public/GenericPlatform/HttpRequestCommon.h b/Engine/Source/Runtime/Online/HTTP/Public/GenericPlatform/HttpRequestCommon.h
--- a/Engine/Source/Runtime/Online/HTTP/Public/GenericPlatform/HttpRequestCommon.h
+++ b/Engine/Source/Runtime/Online/HTTP/Public/GenericPlatform/HttpRequestCommon.h
@@ -58,6 +58,9 @@ public:
 	HTTP_API void StartWaitingInQueue();
 	HTTP_API float GetTimeStartedWaitingInQueue() const;
 
+	HTTP_API virtual void SetURL(const FString& InURL) override;
+	HTTP_API virtual const FString& GetURL() const override;
+
 protected:
 	/**
 	 * Check if this request is valid or allowed, before actually process the request
@@ -162,6 +165,8 @@ protected:
 	/** Record how long it take to connect to the endpoint */
 	double ConnectTime = -1.0;
 
+	FString URL;
+
 	/** Cache the effective URL. When redirected, it will be different with original URL */
 	FString EffectiveURL;
 
diff --git a/Engine/Source/Runtime/Online/HTTP/Public/GenericPlatform/HttpResponseCommon.h b/Engine/Source/Runtime/Online/HTTP/Public/GenericPlatform/HttpResponseCommon.h
--- a/Engine/Source/Runtime/Online/HTTP/Public/GenericPlatform/HttpResponseCommon.h
+++ b/Engine/Source/Runtime/Online/HTTP/Public/GenericPlatform/HttpResponseCommon.h
@@ -20,7 +20,7 @@ PRAGMA_ENABLE_DEPRECATION_WARNINGS
 
 	// IHttpBase
 	HTTP_API virtual FString GetURLParameter(const FString& ParameterName) const override;
-	HTTP_API virtual FString GetURL() const override;
+	HTTP_API virtual const FString& GetURL() const override;
 	HTTP_API virtual const FString& GetEffectiveURL() const override;
 	HTTP_API virtual EHttpRequestStatus::Type GetStatus() const override;
 	HTTP_API virtual EHttpFailureReason GetFailureReason() const override;
diff --git a/Engine/Source/Runtime/Online/HTTP/Public/HttpRequestAdapter.h b/Engine/Source/Runtime/Online/HTTP/Public/HttpRequestAdapter.h
--- a/Engine/Source/Runtime/Online/HTTP/Public/HttpRequestAdapter.h
+++ b/Engine/Source/Runtime/Online/HTTP/Public/HttpRequestAdapter.h
@@ -17,7 +17,7 @@ public:
 	HTTP_API ~FHttpRequestAdapterBase();
 
 	// IHttpRequest interface
-	HTTP_API virtual FString GetURL() const override;
+	HTTP_API virtual const FString& GetURL() const override;
 	HTTP_API virtual FString GetURLParameter(const FString& ParameterName) const override;
 	HTTP_API virtual FString GetHeader(const FString& HeaderName) const override;
 	HTTP_API virtual TArray<FString> GetAllHeaders() const override;
diff --git a/Engine/Source/Runtime/Online/HTTP/Public/Interfaces/IHttpBase.h b/Engine/Source/Runtime/Online/HTTP/Public/Interfaces/IHttpBase.h
--- a/Engine/Source/Runtime/Online/HTTP/Public/Interfaces/IHttpBase.h
+++ b/Engine/Source/Runtime/Online/HTTP/Public/Interfaces/IHttpBase.h
@@ -89,13 +89,12 @@ inline const TCHAR* LexToString(EHttpFailureReason HttpFailureReason)
 class IHttpBase
 {
 public:
-
 	/**
 	 * Get the URL used to send the request.
 	 *
 	 * @return the URL string.
 	 */
-	virtual FString GetURL() const = 0;
+	virtual const FString& GetURL() const = 0;
 
 	/**
 	 * Get the effective URL in case of redirected. If not redirected, it's the same as GetURL
