Address delays/lockups on canceling BPT installs.
1 -- The abort thread waits on the background threads instead of letting the constructor thread wait on them, holding up the UI thread until async jobs complete.
2 -- Clear out the Cloud read jobs so that the read can complete, letting the installer shut down.

#rnx

diff --git a/Engine/Source/Runtime/Online/BuildPatchServices/Private/BuildPatchFileConstructor.cpp b/Engine/Source/Runtime/Online/BuildPatchServices/Private/BuildPatchFileConstructor.cpp
--- a/Engine/Source/Runtime/Online/BuildPatchServices/Private/BuildPatchFileConstructor.cpp
+++ b/Engine/Source/Runtime/Online/BuildPatchServices/Private/BuildPatchFileConstructor.cpp
@@ -1355,6 +1355,12 @@ FBuildPatchFileConstructor::~FBuildPatchFileConstructor()
 	// Wait for threads to shut down.
 	Abort();
 	
+	
+	for (int32 ThreadIndex = 0; ThreadIndex < Threads.Num(); ThreadIndex++)
+	{
+		ThreadCompleteEvents[ThreadIndex]->Wait();
+	}
+
 	for (int32 ThreadIndex = 0; ThreadIndex < Threads.Num(); ThreadIndex++)
 	{
 		FPlatformProcess::ReturnSynchEventToPool(ThreadWakeups[ThreadIndex]);
@@ -1396,11 +1402,6 @@ void FBuildPatchFileConstructor::Abort()
 		ThreadWakeups[ThreadIndex]->Trigger();
 	}
 
-	for (int32 ThreadIndex = 0; ThreadIndex < Threads.Num(); ThreadIndex++)
-	{
-		ThreadCompleteEvents[ThreadIndex]->Wait();
-	}
-
 }
 
 void FBuildPatchFileConstructor::SetChunkLocation(const FGuid& InGuid, EConstructorChunkLocation InNewLocation)
@@ -2664,6 +2665,9 @@ bool FBuildPatchFileConstructor::ConstructFileFromChunks(const FString& BuildFil
 				if (bSuccess &&
 					bShouldAbort)
 				{
+					// The first time we catch an abort, inform the cloud source.
+					CloudSource->Abort();
+
 					ConstructionError = EConstructionError::Aborted;
 					bSuccess = false;
 				}
diff --git a/Engine/Source/Runtime/Online/BuildPatchServices/Private/Installer/CloudChunkSource.cpp b/Engine/Source/Runtime/Online/BuildPatchServices/Private/Installer/CloudChunkSource.cpp
--- a/Engine/Source/Runtime/Online/BuildPatchServices/Private/Installer/CloudChunkSource.cpp
+++ b/Engine/Source/Runtime/Online/BuildPatchServices/Private/Installer/CloudChunkSource.cpp
@@ -890,19 +890,7 @@ namespace BuildPatchServices
 
 		virtual ~FConstructorCloudChunkSource()
 		{
-			// Release any unqueued reads.
-			for (FCloudRead* Read : ReadQueue)
-			{
-				delete Read;
-			}
-
-			// Abort any downloads.
-			for (TPair<int32, FCloudRead*> Pair : InFlightDownloads)
-			{
-				FCloudRead* Read = Pair.Value;
-				DownloadService->RequestAbandon(Read->RequestId);
-				delete Read;
-			}
+			Abort();
 		}
 		
 		virtual FRequestProcessFn CreateRequest(const FGuid& DataId, FMutableMemoryView DestinationBuffer, void* UserPtr, FChunkRequestCompleteDelegate CompleteFn) override
@@ -929,6 +917,7 @@ namespace BuildPatchServices
 			return 0;
 		}
 
+		virtual void Abort() override;
 		virtual void Tick(bool bStartNewDownloads, uint32& OutTimeToNextTickMs) override;
 		virtual void SetWakeupFunction(TUniqueFunction<void()>&& InWakeupMainThreadFn) override
 		{
@@ -1001,6 +990,28 @@ namespace BuildPatchServices
 		uint64 TotalBytes = 0;
 	};
 
+	// Must be called from the same thread as Tick. Can be called multiple times.
+	void FConstructorCloudChunkSource::Abort()
+	{
+		// Release any unqueued reads.
+		for (FCloudRead* Read : ReadQueue)
+		{
+			Read->CompleteFn.Execute(Read->DataId, true, false, Read->UserPtr);
+			delete Read;
+		}
+		ReadQueue.Empty();
+
+		// Abort any downloads.
+		for (TPair<int32, FCloudRead*> Pair : InFlightDownloads)
+		{
+			FCloudRead* Read = Pair.Value;
+			DownloadService->RequestAbandon(Read->RequestId);
+			Read->CompleteFn.Execute(Read->DataId, true, false, Read->UserPtr);
+			delete Read;
+		}
+		InFlightDownloads.Empty();
+	}
+
 	//
 	// Since we need to have a lot of requests in flight at the same time _and_ we need to
 	// synchronize across them to manage how many are in flight at any given moment, we have
diff --git a/Engine/Source/Runtime/Online/BuildPatchServices/Private/Installer/CloudChunkSource.h b/Engine/Source/Runtime/Online/BuildPatchServices/Private/Installer/CloudChunkSource.h
--- a/Engine/Source/Runtime/Online/BuildPatchServices/Private/Installer/CloudChunkSource.h
+++ b/Engine/Source/Runtime/Online/BuildPatchServices/Private/Installer/CloudChunkSource.h
@@ -57,6 +57,9 @@ namespace BuildPatchServices
 	public:
 		virtual ~IConstructorCloudChunkSource() {}
 
+
+		virtual void Abort() = 0;
+
 		virtual void Tick(bool bStartNewDownloads, uint32& OutTimeToNextTickMs) = 0;
 
 		// When downloads are complete we need a way to alert the main thread.
