Fix code on Apple platforms to comply with WebTest "Streaming http upload from invalid file by PUT will fail"

diff --git a/Engine/Source/Runtime/Online/HTTP/Private/Apple/AppleHttp.cpp b/Engine/Source/Runtime/Online/HTTP/Private/Apple/AppleHttp.cpp
--- a/Engine/Source/Runtime/Online/HTTP/Private/Apple/AppleHttp.cpp
+++ b/Engine/Source/Runtime/Online/HTTP/Private/Apple/AppleHttp.cpp
@@ -668,7 +668,7 @@ bool FAppleHttpRequest::SetContentAsStreamedFile(const FString& Filename)
 	else
 	{
 		UE_LOG(LogHttp, Warning, TEXT("FAppleHttpRequest::SetContentAsStreamedFile failed to get file size errno: %d: %s"), errno, UTF8_TO_TCHAR(strerror(errno)));
-		StreamedContentSource.Emplace<FNoStreamSource>();
+		StreamedContentSource.Emplace<FInvalidStreamSource>();
 		ContentBytesLength = 0;
 		return false;
 	}
@@ -726,6 +726,11 @@ struct FAppleHttpRequest::FAppleHttpStreamFactory
 		return nil;
 	}
 	
+	NSInputStream *operator()(FInvalidStreamSource)
+	{
+		return nil;
+	}
+
 	NSInputStream *operator()(const FString& Filename)
 	{
 		return [NSInputStream inputStreamWithFileAtPath: Filename.GetNSString()];
diff --git a/Engine/Source/Runtime/Online/HTTP/Private/Apple/AppleHttp.h b/Engine/Source/Runtime/Online/HTTP/Private/Apple/AppleHttp.h
--- a/Engine/Source/Runtime/Online/HTTP/Private/Apple/AppleHttp.h
+++ b/Engine/Source/Runtime/Online/HTTP/Private/Apple/AppleHttp.h
@@ -116,13 +116,14 @@ private:
 	NSURLSessionTask* Task;
     
 	struct FAppleHttpStreamFactory;
-	using FNoStreamSource = FEmptyVariantState;
+	struct FNoStreamSource{};
+	struct FInvalidStreamSource{};
 	/** Source to create stream from
 		FNoStreamedSource: No streamed data
 		FString: Filename set from SetContentAsStreamedFile
 		TSharedRef<FArchive>: Stream set from SetContentFromStream
 	 */
-	TVariant<FNoStreamSource, FString, TSharedRef<FArchive>> StreamedContentSource;
+	TVariant<FNoStreamSource, FInvalidStreamSource, FString, TSharedRef<FArchive>> StreamedContentSource;
 
 	/** The request payload length in bytes. This must be tracked separately for a file stream */
 	uint64 ContentBytesLength;
