Expose some important occurrences out of BPT to client code for ui/logging/analytics.

#rnx

diff --git a/Engine/Source/Runtime/Online/BuildPatchServices/Private/BuildPatchFileConstructor.cpp b/Engine/Source/Runtime/Online/BuildPatchServices/Private/BuildPatchFileConstructor.cpp
--- a/Engine/Source/Runtime/Online/BuildPatchServices/Private/BuildPatchFileConstructor.cpp
+++ b/Engine/Source/Runtime/Online/BuildPatchServices/Private/BuildPatchFileConstructor.cpp
@@ -26,6 +26,7 @@
 #include "Installer/InstallerError.h"
 #include "Installer/InstallerAnalytics.h"
 #include "Installer/InstallerSharedContext.h"
+#include "Installer/MessagePump.h"
 #include "Templates/Greater.h"
 #include "BuildPatchUtil.h"
 
@@ -1225,7 +1226,7 @@ public:
 FBuildPatchFileConstructor::FBuildPatchFileConstructor(
 	FFileConstructorConfig InConfiguration, IFileSystem* InFileSystem, 
 	IConstructorChunkDbChunkSource* InChunkDbChunkSource, IConstructorCloudChunkSource* InCloudChunkSource, IConstructorInstallChunkSource* InInstallChunkSource, 
-	IChunkReferenceTracker* InChunkReferenceTracker, IInstallerError* InInstallerError, IInstallerAnalytics* InInstallerAnalytics, 
+	IChunkReferenceTracker* InChunkReferenceTracker, IInstallerError* InInstallerError, IInstallerAnalytics* InInstallerAnalytics, IMessagePump* InMessagePump,
 	IFileConstructorStat* InFileConstructorStat, TMap<FGuid, EConstructorChunkLocation>&& InChunkLocations)
 	: Configuration(MoveTemp(InConfiguration))
 	, bIsDownloadStarted(false)
@@ -1242,6 +1243,7 @@ FBuildPatchFileConstructor::FBuildPatchFileConstructor(
 	, ChunkReferenceTracker(InChunkReferenceTracker)
 	, InstallerError(InInstallerError)
 	, InstallerAnalytics(InInstallerAnalytics)
+	, MessagePump(InMessagePump)
 	, FileConstructorStat(InFileConstructorStat)
 	, TotalJobSize(0)
 	, ByteProcessed(0)
@@ -1297,7 +1299,7 @@ FBuildPatchFileConstructor::FBuildPatchFileConstructor(
 	// Create the threads we are allowed to.
 	//
 	bool bHasInstallSource = InstallSource && InstallSource->GetAvailableChunks().Num() != 0;
-	bool bHasChunkDbSource = ChunkDbSource->GetAvailableChunks().Num() != 0;
+	bHasChunkDbSource = ChunkDbSource->GetAvailableChunks().Num() != 0;
 	
 	// For now always create a write thread.
 	int32 ThreadCount = 1;
@@ -2031,6 +2033,13 @@ void FBuildPatchFileConstructor::RequestCompletedFn(const FGuid& Guid, bool bAbo
 
 		}
 
+		if (bHasChunkDbSource)
+		{
+			// Only send this message if we have chunk dbs. The theory is if they don't have chunkdbs then they are expecting
+			// chunks to download from the cloud. If they provide chunkdbs then they are surprised when chunks come from the cloud.
+			MessagePump->SendMessage({FGenericMessage::EType::CloudSourceUsed, Guid});
+		}
+
 		QueueGenericThreadTask(ThreadAssignments[EConstructorChunkLocation::Cloud], CloudSource->CreateRequest(
 			Request->Guid, 
 			Request->ReadBuffer, 
@@ -2343,6 +2352,13 @@ int32 FBuildPatchFileConstructor::StartReadBatch(const FFileManifest& FileManife
 				// If we are already downloading from the cloud, then failures shouldn't try to 
 				// fall back to the cloud.
 				Request.bLaunchedFallback = true;
+
+				if (bHasChunkDbSource)
+				{
+					// Only send this message if we have chunk dbs. The theory is if they don't have chunkdbs then they are expecting
+					// chunks to download from the cloud. If they provide chunkdbs then they are surprised when chunks come from the cloud.
+					MessagePump->SendMessage({FGenericMessage::EType::CloudSourceUsed, RequestPair.Value.Guid});
+				}
 			}
 
 			// We need to kick a request. The question is whether we can request direct
diff --git a/Engine/Source/Runtime/Online/BuildPatchServices/Private/BuildPatchFileConstructor.h b/Engine/Source/Runtime/Online/BuildPatchServices/Private/BuildPatchFileConstructor.h
--- a/Engine/Source/Runtime/Online/BuildPatchServices/Private/BuildPatchFileConstructor.h
+++ b/Engine/Source/Runtime/Online/BuildPatchServices/Private/BuildPatchFileConstructor.h
@@ -42,6 +42,7 @@ namespace BuildPatchServices
 	class IInstallerError;
 	class IInstallerAnalytics;
 	class IFileConstructorStat;
+	class IMessagePump;
 	class IBuildManifestSet;
 	class IBuildInstallerThread;
 	class IConstructorChunkDbChunkSource;
@@ -102,7 +103,7 @@ namespace BuildPatchServices
 		FBuildPatchFileConstructor(
 			FFileConstructorConfig Configuration, IFileSystem* FileSystem, 
 			IConstructorChunkDbChunkSource* ChunkDbChunkSource, IConstructorCloudChunkSource* CloudChunkSource, IConstructorInstallChunkSource* InstallChunkSource, 
-			IChunkReferenceTracker* ChunkReferenceTracker, IInstallerError* InstallerError, IInstallerAnalytics* InstallerAnalytics, 
+			IChunkReferenceTracker* ChunkReferenceTracker, IInstallerError* InstallerError, IInstallerAnalytics* InstallerAnalytics, IMessagePump* MessagePump,
 			IFileConstructorStat* FileConstructorStat, TMap<FGuid, EConstructorChunkLocation>&& ChunkLocations);
 
 		/**
@@ -202,6 +203,9 @@ namespace BuildPatchServices
 		// A flag marking that we have made the initial disk space check following resume logic complete.
 		bool bInitialDiskSizeCheck;
 
+		// If true, the chunkdb source has chunks to provide
+		bool bHasChunkDbSource = false;
+
 		// A flag marking whether we should be paused.
 		FThreadSafeBool bIsPaused;
 
@@ -235,6 +239,8 @@ namespace BuildPatchServices
 		// Pointer to the installer analytics handler.
 		IInstallerAnalytics* InstallerAnalytics;
 
+		IMessagePump* MessagePump = nullptr;
+
 		// Pointer to the stat class.
 		IFileConstructorStat* FileConstructorStat;
 
diff --git a/Engine/Source/Runtime/Online/BuildPatchServices/Private/BuildPatchInstaller.cpp b/Engine/Source/Runtime/Online/BuildPatchServices/Private/BuildPatchInstaller.cpp
--- a/Engine/Source/Runtime/Online/BuildPatchServices/Private/BuildPatchInstaller.cpp
+++ b/Engine/Source/Runtime/Online/BuildPatchServices/Private/BuildPatchInstaller.cpp
@@ -1395,6 +1395,7 @@ namespace BuildPatchServices
 					DownloadService.Get(),
 					ChunkDataSerialization.Get(),
 					DownloadConnectionCount.Get(),
+					MessagePump.Get(),
 					CloudChunkSourceStatistics.Get(),
 					ManifestSet.Get())
 			);
@@ -1421,6 +1422,7 @@ namespace BuildPatchServices
 					ChunkReferenceTracker.Get(), 
 					InstallerError.Get(),
 					InstallerAnalytics.Get(),
+					MessagePump.Get(),
 					FileConstructorStatistics.Get(),
 					MoveTemp(ChunkLocations)
 				)
diff --git a/Engine/Source/Runtime/Online/BuildPatchServices/Private/Installer/CloudChunkSource.cpp b/Engine/Source/Runtime/Online/BuildPatchServices/Private/Installer/CloudChunkSource.cpp
--- a/Engine/Source/Runtime/Online/BuildPatchServices/Private/Installer/CloudChunkSource.cpp
+++ b/Engine/Source/Runtime/Online/BuildPatchServices/Private/Installer/CloudChunkSource.cpp
@@ -874,7 +874,7 @@ namespace BuildPatchServices
 
 	public:
 		FConstructorCloudChunkSource(FConstructorCloudChunkSourceConfig&& InConfiguration, IDownloadService* InDownloadService, IChunkDataSerialization* InChunkDataSerialization,
-			ICloudChunkSourceStat* InCloudChunkSourceStat, IBuildManifestSet* InManifestSet, IDownloadConnectionCount* InDownloadCount)
+			ICloudChunkSourceStat* InCloudChunkSourceStat, IBuildManifestSet* InManifestSet, IDownloadConnectionCount* InDownloadCount, IMessagePump* InMessagePump)
 			: Configuration(MoveTemp(InConfiguration))
 			, DownloadDelegates(MakeShareable(new FDownloadDelegates(*this)))
 			, OnDownloadProgressDelegate(FDownloadProgressDelegate::CreateThreadSafeSP(DownloadDelegates, &FDownloadDelegates::OnDownloadProgress))
@@ -884,6 +884,7 @@ namespace BuildPatchServices
 			, CloudChunkSourceStat(InCloudChunkSourceStat)
 			, ManifestSet(InManifestSet)
 			, DownloadCount(InDownloadCount)
+			, MessagePump(InMessagePump)
 		{
 			CloudDirFailureCount.SetNumZeroed(Configuration.CloudRoots.Num());
 		}
@@ -956,6 +957,7 @@ namespace BuildPatchServices
 		ICloudChunkSourceStat* CloudChunkSourceStat;
 		IBuildManifestSet* ManifestSet;
 		IDownloadConnectionCount* DownloadCount;
+		IMessagePump* MessagePump;
 
 		// Called when we get a completed download and we want to get ticked again.
 		TUniqueFunction<void()> WakeupMainThreadFn;
@@ -1139,6 +1141,8 @@ namespace BuildPatchServices
 
 						CurrentBestCloudDir = MinAtIndex;
 						UE_LOG(LogCloudChunkSource, Warning, TEXT("CDN %s failed download, updating CDN selection to: %s"), *Configuration.CloudRoots[Read->CloudDirUsed], *Configuration.CloudRoots[MinAtIndex]);
+
+						MessagePump->SendMessage({FGenericMessage::EType::CDNDownloadFailed, Read->DataId, Configuration.CloudRoots[Read->CloudDirUsed], Configuration.CloudRoots[MinAtIndex]});
 					}
 					ChunkSuccessTracker.AddFail();
 
@@ -1279,9 +1283,9 @@ namespace BuildPatchServices
 
 
 	IConstructorCloudChunkSource* IConstructorCloudChunkSource::CreateCloudSource(FConstructorCloudChunkSourceConfig&& Configuration, IDownloadService* InDownloadService, 
-		IChunkDataSerialization* InChunkDataSerialization, IDownloadConnectionCount* InDownloadCount, ICloudChunkSourceStat* InCloudChunkSourceStat, 
+		IChunkDataSerialization* InChunkDataSerialization, IDownloadConnectionCount* InDownloadCount, IMessagePump* InMessagePump, ICloudChunkSourceStat* InCloudChunkSourceStat, 
 		IBuildManifestSet* InManifestSet)
 	{
-		return new FConstructorCloudChunkSource(MoveTemp(Configuration), InDownloadService, InChunkDataSerialization, InCloudChunkSourceStat, InManifestSet, InDownloadCount);
+		return new FConstructorCloudChunkSource(MoveTemp(Configuration), InDownloadService, InChunkDataSerialization, InCloudChunkSourceStat, InManifestSet, InDownloadCount, InMessagePump);
 	}
 }
\ No newline at end of file
diff --git a/Engine/Source/Runtime/Online/BuildPatchServices/Private/Installer/CloudChunkSource.h b/Engine/Source/Runtime/Online/BuildPatchServices/Private/Installer/CloudChunkSource.h
--- a/Engine/Source/Runtime/Online/BuildPatchServices/Private/Installer/CloudChunkSource.h
+++ b/Engine/Source/Runtime/Online/BuildPatchServices/Private/Installer/CloudChunkSource.h
@@ -69,7 +69,7 @@ namespace BuildPatchServices
 		virtual void PostRequiredByteCount(uint64 InDownloadExpected) = 0;
 
 		static IConstructorCloudChunkSource* CreateCloudSource(FConstructorCloudChunkSourceConfig&& Configuration, IDownloadService* InDownloadService, 
-			IChunkDataSerialization* InChunkDataSerialization, IDownloadConnectionCount* InDownloadCount, ICloudChunkSourceStat* InCloudChunkSourceStat, 
+			IChunkDataSerialization* InChunkDataSerialization, IDownloadConnectionCount* InDownloadCount, IMessagePump* InMessagePump, ICloudChunkSourceStat* InCloudChunkSourceStat, 
 			IBuildManifestSet* InManifestSet);
 	};
 
diff --git a/Engine/Source/Runtime/Online/BuildPatchServices/Private/Installer/MessagePump.cpp b/Engine/Source/Runtime/Online/BuildPatchServices/Private/Installer/MessagePump.cpp
--- a/Engine/Source/Runtime/Online/BuildPatchServices/Private/Installer/MessagePump.cpp
+++ b/Engine/Source/Runtime/Online/BuildPatchServices/Private/Installer/MessagePump.cpp
@@ -8,7 +8,7 @@
 namespace BuildPatchServices
 {
 	// Union of all possible message types.
-	typedef TUnion<FChunkSourceEvent, FInstallationFileAction> FMessageUnion;
+	typedef TUnion<FChunkSourceEvent, FInstallationFileAction, FGenericMessage> FMessageUnion;
 
 	// Queue type for messages.
 	typedef TQueue<FMessageUnion, EQueueMode::Mpsc> FMessageQueue;
@@ -46,6 +46,7 @@ namespace BuildPatchServices
 		// IMessagePump interface begin.
 		virtual void SendMessage(FChunkSourceEvent Message) override;
 		virtual void SendMessage(FInstallationFileAction Message) override;
+		virtual void SendMessage(FGenericMessage Message) override;
 		virtual void SendRequest(FChunkUriRequest Request, TFunction<void(FChunkUriResponse)> OnResponse) override;
 		virtual void PumpMessages() override;
 		virtual void RegisterMessageHandler(FMessageHandler* MessageHandler) override;
@@ -61,6 +62,11 @@ namespace BuildPatchServices
 		std::atomic<uint32> MessageRequests{0};
 	};
 
+	void FMessagePump::SendMessage(FGenericMessage Message)
+	{
+		MessageQueue.Enqueue(FMessageUnion(MoveTemp(Message)));
+	}
+
 	void FMessagePump::SendMessage(FChunkSourceEvent Message)
 	{
 		MessageQueue.Enqueue(FMessageUnion(MoveTemp(Message)));
@@ -96,6 +102,10 @@ namespace BuildPatchServices
 			{
 				continue;
 			}
+			else if (MessagePumpHelpers::TryPump<FGenericMessage>(Handlers, MessageUnion))
+			{
+				continue;
+			}
 		}
 
 		FChunkUriRequestResponse ChunkUriRequestResponse;
diff --git a/Engine/Source/Runtime/Online/BuildPatchServices/Private/Installer/MessagePump.h b/Engine/Source/Runtime/Online/BuildPatchServices/Private/Installer/MessagePump.h
--- a/Engine/Source/Runtime/Online/BuildPatchServices/Private/Installer/MessagePump.h
+++ b/Engine/Source/Runtime/Online/BuildPatchServices/Private/Installer/MessagePump.h
@@ -23,6 +23,9 @@ namespace BuildPatchServices
 		 */
 		virtual void SendMessage(FChunkSourceEvent Message) = 0;
 
+		// See FGenericMessage for types of messages.
+		virtual void SendMessage(FGenericMessage Message) = 0;
+
 		/**
 		 * Sends an installation file action message.
 		 * @param Message   The message to be sent.
diff --git a/Engine/Source/Runtime/Online/BuildPatchServices/Private/Tests/Mock/MessagePump.mock.h b/Engine/Source/Runtime/Online/BuildPatchServices/Private/Tests/Mock/MessagePump.mock.h
--- a/Engine/Source/Runtime/Online/BuildPatchServices/Private/Tests/Mock/MessagePump.mock.h
+++ b/Engine/Source/Runtime/Online/BuildPatchServices/Private/Tests/Mock/MessagePump.mock.h
@@ -19,6 +19,10 @@ namespace BuildPatchServices
 		{
 		}
 
+		virtual void SendMessage(FGenericMessage Message) override
+		{
+		}
+
 		virtual void PumpMessages() override
 		{
 		}
diff --git a/Engine/Source/Runtime/Online/BuildPatchServices/Public/BuildPatchMessage.h b/Engine/Source/Runtime/Online/BuildPatchServices/Public/BuildPatchMessage.h
--- a/Engine/Source/Runtime/Online/BuildPatchServices/Public/BuildPatchMessage.h
+++ b/Engine/Source/Runtime/Online/BuildPatchServices/Public/BuildPatchMessage.h
@@ -2,6 +2,7 @@
 #pragma once
 
 #include "Containers/UnrealString.h"
+#include "Misc/Guid.h"
 #include "Templates/Function.h"
 
 namespace BuildPatchServices
@@ -65,6 +66,30 @@ namespace BuildPatchServices
 		FString Uri;
 	};
 
+	struct FGenericMessage
+	{
+		// Set of informational messages the installer can send to the client to update
+		// UI or otherwise react.
+		enum class EType : uint32
+		{
+			// Posted when a chunk is requested of the cloud source ONLY WHEN CHUNKDBS ARE PRESENT. This is useful if the installer
+			// is not expected to be downloading anything independently and you want to log this case.
+			// Note that even when chunks are fully provided, cancellation and resumption can lose "harvested"
+			// chunks that were resident only in non-persistent backing store causing a download
+			CloudSourceUsed,
+
+			// Posted when a CDN/CloudDir has failed a download and is dropped in priority.
+			// Payload1 = CDN that failed,
+			// Payload2 = CDN that is considered "best" after this failure.
+			CDNDownloadFailed,
+		};
+
+		EType Type;
+		FGuid ChunkId;
+		FString Payload1, Payload2;
+	};
+	
+
 	// Flag for which Requests a message handler expects to receive, allows for internal implementation optimisation.
 	enum class EMessageRequests : uint32
 	{
@@ -89,6 +114,13 @@ namespace BuildPatchServices
 
 		virtual ~FMessageHandler() = default;
 
+		/**
+		* Handle generic information posting. This can be called from any thread and should not
+		* take appreciable time as it blocks further installation.
+		* See FGenericMessage for payload contents and types
+		*/
+		virtual void HandleMessage(const FGenericMessage& Message) {}
+
 		/**
 		 * Handles a chunk source event message.
 		 * @param Message   The message to be handled.
