Move http task start to StartThreadedRequest for Apple platforms

diff --git a/Engine/Source/Runtime/Online/HTTP/Private/Apple/AppleHttp.cpp b/Engine/Source/Runtime/Online/HTTP/Private/Apple/AppleHttp.cpp
--- a/Engine/Source/Runtime/Online/HTTP/Private/Apple/AppleHttp.cpp
+++ b/Engine/Source/Runtime/Online/HTTP/Private/Apple/AppleHttp.cpp
@@ -11,6 +11,14 @@
 #include "Misc/App.h"
 #include "Misc/Base64.h"
 
+// UE-228716
+TAutoConsoleVariable<bool> CVarAppleHttpStartTaskInHttpThreadEnabled(
+	TEXT("AppleHttp.StartTaskInHttpThreadEnabled"),
+	true,
+	TEXT("Enables this implementation to start the task in the http thread like the rest of implementations"),
+	ECVF_SaveForNextBoot
+);
+
 /**
  * Class to hold data from delegate implementation notifications.
  */
@@ -494,6 +502,7 @@ FAppleHttpRequest::FAppleHttpRequest(NSURLSession* InSession)
 ,	ContentBytesLength(0)
 ,	LastReportedBytesWritten(0)
 ,	LastReportedBytesRead(0)
+,	bStartTaskInHttpThread(CVarAppleHttpStartTaskInHttpThreadEnabled.GetValueOnAnyThread())
 {
 	bUsePlatformActivityTimeout = false;
 	
@@ -716,6 +725,14 @@ bool FAppleHttpRequest::ProcessRequest()
 	
 	StartProcessTime = FPlatformTime::Seconds();
 
+	if (bStartTaskInHttpThread)
+	{
+		SetStatus(EHttpRequestStatus::Processing);
+		SetFailureReason(EHttpFailureReason::None);
+		// AppleEventLoop sets a delegate into the response to be able to notify events
+		InitResponse();
+		FHttpModule::Get().GetHttpManager().AddThreadedRequest(SharedThis(this));
+	}
 	return true;
 }
 
@@ -788,32 +805,39 @@ bool FAppleHttpRequest::SetupRequest()
 		return false;
 	}
 	
-	Task = [Session dataTaskWithRequest: Request];
-	
-	if (Task != nil)
+	if (bStartTaskInHttpThread)
 	{
-		SetStatus(EHttpRequestStatus::Processing);
-		SetFailureReason(EHttpFailureReason::None);
-
-		InitResponse();
-
-		// Both Task and Response keep a strong reference to the delegate
-		TSharedPtr<FAppleHttpResponse> Response = StaticCastSharedPtr<FAppleHttpResponse>(ResponseCommon);
-		Task.delegate = Response->ResponseDelegate;
-
-		//Setup delegates before starting the request
-		FHttpModule::Get().GetHttpManager().AddThreadedRequest(SharedThis(this));
-
-		[[Task retain] resume];
-		UE_LOG(LogHttp, Verbose, TEXT("Task started %p"), this);
 		return true;
 	}
 	else
 	{
-		UE_LOG(LogHttp, Warning, TEXT("SetupRequest failed. Could not initialize NSURLSessionTask."));
-		SetStatus(EHttpRequestStatus::Failed);
-		SetFailureReason(EHttpFailureReason::ConnectionError);
-		return false;
+		Task = [Session dataTaskWithRequest: Request];
+	
+		if (Task != nil)
+		{
+			SetStatus(EHttpRequestStatus::Processing);
+			SetFailureReason(EHttpFailureReason::None);
+
+			InitResponse();
+
+			// Both Task and Response keep a strong reference to the delegate
+			TSharedPtr<FAppleHttpResponse> Response = StaticCastSharedPtr<FAppleHttpResponse>(ResponseCommon);
+			Task.delegate = Response->ResponseDelegate;
+
+			//Setup delegates before starting the request
+			FHttpModule::Get().GetHttpManager().AddThreadedRequest(SharedThis(this));
+
+			[[Task retain] resume];
+			UE_LOG(LogHttp, Verbose, TEXT("Task started %p"), this);
+			return true;
+		}
+		else
+		{
+			UE_LOG(LogHttp, Warning, TEXT("SetupRequest failed. Could not initialize NSURLSessionTask."));
+			SetStatus(EHttpRequestStatus::Failed);
+			SetFailureReason(EHttpFailureReason::ConnectionError);
+			return false;
+		}	
 	}
 }
 
@@ -938,6 +962,36 @@ void FAppleHttpRequest::CheckProgressDelegate()
 
 bool FAppleHttpRequest::StartThreadedRequest()
 {
+	if (bStartTaskInHttpThread)
+	{
+		if (bCanceled)
+		{
+			UE_LOG(LogHttp, Verbose, TEXT("StartThreadedRequest ignored because request has been canceled. %s url=%s"), *GetVerb(), *GetURL());
+			return false;
+		}
+
+		if (Task != nil)
+		{
+			UE_LOG(LogHttp, Verbose, TEXT("StartThreadedRequest ignored because task was already in progress. %s url=%s"), *GetVerb(), *GetURL());
+			return false;
+		}
+	
+		Task = [Session dataTaskWithRequest: Request];
+	
+		if (Task == nil)
+		{
+			UE_LOG(LogHttp, Warning, TEXT("StartThreadedRequest failed. Could not initialize NSURLSessionTask."));
+			SetStatus(EHttpRequestStatus::Failed);
+			SetFailureReason(EHttpFailureReason::ConnectionError);
+			return false;
+		}
+
+		// Both Task and Response keep a strong reference to the delegate
+		TSharedPtr<FAppleHttpResponse> Response = StaticCastSharedPtr<FAppleHttpResponse>(ResponseCommon);
+		Task.delegate = Response->ResponseDelegate;
+
+		[[Task retain] resume];
+	}
 	return true;
 }
 
diff --git a/Engine/Source/Runtime/Online/HTTP/Private/Apple/AppleHttp.h b/Engine/Source/Runtime/Online/HTTP/Private/Apple/AppleHttp.h
--- a/Engine/Source/Runtime/Online/HTTP/Private/Apple/AppleHttp.h
+++ b/Engine/Source/Runtime/Online/HTTP/Private/Apple/AppleHttp.h
@@ -136,6 +136,9 @@ private:
 
 	/** Last reported bytesread */
 	int32 LastReportedBytesRead;
+	
+	// UE-228716: This request should follow same logic as the rest of implementations
+	const bool bStartTaskInHttpThread;
 };
 
 @class FAppleHttpResponseDelegate;
