Refactor: Move HttpRequestCommon and HttpResponseCommon to Private.

diff --git a/Engine/Source/Runtime/Online/HTTP/Private/Apple/AppleEventLoopHttpThread.cpp b/Engine/Source/Runtime/Online/HTTP/Private/Apple/AppleEventLoopHttpThread.cpp
--- a/Engine/Source/Runtime/Online/HTTP/Private/Apple/AppleEventLoopHttpThread.cpp
+++ b/Engine/Source/Runtime/Online/HTTP/Private/Apple/AppleEventLoopHttpThread.cpp
@@ -4,7 +4,6 @@
 
 #include "Apple/AppleHttp.h" 
 
-PRAGMA_DISABLE_DEPRECATION_WARNINGS
 bool FAppleEventLoopHttpThread::StartThreadedRequest(FHttpRequestCommon* Request)
 {
 	FHttpResponsePtr Response = Request->GetResponse();
@@ -31,7 +30,6 @@ void FAppleEventLoopHttpThread::CreateEventLoop()
 	};
 	EventLoop.Emplace(MoveTemp(EventLoopParams));
 }
-PRAGMA_ENABLE_DEPRECATION_WARNINGS
 
 void FAppleEventLoopHttpThread::DestroyEventLoop()
 {
diff --git a/Engine/Source/Runtime/Online/HTTP/Private/Apple/AppleEventLoopHttpThread.h b/Engine/Source/Runtime/Online/HTTP/Private/Apple/AppleEventLoopHttpThread.h
--- a/Engine/Source/Runtime/Online/HTTP/Private/Apple/AppleEventLoopHttpThread.h
+++ b/Engine/Source/Runtime/Online/HTTP/Private/Apple/AppleEventLoopHttpThread.h
@@ -87,10 +87,8 @@ class FAppleEventLoopHttpThread
 {
 protected:
 	//~ Begin FHttpThread Interface
-PRAGMA_DISABLE_DEPRECATION_WARNINGS
 	virtual bool StartThreadedRequest(FHttpRequestCommon* Request) override;
 	virtual void CompleteThreadedRequest(FHttpRequestCommon* Request) override;
-PRAGMA_ENABLE_DEPRECATION_WARNINGS
 	//~ End FHttpThread Interface
 protected:
 	virtual void CreateEventLoop() override;
diff --git a/Engine/Source/Runtime/Online/HTTP/Private/Apple/AppleHttp.h b/Engine/Source/Runtime/Online/HTTP/Private/Apple/AppleHttp.h
--- a/Engine/Source/Runtime/Online/HTTP/Private/Apple/AppleHttp.h
+++ b/Engine/Source/Runtime/Online/HTTP/Private/Apple/AppleHttp.h
@@ -18,9 +18,7 @@ DECLARE_DELEGATE(FNewAppleHttpEventDelegate);
 /**
  * Apple implementation of an Http request
  */
-PRAGMA_DISABLE_DEPRECATION_WARNINGS
 class FAppleHttpRequest : public FHttpRequestCommon
-PRAGMA_ENABLE_DEPRECATION_WARNINGS
 {
 public:
 	// implementation friends
@@ -68,7 +66,6 @@ public:
 	virtual ~FAppleHttpRequest();
 
 PACKAGE_SCOPE:
-PRAGMA_DISABLE_DEPRECATION_WARNINGS
 	using FHttpRequestCommon::StartProcessTime;
 	using FHttpRequestCommon::ConnectTime;
 	using FHttpRequestCommon::BroadcastResponseHeadersReceived;
@@ -78,7 +75,6 @@ PRAGMA_DISABLE_DEPRECATION_WARNINGS
 	using FHttpRequestCommon::HandleStatusCodeReceived;
 	using FHttpRequestCommon::SetEffectiveURL;
 	using FHttpRequestCommon::PassReceivedDataToStream;
-PRAGMA_ENABLE_DEPRECATION_WARNINGS
 
 	bool IsInitializedWithValidStream() const;
 
@@ -146,9 +142,7 @@ private:
 /**
  * Apple implementation of an Http response
  */
-PRAGMA_DISABLE_DEPRECATION_WARNINGS
 class FAppleHttpResponse : public FHttpResponseCommon
-PRAGMA_ENABLE_DEPRECATION_WARNINGS
 {
 public:
 	//~ Begin IHttpBase Interface
diff --git a/Engine/Source/Runtime/Online/HTTP/Private/Curl/CurlHttp.cpp b/Engine/Source/Runtime/Online/HTTP/Private/Curl/CurlHttp.cpp
--- a/Engine/Source/Runtime/Online/HTTP/Private/Curl/CurlHttp.cpp
+++ b/Engine/Source/Runtime/Online/HTTP/Private/Curl/CurlHttp.cpp
@@ -54,10 +54,7 @@ static int SslCertVerify(int PreverifyOk, X509_STORE_CTX* Context)
 		FCurlHttpRequest* Request = static_cast<FCurlHttpRequest*>(SSL_CTX_get_app_data(SslContext));
 		check(Request);
 
-// FHttpRequestCommon deprecation
-PRAGMA_DISABLE_DEPRECATION_WARNINGS
 		const FString Domain = FPlatformHttp::GetUrlDomain(Request->GetURL());
-PRAGMA_ENABLE_DEPRECATION_WARNINGS
 
 		if (!FSslModule::Get().GetCertificateManager().VerifySslCertificates(Context, Domain))
 		{
@@ -86,7 +83,6 @@ static CURLcode sslctx_function(CURL * curl, void * sslctx, void * parm)
 }
 #endif //#if WITH_SSL
 
-PRAGMA_DISABLE_DEPRECATION_WARNINGS
 FCurlHttpRequest::FCurlHttpRequest()
 	: EasyHandle(nullptr)
 	, HeaderList(nullptr)
@@ -1396,6 +1392,5 @@ FString FCurlHttpResponse::GetContentAsString() const
 	return FString(TCHARData.Length(), TCHARData.Get());
 }
 
-PRAGMA_ENABLE_DEPRECATION_WARNINGS
 
 #endif //WITH_CURL
\ No newline at end of file
diff --git a/Engine/Source/Runtime/Online/HTTP/Private/Curl/CurlHttp.h b/Engine/Source/Runtime/Online/HTTP/Private/Curl/CurlHttp.h
--- a/Engine/Source/Runtime/Online/HTTP/Private/Curl/CurlHttp.h
+++ b/Engine/Source/Runtime/Online/HTTP/Private/Curl/CurlHttp.h
@@ -135,9 +135,7 @@ namespace
 /**
  * Curl implementation of an HTTP request
  */
-PRAGMA_DISABLE_DEPRECATION_WARNINGS
 class FCurlHttpRequest : public FHttpRequestCommon
-PRAGMA_ENABLE_DEPRECATION_WARNINGS
 {
 public:
 
@@ -409,9 +407,7 @@ private:
 /**
  * Curl implementation of an HTTP response
  */
-PRAGMA_DISABLE_DEPRECATION_WARNINGS
 class FCurlHttpResponse : public FHttpResponseCommon
-PRAGMA_ENABLE_DEPRECATION_WARNINGS
 {
 public:
 	// implementation friends
diff --git a/Engine/Source/Runtime/Online/HTTP/Private/Curl/CurlHttpThread.cpp b/Engine/Source/Runtime/Online/HTTP/Private/Curl/CurlHttpThread.cpp
--- a/Engine/Source/Runtime/Online/HTTP/Private/Curl/CurlHttpThread.cpp
+++ b/Engine/Source/Runtime/Online/HTTP/Private/Curl/CurlHttpThread.cpp
@@ -46,9 +46,7 @@ void FCurlHttpThread::HttpThreadTick(float DeltaSeconds)
 					CURL* CompletedHandle = Message->easy_handle;
 					curl_multi_remove_handle(FCurlHttpManager::GMultiHandle, CompletedHandle);
 
-PRAGMA_DISABLE_DEPRECATION_WARNINGS
 					FHttpRequestCommon** Request = HandlesToRequests.Find(CompletedHandle);
-PRAGMA_ENABLE_DEPRECATION_WARNINGS
 					if (Request)
 					{
 						FCurlHttpRequest* CurlRequest = static_cast<FCurlHttpRequest*>(*Request);
@@ -70,7 +68,6 @@ PRAGMA_ENABLE_DEPRECATION_WARNINGS
 	FLegacyHttpThread::HttpThreadTick(DeltaSeconds);
 }
 
-PRAGMA_DISABLE_DEPRECATION_WARNINGS
 bool FCurlHttpThread::StartThreadedRequest(FHttpRequestCommon* Request)
 {
 	FCurlHttpRequest* CurlRequest = static_cast<FCurlHttpRequest*>(Request);
@@ -110,6 +107,5 @@ void FCurlHttpThread::CompleteThreadedRequest(FHttpRequestCommon* Request)
 
 	CurlRequest->CleanupRequestHttpThread();
 }
-PRAGMA_ENABLE_DEPRECATION_WARNINGS
 
 #endif
diff --git a/Engine/Source/Runtime/Online/HTTP/Private/Curl/CurlHttpThread.h b/Engine/Source/Runtime/Online/HTTP/Private/Curl/CurlHttpThread.h
--- a/Engine/Source/Runtime/Online/HTTP/Private/Curl/CurlHttpThread.h
+++ b/Engine/Source/Runtime/Online/HTTP/Private/Curl/CurlHttpThread.h
@@ -36,7 +36,6 @@ public:
 protected:
 	//~ Begin FHttpThread Interface
 	virtual void HttpThreadTick(float DeltaSeconds) override;
-PRAGMA_DISABLE_DEPRECATION_WARNINGS
 	virtual bool StartThreadedRequest(FHttpRequestCommon* Request) override;
 	virtual void CompleteThreadedRequest(FHttpRequestCommon* Request) override;
 	//~ End FHttpThread Interface
@@ -44,7 +43,6 @@ protected:
 
 	/** Mapping of libcurl easy handles to HTTP requests */
 	TMap<CURL*, FHttpRequestCommon*> HandlesToRequests;
-PRAGMA_ENABLE_DEPRECATION_WARNINGS
 };
 
 
diff --git a/Engine/Source/Runtime/Online/HTTP/Private/Curl/CurlMultiPollEventLoopHttpThread.cpp b/Engine/Source/Runtime/Online/HTTP/Private/Curl/CurlMultiPollEventLoopHttpThread.cpp
--- a/Engine/Source/Runtime/Online/HTTP/Private/Curl/CurlMultiPollEventLoopHttpThread.cpp
+++ b/Engine/Source/Runtime/Online/HTTP/Private/Curl/CurlMultiPollEventLoopHttpThread.cpp
@@ -68,7 +68,6 @@ FCurlMultiPollEventLoopHttpThread::FCurlMultiPollEventLoopHttpThread()
 {
 }
 
-PRAGMA_DISABLE_DEPRECATION_WARNINGS
 bool FCurlMultiPollEventLoopHttpThread::StartThreadedRequest(FHttpRequestCommon* Request)
 {
 	FCurlHttpRequest* CurlRequest = static_cast<FCurlHttpRequest*>(Request);
@@ -108,7 +107,6 @@ void FCurlMultiPollEventLoopHttpThread::CompleteThreadedRequest(FHttpRequestComm
 
 	CurlRequest->CleanupRequestHttpThread();
 }
-PRAGMA_ENABLE_DEPRECATION_WARNINGS
 
 void FCurlMultiPollEventLoopHttpThread::CreateEventLoop()
 {
@@ -158,9 +156,7 @@ void FCurlMultiPollEventLoopHttpThread::ProcessCurlRequests()
 			CURL* CompletedHandle = Message->easy_handle;
 			curl_multi_remove_handle(FCurlHttpManager::GMultiHandle, CompletedHandle);
 
-PRAGMA_DISABLE_DEPRECATION_WARNINGS
 			FHttpRequestCommon** Request = HandlesToRequests.Find(CompletedHandle);
-PRAGMA_ENABLE_DEPRECATION_WARNINGS
 			if (Request)
 			{
 				FCurlHttpRequest* CurlRequest = static_cast<FCurlHttpRequest*>(*Request);
@@ -181,10 +177,8 @@ PRAGMA_ENABLE_DEPRECATION_WARNINGS
 	// If any requests completed, immediately process requests to handle completion event.
 	if (CompletedRequest)
 	{
-PRAGMA_DISABLE_DEPRECATION_WARNINGS
 		TArray<FHttpRequestCommon*> RequestsToCancel;
 		TArray<FHttpRequestCommon*> RequestsToComplete;
-PRAGMA_ENABLE_DEPRECATION_WARNINGS
 		Process(RequestsToCancel, RequestsToComplete);
 	}
 }
diff --git a/Engine/Source/Runtime/Online/HTTP/Private/Curl/CurlMultiPollEventLoopHttpThread.h b/Engine/Source/Runtime/Online/HTTP/Private/Curl/CurlMultiPollEventLoopHttpThread.h
--- a/Engine/Source/Runtime/Online/HTTP/Private/Curl/CurlMultiPollEventLoopHttpThread.h
+++ b/Engine/Source/Runtime/Online/HTTP/Private/Curl/CurlMultiPollEventLoopHttpThread.h
@@ -77,10 +77,8 @@ public:
 
 protected:
 	//~ Begin FHttpThread Interface
-PRAGMA_DISABLE_DEPRECATION_WARNINGS
 	virtual bool StartThreadedRequest(FHttpRequestCommon* Request) override;
 	virtual void CompleteThreadedRequest(FHttpRequestCommon* Request) override;
-PRAGMA_ENABLE_DEPRECATION_WARNINGS
 	//~ End FHttpThread Interface
 protected:
 	virtual void CreateEventLoop() override;
@@ -93,10 +91,8 @@ protected:
 
 	TOptional<UE::EventLoop::TEventLoop<FCurlMultiPollIOManager>> EventLoop;
 
-PRAGMA_DISABLE_DEPRECATION_WARNINGS
 	/** Mapping of libcurl easy handles to HTTP requests */
 	TMap<CURL*, FHttpRequestCommon*> HandlesToRequests;
-PRAGMA_ENABLE_DEPRECATION_WARNINGS
 };
 
 #endif // WITH_CURL_MULTIPOLL
diff --git a/Engine/Source/Runtime/Online/HTTP/Private/Curl/CurlMultiWaitEventLoopHttpThread.cpp b/Engine/Source/Runtime/Online/HTTP/Private/Curl/CurlMultiWaitEventLoopHttpThread.cpp
--- a/Engine/Source/Runtime/Online/HTTP/Private/Curl/CurlMultiWaitEventLoopHttpThread.cpp
+++ b/Engine/Source/Runtime/Online/HTTP/Private/Curl/CurlMultiWaitEventLoopHttpThread.cpp
@@ -113,7 +113,6 @@ FCurlMultiWaitEventLoopHttpThread::FCurlMultiWaitEventLoopHttpThread()
 {
 }
 
-PRAGMA_DISABLE_DEPRECATION_WARNINGS
 bool FCurlMultiWaitEventLoopHttpThread::StartThreadedRequest(FHttpRequestCommon* Request)
 {
 	FCurlHttpRequest* CurlRequest = static_cast<FCurlHttpRequest*>(Request);
@@ -152,7 +151,6 @@ void FCurlMultiWaitEventLoopHttpThread::CompleteThreadedRequest(FHttpRequestComm
 
 	CurlRequest->CleanupRequestHttpThread();
 }
-PRAGMA_ENABLE_DEPRECATION_WARNINGS
 
 void FCurlMultiWaitEventLoopHttpThread::CreateEventLoop()
 {
@@ -202,9 +200,7 @@ void FCurlMultiWaitEventLoopHttpThread::ProcessCurlRequests()
 			CURL* CompletedHandle = Message->easy_handle;
 			curl_multi_remove_handle(FCurlHttpManager::GMultiHandle, CompletedHandle);
 
-PRAGMA_DISABLE_DEPRECATION_WARNINGS
 			FHttpRequestCommon** Request = HandlesToRequests.Find(CompletedHandle);
-PRAGMA_ENABLE_DEPRECATION_WARNINGS
 			if (Request)
 			{
 				FCurlHttpRequest* CurlRequest = static_cast<FCurlHttpRequest*>(*Request);
@@ -225,10 +221,8 @@ PRAGMA_ENABLE_DEPRECATION_WARNINGS
 	// If any requests completed, immediately process requests to handle completion event.
 	if (CompletedRequest)
 	{
-PRAGMA_DISABLE_DEPRECATION_WARNINGS
 		TArray<FHttpRequestCommon*> RequestsToCancel;
 		TArray<FHttpRequestCommon*> RequestsToComplete;
-PRAGMA_ENABLE_DEPRECATION_WARNINGS
 		Process(RequestsToCancel, RequestsToComplete);
 	}
 }
diff --git a/Engine/Source/Runtime/Online/HTTP/Private/Curl/CurlMultiWaitEventLoopHttpThread.h b/Engine/Source/Runtime/Online/HTTP/Private/Curl/CurlMultiWaitEventLoopHttpThread.h
--- a/Engine/Source/Runtime/Online/HTTP/Private/Curl/CurlMultiWaitEventLoopHttpThread.h
+++ b/Engine/Source/Runtime/Online/HTTP/Private/Curl/CurlMultiWaitEventLoopHttpThread.h
@@ -80,10 +80,8 @@ public:
 
 protected:
 	//~ Begin FHttpThread Interface
-PRAGMA_DISABLE_DEPRECATION_WARNINGS
 	virtual bool StartThreadedRequest(FHttpRequestCommon* Request) override;
 	virtual void CompleteThreadedRequest(FHttpRequestCommon* Request) override;
-PRAGMA_ENABLE_DEPRECATION_WARNINGS
 	//~ End FHttpThread Interface
 protected:
 	virtual void CreateEventLoop() override;
@@ -97,9 +95,7 @@ protected:
 	TOptional<UE::EventLoop::TEventLoop<FCurlMultiWaitIOManager>> EventLoop;
 
 	/** Mapping of libcurl easy handles to HTTP requests */
-PRAGMA_DISABLE_DEPRECATION_WARNINGS
 	TMap<CURL*, FHttpRequestCommon*> HandlesToRequests;
-PRAGMA_ENABLE_DEPRECATION_WARNINGS
 };
 
 #endif // WITH_CURL_MULTIWAIT
diff --git a/Engine/Source/Runtime/Online/HTTP/Private/Curl/CurlSocketEventLoopHttpThread.cpp b/Engine/Source/Runtime/Online/HTTP/Private/Curl/CurlSocketEventLoopHttpThread.cpp
--- a/Engine/Source/Runtime/Online/HTTP/Private/Curl/CurlSocketEventLoopHttpThread.cpp
+++ b/Engine/Source/Runtime/Online/HTTP/Private/Curl/CurlSocketEventLoopHttpThread.cpp
@@ -59,7 +59,6 @@ void FCurlSocketEventLoopHttpThread::HttpThreadTick(float DeltaSeconds)
 	FEventLoopHttpThread::HttpThreadTick(DeltaSeconds);
 }
 
-PRAGMA_DISABLE_DEPRECATION_WARNINGS
 bool FCurlSocketEventLoopHttpThread::StartThreadedRequest(FHttpRequestCommon* Request)
 {
 	FCurlHttpRequest* CurlRequest = static_cast<FCurlHttpRequest*>(Request);
@@ -99,7 +98,6 @@ void FCurlSocketEventLoopHttpThread::CompleteThreadedRequest(FHttpRequestCommon*
 
 	CurlRequest->CleanupRequestHttpThread();
 }
-PRAGMA_ENABLE_DEPRECATION_WARNINGS
 
 void FCurlSocketEventLoopHttpThread::CreateEventLoop()
 {
@@ -248,7 +246,6 @@ void FCurlSocketEventLoopHttpThread::ProcessCurlSocketEvent(curl_socket_t Socket
 	}
 }
 
-PRAGMA_DISABLE_DEPRECATION_WARNINGS
 void FCurlSocketEventLoopHttpThread::ProcessCurlRequests()
 {
 	QUICK_SCOPE_CYCLE_COUNTER(STAT_FCurlSocketEventLoopHttpThread_ProcessCurlRequests);
@@ -296,7 +293,6 @@ void FCurlSocketEventLoopHttpThread::ProcessCurlRequests()
 		Process(RequestsToCancel, RequestsToComplete);
 	}
 }
-PRAGMA_ENABLE_DEPRECATION_WARNINGS
 
 #endif // WITH_CURL_MULTISOCKET
 #endif // WITH_CURL
diff --git a/Engine/Source/Runtime/Online/HTTP/Private/Curl/CurlSocketEventLoopHttpThread.h b/Engine/Source/Runtime/Online/HTTP/Private/Curl/CurlSocketEventLoopHttpThread.h
--- a/Engine/Source/Runtime/Online/HTTP/Private/Curl/CurlSocketEventLoopHttpThread.h
+++ b/Engine/Source/Runtime/Online/HTTP/Private/Curl/CurlSocketEventLoopHttpThread.h
@@ -36,10 +36,8 @@ public:
 protected:
 	//~ Begin FHttpThread Interface
 	virtual void HttpThreadTick(float DeltaSeconds) override;
-PRAGMA_DISABLE_DEPRECATION_WARNINGS
 	virtual bool StartThreadedRequest(FHttpRequestCommon* Request) override;
 	virtual void CompleteThreadedRequest(FHttpRequestCommon* Request) override;
-PRAGMA_ENABLE_DEPRECATION_WARNINGS
 	//~ End FHttpThread Interface
 protected:
 	virtual void CreateEventLoop() override;
@@ -65,9 +63,7 @@ protected:
 	};
 
 	/** Mapping of libcurl easy handles to HTTP requests */
-PRAGMA_DISABLE_DEPRECATION_WARNINGS
 	TMap<CURL*, FHttpRequestCommon*> HandlesToRequests;
-PRAGMA_ENABLE_DEPRECATION_WARNINGS
 
 	UE::EventLoop::FTimerHandle RequestTimeoutTimer;
 
diff --git a/Engine/Source/Runtime/Online/HTTP/Private/EventLoopHttpThread.cpp b/Engine/Source/Runtime/Online/HTTP/Private/EventLoopHttpThread.cpp
--- a/Engine/Source/Runtime/Online/HTTP/Private/EventLoopHttpThread.cpp
+++ b/Engine/Source/Runtime/Online/HTTP/Private/EventLoopHttpThread.cpp
@@ -99,7 +99,6 @@ void FEventLoopHttpThread::UpdateConfigs()
 	FHttpThreadBase::UpdateConfigs();
 }
 
-PRAGMA_DISABLE_DEPRECATION_WARNINGS
 void FEventLoopHttpThread::AddRequest(FHttpRequestCommon* Request)
 {
 	FHttpThreadBase::AddRequest(Request);
@@ -136,7 +135,6 @@ void FEventLoopHttpThread::GetCompletedRequests(TArray<FHttpRequestCommon*>& Out
 {
 	FHttpThreadBase::GetCompletedRequests(OutCompletedRequests);
 }
-PRAGMA_ENABLE_DEPRECATION_WARNINGS
 
 void FEventLoopHttpThread::Tick()
 {
@@ -173,11 +171,9 @@ void FEventLoopHttpThread::ResetTickTimer()
 
 	RequestTickTimer = EventLoop.SetTimer([this]()
 	{
-PRAGMA_DISABLE_DEPRECATION_WARNINGS
 		TArray<FHttpRequestCommon*> RequestsToCancel;
 		TArray<FHttpRequestCommon*> RequestsToComplete;
 		Process(RequestsToCancel, RequestsToComplete);
-PRAGMA_ENABLE_DEPRECATION_WARNINGS
 	},
 	FTimespan::FromSeconds(FHttpModule::Get().GetHttpEventLoopThreadTickIntervalInSeconds()),
 	true /* repeat */);
diff --git a/Engine/Source/Runtime/Online/HTTP/Private/EventLoopHttpThread.h b/Engine/Source/Runtime/Online/HTTP/Private/EventLoopHttpThread.h
--- a/Engine/Source/Runtime/Online/HTTP/Private/EventLoopHttpThread.h
+++ b/Engine/Source/Runtime/Online/HTTP/Private/EventLoopHttpThread.h
@@ -16,11 +16,9 @@ public:
 	virtual void StartThread() override final;
 	virtual void StopThread() override final;
 	virtual void UpdateConfigs() override final;
-PRAGMA_DISABLE_DEPRECATION_WARNINGS
 	virtual void AddRequest(FHttpRequestCommon* Request) override final;
 	virtual void CancelRequest(FHttpRequestCommon* Request) override final;
 	virtual void GetCompletedRequests(TArray<FHttpRequestCommon*>& OutCompletedRequests) override final;
-PRAGMA_ENABLE_DEPRECATION_WARNINGS
 
 	//~ Begin FSingleThreadRunnable Interface
 	// Cannot be overridden to ensure identical behavior with the threaded tick
diff --git a/Engine/Source/Runtime/Online/HTTP/Private/GenericPlatform/HttpRequestCommon.cpp b/Engine/Source/Runtime/Online/HTTP/Private/GenericPlatform/HttpRequestCommon.cpp
--- a/Engine/Source/Runtime/Online/HTTP/Private/GenericPlatform/HttpRequestCommon.cpp
+++ b/Engine/Source/Runtime/Online/HTTP/Private/GenericPlatform/HttpRequestCommon.cpp
@@ -10,7 +10,6 @@
 #include "Misc/CommandLine.h"
 #include "Stats/Stats.h"
 
-PRAGMA_DISABLE_DEPRECATION_WARNINGS
 namespace UE::HttpRequestCommon::Private
 {
 
@@ -788,5 +787,4 @@ void FHttpRequestCommon::LogResponse(const TSharedPtr<IHttpResponse>& InResponse
 
 		StartPos = EndPos + 1;
 	}
-}
-PRAGMA_ENABLE_DEPRECATION_WARNINGS
\ No newline at end of file
+}
\ No newline at end of file
diff --git a/Engine/Source/Runtime/Online/HTTP/Private/GenericPlatform/HttpRequestCommon.h b/Engine/Source/Runtime/Online/HTTP/Private/GenericPlatform/HttpRequestCommon.h
new file mode 100644
--- /dev/null
+++ b/Engine/Source/Runtime/Online/HTTP/Private/GenericPlatform/HttpRequestCommon.h
@@ -0,0 +1,197 @@
+// Copyright Epic Games, Inc. All Rights Reserved.
+
+#pragma once
+
+#include "GenericPlatform/HttpRequestImpl.h"
+#include "GenericPlatform/HttpRequestPayload.h"
+
+class FHttpResponseCommon;
+class IHttpTaskTimerHandle;
+
+/**
+ * Contains implementation of some common functions that don't vary between implementations of different platforms
+ */
+class FHttpRequestCommon : public FHttpRequestImpl
+{
+public:
+	FHttpRequestCommon();
+
+	// IHttpBase
+	HTTP_API virtual FString GetURLParameter(const FString& ParameterName) const override;
+
+	// IHttpRequest
+	HTTP_API virtual EHttpRequestStatus::Type GetStatus() const override;
+	HTTP_API virtual const FString& GetEffectiveURL() const override;
+	HTTP_API virtual EHttpFailureReason GetFailureReason() const override;
+	HTTP_API virtual void SetDelegateThreadPolicy(EHttpRequestDelegateThreadPolicy InDelegateThreadPolicy) override;
+	HTTP_API virtual EHttpRequestDelegateThreadPolicy GetDelegateThreadPolicy() const override;
+	HTTP_API virtual FString GetOption(const FName Option) const override;
+	HTTP_API virtual void SetOption(const FName Option, const FString & OptionValue) override;
+
+	HTTP_API virtual void SetTimeout(float InTimeoutSecs) override;
+	HTTP_API virtual void ClearTimeout() override;
+	HTTP_API virtual void ResetTimeoutStatus() override;
+	HTTP_API virtual TOptional<float> GetTimeout() const override;
+	HTTP_API float GetTimeoutOrDefault() const;
+
+	HTTP_API virtual void SetActivityTimeout(float InTimeoutSecs) override;
+
+	HTTP_API virtual const FHttpResponsePtr GetResponse() const override;
+
+	// Can be called on game thread or http thread depend on the delegate thread policy
+	HTTP_API virtual void FinishRequest() = 0;
+
+	HTTP_API virtual void CancelRequest() override;
+
+	HTTP_API virtual void Shutdown() override;
+
+	HTTP_API virtual void ProcessRequestUntilComplete() override;
+
+	HTTP_API virtual bool SetResponseBodyReceiveStream(TSharedRef<FArchive> Stream) override;
+
+	HTTP_API virtual float GetElapsedTime() const override;
+
+	HTTP_API virtual bool IsThreadedRequestComplete() = 0;
+	HTTP_API virtual bool StartThreadedRequest() = 0;
+	HTTP_API virtual void TickThreadedRequest(float DeltaSeconds) = 0;
+
+	HTTP_API void StartWaitingInQueue();
+	HTTP_API float GetTimeStartedWaitingInQueue() const;
+
+	HTTP_API virtual void SetURL(const FString& InURL) override;
+	HTTP_API virtual const FString& GetURL() const override;
+
+protected:
+	/**
+	 * Check if this request is valid or allowed, before actually process the request
+	 */
+	HTTP_API bool PreProcess();
+	HTTP_API void PostProcess();
+	HTTP_API void PopulateUserAgentHeader();
+	HTTP_API virtual bool SetupRequest() = 0;
+	HTTP_API bool PreCheck() const;
+	HTTP_API virtual void ClearInCaseOfRetry();
+
+	HTTP_API void SetStatus(EHttpRequestStatus::Type InCompletionStatus);
+	HTTP_API void SetFailureReason(EHttpFailureReason InFailureReason);
+
+	/**
+	 * Finish the request when it's not in http manager
+	 */
+	HTTP_API void FinishRequestNotInHttpManager();
+
+	HTTP_API void HandleRequestSucceed(TSharedPtr<IHttpResponse> InResponse);
+	HTTP_API void HandleRequestFailed(TSharedPtr<IHttpResponse> InResponse);
+
+	HTTP_API void StartActivityTimeoutTimer();
+	HTTP_API void StartActivityTimeoutTimerBy(double DelayToTrigger);
+	HTTP_API void ResetActivityTimeoutTimer(FStringView Reason);
+	HTTP_API void OnActivityTimeoutTimerTaskTrigger();
+	HTTP_API void StopActivityTimeoutTimer();
+	HTTP_API void StartTotalTimeoutTimer();
+	HTTP_API void StopTotalTimeoutTimer();
+	HTTP_API void OnTotalTimeoutTimerTaskTrigger();
+
+	HTTP_API virtual void AbortRequest() = 0;
+
+	HTTP_API virtual void CleanupRequest() = 0;
+
+	HTTP_API void HandleStatusCodeReceived(int32 StatusCode);
+	HTTP_API void TriggerStatusCodeReceivedDelegate(int32 StatusCode);
+
+	HTTP_API void SetEffectiveURL(const FString& InEffectiveURL);
+
+	HTTP_API bool PassReceivedDataToStream(void* Ptr, int64 Length);
+	HTTP_API void StopPassingReceivedData();
+
+	HTTP_API float GetActivityTimeoutOrDefault() const;
+
+	HTTP_API bool SetContentAsStreamedFileDefaultImpl(const FString& Filename);
+	HTTP_API bool OpenRequestPayloadDefaultImpl();
+	HTTP_API void CloseRequestPayloadDefaultImpl();
+
+	HTTP_API void LogFailure() const;
+	HTTP_API void LogResponse(const TSharedPtr<IHttpResponse>& InResponse);
+	HTTP_API bool TriggerMockFailure();
+
+	HTTP_API void InitResponse();
+
+	HTTP_API virtual FHttpResponsePtr CreateResponse() = 0;
+	// The function for each impl to make sure when mocking response code, the internal state is valid for finishing request successfully
+	HTTP_API virtual void MockResponseData() = 0;
+
+protected:
+	/** Current status of request being processed */
+	EHttpRequestStatus::Type CompletionStatus = EHttpRequestStatus::NotStarted;
+
+	/** Reason of failure of the HTTP request */
+	EHttpFailureReason FailureReason = EHttpFailureReason::None;
+
+	/** Thread policy about which thread to complete this request */
+	EHttpRequestDelegateThreadPolicy DelegateThreadPolicy = EHttpRequestDelegateThreadPolicy::CompleteOnGameThread;
+
+	/** Timeout in seconds for the entire HTTP request to complete */
+	TOptional<float> TimeoutSecs;
+
+	/** Timeout in seconds for the HTTP request activity timeout */
+	TOptional<float> ActivityTimeoutSecs;
+
+	/** Indicate the request is timed out, it should quit and fail with EHttpFailureReason::TimedOut */
+	std::atomic<bool> bTimedOut = false;
+	/** Indicate the request is activity timed out, it should quit and fail with EHttpFailureReason::ConnectionError */
+	std::atomic<bool> bActivityTimedOut = false;
+	/** Indicate the request is cancelled, it should quit and fail with EHttpFailureReason::Cancelled */
+	std::atomic<bool> bCanceled = false;
+
+	/** TODO: Move this feature into CurlHttp */
+	bool bUsePlatformActivityTimeout = true;
+
+	/** Record when this request started */
+	double RequestStartTimeAbsoluteSeconds;
+
+	/** Record when this request will activity timeout */
+	double ActivityTimeoutAt;
+
+	/** Holder the timer handle, if the request get destroyed before triggering the timeout, use this to remove the timer */
+	TSharedPtr<IHttpTaskTimerHandle> TotalTimeoutHttpTaskTimerHandle;
+	/** Holder the timer handle, if the request get destroyed before triggering the timeout, use this to remove the timer */
+	TSharedPtr<IHttpTaskTimerHandle> ActivityTimeoutHttpTaskTimerHandle;
+	/** Critical section for accessing HttpTaskTimerHandle */
+	FCriticalSection HttpTaskTimerHandleCriticalSection;
+
+	/** Record when the request start to process */
+	double StartProcessTime = 0.0;
+
+	/** Record how long it take to connect to the endpoint */
+	double ConnectTime = -1.0;
+
+	FString URL;
+
+	/** Cache the effective URL. When redirected, it will be different with original URL */
+	FString EffectiveURL;
+
+	/** The response object which we will use to pair with this request */
+	TSharedPtr<FHttpResponseCommon> ResponseCommon;
+
+	/** The stream to receive response body */
+	TSharedPtr<FArchive> ResponseBodyReceiveStream;
+
+	/** Critical section for accessing ResponseBodyReceiveStream */
+	FCriticalSection ResponseBodyReceiveStreamCriticalSection;
+
+	// Flag to indicate the request was initialized with stream. In that case even if stream was set to 
+	// null later on internally, the request itself won't cache received data anymore
+	std::atomic<bool> bInitializedWithValidStream = false;
+
+	/** Payload to use with the request. Typically for POST, PUT, or PATCH */
+	TUniquePtr<FRequestPayload> RequestPayload;
+
+	/** Options for this request. */
+	TMap<const FName, FString> Options;
+
+	/** Total elapsed time in seconds since the start of the request */
+	float ElapsedTime = 0.0f;
+
+	/** Record the time started to wait in the queue */
+	float TimeStartedWaitingInQueue = 0.0f;
+};
diff --git a/Engine/Source/Runtime/Online/HTTP/Private/GenericPlatform/HttpResponseCommon.cpp b/Engine/Source/Runtime/Online/HTTP/Private/GenericPlatform/HttpResponseCommon.cpp
--- a/Engine/Source/Runtime/Online/HTTP/Private/GenericPlatform/HttpResponseCommon.cpp
+++ b/Engine/Source/Runtime/Online/HTTP/Private/GenericPlatform/HttpResponseCommon.cpp
@@ -4,7 +4,6 @@
 #include "GenericPlatform/HttpRequestCommon.h"
 #include "GenericPlatform/GenericPlatformHttp.h"
 
-PRAGMA_DISABLE_DEPRECATION_WARNINGS
 FHttpResponseCommon::FHttpResponseCommon(const FHttpRequestCommon& HttpRequest)
 	: URL(HttpRequest.GetURL())
 	, EffectiveURL(HttpRequest.GetEffectiveURL())
@@ -68,4 +67,3 @@ void FHttpResponseCommon::SetResponseCode(int32 InResponseCode)
 	ResponseCode = InResponseCode;
 }
 
-PRAGMA_ENABLE_DEPRECATION_WARNINGS
diff --git a/Engine/Source/Runtime/Online/HTTP/Private/GenericPlatform/HttpResponseCommon.h b/Engine/Source/Runtime/Online/HTTP/Private/GenericPlatform/HttpResponseCommon.h
new file mode 100644
--- /dev/null
+++ b/Engine/Source/Runtime/Online/HTTP/Private/GenericPlatform/HttpResponseCommon.h
@@ -0,0 +1,38 @@
+// Copyright Epic Games, Inc. All Rights Reserved.
+
+#pragma once
+
+#include "Interfaces/IHttpResponse.h"
+
+class FHttpRequestCommon;
+
+/**
+ * Contains implementation of some common functions that don't vary between implementations of different platforms
+ */
+class FHttpResponseCommon : public IHttpResponse
+{
+	friend FHttpRequestCommon;
+
+public:
+	HTTP_API FHttpResponseCommon(const FHttpRequestCommon& HttpRequest);
+
+	// IHttpBase
+	HTTP_API virtual FString GetURLParameter(const FString& ParameterName) const override;
+	HTTP_API virtual const FString& GetURL() const override;
+	HTTP_API virtual const FString& GetEffectiveURL() const override;
+	HTTP_API virtual EHttpRequestStatus::Type GetStatus() const override;
+	HTTP_API virtual EHttpFailureReason GetFailureReason() const override;
+	HTTP_API virtual int32 GetResponseCode() const override;
+
+protected:
+	HTTP_API void SetRequestStatus(EHttpRequestStatus::Type InCompletionStatus);
+	HTTP_API void SetRequestFailureReason(EHttpFailureReason InFailureReason);
+	HTTP_API void SetEffectiveURL(const FString& InEffectiveURL);
+	HTTP_API void SetResponseCode(int32 InResponseCode);
+
+	FString URL;
+	FString EffectiveURL;
+	EHttpRequestStatus::Type CompletionStatus;
+	EHttpFailureReason FailureReason;
+	int32 ResponseCode = EHttpResponseCodes::Unknown;
+};
diff --git a/Engine/Source/Runtime/Online/HTTP/Private/HttpManager.cpp b/Engine/Source/Runtime/Online/HTTP/Private/HttpManager.cpp
--- a/Engine/Source/Runtime/Online/HTTP/Private/HttpManager.cpp
+++ b/Engine/Source/Runtime/Online/HTTP/Private/HttpManager.cpp
@@ -507,7 +507,6 @@ bool FHttpManager::Tick(float DeltaSeconds)
 			}
 		}
 
-PRAGMA_DISABLE_DEPRECATION_WARNINGS
 		TArray<FHttpRequestCommon*> CompletedThreadedRequests;
 
 		{
@@ -532,7 +531,6 @@ PRAGMA_DISABLE_DEPRECATION_WARNINGS
 				BroadcastHttpRequestCompleted(CompletedRequestRef);
 			}
 		}
-PRAGMA_ENABLE_DEPRECATION_WARNINGS
 	}
 	
 	// Report csv stats.
@@ -572,7 +570,6 @@ void FHttpManager::RemoveRequest(const FHttpRequestRef& Request)
 	Requests.Remove(Request);
 }
 
-PRAGMA_DISABLE_DEPRECATION_WARNINGS
 void FHttpManager::AddThreadedRequest(const TSharedRef<FHttpRequestCommon, ESPMode::ThreadSafe>& Request)
 {
 	check(Thread);
@@ -587,7 +584,6 @@ void FHttpManager::CancelThreadedRequest(const TSharedRef<FHttpRequestCommon, ES
 	check(Thread);
 	Thread->CancelRequest(&Request.Get());
 }
-PRAGMA_ENABLE_DEPRECATION_WARNINGS
 
 bool FHttpManager::IsValidRequest(const IHttpRequest* RequestPtr) const
 {
diff --git a/Engine/Source/Runtime/Online/HTTP/Private/HttpThread.cpp b/Engine/Source/Runtime/Online/HTTP/Private/HttpThread.cpp
--- a/Engine/Source/Runtime/Online/HTTP/Private/HttpThread.cpp
+++ b/Engine/Source/Runtime/Online/HTTP/Private/HttpThread.cpp
@@ -153,7 +153,6 @@ void FHttpThreadBase::StopThread()
 	bIsSingleThread = true;
 }
 
-PRAGMA_DISABLE_DEPRECATION_WARNINGS
 void FHttpThreadBase::AddRequest(FHttpRequestCommon* Request)
 {
 	NewThreadedRequests.Enqueue(Request);
@@ -172,7 +171,6 @@ void FHttpThreadBase::GetCompletedRequests(TArray<FHttpRequestCommon*>& OutCompl
 		OutCompletedRequests.Add(Request);
 	}
 }
-PRAGMA_ENABLE_DEPRECATION_WARNINGS
 
 bool FHttpThreadBase::Init()
 {
@@ -207,7 +205,6 @@ void FHttpThreadBase::HttpThreadTick(float DeltaSeconds)
 {
 }
 
-PRAGMA_DISABLE_DEPRECATION_WARNINGS
 bool FHttpThreadBase::StartThreadedRequest(FHttpRequestCommon* Request)
 {
 	return Request->StartThreadedRequest();
@@ -451,7 +448,6 @@ void FLegacyHttpThread::Tick()
 		Process(RequestsToCancel, RequestsToComplete);
 	}
 }
-PRAGMA_ENABLE_DEPRECATION_WARNINGS
 
 TSharedPtr<IHttpTaskTimerHandle> FLegacyHttpThread::AddHttpThreadTask(TFunction<void()>&& Task, float InDelay)
 {
@@ -487,11 +483,9 @@ bool FLegacyHttpThread::Init()
 UE_DISABLE_OPTIMIZATION_SHIP
 uint32 FLegacyHttpThread::Run()
 {
-PRAGMA_DISABLE_DEPRECATION_WARNINGS
 	// Arrays declared outside of loop to re-use memory
 	TArray<FHttpRequestCommon*> RequestsToCancel;
 	TArray<FHttpRequestCommon*> RequestsToComplete;
-PRAGMA_ENABLE_DEPRECATION_WARNINGS
 	while (!ExitRequest.GetValue())
 	{
 		if (ensureMsgf(!NeedsSingleThreadTick(), TEXT("HTTP Thread was set to singlethread mode while it was running autonomously!")))
diff --git a/Engine/Source/Runtime/Online/HTTP/Private/HttpThread.h b/Engine/Source/Runtime/Online/HTTP/Private/HttpThread.h
--- a/Engine/Source/Runtime/Online/HTTP/Private/HttpThread.h
+++ b/Engine/Source/Runtime/Online/HTTP/Private/HttpThread.h
@@ -53,7 +53,6 @@ public:
 	 */
 	bool IsStopped() const { return bIsStopped; }
 
-PRAGMA_DISABLE_DEPRECATION_WARNINGS
 	/**
 	 * Add a request to begin processing on HTTP thread.
 	 *
@@ -74,7 +73,6 @@ PRAGMA_DISABLE_DEPRECATION_WARNINGS
 	 * @param OutCompletedRequests array of requests that have been completed
 	 */
 	virtual void GetCompletedRequests(TArray<FHttpRequestCommon*>& OutCompletedRequests);
-PRAGMA_ENABLE_DEPRECATION_WARNINGS
 
 	//~ Begin FSingleThreadRunnable Interface
 	virtual void Tick() override;
@@ -111,7 +109,6 @@ protected:
 	 */
 	virtual void HttpThreadTick(float DeltaSeconds);
 	
-PRAGMA_DISABLE_DEPRECATION_WARNINGS
 	/** 
 	 * Start processing a request on the http thread
 	 */
@@ -147,7 +144,6 @@ private:
 	void StartRequestsWaitingInQueue(TArray<FHttpRequestCommon*>& RequestsToComplete);
 	void FinishRequestsFromHttpThreadWithCallbacks(TArray<FHttpRequestCommon*>& RequestsToComplete);
 	void UpdateThreadPriorityIfNeeded();
-PRAGMA_ENABLE_DEPRECATION_WARNINGS
 
 protected:
 	/** Pointer to Runnable Thread */
@@ -170,7 +166,6 @@ private:
 	int32 MaxConcurrentRequests;
 
 protected:
-PRAGMA_DISABLE_DEPRECATION_WARNINGS
 	/** 
 	 * Threaded requests that are waiting to be processed on the http thread.
 	 * Added to on (any) non-HTTP thread, processed then cleared on HTTP thread.
@@ -200,7 +195,6 @@ PRAGMA_DISABLE_DEPRECATION_WARNINGS
 	 * Added to on HTTP thread, processed then cleared on game thread (Single producer, single consumer)
 	 */
 	TSpscQueue<FHttpRequestCommon*> CompletedThreadedRequests;
-PRAGMA_ENABLE_DEPRECATION_WARNINGS
 };
 
 class FLegacyHttpThread	: public FHttpThreadBase
@@ -212,11 +206,9 @@ public:
 
 	virtual void StartThread() override final;
 	virtual void StopThread() override final;
-PRAGMA_DISABLE_DEPRECATION_WARNINGS
 	virtual void AddRequest(FHttpRequestCommon* Request) override final;
 	virtual void CancelRequest(FHttpRequestCommon* Request) override final;
 	virtual void GetCompletedRequests(TArray<FHttpRequestCommon*>& OutCompletedRequests) override final;
-PRAGMA_ENABLE_DEPRECATION_WARNINGS
 
 	//~ Begin FSingleThreadRunnable Interface
 	// Cannot be overridden to ensure identical behavior with the threaded tick
diff --git a/Engine/Source/Runtime/Online/HTTP/Public/GenericPlatform/HttpRequestCommon.h b/Engine/Source/Runtime/Online/HTTP/Public/GenericPlatform/HttpRequestCommon.h
deleted file mode 100644
diff --git a/Engine/Source/Runtime/Online/HTTP/Public/GenericPlatform/HttpResponseCommon.h b/Engine/Source/Runtime/Online/HTTP/Public/GenericPlatform/HttpResponseCommon.h
deleted file mode 100644
diff --git a/Engine/Source/Runtime/Online/HTTP/Public/HttpManager.h b/Engine/Source/Runtime/Online/HTTP/Public/HttpManager.h
--- a/Engine/Source/Runtime/Online/HTTP/Public/HttpManager.h
+++ b/Engine/Source/Runtime/Online/HTTP/Public/HttpManager.h
@@ -19,9 +19,7 @@
 #include "Templates/SharedPointer.h"
 
 class FHttpThreadBase;
-PRAGMA_DISABLE_DEPRECATION_WARNINGS
 class FHttpRequestCommon;
-PRAGMA_ENABLE_DEPRECATION_WARNINGS
 class FOutputDevice;
 class IHttpTaskTimerHandle;
 
@@ -209,7 +207,6 @@ public:
 	 */
 	HTTP_API virtual void FlushTick(float DeltaSeconds);
 
-PRAGMA_DISABLE_DEPRECATION_WARNINGS
 	/** 
 	 * Add a http request to be executed on the http thread
 	 *
@@ -223,7 +220,6 @@ PRAGMA_DISABLE_DEPRECATION_WARNINGS
 	 * @param Request - the request object to cancel
 	 */
 	HTTP_API void CancelThreadedRequest(const TSharedRef<FHttpRequestCommon, ESPMode::ThreadSafe>& Request);
-PRAGMA_ENABLE_DEPRECATION_WARNINGS
 
 	/**
 	 * List all of the Http requests currently being processed
