Explicitly disabled GetTypeHash and comparison operators for TSet and TMap.
Replaced a bespoke operator== which was comparing TSet and TMap with a named function.

diff --git a/Engine/Source/Runtime/Core/Public/Containers/Map.h b/Engine/Source/Runtime/Core/Public/Containers/Map.h
--- a/Engine/Source/Runtime/Core/Public/Containers/Map.h
+++ b/Engine/Source/Runtime/Core/Public/Containers/Map.h
@@ -1066,6 +1066,11 @@ public:
 	FORCEINLINE TRangedForConstIterator begin() const { return TRangedForConstIterator(Pairs.begin()); }
 	FORCEINLINE TRangedForIterator      end() { return TRangedForIterator(Pairs.end()); }
 	FORCEINLINE TRangedForConstIterator end() const { return TRangedForConstIterator(Pairs.end()); }
+
+	// Maps are deliberately prevented from being hashed or compared, because this would hide potentially major performance problems behind default operations.
+	friend uint32 GetTypeHash(const TMapBase& Map) = delete;
+	friend bool operator==(const TMapBase&, const TMapBase&) = delete;
+	friend bool operator!=(const TMapBase&, const TMapBase&) = delete;
 };
 
 /** The base type of sortable maps. */
diff --git a/Engine/Source/Runtime/Core/Public/Containers/Set.h b/Engine/Source/Runtime/Core/Public/Containers/Set.h
--- a/Engine/Source/Runtime/Core/Public/Containers/Set.h
+++ b/Engine/Source/Runtime/Core/Public/Containers/Set.h
@@ -1815,6 +1815,11 @@ public:
 	FORCEINLINE TRangedForConstIterator begin() const { return TRangedForConstIterator(Elements.begin()); }
 	FORCEINLINE TRangedForIterator      end()         { return TRangedForIterator     (Elements.end());   }
 	FORCEINLINE TRangedForConstIterator end() const   { return TRangedForConstIterator(Elements.end());   }
+
+	// Maps are deliberately prevented from being hashed or compared, because this would hide potentially major performance problems behind default operations.
+	friend uint32 GetTypeHash(const TSet& Set) = delete;
+	friend bool operator==(const TSet&, const TSet&) = delete;
+	friend bool operator!=(const TSet&, const TSet&) = delete;
 };
 
 template <typename RangeType>
diff --git a/Engine/Source/Runtime/Online/BuildPatchServices/Private/Tests/TestHelpers.h b/Engine/Source/Runtime/Online/BuildPatchServices/Private/Tests/TestHelpers.h
--- a/Engine/Source/Runtime/Online/BuildPatchServices/Private/Tests/TestHelpers.h
+++ b/Engine/Source/Runtime/Online/BuildPatchServices/Private/Tests/TestHelpers.h
@@ -89,17 +89,11 @@ namespace BuildPatchTestHelpers
 #define ARRAYU64(...) ARRAY(uint64, __VA_ARGS__)
 
 template<typename ElementType>
-bool operator==(const TSet<ElementType>& LHS, const TSet<ElementType>& RHS)
+bool SetsAreEqual(const TSet<ElementType>& LHS, const TSet<ElementType>& RHS)
 {
 	return (LHS.Num() == RHS.Num())
 	    && (LHS.Difference(RHS).Num() == 0)
 	    && (RHS.Difference(LHS).Num() == 0);
 }
 
-template<typename ElementType>
-bool operator!=(const TSet<ElementType>& LHS, const TSet<ElementType>& RHS)
-{
-	return !(LHS == RHS);
-}
-
 #endif //WITH_DEV_AUTOMATION_TESTS
diff --git a/Engine/Source/Runtime/Online/BuildPatchServices/Private/Tests/Unit/ChunkEvictionPolicy.spec.cpp b/Engine/Source/Runtime/Online/BuildPatchServices/Private/Tests/Unit/ChunkEvictionPolicy.spec.cpp
--- a/Engine/Source/Runtime/Online/BuildPatchServices/Private/Tests/Unit/ChunkEvictionPolicy.spec.cpp
+++ b/Engine/Source/Runtime/Online/BuildPatchServices/Private/Tests/Unit/ChunkEvictionPolicy.spec.cpp
@@ -74,7 +74,7 @@ void FChunkEvictionPolicySpec::Define()
 						ExpectedCleanable.Add(NextReferences[Idx]);
 					}
 					ChunkEvictionPolicy->Query(CurrentMap, CurrentMap.Num() + 10, Cleanable, Bootable);
-					TEST_EQUAL(Cleanable, ExpectedCleanable);
+					TEST_TRUE(SetsAreEqual(Cleanable, ExpectedCleanable));
 				});
 
 				It("should provide no chunks if all chunks are referenced.", [this]()
@@ -97,7 +97,7 @@ void FChunkEvictionPolicySpec::Define()
 						ExpectedCleanable.Add(NextReferences[Idx]);
 					}
 					ChunkEvictionPolicy->Query(CurrentMap, CurrentMap.Num(), Cleanable, Bootable);
-					TEST_EQUAL(Cleanable, ExpectedCleanable);
+					TEST_TRUE(SetsAreEqual(Cleanable, ExpectedCleanable));
 				});
 
 				It("should provide no chunks if all chunks are referenced.", [this]()
@@ -120,7 +120,7 @@ void FChunkEvictionPolicySpec::Define()
 						ExpectedCleanable.Add(NextReferences[Idx]);
 					}
 					ChunkEvictionPolicy->Query(CurrentMap, CurrentMap.Num() - 3, Cleanable, Bootable);
-					TEST_EQUAL(Cleanable, ExpectedCleanable);
+					TEST_TRUE(SetsAreEqual(Cleanable, ExpectedCleanable));
 				});
 
 				It("should provide minimum number of bootable chunks.", [this]()
@@ -139,8 +139,8 @@ void FChunkEvictionPolicySpec::Define()
 					ExpectedCleanable.Add(NextReferences[1]);
 					ExpectedBootable.Add(NextReferences[9]);
 					ChunkEvictionPolicy->Query(CurrentMap, CurrentMap.Num() - 3, Cleanable, Bootable);
-					TEST_EQUAL(Cleanable, ExpectedCleanable);
-					TEST_EQUAL(Bootable, ExpectedBootable);
+					TEST_TRUE(SetsAreEqual(Cleanable, ExpectedCleanable));
+					TEST_TRUE(SetsAreEqual(Bootable, ExpectedBootable));
 				});
 
 				It("should provide bootable chunks which are needed the latest.", [this]()
@@ -150,7 +150,7 @@ void FChunkEvictionPolicySpec::Define()
 					ExpectedBootable.Add(NextReferences[8]);
 					ExpectedBootable.Add(NextReferences[7]);
 					ChunkEvictionPolicy->Query(CurrentMap, CurrentMap.Num() - 3, Cleanable, Bootable);
-					TEST_EQUAL(Bootable, ExpectedBootable);
+					TEST_TRUE(SetsAreEqual(Bootable, ExpectedBootable));
 				});
 			});
 		});
diff --git a/Engine/Source/Runtime/Online/BuildPatchServices/Private/Tests/Unit/FileAttribution.spec.cpp b/Engine/Source/Runtime/Online/BuildPatchServices/Private/Tests/Unit/FileAttribution.spec.cpp
--- a/Engine/Source/Runtime/Online/BuildPatchServices/Private/Tests/Unit/FileAttribution.spec.cpp
+++ b/Engine/Source/Runtime/Online/BuildPatchServices/Private/Tests/Unit/FileAttribution.spec.cpp
@@ -167,7 +167,7 @@ void FFileAttributionSpec::Define()
 					{
 						AppliedFiles.Add(Call.Get<1>());
 					}
-					TEST_EQUAL(AppliedFiles, NewOrChangedCompressedFiles);
+					TEST_TRUE(SetsAreEqual(AppliedFiles, NewOrChangedCompressedFiles));
 				});
 
 				xIt("should apply readonly attributes to new or changed readonly files.", [this]()
@@ -183,7 +183,7 @@ void FFileAttributionSpec::Define()
 					{
 						AppliedFiles.Add(Call.Get<1>());
 					}
-					TEST_EQUAL(AppliedFiles, NewOrChangedReadOnlyFiles);
+					TEST_TRUE(SetsAreEqual(AppliedFiles, NewOrChangedReadOnlyFiles));
 				});
 
 				xIt("should apply executable attributes to new or changed executable files.", [this]()
@@ -199,7 +199,7 @@ void FFileAttributionSpec::Define()
 					{
 						AppliedFiles.Add(Call.Get<1>());
 					}
-					TEST_EQUAL(AppliedFiles, NewOrChangedExecutableFiles);
+					TEST_TRUE(SetsAreEqual(AppliedFiles, NewOrChangedExecutableFiles));
 				});
 			});
 
@@ -218,7 +218,7 @@ void FFileAttributionSpec::Define()
 					{
 						AppliedFiles.Add(Call.Get<1>());
 					}
-					TEST_EQUAL(AppliedFiles, FullAllFiles);
+					TEST_TRUE(SetsAreEqual(AppliedFiles, FullAllFiles));
 				});
 
 				xIt("should apply readonly attributes to all files.", [this]()
@@ -234,7 +234,7 @@ void FFileAttributionSpec::Define()
 					{
 						AppliedFiles.Add(Call.Get<1>());
 					}
-					TEST_EQUAL(AppliedFiles, FullAllFiles);
+					TEST_TRUE(SetsAreEqual(AppliedFiles, FullAllFiles));
 				});
 
 				xIt("should apply executable attributes to all files.", [this]()
@@ -250,7 +250,7 @@ void FFileAttributionSpec::Define()
 					{
 						AppliedFiles.Add(Call.Get<1>());
 					}
-					TEST_EQUAL(AppliedFiles, FullAllFiles);
+					TEST_TRUE(SetsAreEqual(AppliedFiles, FullAllFiles));
 				});
 			});
 		});
diff --git a/Engine/Source/Runtime/Online/BuildPatchServices/Private/Tests/Unit/InstallChunkSource.spec.cpp b/Engine/Source/Runtime/Online/BuildPatchServices/Private/Tests/Unit/InstallChunkSource.spec.cpp
--- a/Engine/Source/Runtime/Online/BuildPatchServices/Private/Tests/Unit/InstallChunkSource.spec.cpp
+++ b/Engine/Source/Runtime/Online/BuildPatchServices/Private/Tests/Unit/InstallChunkSource.spec.cpp
@@ -87,7 +87,7 @@ void FInstallChunkSourceSpec::Define()
 				It("should return the available chunks.", [this]()
 				{
 					TSet<FGuid> AvailableChunks = InstallChunkSource->GetAvailableChunks();
-					TEST_EQUAL(AvailableChunks, SomeAvailableChunks);
+					TEST_TRUE(SetsAreEqual(AvailableChunks, SomeAvailableChunks));
 				});
 			});
 		});
diff --git a/Engine/Source/Runtime/Online/BuildPatchServices/Private/Tests/Unit/MemoryChunkStore.spec.cpp b/Engine/Source/Runtime/Online/BuildPatchServices/Private/Tests/Unit/MemoryChunkStore.spec.cpp
--- a/Engine/Source/Runtime/Online/BuildPatchServices/Private/Tests/Unit/MemoryChunkStore.spec.cpp
+++ b/Engine/Source/Runtime/Online/BuildPatchServices/Private/Tests/Unit/MemoryChunkStore.spec.cpp
@@ -83,7 +83,7 @@ void FMemoryChunkStoreSpec::Define()
 					MemoryChunkStore->DumpToOverflow();
 					TSet<FGuid> PutChunks;
 					Algo::Transform(FakeChunkStore->RxPut, PutChunks, [](const FMockChunkStore::FPut& Call) { return Call.Get<1>(); });
-					TEST_EQUAL(SomeChunks, PutChunks);
+					TEST_TRUE(SetsAreEqual(SomeChunks, PutChunks));
 				});
 
 				Describe("when a chunks was previously Get", [this]()
@@ -98,7 +98,7 @@ void FMemoryChunkStoreSpec::Define()
 						MemoryChunkStore->DumpToOverflow();
 						TSet<FGuid> PutChunks;
 						Algo::Transform(FakeChunkStore->RxPut, PutChunks, [](const FMockChunkStore::FPut& Call) { return Call.Get<1>(); });
-						TEST_EQUAL(SomeChunks, PutChunks);
+						TEST_TRUE(SetsAreEqual(SomeChunks, PutChunks));
 					});
 				});
 			});
@@ -134,7 +134,7 @@ void FMemoryChunkStoreSpec::Define()
 						MemoryChunkStore->Put(SomeChunk, TUniquePtr<IChunkDataAccess>(new FFakeChunkDataAccess()));
 						TSet<FGuid> ReleasedChunks;
 						Algo::Transform(MockMemoryChunkStoreStat->RxChunkReleased, ReleasedChunks, [](const FMockMemoryChunkStoreStat::FChunkReleased& Call) { return Call.Get<1>(); });
-						TEST_EQUAL(CleanableChunks, ReleasedChunks);
+						TEST_TRUE(SetsAreEqual(CleanableChunks, ReleasedChunks));
 					});
 				});
 
@@ -150,7 +150,7 @@ void FMemoryChunkStoreSpec::Define()
 						MemoryChunkStore->Put(SomeChunk, TUniquePtr<IChunkDataAccess>(new FFakeChunkDataAccess()));
 						TSet<FGuid> BootedChunks;
 						Algo::Transform(MockMemoryChunkStoreStat->RxChunkBooted, BootedChunks, [](const FMockMemoryChunkStoreStat::FChunkBooted& Call) { return Call.Get<1>(); });
-						TEST_EQUAL(BootableChunks, BootedChunks);
+						TEST_TRUE(SetsAreEqual(BootableChunks, BootedChunks));
 					});
 
 					It("should Put them to the overflow store.", [this]()
@@ -158,7 +158,7 @@ void FMemoryChunkStoreSpec::Define()
 						MemoryChunkStore->Put(SomeChunk, TUniquePtr<IChunkDataAccess>(new FFakeChunkDataAccess()));
 						TSet<FGuid> PutChunks;
 						Algo::Transform(FakeChunkStore->RxPut, PutChunks, [](const FMockChunkStore::FPut& Call) { return Call.Get<1>(); });
-						TEST_EQUAL(BootableChunks, PutChunks);
+						TEST_TRUE(SetsAreEqual(BootableChunks, PutChunks));
 					});
 				});
 			});
diff --git a/Engine/Source/Runtime/Online/BuildPatchServices/Private/Tests/Unit/Verifier.spec.cpp b/Engine/Source/Runtime/Online/BuildPatchServices/Private/Tests/Unit/Verifier.spec.cpp
--- a/Engine/Source/Runtime/Online/BuildPatchServices/Private/Tests/Unit/Verifier.spec.cpp
+++ b/Engine/Source/Runtime/Online/BuildPatchServices/Private/Tests/Unit/Verifier.spec.cpp
@@ -89,7 +89,7 @@ void FVerifierSpec::Define()
 				{
 					Verifier->Verify(OutDatedFiles);
 					TEST_EQUAL(OutDatedFiles.Num(), 0);
-					TEST_EQUAL(LoadedFiles(), AllFiles);
+					TEST_TRUE(SetsAreEqual(LoadedFiles(), AllFiles));
 				});
 			});
 
@@ -104,7 +104,7 @@ void FVerifierSpec::Define()
 				{
 					Verifier->Verify(OutDatedFiles);
 					TEST_EQUAL(OutDatedFiles.Num(), 0);
-					TEST_EQUAL(LoadedFiles(), TouchedFiles);
+					TEST_TRUE(SetsAreEqual(LoadedFiles(), TouchedFiles));
 				});
 			});
 
@@ -120,7 +120,7 @@ void FVerifierSpec::Define()
 					Verifier->Verify(OutDatedFiles);
 					TEST_EQUAL(OutDatedFiles.Num(), 0);
 					TEST_EQUAL(LoadedFiles().Num(), 0);
-					TEST_EQUAL(FilesSizeCheckedFiles(), AllFiles);
+					TEST_TRUE(SetsAreEqual(FilesSizeCheckedFiles(), AllFiles));
 				});
 			});
 
@@ -136,7 +136,7 @@ void FVerifierSpec::Define()
 					Verifier->Verify(OutDatedFiles);
 					TEST_EQUAL(OutDatedFiles.Num(), 0);
 					TEST_EQUAL(LoadedFiles().Num(), 0);
-					TEST_EQUAL(FilesSizeCheckedFiles(), TouchedFiles);
+					TEST_TRUE(SetsAreEqual(FilesSizeCheckedFiles(), TouchedFiles));
 				});
 			});
 
@@ -152,7 +152,7 @@ void FVerifierSpec::Define()
 				It("should provide some files.", [this]()
 				{
 					Verifier->Verify(OutDatedFiles);
-					TEST_EQUAL(TSet<FString>(OutDatedFiles), SomeFiles);
+					TEST_TRUE(SetsAreEqual(TSet<FString>(OutDatedFiles), SomeFiles));
 				});
 			});
 
@@ -168,7 +168,7 @@ void FVerifierSpec::Define()
 				It("should provide some files.", [this]()
 				{
 					Verifier->Verify(OutDatedFiles);
-					TEST_EQUAL(TSet<FString>(OutDatedFiles), SomeFiles);
+					TEST_TRUE(SetsAreEqual(TSet<FString>(OutDatedFiles), SomeFiles));
 				});
 			});
 
@@ -184,7 +184,7 @@ void FVerifierSpec::Define()
 				It("load staged files instead of installed files.", [this]()
 				{
 					Verifier->Verify(OutDatedFiles);
-					TEST_EQUAL(LoadedFiles(), AllFiles);
+					TEST_TRUE(SetsAreEqual(LoadedFiles(), AllFiles));
 				});
 			});
 
