Use threadsafe delegate for FSslCertificateDelegates::VerifySslCertificates

diff --git a/Engine/Source/Runtime/Online/SSL/Public/Interfaces/ISslCertificateManager.h b/Engine/Source/Runtime/Online/SSL/Public/Interfaces/ISslCertificateManager.h
--- a/Engine/Source/Runtime/Online/SSL/Public/Interfaces/ISslCertificateManager.h
+++ b/Engine/Source/Runtime/Online/SSL/Public/Interfaces/ISslCertificateManager.h
@@ -89,6 +89,6 @@ public:
 		FString Subject;
 	};
 
-	DECLARE_DELEGATE_RetVal_TwoParams(bool, FVerifySslCertificates, const FString&, const TArray<FCertInfo>&);
+	using FVerifySslCertificates = TTSDelegate<bool(const FString&, const TArray<FCertInfo>&)>;
 	static SSL_API FVerifySslCertificates VerifySslCertificates;
 };
