Initial download analytics revision

#rnx

diff --git a/Engine/Plugins/Online/Android/AndroidFetchBackgroundDownload/Source/AndroidFetchBackgroundDownload/Java/download/UEDownloadWorker.java b/Engine/Plugins/Online/Android/AndroidFetchBackgroundDownload/Source/AndroidFetchBackgroundDownload/Java/download/UEDownloadWorker.java
--- a/Engine/Plugins/Online/Android/AndroidFetchBackgroundDownload/Source/AndroidFetchBackgroundDownload/Java/download/UEDownloadWorker.java
+++ b/Engine/Plugins/Online/Android/AndroidFetchBackgroundDownload/Source/AndroidFetchBackgroundDownload/Java/download/UEDownloadWorker.java
@@ -45,6 +45,7 @@ import com.tonyodev.fetch2.Request;
 import com.tonyodev.fetch2.exception.FetchException;
 import static com.tonyodev.fetch2.util.FetchUtils.canRetryDownload;
 
+import java.util.Locale;
 import java.util.concurrent.TimeUnit;
 
 import android.app.Activity;
@@ -798,6 +799,15 @@ public class UEDownloadWorker extends UEWorker implements DownloadProgressListen
 		if (!bReceivedResult && !bForceStopped)
 		{
 			UpdateNotification(100, false);
+
+			if (bDidAllRequestsSucceed)
+			{
+				Context context = getApplicationContext();
+				SharedPreferences preferences = context.getSharedPreferences("UEDownloadWorker", Context.MODE_PRIVATE);
+				SharedPreferences.Editor editor = preferences.edit();
+				editor.putString("AllDownloadsCompletedTime", String.format(Locale.getDefault(), "%f", System.currentTimeMillis() / 1000.0f));
+				editor.commit();
+			}
 		
 			nativeAndroidBackgroundDownloadOnAllComplete(bDidAllRequestsSucceed);
 		
diff --git a/Engine/Source/Runtime/Online/BackgroundHTTP/Private/BackgroundHttpNotificationObject.cpp b/Engine/Source/Runtime/Online/BackgroundHTTP/Private/BackgroundHttpNotificationObject.cpp
--- a/Engine/Source/Runtime/Online/BackgroundHTTP/Private/BackgroundHttpNotificationObject.cpp
+++ b/Engine/Source/Runtime/Online/BackgroundHTTP/Private/BackgroundHttpNotificationObject.cpp
@@ -19,14 +19,13 @@ FBackgroundHttpNotificationObject::FBackgroundHttpNotificationObject(FText InNot
 {
 }
 
-FBackgroundHttpNotificationObject::FBackgroundHttpNotificationObject(FText InNotificationTitle, FText InNotificationBody, FText InNotificationAction, const FString& InNotificationActivationString, bool InNotifyOnlyOnFullSuccess, bool InbOnlySendNotificationInBackground, int32 InIdOverride, bool InRecordNotificationTimestamp)
+FBackgroundHttpNotificationObject::FBackgroundHttpNotificationObject(FText InNotificationTitle, FText InNotificationBody, FText InNotificationAction, const FString& InNotificationActivationString, bool InNotifyOnlyOnFullSuccess, bool InbOnlySendNotificationInBackground, int32 InIdOverride)
 	: NotificationTitle(InNotificationTitle)
     , NotificationAction(InNotificationAction)
     , NotificationBody(InNotificationBody)
 	, NotificationActivationString(InNotificationActivationString)
     , bOnlySendNotificationInBackground(InbOnlySendNotificationInBackground)
     , bNotifyOnlyOnFullSuccess(InNotifyOnlyOnFullSuccess)
-	, bRecordNotificationTimestamp(InRecordNotificationTimestamp)
 	, bIsInBackground(false)
 	, NumFailedDownloads(0)
 	, IdOverride(InIdOverride)
@@ -69,12 +68,6 @@ void FBackgroundHttpNotificationObject::OnApp_EnteringBackground()
 
 FBackgroundHttpNotificationObject::~FBackgroundHttpNotificationObject()
 {
-	if (bRecordNotificationTimestamp)
-	{
-		FString Timestamp = FDateTime::UtcNow().ToString();
-		FPlatformMisc::SetStoredValue(TEXT("Epic Games"), TEXT("UEBackgroundHTTPNotification"), TEXT("ScheduleTime"), Timestamp);
-	}
-
 	if (bOnlySendNotificationInBackground)
 	{
 		//These should only be registered if bOnlySendNotificationInBackground is set
diff --git a/Engine/Source/Runtime/Online/BackgroundHTTP/Public/BackgroundHttpNotificationObject.h b/Engine/Source/Runtime/Online/BackgroundHTTP/Public/BackgroundHttpNotificationObject.h
--- a/Engine/Source/Runtime/Online/BackgroundHTTP/Public/BackgroundHttpNotificationObject.h
+++ b/Engine/Source/Runtime/Online/BackgroundHTTP/Public/BackgroundHttpNotificationObject.h
@@ -14,7 +14,7 @@ class FBackgroundHttpNotificationObject
 {
 public:
 	BACKGROUNDHTTP_API FBackgroundHttpNotificationObject(FText InNotificationTitle, FText InNotificationBody, FText InNotificationAction, const FString& InNotificationActivationString, bool InNotifyOnlyOnFullSuccess);
-	BACKGROUNDHTTP_API FBackgroundHttpNotificationObject(FText InNotificationTitle, FText InNotificationBody, FText InNotificationAction, const FString& InNotificationActivationString, bool InNotifyOnlyOnFullSuccess, bool bOnlySendNotificationInBackground, int32 InIdOverride, bool InRecordNotificationTimestamp = false);
+	BACKGROUNDHTTP_API FBackgroundHttpNotificationObject(FText InNotificationTitle, FText InNotificationBody, FText InNotificationAction, const FString& InNotificationActivationString, bool InNotifyOnlyOnFullSuccess, bool bOnlySendNotificationInBackground, int32 InIdOverride);
 	BACKGROUNDHTTP_API ~FBackgroundHttpNotificationObject();
 
 	BACKGROUNDHTTP_API void NotifyOfDownloadResult(bool bWasSuccess);
@@ -31,7 +31,6 @@ private:
 
 	bool bOnlySendNotificationInBackground;
 	bool bNotifyOnlyOnFullSuccess;
-	bool bRecordNotificationTimestamp;
 
 	volatile bool bIsInBackground;
 	volatile int32 NumFailedDownloads;
