Fix for incorrect backing store entry count calculation.

#rnx

diff --git a/Engine/Source/Runtime/Online/BuildPatchServices/Private/BuildPatchFileConstructor.cpp b/Engine/Source/Runtime/Online/BuildPatchServices/Private/BuildPatchFileConstructor.cpp
--- a/Engine/Source/Runtime/Online/BuildPatchServices/Private/BuildPatchFileConstructor.cpp
+++ b/Engine/Source/Runtime/Online/BuildPatchServices/Private/BuildPatchFileConstructor.cpp
@@ -449,7 +449,7 @@ class FChunkBackingStore : IConstructorChunkSource
 
 		UE_LOG(LogBuildPatchServices, VeryVerbose, TEXT("Paging out: %s"), *WriteToString<40>(InGuid));
 
-		uint32 EntriesRequired = Align(InChunk.ChunkSize, BitsPerEntry) >> BitsPerEntry;
+		uint32 EntriesRequired = Align(InChunk.ChunkSize, (1 << BitsPerEntry)) >> BitsPerEntry;
 
 		int32 SpanIndex = 0;
 		while (SpanIndex < FreeDiskSpans.Num() && 
