BPT fixes:
* Handle abort with unstarted cloud reads.
* Handle race condition when handling failed chunk reads redirecting to cloud.
* Handle abandoned downloads during tick.

diff --git a/Engine/Source/Runtime/Online/BuildPatchServices/Private/BuildPatchFileConstructor.cpp b/Engine/Source/Runtime/Online/BuildPatchServices/Private/BuildPatchFileConstructor.cpp
--- a/Engine/Source/Runtime/Online/BuildPatchServices/Private/BuildPatchFileConstructor.cpp
+++ b/Engine/Source/Runtime/Online/BuildPatchServices/Private/BuildPatchFileConstructor.cpp
@@ -2681,13 +2681,18 @@ bool FBuildPatchFileConstructor::ConstructFileFromChunks(const FString& BuildFil
 				if (bSuccess &&
 					bShouldAbort)
 				{
-					// The first time we catch an abort, inform the cloud source.
-					CloudSource->Abort();
-
 					ConstructionError = EConstructionError::Aborted;
 					bSuccess = false;
 				}
 
+				// If we are in a failure case, we want to clear out anything happening with the cloud source.
+				// We'll hit this over and over as we "drain", but we need to do that anyway because further failures
+				// might try to queue more cloud reads.
+				if (!bSuccess)
+				{
+					CloudSource->Abort();
+				}
+
 				if (!bIsPaused && bSuccess)
 				{
 					// Can we start another read?
diff --git a/Engine/Source/Runtime/Online/BuildPatchServices/Private/Installer/CloudChunkSource.cpp b/Engine/Source/Runtime/Online/BuildPatchServices/Private/Installer/CloudChunkSource.cpp
--- a/Engine/Source/Runtime/Online/BuildPatchServices/Private/Installer/CloudChunkSource.cpp
+++ b/Engine/Source/Runtime/Online/BuildPatchServices/Private/Installer/CloudChunkSource.cpp
@@ -896,6 +896,11 @@ namespace BuildPatchServices
 		
 		virtual FRequestProcessFn CreateRequest(const FGuid& DataId, FMutableMemoryView DestinationBuffer, void* UserPtr, FChunkRequestCompleteDelegate CompleteFn) override
 		{
+			//
+			// This function can get called from any thread as the failure case for reading a chunk is to
+			// request it off the cloud.
+			//
+
 			// We don't have a request that can be serviced directly because of how we have multiple in flight,
 			// so we return an empty request function and let the tick do the work.
 			FCloudRead* Read = new FCloudRead();
@@ -903,7 +908,10 @@ namespace BuildPatchServices
 			Read->Destination = DestinationBuffer;
 			Read->CompleteFn = CompleteFn;
 			Read->UserPtr = UserPtr;
+
+			ReadQueueCS.Lock();
 			ReadQueue.Add(Read);
+			ReadQueueCS.Unlock();
 
 			// We need to make sure that we get ticked to start the request as we might not be called from
 			// the dispatch thread.
@@ -977,6 +985,7 @@ namespace BuildPatchServices
 
 		// Read pointers are in one of these until they are either destroyed or passed
 		// to the validation task.
+		FCriticalSection ReadQueueCS;
 		TArray<FCloudRead*> ReadQueue;
 		TMap<int32, FCloudRead*> InFlightDownloads;
 
@@ -996,12 +1005,15 @@ namespace BuildPatchServices
 	void FConstructorCloudChunkSource::Abort()
 	{
 		// Release any unqueued reads.
-		for (FCloudRead* Read : ReadQueue)
 		{
-			Read->CompleteFn.Execute(Read->DataId, true, false, Read->UserPtr);
-			delete Read;
+			FScopeLock _(&ReadQueueCS);
+			for (FCloudRead* Read : ReadQueue)
+			{
+				Read->CompleteFn.Execute(Read->DataId, true, false, Read->UserPtr);
+				delete Read;
+			}
+			ReadQueue.Empty();
 		}
-		ReadQueue.Empty();
 
 		// Abort any downloads.
 		for (TPair<int32, FCloudRead*> Pair : InFlightDownloads)
@@ -1023,8 +1035,16 @@ namespace BuildPatchServices
 	void FConstructorCloudChunkSource::Tick(bool bStartNewDownloads, uint32& OutTimeToNextTickMs)
 	{
 		int32 StartingDownloadCount = InFlightDownloads.Num();
-			
-		UE_LOG(LogCloudChunkSource, Verbose, TEXT("Cloud: Active: %d, Pending %d"), InFlightDownloads.Num(), ReadQueue.Num());
+
+		if (UE_LOG_ACTIVE(LogCloudChunkSource, Verbose))
+		{
+			int32 PendingCount;
+			{
+				FScopeLock _(&ReadQueueCS);
+				PendingCount = ReadQueue.Num();
+			}
+			UE_LOG(LogCloudChunkSource, Verbose, TEXT("Cloud: Active: %d, Pending %d"), InFlightDownloads.Num(), PendingCount);
+		}
 
 		// Process completed downloads.
 		TMap<int32, FDownloadRef> FrameCompletedDownloads;
@@ -1042,7 +1062,15 @@ namespace BuildPatchServices
 				const int32& RequestId = FrameCompletedDownload.Key;
 				const FDownloadRef& Download = FrameCompletedDownload.Value;
 
-				FCloudRead* Read = InFlightDownloads[RequestId];
+				FCloudRead** ReadPtr = InFlightDownloads.Find(RequestId);
+				if (ReadPtr == nullptr)
+				{
+					// We can get here when we abort. We still get a completion callback when we abandon the download
+					// and we no longer have any data structure reference.
+					continue;
+				}
+
+				FCloudRead* Read = *ReadPtr;
 				InFlightDownloads.Remove(RequestId);
 
 				Read->bQueuedForDownload = false;
@@ -1158,21 +1186,26 @@ namespace BuildPatchServices
 					{
 						// Set retry time and put the task back in the queue.
 						Read->RetryTime = FStatsCollector::GetSeconds() + GetRetryDelay(Configuration.RetryDelayTimes, Read->RetryNum);
+						FScopeLock _(&ReadQueueCS);
 						ReadQueue.Add(Read);
 					}
 				}
 			} // end each completed
-		} // end process completed			
+		} // end process completed
 
 		// Update connection status and health.
 		EBuildPatchDownloadHealth DownloadHealth;
 		{
-			bool bAllDownloadsRetrying = ReadQueue.Num() > 0;
-			for (FCloudRead* Read : ReadQueue)
+			bool bAllDownloadsRetrying = false;
 			{
-				if (!Read->RetryNum)
+				FScopeLock _(&ReadQueueCS);
+				bAllDownloadsRetrying = ReadQueue.Num() > 0;
+				for (FCloudRead* Read : ReadQueue)
 				{
-					bAllDownloadsRetrying = false;
+					if (!Read->RetryNum)
+					{
+						bAllDownloadsRetrying = false;
+					}
 				}
 			}
 
@@ -1198,15 +1231,24 @@ namespace BuildPatchServices
 		// Start new downloads, if we aren't gated or paused
 		while (bStartNewDownloads)
 		{
+			FCloudRead* Read = nullptr;
 			const uint32 MaxDownloads = FMath::Max(1U, DownloadCount->GetAdjustedCount(InFlightDownloads.Num(), DownloadHealth));
-			if ((uint32)InFlightDownloads.Num() >= MaxDownloads ||
-				ReadQueue.Num() == 0)
+			if ((uint32)InFlightDownloads.Num() >= MaxDownloads)
 			{
 				break;
 			}
 
-			FCloudRead* Read = ReadQueue[0];
-			ReadQueue.RemoveAt(0);
+			{
+				FScopeLock _(&ReadQueueCS);
+
+				if (ReadQueue.Num() == 0)
+				{
+					break;
+				}
+
+				Read = ReadQueue[0];
+				ReadQueue.RemoveAt(0);
+			}
 
 			Read->CloudDirUsed = CurrentBestCloudDir;
 			Read->bQueuedForDownload = true;
@@ -1236,37 +1278,43 @@ namespace BuildPatchServices
 		// either because we completed a download or because a new request came in. However if we have 
 		// a retry that we want to requeue in a bit we just need to check back in.
 		uint32 WaitTimeMs = TNumericLimits<uint32>::Max(); // default infinite
-		if (ReadQueue.Num())
+		double CurrentTime = FStatsCollector::GetSeconds();
 		{
-			// This should never be that long. We should only have what has been requested
-			// for the current block the constructor is writing.
-			//
-			// Note that we only check retries - for reads that have never queued, the reason they
-			// aren't queued is because we ran out of download slots and we will check those
-			// when we get woken up by a completed download.
-			double ClosestRetry = -1;
-			for (FCloudRead* Read : ReadQueue)
-			{					
-				if (Read->RetryNum && 
-					(ClosestRetry == -1 || Read->RetryTime < ClosestRetry))
-				{
-					ClosestRetry = Read->RetryTime;
-				}
-			}
-
-			if (ClosestRetry != -1)
+			FScopeLock _(&ReadQueueCS);
+			if (ReadQueue.Num())
 			{
-				double TimeToClosestRetry = ClosestRetry - FStatsCollector::GetSeconds();
-				if (TimeToClosestRetry <= 0)
-				{
-					WaitTimeMs = 0;
+				// This should never be that long. We should only have what has been requested
+				// for the current block the constructor is writing.
+				//
+				// Note that we only check retries - for reads that have never queued, the reason they
+				// aren't queued is because we ran out of download slots and we will check those
+				// when we get woken up by a completed download.
+				double ClosestRetry = TNumericLimits<double>::Max();
+				for (FCloudRead* Read : ReadQueue)
+				{					
+					if (Read->RetryNum && Read->RetryTime < ClosestRetry)
+					{
+						ClosestRetry = Read->RetryTime;
+					}
 				}
-				else
+
+				if (ClosestRetry != TNumericLimits<double>::Max())
 				{
-					WaitTimeMs = (uint32)(TimeToClosestRetry * 1000);
+					double TimeToClosestRetry = ClosestRetry - CurrentTime;
+					if (TimeToClosestRetry <= 0)
+					{
+						WaitTimeMs = 0;
+					}
+					else
+					{
+						// Don't wait longer than 30 seconds just for sanity. We're waking up after a timeout
+						// anyway, computers don't care between 30 seconds and 30 minutes here.
+						WaitTimeMs = FMath::Min(30U, (uint32)(TimeToClosestRetry * 1000));
+					}
 				}
 			}
 		}
+
 		if (InFlightDownloads.Num())
 		{
 			// If we have in flight downloads we want to wake up every so often to
@@ -1278,7 +1326,7 @@ namespace BuildPatchServices
 			}
 		}
 
-		OutTimeToNextTickMs = WaitTimeMs;		
+		OutTimeToNextTickMs = WaitTimeMs;
 	}
 
 
