Fix hazard with request completion broadcast

#android

diff --git a/Engine/Source/Runtime/Online/BackgroundHTTP/Private/BackgroundHttpRequestImpl.cpp b/Engine/Source/Runtime/Online/BackgroundHTTP/Private/BackgroundHttpRequestImpl.cpp
--- a/Engine/Source/Runtime/Online/BackgroundHTTP/Private/BackgroundHttpRequestImpl.cpp
+++ b/Engine/Source/Runtime/Online/BackgroundHTTP/Private/BackgroundHttpRequestImpl.cpp
@@ -82,19 +82,16 @@ void FBackgroundHttpRequestImpl::OnBackgroundDownloadComplete()
 }
 
 void FBackgroundHttpRequestImpl::NotifyNotificationObjectOfComplete(bool bWasSuccess)
-{	
-	if (DownloadCompleteNotificationObject.IsValid())
+{
+	// Get notification object and reset it (swap with nullptr)
+	TSharedPtr<FBackgroundHttpNotificationObject, ESPMode::ThreadSafe> LocalDownloadCompleteNotificationObject;
+	DownloadCompleteMutex.Lock();
+	Swap(DownloadCompleteNotificationObject, LocalDownloadCompleteNotificationObject);
+	DownloadCompleteMutex.Unlock();
+	
+	if (LocalDownloadCompleteNotificationObject.IsValid())
 	{
-		const int32 SharedRefCount = DownloadCompleteNotificationObject.GetSharedReferenceCount();
-		UE_LOG(LogBackgroundHttpRequest, VeryVerbose, TEXT("Removing Reference to DownloadCompleteNotificationObject - bWasSuccess:%d | CurrentSharedRefCount:%d"), (int)(bWasSuccess), SharedRefCount);
-
-		if (SharedRefCount == 1)
-		{
-			UE_LOG(LogBackgroundHttpRequest, Display, TEXT("Removing Final Reference to DownloadCompleteNotificationObject!"));
-		}
-
-		DownloadCompleteNotificationObject->NotifyOfDownloadResult(bWasSuccess);
-		DownloadCompleteNotificationObject.Reset();
+		LocalDownloadCompleteNotificationObject->NotifyOfDownloadResult(bWasSuccess);
 	}
 }
 
diff --git a/Engine/Source/Runtime/Online/BackgroundHTTP/Public/BackgroundHttpRequestImpl.h b/Engine/Source/Runtime/Online/BackgroundHTTP/Public/BackgroundHttpRequestImpl.h
--- a/Engine/Source/Runtime/Online/BackgroundHTTP/Public/BackgroundHttpRequestImpl.h
+++ b/Engine/Source/Runtime/Online/BackgroundHTTP/Public/BackgroundHttpRequestImpl.h
@@ -8,6 +8,7 @@
 #include "Interfaces/IBackgroundHttpRequest.h"
 #include "Logging/LogMacros.h"
 #include "Templates/SharedPointer.h"
+#include "Async/Mutex.h"
 
 DECLARE_LOG_CATEGORY_EXTERN(LogBackgroundHttpRequest, Log, All)
 
@@ -47,6 +48,7 @@ public:
 	BACKGROUNDHTTP_API virtual void NotifyRequestMetricsAvailable(const int32 TotalBytesDownloaded, const float DownloadDuration);
 	
 protected:
+	UE::FMutex DownloadCompleteMutex;
 	TSharedPtr<FBackgroundHttpNotificationObject, ESPMode::ThreadSafe> DownloadCompleteNotificationObject;
 	FBackgroundHttpResponsePtr Response;
 	TArray<FString> URLList;
