Use TMpscQueue<FAsyncTask> instead of the deprecated TQueue<FAsyncTask, EQueueMode::Mpsc> in http event loop.

diff --git a/Engine/Source/Runtime/Online/Experimental/EventLoop/Public/EventLoop/EventLoop.h b/Engine/Source/Runtime/Online/Experimental/EventLoop/Public/EventLoop/EventLoop.h
--- a/Engine/Source/Runtime/Online/Experimental/EventLoop/Public/EventLoop/EventLoop.h
+++ b/Engine/Source/Runtime/Online/Experimental/EventLoop/Public/EventLoop/EventLoop.h
@@ -5,6 +5,7 @@
 #include "CoreTypes.h"
 #include "Algo/ForEach.h"
 #include "Async/Mutex.h"
+#include "Containers/MpscQueue.h"
 #include "EventLoop/IEventLoop.h"
 #include "EventLoop/EventLoopTimer.h"
 #include "EventLoop/IEventLoopIOManager.h"
@@ -71,7 +72,7 @@ class TEventLoop final : public IEventLoop
 private:
 	IOManagerType IOManager;
 	TTimerManager<typename Traits::FTimerManagerTraits> TimerManager;
-	TQueue<FAsyncTask, EQueueMode::Mpsc> AsyncTasks;
+	TMpscQueue<FAsyncTask> AsyncTasks;
 	TArray<FOnShutdownComplete> ShutdownRequests;
 	FTimespan LoopTime;
 	TAtomic<uint32> EventLoopThreadId;
diff --git a/Engine/Source/Runtime/Online/Experimental/EventLoop/Public/EventLoop/EventLoopManagedStorage.h b/Engine/Source/Runtime/Online/Experimental/EventLoop/Public/EventLoop/EventLoopManagedStorage.h
--- a/Engine/Source/Runtime/Online/Experimental/EventLoop/Public/EventLoop/EventLoopManagedStorage.h
+++ b/Engine/Source/Runtime/Online/Experimental/EventLoop/Public/EventLoop/EventLoopManagedStorage.h
@@ -4,7 +4,7 @@
 #include "CoreTypes.h"
 #include "Containers/Map.h"
 #include "Containers/SparseArray.h"
-#include "Containers/Queue.h"
+#include "Containers/MpscQueue.h"
 #include "EventLoop/EventLoopHandle.h"
 #include "HAL/PlatformTLS.h"
 #include "Misc/CoreMiscDefines.h"
@@ -52,10 +52,6 @@ struct FManagedStorageTraitsBase
  */
 struct FManagedStorageDefaultTraits : public FManagedStorageTraitsBase
 {
-	/**
-	 * The access mode used by the request queues.
-	 */
-	static constexpr EQueueMode QueueMode = EQueueMode::Mpsc;
 
 	/**
 	 * The allocator to use for the internal handle array.
@@ -189,9 +185,9 @@ private:
 	/* Mapping of external handles to the array index of the element in Storage. */
 	TMap<FExternalHandle, int32> StorageHandleIndex;
 	/** Requests to add entries to storage. Processed in Update(). */
-	TQueue<FAddRequest, Traits::QueueMode> AddRequests;
+	TMpscQueue<FAddRequest> AddRequests;
 	/** Requests to remove entries from storage. Processed in Update(). */
-	TQueue<FRemoveRequest, Traits::QueueMode> RemoveRequests;
+	TMpscQueue<FRemoveRequest> RemoveRequests;
 	/**
 	 * The thread id set when calling Init. This thread becomes the 'Manager' thread and is safe
 	 * to call non-thread safe methods from.
@@ -305,7 +301,7 @@ public:
 	{
 		// Assign new handle for data.
 		FExternalHandle Handle(FExternalHandle::GenerateNewHandle);
-		AddRequests.Enqueue({Handle, MoveTemp(Data)});
+		AddRequests.Enqueue(FAddRequest{Handle, MoveTemp(Data)});
 		return Handle;
 	}
 
@@ -319,7 +315,7 @@ public:
 	 */
 	void Remove(const FExternalHandle Handle, FManagedStorageOnRemoveComplete&& OnRemoveComplete = FManagedStorageOnRemoveComplete())
 	{
-		RemoveRequests.Enqueue({Handle, MoveTemp(OnRemoveComplete)});
+		RemoveRequests.Enqueue(FRemoveRequest{Handle, MoveTemp(OnRemoveComplete)});
 	}
 
 	/**
