Fix the issue console variables with "http." prefix won't take effect when changing it through console command;
Make sure to clear the parsed url patterns for mocking failures, when parsing the new cvar value.

diff --git a/Engine/Source/Runtime/Online/HTTP/Private/HttpManager.cpp b/Engine/Source/Runtime/Online/HTTP/Private/HttpManager.cpp
--- a/Engine/Source/Runtime/Online/HTTP/Private/HttpManager.cpp
+++ b/Engine/Source/Runtime/Online/HTTP/Private/HttpManager.cpp
@@ -745,6 +745,8 @@ void FHttpManager::UpdateUrlPatternsToMockFailure(IConsoleVariable* CVar)
 	TArray<FString> UrlPatternsToMockFailureStrings;
 	UrlPatternsToMockFailureStr.ParseIntoArray(UrlPatternsToMockFailureStrings, TEXT(" "));
 
+	UrlPatternsToMockFailure.Empty();
+
 	for (const FString& UrlPatternToMockFailureString : UrlPatternsToMockFailureStrings)
 	{
 		TArray<FString> UrlPattern;
diff --git a/Engine/Source/Runtime/Online/HTTP/Private/HttpModule.cpp b/Engine/Source/Runtime/Online/HTTP/Private/HttpModule.cpp
--- a/Engine/Source/Runtime/Online/HTTP/Private/HttpModule.cpp
+++ b/Engine/Source/Runtime/Online/HTTP/Private/HttpModule.cpp
@@ -270,12 +270,19 @@ bool FHttpModule::HandleHTTPCommand(const TCHAR* Cmd, FOutputDevice& Ar)
 
 bool FHttpModule::Exec_Runtime(UWorld* InWorld, const TCHAR* Cmd, FOutputDevice& Ar)
 {
-	// Ignore any execs that don't start with HTTP
-	if (FParse::Command(&Cmd, TEXT("HTTP")))
+	if (FCString::Strnicmp(Cmd, TEXT("HTTP."), 5) == 0)
 	{
-		return HandleHTTPCommand( Cmd, Ar );
+		// It's changing console variable
+		return false;
 	}
-	return false;
+
+	if (!FParse::Command(&Cmd, TEXT("HTTP")))
+	{
+		// Ignore any execs that don't start with HTTP
+		return false;
+	}
+
+	return HandleHTTPCommand(Cmd, Ar);
 }
 
 FHttpModule& FHttpModule::Get()
