Fix the compile error of EventLoopTests project caused by changelist 38780411.

diff --git a/Engine/Source/Runtime/Online/Experimental/EventLoopTests/Tests/EventLoop/EventLoopManagedStorageTests.cpp b/Engine/Source/Runtime/Online/Experimental/EventLoopTests/Tests/EventLoop/EventLoopManagedStorageTests.cpp
--- a/Engine/Source/Runtime/Online/Experimental/EventLoopTests/Tests/EventLoop/EventLoopManagedStorageTests.cpp
+++ b/Engine/Source/Runtime/Online/Experimental/EventLoopTests/Tests/EventLoop/EventLoopManagedStorageTests.cpp
@@ -34,7 +34,6 @@ namespace UE::EventLoop
 		}
 
 		static constexpr bool bStorageAccessThreadChecksEnabled = true;
-		static constexpr EQueueMode QueueMode = EQueueMode::Mpsc;
 		using InternalHandleArryAllocatorType = TInlineAllocator<32>;
 		using FExternalHandle = TResourceHandle<FMockManagedStorageHandleTraits>;
 
diff --git a/Engine/Source/Runtime/Online/Experimental/EventLoopTests/Tests/EventLoop/EventLoopTests.cpp b/Engine/Source/Runtime/Online/Experimental/EventLoopTests/Tests/EventLoop/EventLoopTests.cpp
--- a/Engine/Source/Runtime/Online/Experimental/EventLoopTests/Tests/EventLoop/EventLoopTests.cpp
+++ b/Engine/Source/Runtime/Online/Experimental/EventLoopTests/Tests/EventLoop/EventLoopTests.cpp
@@ -87,7 +87,6 @@ namespace UE::EventLoop
 				}
 
 				static constexpr bool bStorageAccessThreadChecksEnabled = true;
-				static constexpr EQueueMode QueueMode = EQueueMode::Mpsc;
 				using InternalHandleArryAllocatorType = TInlineAllocator<32>;
 				using FExternalHandle = FTimerHandle;
 
diff --git a/Engine/Source/Runtime/Online/Experimental/EventLoopTests/Tests/EventLoop/EventLoopTimerTests.cpp b/Engine/Source/Runtime/Online/Experimental/EventLoopTests/Tests/EventLoop/EventLoopTimerTests.cpp
--- a/Engine/Source/Runtime/Online/Experimental/EventLoopTests/Tests/EventLoop/EventLoopTimerTests.cpp
+++ b/Engine/Source/Runtime/Online/Experimental/EventLoopTests/Tests/EventLoop/EventLoopTimerTests.cpp
@@ -44,7 +44,6 @@ namespace UE::EventLoop
 			}
 
 			static constexpr bool bStorageAccessThreadChecksEnabled = true;
-			static constexpr EQueueMode QueueMode = EQueueMode::Mpsc;
 			using InternalHandleArryAllocatorType = TInlineAllocator<32>;
 			using FExternalHandle = TResourceHandle<FMockTimerHandleTraits>;
 
