Add cvar to be able to configure the duration in http queue before outputing warning log;
Disable that warning log for unattended mode as well.

diff --git a/Engine/Source/Runtime/Online/HTTP/Private/HttpThread.cpp b/Engine/Source/Runtime/Online/HTTP/Private/HttpThread.cpp
--- a/Engine/Source/Runtime/Online/HTTP/Private/HttpThread.cpp
+++ b/Engine/Source/Runtime/Online/HTTP/Private/HttpThread.cpp
@@ -7,6 +7,7 @@
 #include "HAL/RunnableThread.h"
 #include "HAL/IConsoleManager.h"
 #include "HttpManager.h"
+#include "Misc/App.h"
 #include "Misc/ConfigCacheIni.h"
 #include "Misc/CommandLine.h"
 #include "Misc/Fork.h"
@@ -24,6 +25,12 @@ TAutoConsoleVariable<int32> CVarHttpMaxConcurrentRequests(
 	ECVF_SaveForNextBoot
 );
 
+TAutoConsoleVariable<int32> CVarHttpDurationInQueueToWarnInSeconds(
+	TEXT("http.DurationInQueueToWarnInSeconds"),
+	10,
+	TEXT("If http request waited more than this duration in the queue, output a warning before starting")
+);
+
 // Thread priority cvar (settable at runtime)
 // We declare these explicitly rather than just casting the cvar in case the enum changes in future
 const int32 GHttpThreadPriorities[] =
@@ -290,7 +297,7 @@ void FHttpThreadBase::StartRequestsWaitingInQueue(TArray<FHttpRequestCommon*>& R
 			RateLimitedThreadedRequests.RemoveAt(0);
 
 			float DurationInQueue = FPlatformTime::Seconds() - ReadyThreadedRequest->GetTimeStartedWaitingInQueue();
-			UE_CLOG(DurationInQueue > 10.0f, LogHttp, Warning, TEXT("Request (%p) waited in queue for %.2fs before starting"), ReadyThreadedRequest, DurationInQueue);
+			UE_CLOG(!FApp::IsUnattended() && DurationInQueue > CVarHttpDurationInQueueToWarnInSeconds.GetValueOnAnyThread(), LogHttp, Warning, TEXT("Request (%p) waited in queue for %.2fs before starting"), ReadyThreadedRequest, DurationInQueue);
 			float StartImmediately = 0.01f;
 			if (DurationInQueue > StartImmediately)
 			{
