Log the reason why we failed and the files that failed during BPS verification.

#rnx

diff --git a/Engine/Source/Runtime/Online/BuildPatchServices/Private/Installer/Verifier.cpp b/Engine/Source/Runtime/Online/BuildPatchServices/Private/Installer/Verifier.cpp
--- a/Engine/Source/Runtime/Online/BuildPatchServices/Private/Installer/Verifier.cpp
+++ b/Engine/Source/Runtime/Online/BuildPatchServices/Private/Installer/Verifier.cpp
@@ -264,6 +264,7 @@ namespace BuildPatchServices
 			if (FileVerifyResult != EVerifyResult::Success)
 			{
 				CorruptFiles.Add(ThreadJobList[i].FileName);
+				UE_LOG(LogVerifier, Warning, TEXT("File verification failed on: %s (cause = %d)"), *ThreadJobList[i].FileName, FileVerifyResult);
 				if (VerifyResult == EVerifyResult::Success)
 				{
 					VerifyResult = FileVerifyResult;
