LwsWebSockets: Added ini configurable disabling of FLwsWebSocketsManager's default polling behavior.

With the new bPollService option set to false (defaults to true to maintain current behavior), FLwsWebSocketsManager can now make blocking calls to lws_service, allowing for low-latency websocket reads that may otherwise have been rate limited by ThreadTargetFrameTimeInSeconds explicit sleeping.

This is a restoration of changes that were removed in 20445534, though now behind ini configuration.

Also shifted  [WebSockets.LibWebSockets] .ini section reading to new IWebSocketsManager::UpdateConfigs function, allowing ini settings to be re-read on FCoreDelegates::TSOnConfigSectionsChanged, similar to FHttpManager::UpdateConfigs.

diff --git a/Engine/Source/Runtime/Online/WebSockets/Private/Lws/LwsWebSocket.cpp b/Engine/Source/Runtime/Online/WebSockets/Private/Lws/LwsWebSocket.cpp
--- a/Engine/Source/Runtime/Online/WebSockets/Private/Lws/LwsWebSocket.cpp
+++ b/Engine/Source/Runtime/Online/WebSockets/Private/Lws/LwsWebSocket.cpp
@@ -169,6 +169,8 @@ void FLwsWebSocket::Close(int32 Code, const FString& Reason)
 void FLwsWebSocket::Send(const void* Data, SIZE_T Size, bool bIsBinary)
 {
 	SendQueue.Enqueue(new FLwsSendBuffer(static_cast<const uint8*>(Data), Size, bIsBinary));
+	FLwsWebSocketsManager& WebSocketsManager = FLwsWebSocketsManager::Get();
+	WebSocketsManager.WakeService();
 }
 
 void FLwsWebSocket::Send(const FString& Data)
diff --git a/Engine/Source/Runtime/Online/WebSockets/Private/Lws/LwsWebSocketsManager.cpp b/Engine/Source/Runtime/Online/WebSockets/Private/Lws/LwsWebSocketsManager.cpp
--- a/Engine/Source/Runtime/Online/WebSockets/Private/Lws/LwsWebSocketsManager.cpp
+++ b/Engine/Source/Runtime/Online/WebSockets/Private/Lws/LwsWebSocketsManager.cpp
@@ -51,31 +51,58 @@ FLwsWebSocketsManager::FLwsWebSocketsManager()
 	LwsContext(nullptr)
 	, Thread(nullptr)
 {
-	ThreadTargetFrameTimeInSeconds = 1.0f / 30.0f; // 30Hz
-	GConfig->GetDouble(TEXT("WebSockets.LibWebSockets"), TEXT("ThreadTargetFrameTimeInSeconds"), ThreadTargetFrameTimeInSeconds, GEngineIni);
+}
 
-	ThreadMinimumSleepTimeInSeconds = 0.0f;
-	GConfig->GetDouble(TEXT("WebSockets.LibWebSockets"), TEXT("ThreadMinimumSleepTimeInSeconds"), ThreadMinimumSleepTimeInSeconds, GEngineIni);
+FLwsWebSocketsManager& FLwsWebSocketsManager::Get()
+{
+	FLwsWebSocketsManager* Manager = static_cast<FLwsWebSocketsManager*>(FWebSocketsModule::Get().WebSocketsManager);
+	check(Manager);
+	return *Manager;
+}
 
+void FLwsWebSocketsManager::UpdateConfigs()
+{
+	const bool bWasPollingService = bPollService;
+	GConfig->GetBool(TEXT("WebSockets.LibWebSockets"), TEXT("bPollService"), bPollService, GEngineIni);
+	GConfig->GetDouble(TEXT("WebSockets.LibWebSockets"), TEXT("ThreadTargetFrameTimeInSeconds"), ThreadTargetFrameTimeInSeconds, GEngineIni);
+	GConfig->GetDouble(TEXT("WebSockets.LibWebSockets"), TEXT("ThreadMinimumSleepTimeInSeconds"), ThreadMinimumSleepTimeInSeconds, GEngineIni);
+	
 	GConfig->GetBool(TEXT("LwsWebSocket"), TEXT("bDisableDomainAllowlist"), bDisableDomainAllowlist, GEngineIni);
-
+	
 	GConfig->GetBool(TEXT("LwsWebSocket"), TEXT("bDisableCertValidation"), bDisableCertValidation, GEngineIni);
 	
 #if !UE_BUILD_SHIPPING
 	// Allow -lws-insecure= to override ini flags
 	FParse::Bool(FCommandLine::Get(), TEXT("lws-insecure="), bDisableCertValidation);
 #endif
-}
+	
+	if (bWasPollingService != bPollService)
+	{
+		if (bPollService)
+		{
+			UE_LOG(LogWebSockets, Log, TEXT("FLwsWebSocketsManager::UpdateConfigs: Service polling enabled by config"));
 
-FLwsWebSocketsManager& FLwsWebSocketsManager::Get()
-{
-	FLwsWebSocketsManager* Manager = static_cast<FLwsWebSocketsManager*>(FWebSocketsModule::Get().WebSocketsManager);
-	check(Manager);
-	return *Manager;
+			if (LwsContext)
+			{
+				// Cancel blocking service call in FLwsWebSocketsManager::Tick to allow sleep, service, sleep polling loop to start in
+				// FLwsWebSocketsManager::Run   
+				lws_cancel_service(LwsContext);
+			}
+		}
+		else
+		{
+			UE_LOG(LogWebSockets, Log, TEXT("FLwsWebSocketsManager::UpdateConfigs: Service polling disabled by config"));
+		}
+	}
 }
 
 void FLwsWebSocketsManager::InitWebSockets(TArrayView<const FString> Protocols)
 {
+	UE_LOG(LogWebSockets, Log, TEXT("FLwsWebSocketsManager::InitWebSockets"));
+	
+	// Pull initial config values. UpdateConfigs may be called again later 
+	UpdateConfigs();
+	
 	QUICK_SCOPE_CYCLE_COUNTER(STAT_FLwsWebSocketsManager_InitWebSockets);
 	check(!Thread && LwsProtocols.Num() == 0);
 
@@ -246,10 +273,13 @@ uint32 FLwsWebSocketsManager::Run()
 		double BeginTime = FPlatformTime::Seconds();
 		Tick();
 		double EndTime = FPlatformTime::Seconds();
-
 		double TotalTime = EndTime - BeginTime;
-		double SleepTime = FMath::Max(ThreadTargetFrameTimeInSeconds - TotalTime, ThreadMinimumSleepTimeInSeconds);
-		FPlatformProcess::SleepNoStats(SleepTime);
+
+		if (bPollService)
+		{
+			double SleepTime = FMath::Max(ThreadTargetFrameTimeInSeconds - TotalTime, ThreadMinimumSleepTimeInSeconds);
+			FPlatformProcess::SleepNoStats(SleepTime);
+		}
 	}
 
 	return 0;
@@ -399,7 +429,7 @@ void FLwsWebSocketsManager::Tick()
 	}
 	if (LwsContext)
 	{
-		lws_service(LwsContext, 0);
+		lws_service(LwsContext, bPollService ? 0 : MAX_int32);
 	}
 	for (FLwsWebSocket* Socket : SocketsDestroyedDuringService)
 	{
@@ -429,6 +459,18 @@ void FLwsWebSocketsManager::StartProcessingWebSocket(FLwsWebSocket* Socket)
 {
 	Sockets.Emplace(Socket->AsShared());
 	SocketsToStart.Enqueue(Socket);
+	WakeService();
+}
+
+void FLwsWebSocketsManager::WakeService()
+{
+	if (bPollService)
+	{
+		return;
+	}
+	
+	// Safe to call from other threads
+	lws_cancel_service(LwsContext);
 }
 
 bool FLwsWebSocketsManager::GameThreadTick(float DeltaTime)
diff --git a/Engine/Source/Runtime/Online/WebSockets/Private/Lws/LwsWebSocketsManager.h b/Engine/Source/Runtime/Online/WebSockets/Private/Lws/LwsWebSocketsManager.h
--- a/Engine/Source/Runtime/Online/WebSockets/Private/Lws/LwsWebSocketsManager.h
+++ b/Engine/Source/Runtime/Online/WebSockets/Private/Lws/LwsWebSocketsManager.h
@@ -43,10 +43,17 @@ public:
 	 */
 	void StartProcessingWebSocket(FLwsWebSocket* Socket);
 
+	/**
+	 * Trigger/Wake up the websocket thread to process new requests.
+	 * Note: This is a no-op if bPollService = true
+	 */
+	void WakeService();
+
 	// IWebSocketsManager
 	virtual void InitWebSockets(TArrayView<const FString> Protocols) override;
 	virtual void ShutdownWebSockets() override;
 	virtual TSharedRef<IWebSocket> CreateWebSocket(const FString& Url, const TArray<FString>& Protocols, const TMap<FString, FString>& UpgradeHeaders) override;
+	virtual void UpdateConfigs() override;
 
 private:
 
@@ -104,10 +111,17 @@ private:
 	/** signal request to stop and exit thread */
 	FThreadSafeCounter ExitRequest;
 
+	/**
+	 * If true (default), our thread will explicitly sleep (according to ThreadTargetFrameTimeInSeconds & ThreadMinimumSleepTimeInSeconds) in
+	 * between calls to lws_service with 0 timeout.
+	 * 
+	 * If false, lws_service will use a MAX_int32 timeout, allowing lws to block and handle wake / sleep on activity internally itself.
+	 */
+	bool bPollService = true;
 	/** Target frame time for our thread's tick */
-	double ThreadTargetFrameTimeInSeconds;
+	double ThreadTargetFrameTimeInSeconds = 1.0f / 30.0f; // 30Hz;
 	/** Minimum time to sleep in our thread's tick, even if the sleep makes us exceed our target frame time */
-	double ThreadMinimumSleepTimeInSeconds;
+	double ThreadMinimumSleepTimeInSeconds = 0.0f;
 
 	bool bDisableDomainAllowlist = false;
 	bool bDisableCertValidation = false;
diff --git a/Engine/Source/Runtime/Online/WebSockets/Private/WebSocketsModule.cpp b/Engine/Source/Runtime/Online/WebSockets/Private/WebSocketsModule.cpp
--- a/Engine/Source/Runtime/Online/WebSockets/Private/WebSocketsModule.cpp
+++ b/Engine/Source/Runtime/Online/WebSockets/Private/WebSocketsModule.cpp
@@ -6,6 +6,7 @@
 #if WITH_WEBSOCKETS
 #include "PlatformWebSocket.h"
 #include "Misc/ConfigCacheIni.h"
+#include "Misc/CoreDelegates.h"
 #endif // #if WITH_WEBSOCKETS
 
 DEFINE_LOG_CATEGORY(LogWebSockets);
@@ -41,12 +42,17 @@ void FWebSocketsModule::StartupModule()
 
 	WebSocketsManager = new FPlatformWebSocketsManager;
 	WebSocketsManager->InitWebSockets(Protocols);
+
+	// Subscrive to config changes so we can forward to WebSocketsManager  
+	FCoreDelegates::TSOnConfigSectionsChanged().AddRaw(this, &FWebSocketsModule::OnConfigSectionsChanged);
 #endif
 }
 
 void FWebSocketsModule::ShutdownModule()
 {
 #if WITH_WEBSOCKETS
+	FCoreDelegates::TSOnConfigSectionsChanged().RemoveAll(this);
+	
 	if (WebSocketsManager)
 	{
 		WebSocketsManager->ShutdownWebSockets();
@@ -58,6 +64,16 @@ void FWebSocketsModule::ShutdownModule()
 	Singleton = nullptr;
 }
 
+#if WITH_WEBSOCKETS
+void FWebSocketsModule::OnConfigSectionsChanged(const FString& IniFilename, const TSet<FString>& SectionNames)
+{
+	if (WebSocketsManager)
+	{
+		WebSocketsManager->UpdateConfigs();
+	}
+}
+#endif
+
 FWebSocketsModule& FWebSocketsModule::Get()
 {
 	if (nullptr == Singleton)
diff --git a/Engine/Source/Runtime/Online/WebSockets/Public/IWebSocketsManager.h b/Engine/Source/Runtime/Online/WebSockets/Public/IWebSocketsManager.h
--- a/Engine/Source/Runtime/Online/WebSockets/Public/IWebSocketsManager.h
+++ b/Engine/Source/Runtime/Online/WebSockets/Public/IWebSocketsManager.h
@@ -35,4 +35,9 @@ public:
 	 * Virtual destructor
 	 */
 	virtual ~IWebSocketsManager() = default;
+
+	/**
+	 * Update configuration. Called when config has been updated and we need to apply any changes.
+	 */
+	virtual void UpdateConfigs() {}
 };
diff --git a/Engine/Source/Runtime/Online/WebSockets/Public/WebSocketsModule.h b/Engine/Source/Runtime/Online/WebSockets/Public/WebSocketsModule.h
--- a/Engine/Source/Runtime/Online/WebSockets/Public/WebSocketsModule.h
+++ b/Engine/Source/Runtime/Online/WebSockets/Public/WebSocketsModule.h
@@ -94,6 +94,9 @@ private:
 	IWebSocketsManager* WebSocketsManager;
 	friend class FLwsWebSocketsManager;
 	friend class FLwsWebSocket;
+
+	/** Delegate for config file changes. */
+	void OnConfigSectionsChanged(const FString& IniFilename, const TSet<FString>& SectionNames);
 #endif // #if WITH_WEBSOCKETS
 
 	/** singleton for the module while loaded and available */
