BPT:
* Multiple files in flight allowing for better download connection utilization.
* Allow instance overrides for buffer sizes and thread usage.
* Fix issue constructing empty files

#rnx

diff --git a/Engine/Source/Runtime/Online/BuildPatchServices/Private/BuildPatchFileConstructor.cpp b/Engine/Source/Runtime/Online/BuildPatchServices/Private/BuildPatchFileConstructor.cpp
--- a/Engine/Source/Runtime/Online/BuildPatchServices/Private/BuildPatchFileConstructor.cpp
+++ b/Engine/Source/Runtime/Online/BuildPatchServices/Private/BuildPatchFileConstructor.cpp
@@ -29,6 +29,7 @@
 #include "Installer/MessagePump.h"
 #include "Templates/Greater.h"
 #include "BuildPatchUtil.h"
+#include "ProfilingDebugging/CountersTrace.h"
 
 using namespace BuildPatchServices;
 
@@ -126,19 +127,107 @@ namespace FileConstructorHelpers
 	}
 }
 
+
+struct FAdministrationScope
+{
+	ISpeedRecorder::FRecord ActivityRecord;
+	IFileConstructorStat* FileConstructorStat;
+	FAdministrationScope(IFileConstructorStat* InFileConstructorStat)
+	{
+		FileConstructorStat = InFileConstructorStat;
+		FileConstructorStat->OnBeforeAdminister();
+		ActivityRecord.CyclesStart = FStatsCollector::GetCycles();
+	}
+	~FAdministrationScope()
+	{
+		ActivityRecord.CyclesEnd = FStatsCollector::GetCycles();
+		ActivityRecord.Size = 0;
+		FileConstructorStat->OnAfterAdminister(ActivityRecord);
+	}
+};
+
+struct FReadScope
+{
+	ISpeedRecorder::FRecord ActivityRecord;
+	IFileConstructorStat* FileConstructorStat;
+	FReadScope(IFileConstructorStat* InFileConstructorStat, int64 Size)
+	{
+		FileConstructorStat = InFileConstructorStat;		
+		ActivityRecord.CyclesStart = FStatsCollector::GetCycles();
+		ActivityRecord.Size = Size;
+		FileConstructorStat->OnBeforeRead();
+	}
+	~FReadScope()
+	{
+		ActivityRecord.CyclesEnd = FStatsCollector::GetCycles();
+		FileConstructorStat->OnAfterRead(ActivityRecord);
+	}
+};
+
 enum class EConstructionError : uint8
 {
 	None = 0,
 	CannotCreateFile,
 	OutOfDiskSpace,
+	FailedInitialSizeCheck,
 	MissingChunk,
 	SerializeError,
 	TrackingError,
 	OutboundDataError,
 	InternalConsistencyError,
-	Aborted
+	Aborted,
+	MissingFileInfo
+};
+
+
+// Since we can have more than one file in flight, store state here.
+struct FFileConstructionState
+{
+	FGuid ErrorContextGuid;
+	EConstructionError ConstructionError = EConstructionError::None;
+	int32 CreateFilePlatformLastError = 0;
+
+	FSHA1 HashState;
+
+	// If this is true then we didn't actually have to make the file, it was already done or a symlink or something.
+	bool bSkippedConstruction = false;
+	bool bSuccess = true;
+	bool bIsResumedFile = false;
+
+	// We track how far we are in the file when we write into the write buffer so that
+	// we advance progress bars smoothly instead of in huge writebuffer sized chunks.
+	UE::FMutex ProgressLock;
+	uint64 Progress = 0;
+	uint64 LastSeenProgress = 0; // not locked.	
+
+	// Where we started constructing the file. Can be non zero due to resume.
+	int64 StartPosition = 0;
+
+	int32 OutstandingBatches = 0;
+
+	// Can be nonzero in the first batch due to resume.
+	int32 NextChunkPartToRead = 0;
+
+	TUniquePtr<FArchive> NewFile;
+
+	const FFileManifest* FileManifest;
+	FString BuildFilename; // name as it is in the manifest
+	FString NewFilename; // with output path and such
+
+	FFileConstructionState(const FFileManifest* InFileManifest, FString&& InBuildFilename, FString&& InNewFilename) 
+	: FileManifest(InFileManifest)
+	, BuildFilename(InBuildFilename)
+	, NewFilename(InNewFilename)
+	{
+		if (!InFileManifest)
+		{
+			bSuccess = false;
+			ConstructionError = EConstructionError::MissingFileInfo;
+		}
+	}
 };
 
+
 /**
  * This struct handles loading and saving of simple resume information, that will allow us to decide which
  * files should be resumed from. It will also check that we are creating the same version and app as we expect to be.
@@ -300,6 +389,9 @@ class FChunkBackingStore : IConstructorChunkSource
 	uint64 CurrentMemoryLoad = 0;
 	uint64 PeakMemoryLoad = 0;
 
+	int32 ReadCount = 0;
+	int32 WriteCount = 0;
+
 	static constexpr uint64 ChunkStoreMemoryLimitDisabledSentinel = TNumericLimits<uint64>::Max();
 	uint64 MaxMemoryBytes = 1 << 30; // 1GB.
 
@@ -556,6 +648,8 @@ class FChunkBackingStore : IConstructorChunkSource
 
 		// Write
 		{
+			WriteCount++;
+
 			TRACE_CPUPROFILER_EVENT_SCOPE(BSPO_Write);
 			if (!BackingStoreFileHandle->Seek(UsedSpan.Offset()))
 			{
@@ -867,6 +961,8 @@ public:
 		{
 			UE_LOG(LogBuildPatchServices, Warning, TEXT("Unable to delete disk backing store: %s"), *BackingStoreFileName);
 		}
+
+		UE_LOG(LogBuildPatchServices, Verbose, TEXT("BackingStore Reads: %d Writes: %d"), ReadCount, WriteCount);
 	}
 
 	// Set the amount of disk space the installation needs to we can ensure that we don't 
@@ -1153,6 +1249,7 @@ public:
 
 		if (bSuccess)
 		{
+			ReadCount++;
 			BackingStoreFileHandle->Seek(UsedEntry->Offset());
 			if (!BackingStoreFileHandle->Read((uint8*)DestinationBuffer.GetData(), DestinationBuffer.GetSize()))
 			{
@@ -1292,40 +1389,52 @@ FBuildPatchFileConstructor::FBuildPatchFileConstructor(
 		InstallSource->SetFileRetirementPositions(MoveTemp(FileRetirementPositions));	
 	}
 
-	WriteBuffers[0].Reserve(WriteBufferSize);
-	WriteBuffers[1].Reserve(WriteBufferSize);
-
 	//
 	// Create the threads we are allowed to.
 	//
 	bool bHasInstallSource = InstallSource && InstallSource->GetAvailableChunks().Num() != 0;
 	bHasChunkDbSource = ChunkDbSource->GetAvailableChunks().Num() != 0;
 	
-	// For now always create a write thread.
-	int32 ThreadCount = 1;
-	WriteThreadIndex = 0;
-
 	// Default everything to running synchronously.
 	for (int8& ThreadAssignment : ThreadAssignments)
 	{
 		ThreadAssignment = -1;
 	}
+	WriteThreadIndex = -1;
+
+	bool bSpawnAdditionalIOThreads = Configuration.bDefaultSpawnAdditionalIOThreads;
+	if (GConfig->GetBool(TEXT("Portal.BuildPatch"), TEXT("ConstructorSpawnAdditionalIOThreads"), bSpawnAdditionalIOThreads, GEngineIni))
+	{
+		UE_LOG(LogBuildPatchServices, Verbose, TEXT("Got INI ConstructorSpawnAdditionalIOThreads = %d"), bSpawnAdditionalIOThreads);
+	}
+
+	if (Configuration.SpawnAdditionalIOThreads.IsSet())
+	{
+		bSpawnAdditionalIOThreads = Configuration.SpawnAdditionalIOThreads.Get(bSpawnAdditionalIOThreads);
+		UE_LOG(LogBuildPatchServices, Verbose, TEXT("Got override ConstructorSpawnAdditionalIOThreads = %d"), bSpawnAdditionalIOThreads);
+	}
 
 	// For now we have to strictly assign jobs to threads so that we don't accidentally
 	// hit the same file handle on multiple threads. Once we have proper ReadAtOffset support
 	// we can go nuts (and just use UE::Tasks)
 	// LONGTERM try using UE::Pipe and just blasting everything on tasks?
-	if (bHasInstallSource)
-	{
-		ThreadAssignments[EConstructorChunkLocation::Install] = ThreadCount;
-		ThreadCount++;
-	}
-	if (bHasChunkDbSource)
+	int32 ThreadCount = 0;
+	if (bSpawnAdditionalIOThreads)
 	{
-		ThreadAssignments[EConstructorChunkLocation::ChunkDb] = ThreadCount;
+		WriteThreadIndex = 0;
 		ThreadCount++;
-	}
 
+		if (bHasInstallSource)
+		{
+			ThreadAssignments[EConstructorChunkLocation::Install] = ThreadCount;
+			ThreadCount++;
+		}
+		if (bHasChunkDbSource)
+		{
+			ThreadAssignments[EConstructorChunkLocation::ChunkDb] = ThreadCount;
+			ThreadCount++;
+		}
+	}
 	
 	WakeUpDispatchThreadEvent = FPlatformProcess::GetSynchEventFromPool();
 	CloudSource->SetWakeupFunction([WakeUpDispatchThreadEvent = WakeUpDispatchThreadEvent]()
@@ -1410,6 +1519,7 @@ void FBuildPatchFileConstructor::SetChunkLocation(const FGuid& InGuid, EConstruc
 {
 	// NOTE - reentrant. We assume a) that ChunkLocations is filled before threading
 	// and that b) no guids are used across concurrent actions.
+	FRWScopeLock ChunkLock(ChunkLocationsLock, SLT_Write);
 	EConstructorChunkLocation* Location = ChunkLocations.Find(InGuid);
 	if (!Location)
 	{
@@ -1424,10 +1534,8 @@ void FBuildPatchFileConstructor::SetChunkLocation(const FGuid& InGuid, EConstruc
 
 			UE_LOG(LogBuildPatchServices, VeryVerbose, TEXT("Migrating chunk to cloud: %s, %llu bytes"), *WriteToString<40>(InGuid), ChunkSize);
 
-			DownloadRequirementLock.Lock();
 			DownloadRequirement += ChunkSize;
 			CloudSource->PostRequiredByteCount(DownloadRequirement);
-			DownloadRequirementLock.Unlock();
 		}
 
 		*Location = InNewLocation;
@@ -1497,20 +1605,206 @@ void FBuildPatchFileConstructor::GenericThreadFn(int32 ThreadIndex)
 	ThreadCompleteEvents[ThreadIndex]->Trigger();
 }
 
+bool FBuildPatchFileConstructor::HarvestChunksForCompletedFile(const FString& CompletedFullPathFileName)
+{
+	UE_LOG(LogBuildPatchServices, Verbose, TEXT("Harvesting source: %s"), *CompletedFullPathFileName);
+
+	TRACE_CPUPROFILER_EVENT_SCOPE(Harvest);
+	// We need to grab any chunks from install sources that are no longer available.
+	// Anything that's already been loaded is already placed into the memory store appropriately,
+	// but anything that _hasn't_ needs to be pulled out.	
+	TSet<FGuid> FileChunks;
+	InstallSource->GetChunksForFile(CompletedFullPathFileName, FileChunks);
+
+	struct FNeededChunk
+	{
+		FGuid Id;
+		int32 LastUsageIndex;
+		int32 NextUsageIndex;
+		int32 ChunkSize;
+	};
+
+	TArray<FNeededChunk> ChunksFromFileWeNeed;
+
+	{
+		FRWScopeLock ChunkLock(ChunkLocationsLock, SLT_ReadOnly);
+
+		int32 CurrentUsageIndex = ChunkReferenceTracker->GetCurrentUsageIndex();
+		for (const FGuid& FileChunk : FileChunks)
+		{
+			EConstructorChunkLocation* Location = ChunkLocations.Find(FileChunk);
+			if (Location && *Location == EConstructorChunkLocation::Install)
+			{
+				int32 LastUsageIndex = 0;
+				int32 NextUsageIndex = ChunkReferenceTracker->GetNextUsageForChunk(FileChunk, LastUsageIndex);
+
+				if (NextUsageIndex == -1 ||
+					LastUsageIndex < CurrentUsageIndex)
+				{
+					// The chunk is no longer needed
+					*Location = EConstructorChunkLocation::Retired;
+					continue;
+				}
+
+				int32 ChunkSize = Configuration.ManifestSet->GetChunkInfo(FileChunk)->WindowSize;
+
+				ChunksFromFileWeNeed.Add({FileChunk, LastUsageIndex, NextUsageIndex, ChunkSize});
+			}
+		}
+	}
+
+	if (ChunksFromFileWeNeed.Num() == 0)
+	{
+		return true;
+	}
+
+	// Try to load all the chunks that are about to go away. If it fails we don't particularly
+	// care since we would have fallen back to cloud anyway.
+
+	// There's some care here - if we just kick off a ton of reads, all those backing store
+	// entries are locked during the reads so we have to allocate space and can't page anything
+	// out. This is fine if we can load the whole file, but under constrained memory we want to
+	// only keep stuff that's going to get used soon - so we load the stuff we _aren't_ going to use
+	// soon so it can get paged out. Then we load in batches so we can release the locks and let them
+	// page out. 
+	// LONGTERM - detect this condition and write directly to the disk backing store? Ideally this
+	// would be something we can retain across restarts as right now any harvested chunks get lost
+	// on abort and cause a download from the cloud source (not chunkdb!)
+	Algo::SortBy(ChunksFromFileWeNeed, &FNeededChunk::NextUsageIndex, TGreater<int32>());
+
+	constexpr int32 HarvestBatchSize = 16 << 20; // 16 MB
+
+	bool bHarvestSuccess = true;
+	for (int32 ChunkIndex = 0; ChunkIndex < ChunksFromFileWeNeed.Num();)
+	{
+		int32 BatchSize = 0;
+		int32 ChunkEndIndex = ChunkIndex + 1;
+		for (; ChunkEndIndex < ChunksFromFileWeNeed.Num(); ChunkEndIndex++)
+		{
+			if (BatchSize + ChunksFromFileWeNeed[ChunkEndIndex].ChunkSize > HarvestBatchSize)
+			{
+				break;
+			}
+			BatchSize += ChunksFromFileWeNeed[ChunkEndIndex].ChunkSize;
+		}
+
+		PendingHarvestRequests.store(ChunkEndIndex - ChunkIndex, std::memory_order_release);
+
+		for (int32 DispatchIndex = ChunkIndex; DispatchIndex < ChunkEndIndex; DispatchIndex++)
+		{
+			FNeededChunk& Chunk = ChunksFromFileWeNeed[DispatchIndex];
+
+			SetChunkLocation(Chunk.Id, EConstructorChunkLocation::Memory);
+
+			FMutableMemoryView Destination = BackingStore->ReserveAndLockEntry(Chunk.Id, Chunk.ChunkSize, Chunk.LastUsageIndex);
+			if (Destination.GetSize() == 0)
+			{
+				// Call the completion function so we decrement the request count,
+				// but this is a consistency failure so we can't use the results.
+				ChunkHarvestCompletedFn(Chunk.Id, false, false, 0);
+				bHarvestSuccess = false;
+			}
+
+			IConstructorChunkSource::FRequestProcessFn HarvestFn = InstallSource->CreateRequest(Chunk.Id, Destination, 0, IConstructorChunkSource::FChunkRequestCompleteDelegate::CreateRaw(this, &FBuildPatchFileConstructor::ChunkHarvestCompletedFn));
+
+			HarvestFn(false);
+		}
+
+		// The read is synchronous but the verification is not, so we still need to do the wait.
+		{
+			TRACE_CPUPROFILER_EVENT_SCOPE(HarvestWait);
+			WakeUpDispatchThreadEvent->Wait();
+		}
+
+		if (bHarvestSuccess)
+		{
+			// Unlock any memory store entries
+			FRWScopeLock ChunkLock(ChunkLocationsLock, SLT_ReadOnly);
+
+			for (int32 DispatchIndex = ChunkIndex; DispatchIndex < ChunkEndIndex; DispatchIndex++)
+			{
+				// The read could have failed - in which case the location was switched to cloud from Memory.
+				// We handle this here so we don't have to deal with synchronization in the completion function.
+				if (ChunkLocations[ChunksFromFileWeNeed[DispatchIndex].Id] == EConstructorChunkLocation::Cloud)
+				{
+					if (!BackingStore->DereserveHarvestingEntry(ChunksFromFileWeNeed[DispatchIndex].Id))
+					{
+						bHarvestSuccess = false;
+					}
+				}
+				else
+				{
+					if (!BackingStore->CommitAndReleaseEntry(ChunksFromFileWeNeed[DispatchIndex].Id))
+					{
+						bHarvestSuccess = false;
+					}
+				}
+			}
+		}
+
+		if (bHarvestSuccess && !bAllowMultipleFilesInFlight)
+		{
+			bHarvestSuccess = BackingStore->CheckNoLocks(true);
+		}
+
+		ChunkIndex = ChunkEndIndex;
+
+		if (bShouldAbort || !bHarvestSuccess)
+		{
+			break;
+		}
+	}
+
+	return bHarvestSuccess;
+}
+
 void FBuildPatchFileConstructor::Run()
 {
 	FileConstructorStat->OnTotalRequiredUpdated(TotalJobSize);
 
+	
+	// We'd really like to have a sense of what each chunk looks like, size-wise, so that
+	// we know things like how many downloads we expect per batch.
+	// Map of window sizes to counts of that size.
+	TMap<uint32, int32> WindowSizes;
+
+	// lock not requried - no threads yet.
 	DownloadRequirement = 0;
 	for (const TPair<FGuid, EConstructorChunkLocation>& ChunkLocation : ChunkLocations)
 	{
+		FChunkInfo const* ChunkInfo = Configuration.ManifestSet->GetChunkInfo(ChunkLocation.Key);
+		WindowSizes.FindOrAdd(ChunkInfo->WindowSize)++;
+
 		if (ChunkLocation.Value == EConstructorChunkLocation::Cloud)
 		{
-			DownloadRequirement += Configuration.ManifestSet->GetDownloadSize(ChunkLocation.Key);
+			DownloadRequirement += ChunkInfo->FileSize;
 		}
 	}
 	CloudSource->PostRequiredByteCount(DownloadRequirement);
 
+	{
+		ExpectedChunkSize = 0;
+		int32 LargestWindowCount = 0;
+		for (TPair<uint32, int32>& Entry : WindowSizes)
+		{
+			if (Entry.Value > LargestWindowCount)
+			{
+				LargestWindowCount = Entry.Value;
+				ExpectedChunkSize = Entry.Key;
+			}
+		}
+
+		if (ExpectedChunkSize)
+		{
+			UE_LOG(LogBuildPatchServices, Display, TEXT("Expected chunk size: %d count %d"), ExpectedChunkSize, LargestWindowCount);
+		}
+		else
+		{
+			UE_LOG(LogBuildPatchServices, Display, TEXT("Can't find largest chunk size, using 1MB"));
+			ExpectedChunkSize = 1 << 20;
+		}
+	}
+
 	// Check for resume data, we need to also look for a legacy resume file to use instead in case we are resuming from an install of previous code version.
 	const FString LegacyResumeDataFilename = Configuration.StagingDirectory / TEXT("$resumeData");
 	const FString ResumeDataFilename = Configuration.MetaDirectory / TEXT("$resumeData");
@@ -1548,233 +1842,7 @@ void FBuildPatchFileConstructor::Run()
 	FileConstructorStat->OnResumeStarted();
 
 	// While we have files to construct, run.
-	FString FileToConstruct;
-	while (GetFileToConstruct(FileToConstruct) && !bShouldAbort)
-	{
-		// Get the file manifest.
-		const FFileManifest* FileManifest = Configuration.ManifestSet->GetNewFileManifest(FileToConstruct);
-		bool bFileSuccess = FileManifest != nullptr;
-		if (bFileSuccess)
-		{
-			const int64 FileSize = FileManifest->FileSize;
-			FileConstructorStat->OnFileStarted(FileToConstruct, FileSize);
-
-			// Check resume status for this file.
-			const bool bFilePreviouslyComplete = ResumeData.FilesCompleted.Contains(FileToConstruct);
-			const bool bFilePreviouslyStarted = ResumeData.FilesStarted.Contains(FileToConstruct);
-
-			// Construct or skip the file.
-			if (bFilePreviouslyComplete)
-			{
-				bFileSuccess = true;
-				CountBytesProcessed(FileSize);
-				GLog->Logf(TEXT("FBuildPatchFileConstructor: Skipping completed file %s"), *FileToConstruct);
-				// Go through each chunk part, and dereference it from the reference tracker.
-				for (const FChunkPart& ChunkPart : FileManifest->ChunkParts)
-				{
-					bFileSuccess = ChunkReferenceTracker->PopReference(ChunkPart.Guid) && bFileSuccess;
-				}
-			}
-			else
-			{
-				bFileSuccess = ConstructFileFromChunks(FileToConstruct, *FileManifest, bFilePreviouslyStarted);
-			}
-		}
-		else
-		{
-			// Only report or log if the first error
-			if (InstallerError->HasError() == false)
-			{
-				InstallerAnalytics->RecordConstructionError(FileToConstruct, INDEX_NONE, TEXT("Missing File Manifest"));
-				UE_LOG(LogBuildPatchServices, Error, TEXT("FBuildPatchFileConstructor: Missing file manifest for %s"), *FileToConstruct);
-			}
-			// Always set
-			InstallerError->SetError(EBuildPatchInstallError::FileConstructionFail, ConstructionErrorCodes::MissingFileInfo);
-		}
-
-		if (bFileSuccess)
-		{
-			// If we are destructive, remove the old file.
-			if (Configuration.InstallMode == EInstallMode::DestructiveInstall)
-			{
-				const bool bRequireExists = false;
-				const bool bEvenReadOnly = true;
-				FString FileToDelete = Configuration.InstallDirectory / FileToConstruct;
-				FPaths::NormalizeFilename(FileToDelete);
-				FPaths::CollapseRelativeDirectories(FileToDelete);
-				if (FileSystem->FileExists(*FileToDelete))
-				{
-					bool bHarvestSuccess = true;
-					if (InstallSource)
-					{
-						UE_LOG(LogBuildPatchServices, Verbose, TEXT("Harvesting source: %s"), *FileToDelete);
-
-						TRACE_CPUPROFILER_EVENT_SCOPE(Harvest);
-						// We need to grab any chunks from install sources that are no longer available.
-						// Anything that's already been loaded is already placed into the memory store appropriately,
-						// but anything that _hasn't_ needs to be pulled out.	
-						TSet<FGuid> FileChunks;
-						InstallSource->GetChunksForFile(FileToDelete, FileChunks);
-
-						struct FNeededChunk
-						{
-							FGuid Id;
-							int32 LastUsageIndex;
-							int32 NextUsageIndex;
-							int32 ChunkSize;
-						};
-
-						TArray<FNeededChunk> ChunksFromFileWeNeed;
-
-						int32 CurrentUsageIndex = ChunkReferenceTracker->GetCurrentUsageIndex();
-						for (const FGuid& FileChunk : FileChunks)
-						{
-							EConstructorChunkLocation* Location = ChunkLocations.Find(FileChunk);
-							if (Location && *Location == EConstructorChunkLocation::Install)
-							{
-								int32 LastUsageIndex = 0;
-								int32 NextUsageIndex = ChunkReferenceTracker->GetNextUsageForChunk(FileChunk, LastUsageIndex);
-
-								if (NextUsageIndex == -1 ||
-									LastUsageIndex < CurrentUsageIndex)
-								{
-									// The chunk is no longer needed
-									*Location = EConstructorChunkLocation::Retired;
-									continue;
-								}
-
-								int32 ChunkSize = Configuration.ManifestSet->GetChunkInfo(FileChunk)->WindowSize;
-
-								ChunksFromFileWeNeed.Add({FileChunk, LastUsageIndex, NextUsageIndex, ChunkSize});
-							}
-						}
-
-						if (ChunksFromFileWeNeed.Num())
-						{
-							// Try to load all the chunks that are about to go away. If it fails we don't particularly
-							// care since we would have fallen back to cloud anyway.
-
-							// There's some care here - if we just kick off a ton of reads, all those backing store
-							// entries are locked during the reads so we have to allocate space and can't page anything
-							// out. This is fine if we can load the whole file, but under constrained memory we want to
-							// only keep stuff that's going to get used soon - so we load the stuff we _aren't_ going to use
-							// soon so it can get paged out. Then we load in batches so we can release the locks and let them
-							// page out. 
-							// LONGTERM - detect this condition and write directly to the disk backing store? Ideally this
-							// would be something we can retain across restarts as right now any harvested chunks get lost
-							// on abort and cause a download from the cloud source (not chunkdb!)
-							Algo::SortBy(ChunksFromFileWeNeed, &FNeededChunk::NextUsageIndex, TGreater<int32>());
-
-							constexpr int32 HarvestBatchSize = 16 << 20; // 16 MB
-
-							for (int32 ChunkIndex = 0; ChunkIndex < ChunksFromFileWeNeed.Num();)
-							{
-								int32 BatchSize = 0;
-								int32 ChunkEndIndex = ChunkIndex + 1;
-								for (; ChunkEndIndex < ChunksFromFileWeNeed.Num(); ChunkEndIndex++)
-								{
-									if (BatchSize + ChunksFromFileWeNeed[ChunkEndIndex].ChunkSize > HarvestBatchSize)
-									{
-										break;
-									}
-									BatchSize += ChunksFromFileWeNeed[ChunkEndIndex].ChunkSize;
-								}
-
-								PendingHarvestRequests.store(ChunkEndIndex - ChunkIndex, std::memory_order_release);
-								
-								for (int32 DispatchIndex = ChunkIndex; DispatchIndex < ChunkEndIndex; DispatchIndex++)
-								{
-									FNeededChunk& Chunk = ChunksFromFileWeNeed[DispatchIndex];
-
-									SetChunkLocation(Chunk.Id, EConstructorChunkLocation::Memory);
-
-									FMutableMemoryView Destination = BackingStore->ReserveAndLockEntry(Chunk.Id, Chunk.ChunkSize, Chunk.LastUsageIndex);
-									if (Destination.GetSize() == 0)
-									{
-										// Call the completion function so we decrement the request count,
-										// but this is a consistency failure so we can't use the results.
-										ChunkHarvestCompletedFn(Chunk.Id, false, false, 0);
-										bHarvestSuccess = false;
-									}
-
-									IConstructorChunkSource::FRequestProcessFn HarvestFn = InstallSource->CreateRequest(Chunk.Id, Destination, 0, IConstructorChunkSource::FChunkRequestCompleteDelegate::CreateRaw(this, &FBuildPatchFileConstructor::ChunkHarvestCompletedFn));
-									
-									HarvestFn(false);
-								}
-
-								// The read is synchronous but the verification is not, so we still need to do the wait.
-								{
-									TRACE_CPUPROFILER_EVENT_SCOPE(HarvestWait);
-									WakeUpDispatchThreadEvent->Wait();
-								}
-							
-								if (bHarvestSuccess)
-								{
-									// Unlock any memory store entries
-									for (int32 DispatchIndex = ChunkIndex; DispatchIndex < ChunkEndIndex; DispatchIndex++)
-									{
-										// The read could have failed - in which case the location was switched to cloud from Memory.
-										// We handle this here so we don't have to deal with synchronization in the completion function.
-										if (ChunkLocations[ChunksFromFileWeNeed[DispatchIndex].Id] == EConstructorChunkLocation::Cloud)
-										{
-											if (!BackingStore->DereserveHarvestingEntry(ChunksFromFileWeNeed[DispatchIndex].Id))
-											{
-												bHarvestSuccess = false;
-											}
-										}
-										else
-										{
-											if (!BackingStore->CommitAndReleaseEntry(ChunksFromFileWeNeed[DispatchIndex].Id))
-											{
-												bHarvestSuccess = false;
-											}
-										}
-									}
-								}
-								
-								if (bHarvestSuccess)
-								{
-									bHarvestSuccess = BackingStore->CheckNoLocks(true);
-								}
-								
-								ChunkIndex = ChunkEndIndex;
-
-								if (bShouldAbort || !bHarvestSuccess)
-								{
-									break;
-								}
-							}
-						} // end if we need to harvest
-
-					} // end if install source exists.
-
-					if (!bHarvestSuccess)
-					{
-						Abort();
-						InstallerError->SetError(EBuildPatchInstallError::FileConstructionFail, ConstructionErrorCodes::InternalConsistencyFailure);
-					}
-
-					OnBeforeDeleteFile().Broadcast(FileToDelete);
-					{
-						// This can take forever due to file system filters
-						TRACE_CPUPROFILER_EVENT_SCOPE(Delete);
-						IFileManager::Get().Delete(*FileToDelete, bRequireExists, bEvenReadOnly);
-					}
-				}
-			}
-		}
-		else
-		{
-			// This will only record and log if a failure was not already registered.
-			Abort();
-			InstallerError->SetError(EBuildPatchInstallError::FileConstructionFail, ConstructionErrorCodes::UnknownFail);
-			UE_LOG(LogBuildPatchServices, Error, TEXT("FBuildPatchFileConstructor: Failed to build %s "), *FileToConstruct);
-		}
-		FileConstructorStat->OnFileCompleted(FileToConstruct, bFileSuccess);
-
-		// Wait while paused.
-		FileConstructorHelpers::WaitWhilePaused(bIsPaused, bShouldAbort);
-	}
+	ConstructFiles(ResumeData);
 
 	// Mark resume complete if we didn't have work to do.
 	if (!bIsDownloadStarted)
@@ -1938,6 +2006,8 @@ struct FRequestInfo
 	// Splats offset into this.
 	FMutableMemoryView DestinationBuffer;
 
+	FFileConstructionState* File = nullptr;
+
 	// We can only read direct in some cases.
 	bool bReadIntoMemoryStore = false;
 
@@ -1955,28 +2025,36 @@ struct FBatchState
 {
 	int32_t BatchId = 0;
 
-	TMap<FGuid, FRequestInfo> Requests;			
+	TMap<FGuid, FRequestInfo> Requests;
 	std::atomic_int32_t PendingRequestCount = 0;
-	std::atomic_int32_t FailedRequestCount = 0;			
+	std::atomic_int32_t FailedRequestCount = 0;
+
+	// Reads for the batch end up here, and will be written to the target
+	// file in order.
+	FMutableMemoryView BatchBuffer;
 
-	int32 BatchIndex = 0;
 	EConstructionError Error = EConstructionError::None;
 	FGuid ErrorContextGuid;
 
-	int32 StartChunkFileIndex = 0;
+	FFileConstructionState* OwningFile = nullptr;
+
+	int32 StartChunkPartIndex = 0;
 	int32 ChunkCount = 0;
 
+	// If true, this batch never reads or writes, it exists to complete the empty file in order.
+	bool bIsEmptyFileSentinel = false;
 	bool bNeedsWrite = false;
-
 	bool bIsReading = false;
+	bool bIsWriting = false;	
 
-	bool BufferIsAvailableForRead() const { return !bIsWriting && !bIsReading && !bNeedsWrite; }
-
-	bool bIsWriting = false;
-
-	// Set by the writing thread when the write completed prior to waking up the dispatch thread.
+	// Set by the completing threads when the batch is done.
 	std::atomic<bool> bIsFinished = true;
 	
+	FBatchState()
+	{
+		static std::atomic_int32_t UniqueBatchId = 1;
+		BatchId = UniqueBatchId.fetch_add(1);
+	}
 };
 
 // Called from basically any thread.
@@ -2026,10 +2104,10 @@ void FBuildPatchFileConstructor::RequestCompletedFn(const FGuid& Guid, bool bAbo
 			// need to tell the user about the download requirement though:
 			uint64 ChunkSize = Configuration.ManifestSet->GetDownloadSize(Guid);
 
-			DownloadRequirementLock.Lock();
+			ChunkLocationsLock.WriteLock();
 			DownloadRequirement += ChunkSize;
 			CloudSource->PostRequiredByteCount(DownloadRequirement);
-			DownloadRequirementLock.Unlock();
+			ChunkLocationsLock.WriteUnlock();
 
 		}
 
@@ -2076,9 +2154,9 @@ void FBuildPatchFileConstructor::RequestCompletedFn(const FGuid& Guid, bool bAbo
 				TotalToDestinationBuffer += Request->Splats[0].BytesToCopy;
 			}
 
-			FileProgressLock.Lock();
-			FileProgress += TotalToDestinationBuffer;
-			FileProgressLock.Unlock();
+			Request->File->ProgressLock.Lock();
+			Request->File->Progress += TotalToDestinationBuffer;
+			Request->File->ProgressLock.Unlock();
 		}
 		
 		if (Request->Batch->PendingRequestCount.fetch_add(-1, std::memory_order_relaxed) == 1)
@@ -2099,16 +2177,17 @@ IConstructorChunkSource::FRequestProcessFn FBuildPatchFileConstructor::CreateWri
 			return;
 		}
 
-		TArray<uint8>& WriteBuffer = WriteBuffers[Batch->BatchIndex];
+		// Has to be mutable because of the serialize call.
+		FMutableMemoryView WriteBuffer = Batch->BatchBuffer;
 
 		// Manage write limits.
 		if (bStallWhenFileSystemThrottled)
 		{
-			int64 AvailableBytes = FileSystem->GetAllowedBytesToWriteThrottledStorage(*File->GetArchiveName());
-			while (WriteBuffer.Num() > AvailableBytes)
+			uint64 AvailableBytes = FileSystem->GetAllowedBytesToWriteThrottledStorage(*File->GetArchiveName());
+			while (WriteBuffer.GetSize() > AvailableBytes)
 			{
-				UE_LOG(LogBuildPatchServices, Display, TEXT("Avaliable write bytes to write throttled storage exhausted (%s).  Sleeping %ds.  Bytes needed: %u, bytes available: %lld")
-					, *File->GetArchiveName(), SleepTimeWhenFileSystemThrottledSeconds, WriteBuffer.Num(), AvailableBytes);
+				UE_LOG(LogBuildPatchServices, Display, TEXT("Avaliable write bytes to write throttled storage exhausted (%s).  Sleeping %ds.  Bytes needed: %llu, bytes available: %llu")
+					, *File->GetArchiveName(), SleepTimeWhenFileSystemThrottledSeconds, WriteBuffer.GetSize(), AvailableBytes);
 				FPlatformProcess::Sleep(SleepTimeWhenFileSystemThrottledSeconds);
 				AvailableBytes = FileSystem->GetAllowedBytesToWriteThrottledStorage(*File->GetArchiveName());
 
@@ -2125,10 +2204,10 @@ IConstructorChunkSource::FRequestProcessFn FBuildPatchFileConstructor::CreateWri
 		
 		{
 			TRACE_CPUPROFILER_EVENT_SCOPE(WriteThread_Serialize)
-			File->Serialize(WriteBuffer.GetData(), WriteBuffer.Num());
+			File->Serialize(WriteBuffer.GetData(), WriteBuffer.GetSize());
 		}
 
-		ActivityRecord.Size = WriteBuffer.Num();
+		ActivityRecord.Size = WriteBuffer.GetSize();
 		ActivityRecord.CyclesEnd = FStatsCollector::GetCycles();
 		FileConstructorStat->OnAfterWrite(ActivityRecord);
 		
@@ -2179,9 +2258,9 @@ void FBuildPatchFileConstructor::CompleteReadBatch(const FFileManifest& FileMani
 						TotalToDestinationBuffer += Splat.BytesToCopy;
 					}
 
-					FileProgressLock.Lock();
-					FileProgress += TotalToDestinationBuffer;
-					FileProgressLock.Unlock();
+					Request.File->ProgressLock.Lock();
+					Request.File->Progress += TotalToDestinationBuffer;
+					Request.File->ProgressLock.Unlock();
 				}
 				else
 				{
@@ -2213,10 +2292,10 @@ void FBuildPatchFileConstructor::CompleteReadBatch(const FFileManifest& FileMani
 		// Retire the chunks we've used. This has to be in order 
 		for (int32 i = 0; i < Batch.ChunkCount; i++)
 		{
-			if (!ChunkReferenceTracker->PopReference(FileManifest.ChunkParts[Batch.StartChunkFileIndex + i].Guid))
+			if (!ChunkReferenceTracker->PopReference(FileManifest.ChunkParts[Batch.StartChunkPartIndex + i].Guid))
 			{
 				Batch.Error = EConstructionError::TrackingError;
-				Batch.ErrorContextGuid = FileManifest.ChunkParts[Batch.StartChunkFileIndex + i].Guid;
+				Batch.ErrorContextGuid = FileManifest.ChunkParts[Batch.StartChunkPartIndex + i].Guid;
 			}
 		}
 	}
@@ -2237,23 +2316,20 @@ void FBuildPatchFileConstructor::CompleteReadBatch(const FFileManifest& FileMani
 // Note that there could be outstanding reads and we can no longer rely on the completion functions
 // to wake up the dispatch thread (which we must run on), so we can't ever wait if this returns -1.
 // Check Batch.Error for the error on -1 return.
-int32 FBuildPatchFileConstructor::StartReadBatch(const FFileManifest& FileManifest, int32 StartFileChunkPartIndex, FBatchState& Batch)
+void FBuildPatchFileConstructor::StartReadBatch(FFileConstructionState& CurrentFile, FBatchState& Batch)
 {
-	static std::atomic_int32_t UniqueBatchId = 1;
-
+	CurrentFile.OutstandingBatches++;
 	Batch.Requests.Reset();
-	Batch.BatchId = UniqueBatchId.fetch_add(1);
-	Batch.StartChunkFileIndex = StartFileChunkPartIndex;
-
+	Batch.StartChunkPartIndex = CurrentFile.NextChunkPartToRead;
 
-	int32 EndChunkIdx = StartFileChunkPartIndex;
+	int32 EndChunkIdx = Batch.StartChunkPartIndex;
 	{
 		TRACE_CPUPROFILER_EVENT_SCOPE(CFFC_Count);
 		uint32 BufferFillLevel = 0;
-		for (; EndChunkIdx < FileManifest.ChunkParts.Num(); ++EndChunkIdx)
+		for (; EndChunkIdx < CurrentFile.FileManifest->ChunkParts.Num(); ++EndChunkIdx)
 		{
-			const FChunkPart& ChunkPart = FileManifest.ChunkParts[EndChunkIdx];
-			if (ChunkPart.Size + BufferFillLevel > WriteBufferSize)
+			const FChunkPart& ChunkPart = CurrentFile.FileManifest->ChunkParts[EndChunkIdx];
+			if (ChunkPart.Size + BufferFillLevel > Batch.BatchBuffer.GetSize())
 			{
 				break;
 			}
@@ -2262,6 +2338,7 @@ int32 FBuildPatchFileConstructor::StartReadBatch(const FFileManifest& FileManife
 			Request.Guid = ChunkPart.Guid;
 			Request.Batch = &Batch;
 			Request.ChunkSize = Configuration.ManifestSet->GetChunkInfo(ChunkPart.Guid)->WindowSize;
+			Request.File = &CurrentFile;
 			
 			FRequestInfo::FRequestSplat& Splat = Request.Splats.AddDefaulted_GetRef();
 			Splat.DestinationOffset = BufferFillLevel;
@@ -2271,14 +2348,13 @@ int32 FBuildPatchFileConstructor::StartReadBatch(const FFileManifest& FileManife
 			BufferFillLevel += ChunkPart.Size;
 		}
 
-		WriteBuffers[Batch.BatchIndex].SetNumUninitialized(BufferFillLevel, EAllowShrinking::No);
+		// Trim our view to the amount we actually used. The calling code will reclaim the rest.
+		Batch.BatchBuffer.LeftInline(BufferFillLevel);
 	}
 
-	Batch.ChunkCount = EndChunkIdx - StartFileChunkPartIndex;
+	Batch.ChunkCount = EndChunkIdx - Batch.StartChunkPartIndex;
 	Batch.PendingRequestCount = Batch.Requests.Num();
 	Batch.FailedRequestCount = 0;
-
-	UE_LOG(LogBuildPatchServices, Verbose, TEXT("Starting ReadBatch: %d, Chunks=%d"), Batch.BatchId, Batch.ChunkCount);
 	
 	//
 	// IMPORTANT!!!
@@ -2310,10 +2386,10 @@ int32 FBuildPatchFileConstructor::StartReadBatch(const FFileManifest& FileManife
 
 
 		FRequestInfo& Request = RequestPair.Value;
-		Request.DestinationBuffer = MakeMemoryView(WriteBuffers[Batch.BatchIndex]);
+		Request.DestinationBuffer = Batch.BatchBuffer;
 
 		
-		EConstructorChunkLocation* ChunkLocationPtr = ChunkLocations.Find(Request.Guid);
+		const EConstructorChunkLocation* ChunkLocationPtr = ChunkLocations.Find(Request.Guid);
 		if (ChunkLocationPtr == nullptr)
 		{
 			Batch.Error = EConstructionError::InternalConsistencyError;
@@ -2404,7 +2480,7 @@ int32 FBuildPatchFileConstructor::StartReadBatch(const FFileManifest& FileManife
 		}
 	}
 
-	return EndChunkIdx;
+	CurrentFile.NextChunkPartToRead = EndChunkIdx;
 }
 
 void FBuildPatchFileConstructor::WakeUpDispatch()
@@ -2412,99 +2488,56 @@ void FBuildPatchFileConstructor::WakeUpDispatch()
 	WakeUpDispatchThreadEvent->Trigger();
 }
 
-bool FBuildPatchFileConstructor::ConstructFileFromChunks(const FString& BuildFilename, const FFileManifest& FileManifest, bool bResumeExisting)
+void FBuildPatchFileConstructor::InitFile(FFileConstructionState& CurrentFile, const FResumeData& ResumeData)
 {
-	TRACE_CPUPROFILER_EVENT_SCOPE(ConstructFileFromChunks);
-	UE_LOG(LogBuildPatchServices, Verbose, TEXT("Starting File: %s"), *BuildFilename);
+	if (!CurrentFile.bSuccess)
+	{
+		return;
+	}
 
-	bool bSuccess = true;
+	const int64 FileSize = CurrentFile.FileManifest->FileSize;
 
-	FGuid ErrorContextGuid;
-	EConstructionError ConstructionError = EConstructionError::None;
+	// Check resume status for this file.
+	const bool bFilePreviouslyComplete = ResumeData.FilesCompleted.Contains(CurrentFile.BuildFilename);	
 
-	FString NewFilename = Configuration.StagingDirectory / BuildFilename;
+	// Construct or skip the file.
+	if (bFilePreviouslyComplete)
+	{
+		CountBytesProcessed(FileSize);
 
-	// Calculate the hash as we write the data
-	FSHA1 HashState;
+		UE_LOG(LogBuildPatchServices, Display, TEXT("Skipping completed file %s"), *CurrentFile.BuildFilename);
+		// Go through each chunk part, and dereference it from the reference tracker.
+		for (const FChunkPart& ChunkPart : CurrentFile.FileManifest->ChunkParts)
+		{
+			if (!ChunkReferenceTracker->PopReference(ChunkPart.Guid))
+			{
+				CurrentFile.bSuccess = false;
+				CurrentFile.ConstructionError = EConstructionError::TrackingError;
+				break;
+			}
+		}
+
+		CurrentFile.bSkippedConstruction = true;
+		return;
+	}
 
-	if (!FileManifest.SymlinkTarget.IsEmpty())
+	if (!CurrentFile.bSuccess &&
+		!CurrentFile.FileManifest->SymlinkTarget.IsEmpty())
 	{
 #if PLATFORM_MAC
-		bSuccess = symlink(TCHAR_TO_UTF8(*FileManifest.SymlinkTarget), TCHAR_TO_UTF8(*NewFilename)) == 0;
+		CurrentFile.bSuccess = symlink(TCHAR_TO_UTF8(*CurrentFile.FileManifest->SymlinkTarget), TCHAR_TO_UTF8(*CurrentFile.NewFilename)) == 0;
+		CurrentFile.bSkippedConstruction = true;
 #else
 		const bool bSymlinkNotImplemented = false;
 		check(bSymlinkNotImplemented);
-		bSuccess = false;
+		CurrentFile.bSuccess = false;
 #endif
-		return bSuccess;
-	}
-
-	// Check for resuming of existing file
-	int64 StartPosition = 0;
-	int32 StartChunkPart = 0;
-	if (bResumeExisting)
-	{
-		// We have to read in the existing file so that the hash check can still be done.
-		TUniquePtr<FArchive> NewFileReader(IFileManager::Get().CreateFileReader(*NewFilename));
-		if (NewFileReader.IsValid())
-		{
-			// Start with a sensible buffer size for reading. 4 MiB.
-			const int32 ReadBufferSize = 4 * 1024 * 1024;
-			// Read buffer
-			TArray<uint8> ReadBuffer;
-			ReadBuffer.Empty(ReadBufferSize);
-			ReadBuffer.SetNumUninitialized(ReadBufferSize);
-
-			// Reuse the entire file. Previously this truncated to size - 1kb but that's unlikely to catch our actual
-			// issue because its less than a sector size and makes it so that a graceful resume requires potentially retired
-			// chunks.
-			StartPosition = NewFileReader->TotalSize();
-
-			// We'll also find the correct chunkpart to start writing from
-			int64 ByteCounter = 0;
-			for (int32 ChunkPartIdx = StartChunkPart; ChunkPartIdx < FileManifest.ChunkParts.Num() && !bShouldAbort; ++ChunkPartIdx)
-			{
-				const FChunkPart& ChunkPart = FileManifest.ChunkParts[ChunkPartIdx];
-				const int64 NextBytePosition = ByteCounter + ChunkPart.Size;
-				if (NextBytePosition <= StartPosition)
-				{
-					// Ensure buffer is large enough
-					ReadBuffer.SetNumUninitialized(ChunkPart.Size, EAllowShrinking::No);
-					ISpeedRecorder::FRecord ActivityRecord;
-					// Read data for hash check
-					FileConstructorStat->OnBeforeRead();
-					ActivityRecord.CyclesStart = FStatsCollector::GetCycles();
-					NewFileReader->Serialize(ReadBuffer.GetData(), ChunkPart.Size);
-					ActivityRecord.CyclesEnd = FStatsCollector::GetCycles();
-					ActivityRecord.Size = ChunkPart.Size;
-					HashState.Update(ReadBuffer.GetData(), ChunkPart.Size);
-					FileConstructorStat->OnAfterRead(ActivityRecord);
-					// Count bytes read from file
-					ByteCounter = NextBytePosition;
-					// Set to resume from next chunk part
-					StartChunkPart = ChunkPartIdx + 1;
-					// Inform the reference tracker of the chunk part skip
-					bSuccess = ChunkReferenceTracker->PopReference(ChunkPart.Guid) && bSuccess;
-					CountBytesProcessed(ChunkPart.Size);
-					FileConstructorStat->OnFileProgress(BuildFilename, NewFileReader->Tell());
-					// Wait if paused
-					FileConstructorHelpers::WaitWhilePaused(bIsPaused, bShouldAbort);
-				}
-				else
-				{
-					// No more parts on disk
-					break;
-				}
-			}
-			// Set start position to the byte we got up to
-			StartPosition = ByteCounter;
-			// Close file
-			NewFileReader->Close();
-		}
 	}
+}
 
-	// If we haven't done so yet, make the initial disk space check. We do this after resume
-	// so that we know how much to discount from our current file size.
+bool FBuildPatchFileConstructor::HandleInitialDiskSizeCheck(const FFileManifest& FileManifest, int64 StartPosition)
+{
+	// Returns false if we failed the disk space check.
 	if (!bInitialDiskSizeCheck)
 	{
 		bInitialDiskSizeCheck = true;
@@ -2551,341 +2584,260 @@ bool FBuildPatchFileConstructor::ConstructFileFromChunks(const FString& BuildFil
 		}
 	}
 
-	// Now we can make sure the chunk cache knows to start downloading chunks
-	if (!bIsDownloadStarted)
+	return true;
+}
+
+void FBuildPatchFileConstructor::ResumeFile(FFileConstructionState& FileToResume)
+{
+	if (!FileToResume.bSuccess || FileToResume.bSkippedConstruction)
 	{
-		bIsDownloadStarted = true;
-		FileConstructorStat->OnResumeCompleted();
+		return;
 	}
 
-	// Attempt to create the file
-	ISpeedRecorder::FRecord ActivityRecord;
-	FileConstructorStat->OnBeforeAdminister();
-	ActivityRecord.CyclesStart = FStatsCollector::GetCycles();
-	TUniquePtr<FArchive> NewFile = FileSystem->CreateFileWriter(*NewFilename, bResumeExisting ? EWriteFlags::Append : EWriteFlags::None);
-	int32 CreateFilePlatformLastError = FPlatformMisc::GetLastError();
-	ActivityRecord.CyclesEnd = FStatsCollector::GetCycles();
-	ActivityRecord.Size = 0;
-	FileConstructorStat->OnAfterAdminister(ActivityRecord);
-	bSuccess = NewFile != nullptr;
-	if (bSuccess)
+	// We have to read in the existing file so that the hash check can still be done.
+	TUniquePtr<FArchive> NewFileReader(IFileManager::Get().CreateFileReader(*FileToResume.NewFilename));
+	if (!NewFileReader.IsValid())
 	{
-		FileProgress = StartPosition;
+		// We don't fail if we can't read in the previous file - we try and rebuild it from scratch.
+		// (Note that the likely outcome here is we can't open the file for write either and fail to
+		// install - we're only here if we were supposed to be resuming!)
+		return;
+	}
 
-		// Seek to file write position
-		if (NewFile->Tell() != StartPosition)
-		{
-			FileConstructorStat->OnBeforeAdminister();
-			ActivityRecord.CyclesStart = FStatsCollector::GetCycles();
-			NewFile->Seek(StartPosition);
-			ActivityRecord.CyclesEnd = FStatsCollector::GetCycles();
-			ActivityRecord.Size = 0;
-			FileConstructorStat->OnAfterAdminister(ActivityRecord);
-		}
+	const int32 ReadBufferSize = 4 * 1024 * 1024;
+	// Read buffer
+	TArray<uint8> ReadBuffer;
+	ReadBuffer.SetNumUninitialized(ReadBufferSize);
 
-		constexpr int32 BatchCount = UE_ARRAY_COUNT(WriteBuffers);
-		FBatchState BatchState[BatchCount];
-		for (int32 BatchIndex = 0; BatchIndex < BatchCount; BatchIndex++)
+	// We don't allow resuming mid-chunk for simplicity, so eat entire chunks until
+	// we can't anymore.
+	int64 OnDiskSize = NewFileReader->TotalSize();
+	int64 ByteCounter = 0;
+	for (int32 ChunkPartIdx = 0; ChunkPartIdx < FileToResume.FileManifest->ChunkParts.Num() && !bShouldAbort; ++ChunkPartIdx)
+	{
+		const FChunkPart& ChunkPart = FileToResume.FileManifest->ChunkParts[ChunkPartIdx];
+		const int64 NextBytePosition = ByteCounter + ChunkPart.Size;
+		ByteCounter = NextBytePosition;
+		if (NextBytePosition <= OnDiskSize)
 		{
-			BatchState[BatchIndex].BatchIndex = BatchIndex;
-		}
-
-		FBatchState* WritingBatch = nullptr;
-	
-		// Fill read/write buffers in batches until we complete the file.
-		int32 CurrentFileChunkPartIndex = StartChunkPart;
-
-		uint64 LastSeenFileProgress = FileProgress; // no need to lock - nothing async yet.
+			// Ensure buffer is large enough
+			ReadBuffer.SetNumUninitialized(ChunkPart.Size, EAllowShrinking::No);
 
-		for (;;)
-		{
-			// Check state. We can continue to run this after a failure has been encountered in order
-			// to drain any async tasks so we know we can shut down safely.
 			{
-				TRACE_CPUPROFILER_EVENT_SCOPE(CFFC_StateCheck);
-
-				// Is our next write done?
-				if (WritingBatch && WritingBatch->bIsFinished.load(std::memory_order_acquire))
-				{
-					WritingBatch->bIsWriting = false;
-					WritingBatch = nullptr;
-				
-					if (NewFile->IsError())
-					{
-						bSuccess = false;
-
-						// SerializeError signals the error management code to check for low disk space
-						// since we currently have no way of routing out of disk errors from the UE low level
-						// writing code.
-						ConstructionError = EConstructionError::SerializeError;
-					}
-					
-					// If that was the last write, then we are done.
-					if (CurrentFileChunkPartIndex == FileManifest.ChunkParts.Num())
-					{
-						int32 WriteBatchCount = 0;
-						for (uint32 i=0; i < BatchCount; i++)
-						{
-							WriteBatchCount += BatchState[i].bNeedsWrite ? 1 : 0;
-						}
-
-						if (WriteBatchCount == 0)
-						{
-							break;
-						}
-					}
-				}
-
-				// are any reads done?
-				for (;;)
-				{
-					// We have to retire reads in order because later reads could be waiting for data to
-					// get placed correctly by the earlier ones.
-					int32 EarliestReadingBatch = -1;
-					for (int32 i = 0; i < BatchCount; i++)
-					{
-						if (BatchState[i].bIsReading)
-						{
-							if (EarliestReadingBatch == -1 || BatchState[i].StartChunkFileIndex < BatchState[EarliestReadingBatch].StartChunkFileIndex)
-							{
-								EarliestReadingBatch = i;
-							}
-						}
-					}
+				FReadScope _(FileConstructorStat, ChunkPart.Size);
+				NewFileReader->Serialize(ReadBuffer.GetData(), ChunkPart.Size);
+			}
 
-					if (EarliestReadingBatch == -1)
-					{
-						break; // Nothing reading.
-					}
+			FileToResume.HashState.Update(ReadBuffer.GetData(), ChunkPart.Size);
 
-					if (!BatchState[EarliestReadingBatch].bIsFinished.load(std::memory_order_acquire))
-					{
-						break; // Read we care about isn't done.
-					}
+			// Update resume start position
+			FileToResume.StartPosition = NextBytePosition;
+			FileToResume.NextChunkPartToRead = ChunkPartIdx + 1;
 
-					CompleteReadBatch(FileManifest, BatchState[EarliestReadingBatch]);
-					BatchState[EarliestReadingBatch].bIsReading = false;
-					BatchState[EarliestReadingBatch].bNeedsWrite = true;
+			// Inform the reference tracker of the chunk part skip
+			if (!ChunkReferenceTracker->PopReference(ChunkPart.Guid))
+			{
+				FileToResume.bSuccess = false;
+				FileToResume.ConstructionError = EConstructionError::TrackingError;
+				FileToResume.ErrorContextGuid = ChunkPart.Guid;
+				break;
+			}
 
-					if (ConstructionError != EConstructionError::None ||
-						BatchState[EarliestReadingBatch].Error != EConstructionError::None)
-					{
-						ErrorContextGuid = BatchState[EarliestReadingBatch].ErrorContextGuid;
-						ConstructionError = BatchState[EarliestReadingBatch].Error;
-						bSuccess = false;
-					}
+			CountBytesProcessed(ChunkPart.Size);
+			FileConstructorStat->OnFileProgress(FileToResume.BuildFilename, NewFileReader->Tell());
+			// Wait if paused
+			FileConstructorHelpers::WaitWhilePaused(bIsPaused, bShouldAbort);
+		}
+		else
+		{
+			// We can't consume any more full chunks from the part list, bail.
+			break;
+		}
+	}
 
-					// Loop around to check if another is done.
-				}
+	NewFileReader->Close();
+	FileToResume.bIsResumedFile = true;
+}
 
-				if (bSuccess &&
-					bShouldAbort)
-				{
-					ConstructionError = EConstructionError::Aborted;
-					bSuccess = false;
-				}
+void FBuildPatchFileConstructor::OpenFileToConstruct(FFileConstructionState& CurrentFile)
+{
+	if (!CurrentFile.bSuccess || CurrentFile.bSkippedConstruction)
+	{
+		return;
+	}
 
-				// If we are in a failure case, we want to clear out anything happening with the cloud source.
-				// We'll hit this over and over as we "drain", but we need to do that anyway because further failures
-				// might try to queue more cloud reads.
-				if (!bSuccess)
-				{
-					CloudSource->Abort();
-				}
+	// Attempt to create the file
+	{
+		FAdministrationScope _(FileConstructorStat);
+		CurrentFile.NewFile = FileSystem->CreateFileWriter(*CurrentFile.NewFilename, CurrentFile.bIsResumedFile ? EWriteFlags::Append : EWriteFlags::None);
+		CurrentFile.CreateFilePlatformLastError = FPlatformMisc::GetLastError();
+	}
 
-				if (!bIsPaused && bSuccess)
-				{
-					// Can we start another read?
-					// LONGTERM - if we are memory constrained and have an install source we can consider only dispatching
-					// single reads. This dramatically lowers our memory requirements as install sources often use a lot
-					// of small pieces of chunks - and we have to load the entire chunk into memory.
-					for (int32 i = 0; i < BatchCount; i++)
-					{
-						if (CurrentFileChunkPartIndex < FileManifest.ChunkParts.Num() &&
-							BatchState[i].BufferIsAvailableForRead())
-						{
-							BatchState[i].bNeedsWrite = true;
-							BatchState[i].bIsReading = true;
-							BatchState[i].bIsFinished.store(false, std::memory_order_release);
-							CurrentFileChunkPartIndex = StartReadBatch(FileManifest, CurrentFileChunkPartIndex, BatchState[i]);
-							if (CurrentFileChunkPartIndex == -1)
-							{
-								ConstructionError = BatchState[i].Error;
-								bSuccess = false;
-								break;
-							}
-						}
-					}
+	CurrentFile.bSuccess = CurrentFile.NewFile != nullptr;
+	if (!CurrentFile.bSuccess)
+	{
+		CurrentFile.ConstructionError = EConstructionError::CannotCreateFile;
+		return;
+	}
 
-					// Can we start a write?
-					if (bSuccess && WritingBatch == nullptr)
-					{
-						// We always have to issue writes in order and there can only be one.
-						int32 WriteBufferIndex = -1;
-						for (int32 i=0; i<BatchCount; i++)
-						{
-							if (BatchState[i].bNeedsWrite)
-							{
-								if (WriteBufferIndex == -1 ||
-									BatchState[WriteBufferIndex].StartChunkFileIndex > BatchState[i].StartChunkFileIndex)
-								{
-									WriteBufferIndex = i;
-								}
-							}
-						}
+	// Seek to file write position
+	if (CurrentFile.NewFile->Tell() != CurrentFile.StartPosition)
+	{
+		FAdministrationScope _(FileConstructorStat);
 
-						// However the next write might not be ready (fix possible of out of order read completion)
-						if (WriteBufferIndex >= 0 &&
-							!BatchState[WriteBufferIndex].bIsReading)
-						{
-							BatchState[WriteBufferIndex].bNeedsWrite = false;
-							BatchState[WriteBufferIndex].bIsWriting = true;
-							BatchState[WriteBufferIndex].bIsFinished.store(false, std::memory_order_relaxed);
+		// Currently no way of checking if the seek succeeded. If it didn't and further reads succeed, then
+		// we can end up with a bad file on disk and not know it as the hash is assuming this worked - requires
+		// full load-and-hash verification to find.
+		CurrentFile.NewFile->Seek(CurrentFile.StartPosition);
+	}
 
-							WritingBatch = BatchState + WriteBufferIndex;
-							
-							// Launch the write and hash the buffer on this thread.
-							UE_LOG(LogBuildPatchServices, Verbose, TEXT("Writing Batch: %d"), BatchState[WriteBufferIndex].BatchId);
+	CurrentFile.Progress = CurrentFile.StartPosition;
+	CurrentFile.LastSeenProgress = CurrentFile.Progress;
+}
 
-							IConstructorChunkSource::FRequestProcessFn WriteFn = CreateWriteRequest(NewFile.Get(), BatchState[WriteBufferIndex]);
-							QueueGenericThreadTask(WriteThreadIndex, MoveTemp(WriteFn));
-					
-							{
-								TRACE_CPUPROFILER_EVENT_SCOPE(CFFC_Hash);
-								HashState.Update(WriteBuffers[WriteBufferIndex].GetData(), WriteBuffers[WriteBufferIndex].Num());
-							}
-						}
-					} // end if checking for write
-				} // end not paused
-			} // end state check
-
-			// If the file progress changed since we last saw it, post the update.
-			// Note that we want this to update reasonably often but we're about to wait
-			// potentially until all reads complete - however the only time things actually
-			// take a long time wall-clock wise is when we are downloading, and the cloud
-			// source will then prevent us from sleeping too long, so we actually catch these
-			// updates.
-			// We do this from here to ensure we always increase rather than risk multi thread
-			// races.
-			{
-				uint64 CurrentFileProgress;
-				FileProgressLock.Lock();
-				CurrentFileProgress = FileProgress;
-				FileProgressLock.Unlock();
-
-				if (CurrentFileProgress != LastSeenFileProgress)
-				{
-					// this updates the overall install progress.
-					CountBytesProcessed(CurrentFileProgress - LastSeenFileProgress);
-					LastSeenFileProgress = CurrentFileProgress;
-					FileConstructorStat->OnFileProgress(BuildFilename, CurrentFileProgress);
-				}
-			}
 
+void FBuildPatchFileConstructor::CompleteConstructedFile(FFileConstructionState& CurrentFile)
+{
+	TRACE_CPUPROFILER_EVENT_SCOPE(Constructor_CompleteFile);
+	if (!CurrentFile.bSkippedConstruction)
+	{
+		if (CurrentFile.bSuccess && CurrentFile.NewFile->IsError())
+		{
+			CurrentFile.ConstructionError = EConstructionError::SerializeError;
+			CurrentFile.bSuccess = false;
+		}
 
-			uint32 WaitTimeMs = TNumericLimits<uint32>::Max();
-			{
-				TRACE_CPUPROFILER_EVENT_SCOPE(CFFC_TickCloud);
+		// Close the file writer
+		bool bArchiveSuccess = false;
+		if (CurrentFile.NewFile)
+		{
+			FAdministrationScope _(FileConstructorStat);
+			bArchiveSuccess = CurrentFile.NewFile->Close();
+			CurrentFile.NewFile.Reset();
+		}
+		else
+		{
+			bArchiveSuccess = true;
+		}
 
-				// Cloud source may want to be woken up after a specific time.
-				CloudSource->Tick(!bIsPaused && bSuccess, WaitTimeMs);
-			}
+		// Check for final success
+		if (CurrentFile.bSuccess && !bArchiveSuccess)
+		{
+			CurrentFile.ConstructionError = EConstructionError::SerializeError;
+			CurrentFile.bSuccess = false;
+		}
 
-			int32 ReadingBatchCount = 0;
-			for (int32 i = 0; i < BatchCount; i++)
+		// We can't check for zero locks if we have multiple files in flight because the other
+		// files hold locks.
+		if (CurrentFile.bSuccess && !bAllowMultipleFilesInFlight)
+		{
+			if (!BackingStore->CheckNoLocks(false))
 			{
-				ReadingBatchCount += BatchState[i].bIsReading ? 1 : 0;
+				CurrentFile.bSuccess = false;
+				CurrentFile.ConstructionError = EConstructionError::InternalConsistencyError;
 			}
+		}
+	
+		// Verify the hash for the file that we created
+		if (CurrentFile.bSuccess)
+		{
+			CurrentFile.HashState.Final();
 
-			if (!bSuccess)
+			FSHAHash HashValue;
+			CurrentFile.HashState.GetHash(HashValue.Hash);
+			CurrentFile.bSuccess = HashValue == CurrentFile.FileManifest->FileHash;
+			if (!CurrentFile.bSuccess)
 			{
-				// We can only bail when all our async jobs have completed.
-				if (WritingBatch == nullptr && ReadingBatchCount == 0)
-				{
-					break;
-				}
+				CurrentFile.ConstructionError = EConstructionError::OutboundDataError;
 			}
+		}
 
+#if PLATFORM_MAC
+		if (CurrentFile.bSuccess && 
+			EnumHasAllFlags(CurrentFile.FileManifest.FileMetaFlags, EFileMetaFlags::UnixExecutable))
+		{
+			// Enable executable permission bit
+			struct stat FileInfo;
+			if (stat(TCHAR_TO_UTF8(*CurrentFile.NewFilename), &FileInfo) == 0)
 			{
-				if (WaitTimeMs == TNumericLimits<uint32>::Max())
+				CurrentFile.bSuccess = chmod(TCHAR_TO_UTF8(*NewFilename), FileInfo.st_mode | S_IXUSR | S_IXGRP | S_IXOTH) == 0;
+				if (!CurrentFile.bSuccess)
 				{
-					UE_LOG(LogBuildPatchServices, VeryVerbose, TEXT("Constructor waiting infinite. Reading batches: %d, Writing Batch: %d"), ReadingBatchCount, WritingBatch ? 1 : 0);
+					UE_LOG(LogBuildPatchServices, Error, TEXT("FBuildPatchFileConstructor: Failed to set exec bit %s"), *CurrentFile.BuildFilename);
 				}
-
-				TRACE_CPUPROFILER_EVENT_SCOPE(CFFC_Waiting);
-				// We have a bunch of stuff outstanding that will wake us up if something happens.
-				WakeUpDispatchThreadEvent->Wait(WaitTimeMs);
 			}
+		}
+#endif
 
-		} // end looping over
-
-		bSuccess = bSuccess && !NewFile->IsError();
-
-		// Update this for disk space requirements tracking below on error
-		StartPosition = NewFile->Tell();
-
-		// Close the file writer
-		FileConstructorStat->OnBeforeAdminister();
-		ActivityRecord.CyclesStart = FStatsCollector::GetCycles();
-
-		const bool bArchiveSuccess = NewFile->Close();
-		NewFile.Reset();
-		ActivityRecord.CyclesEnd = FStatsCollector::GetCycles();
-		ActivityRecord.Size = 0;
-		FileConstructorStat->OnAfterAdminister(ActivityRecord);
-
-		// Check for final success
-		if (ConstructionError == EConstructionError::None && !bArchiveSuccess)
+#if PLATFORM_ANDROID
+		if (CurrentFile.bSuccess)
 		{
-			ConstructionError = EConstructionError::SerializeError;
-			bSuccess = false;
+			IFileManager::Get().SetTimeStamp(*CurrentFile.NewFilename, FDateTime::UtcNow());
 		}
-	}
-	else
+#endif
+	} // end if we did actual construction work for this file.
+
+	if (CurrentFile.bSuccess)
 	{
-		ConstructionError = EConstructionError::CannotCreateFile;
+		ChunkDbSource->ReportFileCompletion(ChunkReferenceTracker->GetRemainingChunkCount());
 	}
 
-	bSuccess = bSuccess && BackingStore->CheckNoLocks(false);
+	// If we are destructive, remove the old file.
+	if (CurrentFile.bSuccess &&
+		Configuration.InstallMode == EInstallMode::DestructiveInstall)
+	{
+		FString FileToDelete = Configuration.InstallDirectory / CurrentFile.BuildFilename;
+		FPaths::NormalizeFilename(FileToDelete);
+		FPaths::CollapseRelativeDirectories(FileToDelete);
+		if (FileSystem->FileExists(*FileToDelete))
+		{
+			if (InstallSource && !HarvestChunksForCompletedFile(FileToDelete))
+			{
+				CurrentFile.bSuccess = false;
+				CurrentFile.ConstructionError = EConstructionError::InternalConsistencyError;
+			} // end if install source exists.
 
-	// Check for error state
-	if (!bSuccess)
+			OnBeforeDeleteFile().Broadcast(FileToDelete);
+			{
+				// This can take forever due to file system filters. If we throw this on an async
+				// job then we can go over our calculated disk space.
+				TRACE_CPUPROFILER_EVENT_SCOPE(Delete);
+				const bool bRequireExists = false;
+				const bool bEvenReadOnly = true;
+				IFileManager::Get().Delete(*FileToDelete, bRequireExists, bEvenReadOnly);
+			}
+		}
+	}
+
+	FileConstructorStat->OnFileCompleted(CurrentFile.BuildFilename, CurrentFile.bSuccess);
+	
+	// Report errors.
+	if (!CurrentFile.bSuccess)
 	{
-		if (ConstructionError == EConstructionError::SerializeError)
+		if (CurrentFile.ConstructionError == EConstructionError::SerializeError)
 		{
 			// Serialize error is our catchall file error right now. This should probably get
 			// migrated such that it's when we fail to load an existing chunk (i.e. corruption)
 			// but instead that shows up as a missing chunk.
 
-			// space issue if it looks like that's the problem.			
+			// Right now however we get this whenever we fail to write due to our disk space disappearing,
+			// and we'd like to report that, so we check and see if we're low on space if it looks like
+			// that's the problem.
 
 			uint64 TotalSize = 0;
 			uint64 FreeSize = 0;
-			uint64 LocalDiskSpaceRequired = TNumericLimits<uint64>::Max();
 			if (FPlatformMisc::GetDiskTotalAndFreeSpace(Configuration.InstallDirectory, TotalSize, FreeSize))
 			{
-				LocalDiskSpaceRequired = CalculateInProgressDiskSpaceRequired(FileManifest, StartPosition);
+				// Don't bother recalculating our disk space requirements - our original requirement is still
+				// valid, the only thing that's changed is someone took our space.
+				AvailableDiskSpace.store(FreeSize, std::memory_order_release);
 
-				// If we are delete-during-install this gets more complicated because we'll be freeing up
-				// space as we add.
-				if (Configuration.bDeleteChunkDBFilesAfterUse)
-				{
-					LocalDiskSpaceRequired = CalculateDiskSpaceRequirementsWithDeleteDuringInstall(ConstructionStack);
-				}
-			}
-
-			if (LocalDiskSpaceRequired != TNumericLimits<uint64>::Max())
-			{
 				// We're responding to an actual failure, which would have happened because we literally weren't able
 				// to write our write butter. Because of transient stuff this might not be correct so we double our
 				// write buffer size for this check.
-				if (FreeSize < (2 * WriteBufferSize))
+				if (FreeSize < ((uint64)MaxWriteBatchSize << 1))
 				{
-					// We've already failed so it makes sense to reevaluate how much extra we need. 
-					// I'm not sure I like using the same error wording for initial and ongoing disk space failure, but whatevs
-					AvailableDiskSpace.store(FreeSize, std::memory_order_release);
-					RequiredDiskSpace.store(LocalDiskSpaceRequired, std::memory_order_release);
-					ConstructionError = EConstructionError::OutOfDiskSpace;
+					CurrentFile.ConstructionError = EConstructionError::OutOfDiskSpace;
 				}
 				else
 				{
@@ -2901,8 +2853,19 @@ bool FBuildPatchFileConstructor::ConstructFileFromChunks(const FString& BuildFil
 		}
 
 		const bool bReportAnalytic = InstallerError->HasError() == false;
-		switch (ConstructionError)
+		switch (CurrentFile.ConstructionError)
 		{
+		case EConstructionError::OutboundDataError:
+			{
+				// Only report if the first error
+				if (bReportAnalytic)
+				{
+					InstallerAnalytics->RecordConstructionError(CurrentFile.BuildFilename, INDEX_NONE, TEXT("Serialised Verify Fail"));
+				}
+				UE_LOG(LogBuildPatchServices, Error, TEXT("FBuildPatchFileConstructor: Verify failed after constructing %s"), *CurrentFile.BuildFilename);
+				InstallerError->SetError(EBuildPatchInstallError::FileConstructionFail, ConstructionErrorCodes::OutboundCorrupt);
+				break;
+			}
 		case EConstructionError::OutOfDiskSpace:
 			{
 				uint64 LocalAvailableDiskSpace = AvailableDiskSpace.load(std::memory_order_acquire);
@@ -2919,20 +2882,20 @@ bool FBuildPatchFileConstructor::ConstructFileFromChunks(const FString& BuildFil
 			{
 				if (bReportAnalytic)
 				{
-					InstallerAnalytics->RecordConstructionError(BuildFilename, CreateFilePlatformLastError, TEXT("Could Not Create File"));
+					InstallerAnalytics->RecordConstructionError(CurrentFile.BuildFilename, CurrentFile.CreateFilePlatformLastError, TEXT("Could Not Create File"));
 				}
-				UE_LOG(LogBuildPatchServices, Error, TEXT("FBuildPatchFileConstructor: Could not create %s (LastError=%d)"), *BuildFilename, CreateFilePlatformLastError);
-				InstallerError->SetError(EBuildPatchInstallError::FileConstructionFail, ConstructionErrorCodes::FileCreateFail, CreateFilePlatformLastError);
+				UE_LOG(LogBuildPatchServices, Error, TEXT("FBuildPatchFileConstructor: Could not create %s (LastError=%d)"), *CurrentFile.BuildFilename, CurrentFile.CreateFilePlatformLastError);
+				InstallerError->SetError(EBuildPatchInstallError::FileConstructionFail, ConstructionErrorCodes::FileCreateFail, CurrentFile.CreateFilePlatformLastError);
 				break;
 			}
 		case EConstructionError::MissingChunk:
 			{
 				if (bReportAnalytic)
 				{
-					InstallerAnalytics->RecordConstructionError(BuildFilename, INDEX_NONE, TEXT("Missing Chunk"));
+					InstallerAnalytics->RecordConstructionError(CurrentFile.BuildFilename, INDEX_NONE, TEXT("Missing Chunk"));
 				}
 
-				UE_LOG(LogBuildPatchServices, Error, TEXT("FBuildPatchFileConstructor: Failed %s due to missing chunk %s (can be 0000 if unknown)"), *BuildFilename, *WriteToString<64>(ErrorContextGuid));
+				UE_LOG(LogBuildPatchServices, Error, TEXT("FBuildPatchFileConstructor: Failed %s due to missing chunk %s (can be 0000 if unknown)"), *CurrentFile.BuildFilename, *WriteToString<64>(CurrentFile.ErrorContextGuid));
 				InstallerError->SetError(EBuildPatchInstallError::FileConstructionFail, ConstructionErrorCodes::MissingChunkData);
 				break;
 			}
@@ -2940,9 +2903,9 @@ bool FBuildPatchFileConstructor::ConstructFileFromChunks(const FString& BuildFil
 			{
 				if (bReportAnalytic)
 				{
-					InstallerAnalytics->RecordConstructionError(BuildFilename, INDEX_NONE, TEXT("Serialization Error"));
+					InstallerAnalytics->RecordConstructionError(CurrentFile.BuildFilename, INDEX_NONE, TEXT("Serialization Error"));
 				}
-				UE_LOG(LogBuildPatchServices, Error, TEXT("FBuildPatchFileConstructor: Failed %s due to serialization error"), *BuildFilename);
+				UE_LOG(LogBuildPatchServices, Error, TEXT("FBuildPatchFileConstructor: Failed %s due to serialization error"), *CurrentFile.BuildFilename);
 				InstallerError->SetError(EBuildPatchInstallError::FileConstructionFail, ConstructionErrorCodes::SerializationError);
 				break;
 			}
@@ -2950,9 +2913,9 @@ bool FBuildPatchFileConstructor::ConstructFileFromChunks(const FString& BuildFil
 			{
 				if (bReportAnalytic)
 				{
-					InstallerAnalytics->RecordConstructionError(BuildFilename, INDEX_NONE, TEXT("Tracking Error"));
+					InstallerAnalytics->RecordConstructionError(CurrentFile.BuildFilename, INDEX_NONE, TEXT("Tracking Error"));
 				}
-				UE_LOG(LogBuildPatchServices, Error, TEXT("FBuildPatchFileConstructor: Failed %s due to untracked chunk %s (can be 0000 if unknown)"), *BuildFilename, *WriteToString<64>(ErrorContextGuid));
+				UE_LOG(LogBuildPatchServices, Error, TEXT("FBuildPatchFileConstructor: Failed %s due to untracked chunk %s (can be 0000 if unknown)"), *CurrentFile.BuildFilename, *WriteToString<64>(CurrentFile.ErrorContextGuid));
 				InstallerError->SetError(EBuildPatchInstallError::FileConstructionFail, ConstructionErrorCodes::TrackingError);
 				break;
 			}
@@ -2960,95 +2923,737 @@ bool FBuildPatchFileConstructor::ConstructFileFromChunks(const FString& BuildFil
 			{
 				if (bReportAnalytic)
 				{
-					InstallerAnalytics->RecordConstructionError(BuildFilename, INDEX_NONE, TEXT("Internal Consistency Error"));
+					InstallerAnalytics->RecordConstructionError(CurrentFile.BuildFilename, INDEX_NONE, TEXT("Internal Consistency Error"));
 				}
-				UE_LOG(LogBuildPatchServices, Error, TEXT("FBuildPatchFileConstructor: Failed %s due to internal consistency checking failure"), *BuildFilename);
+				UE_LOG(LogBuildPatchServices, Error, TEXT("FBuildPatchFileConstructor: Failed %s due to internal consistency checking failure"), *CurrentFile.BuildFilename);
 				InstallerError->SetError(EBuildPatchInstallError::FileConstructionFail, ConstructionErrorCodes::InternalConsistencyFailure);
 				break;
 			}
-			case EConstructionError::Aborted:
+		case EConstructionError::MissingFileInfo:
+			{
+				if (bReportAnalytic)
+				{
+					InstallerAnalytics->RecordConstructionError(CurrentFile.BuildFilename, INDEX_NONE, TEXT("Missing File Manifest"));					
+				}
+				UE_LOG(LogBuildPatchServices, Error, TEXT("FBuildPatchFileConstructor: Missing file manifest for %s"), *CurrentFile.BuildFilename);
+				InstallerError->SetError(EBuildPatchInstallError::FileConstructionFail, ConstructionErrorCodes::MissingFileInfo);
+				break;
+			}
+		case EConstructionError::FailedInitialSizeCheck:
+			{
+				// Error already set back when we had the info.
+				break;
+			}
+		case EConstructionError::Aborted:
 			{
 				// We don't set errors on abort.
 				break;
 			}
 		}
+
+		// Delete the staging file if unsuccessful by means of any failure that could leave the file in unknown state.
+		switch (CurrentFile.ConstructionError)
+		{
+			// Errors we expect that we can conceptually resume:
+			case EConstructionError::OutOfDiskSpace:
+			case EConstructionError::MissingChunk:
+			case EConstructionError::Aborted:
+			{
+				break;
+			}
+
+			// Errors we expect there to be issues with the outbound file:
+			case EConstructionError::CannotCreateFile:
+			case EConstructionError::SerializeError:
+			case EConstructionError::TrackingError:
+			case EConstructionError::OutboundDataError:
+			case EConstructionError::InternalConsistencyError:
+			{
+				if (!FileSystem->DeleteFile(*CurrentFile.NewFilename))
+				{
+					UE_LOG(LogBuildPatchServices, Warning, TEXT("FBuildPatchFileConstructor: Error deleting file: %s (Error Code %i)"), *CurrentFile.NewFilename, FPlatformMisc::GetLastError());
+				}
+				break;
+			}
+		}
+
+		// Stop trying to install further files.
+		Abort();
 	}
+}
 
-	// Verify the hash for the file that we created
-	if (bSuccess)
+void FBuildPatchFileConstructor::StartFile(FFileConstructionState& CurrentFile, const FResumeData& ResumeData)
+{
+	TRACE_CPUPROFILER_EVENT_SCOPE(Constructor_StartFile);
+
+	InitFile(CurrentFile, ResumeData);
+
+	if (CurrentFile.bSkippedConstruction)
 	{
-		HashState.Final();
+		return;
+	}
 
-		FSHAHash HashValue;
-		HashState.GetHash(HashValue.Hash);
-		bSuccess = HashValue == FileManifest.FileHash;
-		if (!bSuccess)
+	if (ResumeData.FilesStarted.Contains(CurrentFile.BuildFilename))
+	{
+		ResumeFile(CurrentFile);
+
+		// Generally speaking we only expect there to be one file to resume (as of this writing there
+		// is no way for the FilesStarted set to have more than one file), so we update the download requirements
+		// after the first one.
+
+		// Need to sum up the size of all remaining chunks that we need. We can't just look at the
+		// chunk locations because we don't know if it's still needed or not.
+		TSet<FGuid> RemainingNeededChunks = ChunkReferenceTracker->GetReferencedChunks();
+
+		uint64 NewDownloadRequirement = 0;
+		FRWScopeLock ChunkLock(ChunkLocationsLock, SLT_Write); // write because we're updating the download requirement
+		for (FGuid& NeededChunk : RemainingNeededChunks)
 		{
-			ConstructionError = EConstructionError::OutboundDataError;
-			// Only report if the first error
-			if (InstallerError->HasError() == false)
+			EConstructorChunkLocation* LocationPtr = ChunkLocations.Find(NeededChunk);
+			if (LocationPtr && *LocationPtr == EConstructorChunkLocation::Cloud)
 			{
-				InstallerAnalytics->RecordConstructionError(BuildFilename, INDEX_NONE, TEXT("Serialised Verify Fail"));
+				NewDownloadRequirement += Configuration.ManifestSet->GetChunkInfo(NeededChunk)->FileSize;
 			}
-			UE_LOG(LogBuildPatchServices, Error, TEXT("FBuildPatchFileConstructor: Verify failed after constructing %s"), *BuildFilename);
-			InstallerError->SetError(EBuildPatchInstallError::FileConstructionFail, ConstructionErrorCodes::OutboundCorrupt);
 		}
+
+		DownloadRequirement = NewDownloadRequirement;
+		CloudSource->PostRequiredByteCount(DownloadRequirement);
 	}
 
-#if PLATFORM_MAC
-	if (bSuccess && EnumHasAllFlags(FileManifest.FileMetaFlags, EFileMetaFlags::UnixExecutable))
+	// If we haven't done so yet, make the initial disk space check. We do this after resume
+	// so that we know how much to discount from our current file size.
+	if (!HandleInitialDiskSizeCheck(*CurrentFile.FileManifest, CurrentFile.StartPosition))
 	{
-		// Enable executable permission bit
-		struct stat FileInfo;
-		if (stat(TCHAR_TO_UTF8(*NewFilename), &FileInfo) == 0)
-		{
-			bSuccess = chmod(TCHAR_TO_UTF8(*NewFilename), FileInfo.st_mode | S_IXUSR | S_IXGRP | S_IXOTH) == 0;
-		}
+		CurrentFile.bSuccess = false;
+		CurrentFile.ConstructionError = EConstructionError::FailedInitialSizeCheck;
 	}
-#endif
 
-#if PLATFORM_ANDROID
-	if (bSuccess)
+	if (!bIsDownloadStarted)
 	{
-		IFileManager::Get().SetTimeStamp(*NewFilename, FDateTime::UtcNow());
+		bIsDownloadStarted = true;
+		FileConstructorStat->OnResumeCompleted();
 	}
-#endif
 
-	if (bSuccess)
+	OpenFileToConstruct(CurrentFile);
+
+}
+
+void FBuildPatchFileConstructor::ConstructFiles(const FResumeData& ResumeData)
+{
+	MaxWriteBatchSize = FFileConstructorConfig::DefaultIOBatchSizeMB;
+	if (GConfig->GetInt(TEXT("Portal.BuildPatch"), TEXT("ConstructorIOBatchSizeMB"), (int32&)MaxWriteBatchSize, GEngineIni))
 	{
-		ChunkDbSource->ReportFileCompletion(ChunkReferenceTracker->GetRemainingChunkCount());
+		UE_LOG(LogBuildPatchServices, Verbose, TEXT("Got INI ConstructorIOBatchSizeMB = %d"), MaxWriteBatchSize);
+	}
+	if (Configuration.IOBatchSizeMB.IsSet())
+	{
+		MaxWriteBatchSize = Configuration.IOBatchSizeMB.Get(MaxWriteBatchSize);
+		UE_LOG(LogBuildPatchServices, Verbose, TEXT("Got override ConstructorIOBatchSizeMB = %d"), MaxWriteBatchSize);
+	}
+	MaxWriteBatchSize <<= 20; // to MB;
+
+	IOBufferSize = FFileConstructorConfig::DefaultIOBufferSizeMB;
+	if (GConfig->GetInt(TEXT("Portal.BuildPatch"), TEXT("ConstructorIOBufferSizeMB"), (int32&)IOBufferSize, GEngineIni))
+	{
+		UE_LOG(LogBuildPatchServices, Verbose, TEXT("Got INI ConstructorIOBufferSizeMB = %d"), IOBufferSize);
+	}
+	if (Configuration.IOBufferSizeMB.IsSet())
+	{
+		IOBufferSize = Configuration.IOBufferSizeMB.Get(IOBufferSize);
+		UE_LOG(LogBuildPatchServices, Verbose, TEXT("Got override ConstructorIOBufferSizeMB = %d"), IOBufferSize);
 	}
+	IOBufferSize <<= 20; // to MB;
 
-	// Delete the staging file if unsuccessful by means of any failure that could leave the file in unknown state.
-	if (!bSuccess)
+	// Ensure that our batch size can always make progress (all chunks can fit)
+	uint32 LargestChunkSize = 0;
 	{
-		switch (ConstructionError)
+		FScopeLock Lock(&ThreadLock);
+
+		for (const FString& FileToConstruct : ConstructionStack)
 		{
-			// Errors we expect that we can conceptually resume:
-		case EConstructionError::OutOfDiskSpace:
-		case EConstructionError::MissingChunk:
-		case EConstructionError::Aborted:
+			const FFileManifest* FileManifest = Configuration.ManifestSet->GetNewFileManifest(FileToConstruct);
+			for (const FChunkPart& ChunkPart : FileManifest->ChunkParts)
 			{
-				break;
+				if (ChunkPart.Size > LargestChunkSize)
+				{
+					LargestChunkSize = ChunkPart.Size;
+				}
 			}
+		}
 
-			// Errors we expect there to be issues with the outbound file:
-		case EConstructionError::CannotCreateFile:
-		case EConstructionError::SerializeError:
-		case EConstructionError::TrackingError:
-		case EConstructionError::OutboundDataError:
-		case EConstructionError::InternalConsistencyError:
+		if (LargestChunkSize > MaxWriteBatchSize)
+		{
+			MaxWriteBatchSize = LargestChunkSize;
+			UE_LOG(LogBuildPatchServices, Display, TEXT("Increasing batch size to fit any chunk size: %u bytes"), MaxWriteBatchSize);
+
+			if (MaxWriteBatchSize > IOBufferSize)
+			{
+				IOBufferSize = MaxWriteBatchSize;
+				UE_LOG(LogBuildPatchServices, Display, TEXT("Increasing IO buffer size to fit batch size: %u bytes"), IOBufferSize);
+			}
+		}
+	}
+
+	uint64 ConstructStartCycles = FPlatformTime::Cycles64();
+	uint64 WriteDoneCycles = 0;
+	uint64 ReadDoneCycles = 0;
+	uint64 ReadCheckCycles = 0;
+	uint64 CloudTickCycles = 0;
+	uint64 HashCycles = 0;
+	uint64 WaitCycles = 0;
+	uint64 WriteStartCycles = 0;
+
+	// List of files that are currently opened. The last one is the one we are starting reads on,
+	// the first one is the one we are currently writing.
+	TArray<TUniquePtr<FFileConstructionState>> ActiveFiles;
+
+	// List of batches in flight. These must be dispatched in order.
+	TArray<TUniquePtr<FBatchState>> Batches;
+
+	// This is the buffer we issue reads into and writes out of. We segment it up
+	// based on batch sizing.
+	TArray<uint8> IOBUffer;
+	IOBUffer.SetNumUninitialized(IOBufferSize);
+
+	struct Range
+	{
+		uint64 Start=0, End=0;
+	};
+	TArray<Range, TInlineAllocator<3>> IOBufferFreeBlockList;
+	IOBufferFreeBlockList.Add({0, IOBufferSize});
+
+	auto SortAndCoalesceFreeList = [](TArray<Range, TInlineAllocator<3>>& FreeList)
+	{
+		// Sort the list on Start.
+		Algo::SortBy(FreeList, &Range::Start, TLess<uint64>());
+
+		// go over everything and coalesce anything that's adjacent.
+		for (int32 i=0; i<FreeList.Num() - 1; i++)
+		{
+			if (FreeList[i].End == FreeList[i+1].Start)
+			{
+				// Extend us.
+				FreeList[i].End = FreeList[i+1].End;
+
+				// Remove them.
+				FreeList.RemoveAt(i+1);
+
+				// Recheck us.
+				i--;
+			}
+		}
+	};
+
+	//
+	// The abort handling with this loop is:
+	// During the loop, if we see the abort signal, we mark all active files as failed.
+	// That prevents any new work from starting, but we have to have the outstanding work
+	// complete in order to be thread safe. Once that work completes, we then break out of the loop
+	// (bDoneWithFiles will still be true).
+	//
+	bool bStillProcessingFiles = true;
+	for (;bStillProcessingFiles;)
+	{		
+		
+
+		bool bHasAnyFileFailed = false;
+		// Check state. We can continue to run this after a failure has been encountered in order
+		// to drain any async tasks so we know we can shut down safely.
+		{
+			// Is our next write done?
+			if (Batches.Num() &&
+				((Batches[0]->bIsWriting && Batches[0]->bIsFinished.load(std::memory_order_acquire)) || Batches[0]->bIsEmptyFileSentinel))
 			{
-				if (!FileSystem->DeleteFile(*NewFilename))
+				uint64 WriteDoneStartCycles = FPlatformTime::Cycles64();
+
+				// Could be empty buffer if sentinel
+				if (Batches[0]->BatchBuffer.GetSize())
 				{
-					UE_LOG(LogBuildPatchServices, Warning, TEXT("FBuildPatchFileConstructor: Error deleting file: %s (Error Code %i)"), *NewFilename, FPlatformMisc::GetLastError());
+					uint64 BufferBase = (uint64)((uint8*)Batches[0]->BatchBuffer.GetData() - (uint8*)IOBUffer.GetData());
+
+					IOBufferFreeBlockList.Add({BufferBase, BufferBase + Batches[0]->BatchBuffer.GetSize()});
+					SortAndCoalesceFreeList(IOBufferFreeBlockList);
 				}
+
+				FFileConstructionState* File = Batches[0]->OwningFile;
+				File->OutstandingBatches--;
+
+				Batches.RemoveAt(0);
+				
+				if (File->NewFile->IsError())
+				{
+					File->bSuccess = false;
+
+					// SerializeError signals the error management code to check for low disk space
+					// since we currently have no way of routing out of disk errors from the UE low level
+					// writing code.
+					File->ConstructionError = EConstructionError::SerializeError;
+				}
+					
+				// If we have nothing further to read and all the file's reads are done,
+				// since we only ever have 1 write batch we know it must be done.
+				if (File->NextChunkPartToRead == File->FileManifest->ChunkParts.Num() &&
+					File->OutstandingBatches == 0)
+				{
+					CompleteConstructedFile(*File);
+
+					UE_LOG(LogBuildPatchServices, Verbose, TEXT("Completed file: %s, New ActiveFileCount = %d"), *File->BuildFilename, ActiveFiles.Num() - 1);
+
+					// Since we write in order, we must be the first in the queue.
+					check (ActiveFiles[0].Get() == File);
+
+					// This frees the file
+					ActiveFiles.RemoveAt(0);
+
+					// We completed the previous file, if there's another active file we need to tell
+					// the stats that it has started. If there isn't, we'll fire it off when we start it.
+					if (ActiveFiles.Num())
+					{
+						FileConstructorStat->OnFileStarted(ActiveFiles[0]->BuildFilename, ActiveFiles[0]->FileManifest->FileSize);
+					}
+				}
+
+				WriteDoneCycles += FPlatformTime::Cycles64() - WriteDoneStartCycles;
+			}
+
+			// Check for completed reads.
+			if (Batches.Num())
+			{
+				uint64 ReadDoneStartCycles = FPlatformTime::Cycles64();
+
+				// We have to retire reads in order because later reads could be waiting for data to
+				// get placed correctly by the earlier one. e.g. read 1 could be placing in the memory store
+				// and read 2 could be wanting to read from that when it retires.
+				// Since we queue the batches in order, we can look at the first reading batch and know it's the next one.
+				for (TUniquePtr<FBatchState>& Batch : Batches)
+				{
+					if (Batch->bIsReading)
+					{
+						if (Batch->bIsFinished.load(std::memory_order_acquire))
+						{
+							FFileConstructionState* File = Batch->OwningFile;
+
+							CompleteReadBatch(*File->FileManifest, *Batch.Get());
+							Batch->bIsReading = false;
+							Batch->bNeedsWrite = true;
+
+							if (File->ConstructionError != EConstructionError::None ||
+								Batch->Error != EConstructionError::None)
+							{
+								File->ErrorContextGuid = Batch->ErrorContextGuid;
+								File->ConstructionError = Batch->Error;
+								File->bSuccess = false;
+							}
+						}
+						else
+						{
+							// We have to retire reads in order, so once we hit one that isn't finished we stop.
+							break;
+						}
+					}
+					// else - might be a writing batch before the first read.
+				}
+
+				ReadDoneCycles += FPlatformTime::Cycles64() - ReadDoneStartCycles;
+			}
+
+			// If we've aborted, fail the files.
+			for (TUniquePtr<FFileConstructionState>& File : ActiveFiles)
+			{
+				if (File->bSuccess &&
+					bShouldAbort)
+				{
+					File->ConstructionError = EConstructionError::Aborted;
+					File->bSuccess = false;
+				}
+
+				bHasAnyFileFailed |= !File->bSuccess;
+			}
+
+			// If we are in a failure case, we want to clear out anything happening with the cloud source.
+			// We'll hit this over and over as we "drain", but we need to do that anyway because further failures
+			// might try to queue more cloud reads.
+			if (bHasAnyFileFailed)
+			{
+				CloudSource->Abort();
+			}
+
+			if (!bIsPaused && !bHasAnyFileFailed)
+			{
+				// Can we start another read?
+				// LONGTERM - if we are memory constrained and have an install source we can consider only dispatching
+				// single reads. This dramatically lowers our memory requirements as install sources often use a lot
+				// of small pieces of chunks - and we have to load the entire chunk into memory.
+				bool bCheckForAnotherRead = true;
+				while (bCheckForAnotherRead)
+				{
+					double ReadCheckStartCycles = FPlatformTime::Cycles64();
+
+					TRACE_CPUPROFILER_EVENT_SCOPE(CFFC_ReadCheck);
+
+					// We only check for another if we got one queued so we fill up the buffer space asap.
+					bCheckForAnotherRead = false;
+
+					int32 BiggestFreeBlockSize = 0;
+					int32 BiggestFreeBlockSlotIndex = 0;
+					for (int32 Slot = 0; Slot < IOBufferFreeBlockList.Num(); Slot++)
+					{
+						if (IOBufferFreeBlockList[Slot].End - IOBufferFreeBlockList[Slot].Start > BiggestFreeBlockSize)
+						{
+							BiggestFreeBlockSize = IOBufferFreeBlockList[Slot].End - IOBufferFreeBlockList[Slot].Start;
+							BiggestFreeBlockSlotIndex = Slot;
+						}
+					}
+
+					if (BiggestFreeBlockSize > 0)
+					{
+						// Default to continue to work on the last file.
+						FFileConstructionState* FileToStart = nullptr;
+						int32 ActiveFileCount = ActiveFiles.Num();
+						if (ActiveFileCount && ActiveFiles[ActiveFileCount-1]->NextChunkPartToRead < ActiveFiles[ActiveFileCount-1]->FileManifest->ChunkParts.Num())
+						{
+							FileToStart = ActiveFiles[ActiveFileCount-1].Get();
+						}
+
+						// No more work do to on active files - is there another to start?
+						if (!FileToStart)
+						{
+							bool bAllowAnotherFile = bAllowMultipleFilesInFlight || (ActiveFiles.Num() == 0);
+
+							FString AnotherFileToConstruct;
+							if (bAllowAnotherFile && GetFileToConstruct(AnotherFileToConstruct))
+							{
+								FString StagingFileName = Configuration.StagingDirectory / AnotherFileToConstruct;
+								TUniquePtr<FFileConstructionState> AnotherFile = MakeUnique<FFileConstructionState>(
+									Configuration.ManifestSet->GetNewFileManifest(AnotherFileToConstruct), 
+									MoveTemp(AnotherFileToConstruct), 
+									MoveTemp(StagingFileName));
+
+								TRACE_BOOKMARK(TEXT("Starting File: %s [%u]"), *AnotherFile->BuildFilename, AnotherFile->FileManifest->ChunkParts.Num());
+
+								UE_LOG(LogBuildPatchServices, Verbose, TEXT("Starting File: %s [%s bytes, %d chunks], New ActiveFileCount = %d"), 
+									*AnotherFile->BuildFilename, 
+									*FormatNumber(AnotherFile->FileManifest->FileSize),
+									AnotherFile->FileManifest->ChunkParts.Num(),
+									ActiveFiles.Num() + 1);
+
+								StartFile(*AnotherFile.Get(), ResumeData);
+
+								if (AnotherFile->bSkippedConstruction)
+								{
+									// Nothing else needs to happen with this file - we'll loop around to try for another.
+									bCheckForAnotherRead = true;
+								}
+								else
+								{
+									// Only start the file if there's not currently an active file - otherwise we
+									// do it when the current one finishes.
+									if (ActiveFileCount == 0)
+									{
+										FileConstructorStat->OnFileStarted(AnotherFile->BuildFilename, AnotherFile->FileManifest->FileSize);
+									}
+
+									FileToStart = AnotherFile.Get();
+									ActiveFiles.Push(MoveTemp(AnotherFile));
+								}
+							}
+							else
+							{
+								// If we don't have any files to construct, then we can't start any more reads at all.
+								// If there are no active files, then we are done.
+								if (ActiveFiles.Num() == 0)
+								{
+									bStillProcessingFiles = false;
+								}
+							}
+						} // end if we are starting a new file.
+
+						// It's possible the file failed during creation and we need to start the failure process. We need to
+						// set this since we did the scan before here.
+						if (FileToStart && !FileToStart->bSuccess)
+						{
+							bHasAnyFileFailed = true;
+						}
+
+						if (FileToStart && FileToStart->bSuccess)
+						{
+							if (FileToStart->FileManifest->ChunkParts.Num() == 0)
+							{
+								// We have a file that will never launch any batches, which means it'll never hit the finalization
+								// logic. We can't complete it here because then we are out of order. So we need to inject a placeholder
+								// batch that will auto pass the write check and also prevent us from sleeping on an event.
+								TUniquePtr<FBatchState> Batch = MakeUnique<FBatchState>();
+								Batch->bNeedsWrite = false;
+								Batch->bIsReading = false;
+								Batch->OwningFile = FileToStart;
+								Batch->bIsEmptyFileSentinel = true;
+								Batch->bIsFinished.store(false, std::memory_order_release);
+								FileToStart->OutstandingBatches++;
+
+								Batches.Add(MoveTemp(Batch));
+
+								bCheckForAnotherRead = true;
+							}
+							else
+							{
+
+								uint32 MaxBufferSize = BiggestFreeBlockSize;
+								if (MaxBufferSize > MaxWriteBatchSize)
+								{
+									MaxBufferSize = MaxWriteBatchSize;
+								}
+
+								// We want to do big batches as much as possible. If we only have space for a single chunk,
+								// that's fine if it's a single chunk file, but generally we want to try and favor large batches.
+								// If we're at the point where we are worried about this, we have enough outstanding work to keep
+								// the pipelines full so we can afford to wait for room.
+								uint32 MaxFileBatchSize = 0;
+								const TArray<FChunkPart>& FileChunkParts = FileToStart->FileManifest->ChunkParts;
+								uint32 NextChunkSize = FileChunkParts[FileToStart->NextChunkPartToRead].Size;
+
+								for (int32 ChunkPartIndex = FileToStart->NextChunkPartToRead;
+									ChunkPartIndex < FileChunkParts.Num();
+									ChunkPartIndex++)
+								{
+									MaxFileBatchSize += FileChunkParts[ChunkPartIndex].Size;
+									if (MaxFileBatchSize >= MaxWriteBatchSize)
+									{
+										// If it's big enough for the max batch size we no longer care.
+										MaxFileBatchSize = MaxWriteBatchSize;
+										break;
+									}
+								}
+
+								// If the file can support a large batch, we want to wait until we have room for a reasonable size.
+								// We know we can eventually read the next chunk because during init we sized the buffers such that
+								// we could.
+								if (NextChunkSize > MaxBufferSize &&
+									Batches.Num() == 0)
+								{
+									FileToStart->bSuccess = false;
+									FileToStart->ConstructionError = EConstructionError::InternalConsistencyError;
+									UE_LOG(LogBuildPatchServices, Error, TEXT("Chunk size encountered larger than batch buffer size! %u vs %u"), NextChunkSize, MaxBufferSize);
+									bHasAnyFileFailed = true;
+
+									// We'll fail the next conditional below and then start the failure process on the next loop.
+								}
+
+								if (MaxBufferSize >= MaxFileBatchSize)
+								{
+									TUniquePtr<FBatchState> Batch = MakeUnique<FBatchState>();
+									Batch->bNeedsWrite = true;
+									Batch->bIsReading = true;
+									Batch->OwningFile = FileToStart;
+									Batch->BatchBuffer = FMutableMemoryView(IOBUffer.GetData() + IOBufferFreeBlockList[BiggestFreeBlockSlotIndex].Start, MaxBufferSize);
+									Batch->bIsFinished.store(false, std::memory_order_release);
+									StartReadBatch(*FileToStart, *Batch.Get());
+
+									// The read might not have used the whole thing, so only consume as much
+									// as it needed.
+									IOBufferFreeBlockList[BiggestFreeBlockSlotIndex].Start += Batch->BatchBuffer.GetSize();
+									if (IOBufferFreeBlockList[BiggestFreeBlockSlotIndex].Start == IOBufferFreeBlockList[BiggestFreeBlockSlotIndex].End)
+									{
+										// Ate whole thing
+										IOBufferFreeBlockList.RemoveAt(BiggestFreeBlockSlotIndex);
+										SortAndCoalesceFreeList(IOBufferFreeBlockList);
+									}
+
+									bCheckForAnotherRead = true;
+								
+									UE_LOG(LogBuildPatchServices, VeryVerbose, TEXT("Starting ReadBatch: %d, Chunks=%d, Bytes=%s, Batches=%d"), 
+										Batch->BatchId, 
+										Batch->ChunkCount, 
+										*FormatNumber(Batch->BatchBuffer.GetSize()),
+										Batches.Num() + 1
+									);
+
+									Batches.Push(MoveTemp(Batch));
+								}
+							} // end if file has parts
+						} // end if we have a file to read from.
+					} // end if free block exists.
+
+					ReadCheckCycles += FPlatformTime::Cycles64() - ReadCheckStartCycles;
+				} // end looping on whether we should start a read.
+
+				// Can we start a write? 
+				// We always have to issue writes in order and there can only be one, so it must be the
+				// first active file, and the first batch. 
+				if (Batches.Num() && Batches[0]->bNeedsWrite && !Batches[0]->bIsReading)
+				{
+					uint64 WriteStartStartCycles = FPlatformTime::Cycles64();
+
+					FFileConstructionState* FirstFile = ActiveFiles[0].Get();
+
+					Batches[0]->bNeedsWrite = false;
+					Batches[0]->bIsWriting = true;
+					Batches[0]->bIsFinished.store(false, std::memory_order_relaxed);
+
+					// Launch the write and hash the buffer on this thread.
+					UE_LOG(LogBuildPatchServices, VeryVerbose, TEXT("Writing Batch: %d, file %s [%d - %d]"),
+						Batches[0]->BatchId, 
+						*FirstFile->BuildFilename,
+						Batches[0]->StartChunkPartIndex,
+						Batches[0]->StartChunkPartIndex + Batches[0]->ChunkCount
+						);
+
+					IConstructorChunkSource::FRequestProcessFn WriteFn = CreateWriteRequest(FirstFile->NewFile.Get(), *Batches[0].Get());
+					QueueGenericThreadTask(WriteThreadIndex, MoveTemp(WriteFn));
+					
+					uint64 HashStartCycles = FPlatformTime::Cycles64();
+					WriteStartCycles += HashStartCycles - WriteStartStartCycles;
+
+					{
+						TRACE_CPUPROFILER_EVENT_SCOPE(CFFC_Hash);
+						FirstFile->HashState.Update((const uint8*)Batches[0]->BatchBuffer.GetData(), Batches[0]->BatchBuffer.GetSize());
+					}
+					HashCycles += FPlatformTime::Cycles64() - HashStartCycles;
+
+				} // end if checking for write
+			} // end not paused
+		} // end state check
+
+		// If the file progress changed since we last saw it, post the update.
+		// Note that we want this to update reasonably often but we're about to wait
+		// potentially until all reads complete - however the only time things actually
+		// take a long time wall-clock wise is when we are downloading, and the cloud
+		// source will then prevent us from sleeping too long, so we actually catch these
+		// updates.
+		// We do this from here to ensure we always increase rather than risk multi thread
+		// races.
+		if (ActiveFiles.Num())
+		{
+			// We only post the progress for the first file in the active list - this means
+			// that when we finish that file we'll likely jump to the middle progress for the next
+			// file, but we don't have a way to post the progress per file.
+			FFileConstructionState* File = ActiveFiles[0].Get();
+
+			uint64 CurrentFileProgress;
+			File->ProgressLock.Lock();
+			CurrentFileProgress = File->Progress;
+			File->ProgressLock.Unlock();
+					
+			if (CurrentFileProgress != File->LastSeenProgress)
+			{
+				// this updates the overall install progress.
+				CountBytesProcessed(CurrentFileProgress - File->LastSeenProgress);
+				File->LastSeenProgress = CurrentFileProgress;
+				FileConstructorStat->OnFileProgress(File->BuildFilename, CurrentFileProgress);
+			}
+		}
+
+		uint32 WaitTimeMs = TNumericLimits<uint32>::Max();
+		{
+			uint64 CloudTickStartCycles = FPlatformTime::Cycles64();
+
+			TRACE_CPUPROFILER_EVENT_SCOPE(CFFC_TickCloud);
+
+			// Max downloads is tricky - the internet makes no guarantees about which of our downloads finishes first. So while
+			// we want to have as many outstanding as possible to cover up connection overhead / resends and all that, if we enqueue 
+			// downloads from several batches, we can end up where the first batch can't complete because it's waiting on a download 
+			// that isn't finishing due to congestion from the next batch's download. Then we counterintuitively end up with FEWER 
+			// outstanding downloads because we can't launch more batches due to waiting on the front of this long chain. This is easily reproducible
+			// where Insights will show highly out of order completion. This ordering unfortunately scales with the number of
+			// outstanding downloads to a certain extent: if you allow 16 downloads then you'll end up waiting on a download
+			// 16 issues old - 32 downloads and you'll wait on one 32 issues old. We try to bound this by capping the issued
+			// downloads here, and in the cloud source we prevent queues if we get too far ahead of the last download.
+			// Note this also gets adjusted by the connection health stuff internal to the cloud source.
+			const uint32 MaxDownloads = FMath::Max(1U, FMath::DivideAndRoundUp(MaxWriteBatchSize, ExpectedChunkSize));
+			
+			// WaitTimeMs is an OUT param
+			CloudSource->Tick(!bIsPaused && !bHasAnyFileFailed, WaitTimeMs, MaxDownloads);
+
+			CloudTickCycles += FPlatformTime::Cycles64() - CloudTickStartCycles;
+		}
+
+		int32 ActiveReadBatches = 0;
+		int32 ActiveWriteBatches = 0;
+		int32 EmptyFileBatches = 0;
+		for (const TUniquePtr<FBatchState>& _Batch : Batches)
+		{
+			EmptyFileBatches += _Batch->bIsEmptyFileSentinel ? 1 : 0;
+			ActiveReadBatches += _Batch->bIsReading ? 1 : 0;
+			ActiveWriteBatches += _Batch->bIsWriting ? 1 : 0;
+		}
+
+		// Empty files don't have async jobs.
+		bool bAsyncJobExists = ActiveReadBatches || ActiveWriteBatches;
+
+		if (bHasAnyFileFailed)
+		{
+			// We can only bail when all our async jobs have completed.
+			if (!bAsyncJobExists)
+			{
 				break;
 			}
 		}
+
+		TRACE_INT_VALUE(TEXT("BPS.FC.ActiveFiles"), ActiveFiles.Num());
+		TRACE_INT_VALUE(TEXT("BPS.FC.ActiveReadBatches"), ActiveReadBatches);
+
+		if (bStillProcessingFiles && 
+			bAsyncJobExists)
+		{
+			if (WaitTimeMs == TNumericLimits<uint32>::Max())
+			{
+				WaitTimeMs = 15*1000;
+			}
+
+			uint64 WaitStartCycles = FPlatformTime::Cycles64();
+			TRACE_CPUPROFILER_EVENT_SCOPE(CFFC_Waiting);
+			// We have a bunch of stuff outstanding that will wake us up if something happens.
+			WakeUpDispatchThreadEvent->Wait(WaitTimeMs);
+
+			WaitCycles += FPlatformTime::Cycles64() - WaitStartCycles;
+		}
+	} // end loop until we complete all the files.
+
+	// Any remaining active files (due to abort/failure) need to be failed and completed
+	// so that errors get reported. We want to report the non-abort failures first, because
+	// anything else that was in flight gets reported as an abort and the first error is the
+	// one we actually care about.
+	if (ActiveFiles.Num())
+	{
+		for (TUniquePtr<FFileConstructionState>& File : ActiveFiles)
+		{
+			if (!File->bSuccess && File->ConstructionError != EConstructionError::Aborted)
+			{
+				CompleteConstructedFile(*File.Get());
+			}
+		}
+
+		// Now handle everything else.
+		for (TUniquePtr<FFileConstructionState>& File : ActiveFiles)
+		{
+			if (!File->bSuccess && File->ConstructionError != EConstructionError::Aborted)
+			{
+				continue; // handled in previous loop.
+			}
+
+			if (File->bSuccess)
+			{
+				File->bSuccess = false;
+				File->ConstructionError = EConstructionError::Aborted;
+			}
+			CompleteConstructedFile(*File.Get());
+		}
 	}
 
-	return bSuccess;
+	uint64 ConstructCycles = FPlatformTime::Cycles64() - ConstructStartCycles;
+	uint64 UnaccountedForCycles = ConstructCycles - HashCycles - WaitCycles - ReadCheckCycles - ReadDoneCycles - WriteStartCycles - WriteDoneCycles - CloudTickCycles;
+
+	double ConstructSec = FPlatformTime::ToSeconds64(ConstructCycles);
+	UE_LOG(LogBuildPatchServices, Display, TEXT("Construction done: %.2f sec. Hash %.1f%% Wait %.1f%% ReadCheck %.1f%% WriteStart %.1f%% ReadDone %.1f%% WriteDone %.1f%% CloudTick %.1f%% Unaccounted %.1f%%"),
+		ConstructSec,
+		100.0 * HashCycles / ConstructCycles,
+		100.0 * WaitCycles / ConstructCycles,
+		100.0 * ReadCheckCycles / ConstructCycles,
+		100.0 * WriteStartCycles / ConstructCycles,
+		100.0 * ReadDoneCycles / ConstructCycles,
+		100.0 * WriteDoneCycles / ConstructCycles,
+		100.0 * CloudTickCycles / ConstructCycles,
+		100.0 * UnaccountedForCycles / ConstructCycles
+		);
 }
 
 void FBuildPatchFileConstructor::DeleteDirectoryContents(const FString& RootDirectory)
diff --git a/Engine/Source/Runtime/Online/BuildPatchServices/Private/BuildPatchFileConstructor.h b/Engine/Source/Runtime/Online/BuildPatchServices/Private/BuildPatchFileConstructor.h
--- a/Engine/Source/Runtime/Online/BuildPatchServices/Private/BuildPatchFileConstructor.h
+++ b/Engine/Source/Runtime/Online/BuildPatchServices/Private/BuildPatchFileConstructor.h
@@ -16,6 +16,8 @@
 
 // Forward declarations
 struct FBatchState;
+struct FResumeData;
+struct FFileConstructionState;
 class FBuildPatchAppManifest;
 enum class EConstructionError : uint8;
 
@@ -79,6 +81,14 @@ namespace BuildPatchServices
 		IBuildInstallerSharedContext* SharedContext;
 
 		bool bDeleteChunkDBFilesAfterUse = false;
+
+		static const bool bDefaultSpawnAdditionalIOThreads = true;
+		static const int32 DefaultIOBatchSizeMB = 10;
+		static const int32 DefaultIOBufferSizeMB = 64;
+
+		TOptional<bool> SpawnAdditionalIOThreads;
+		TOptional<int32> IOBatchSizeMB;
+		TOptional<int32> IOBufferSizeMB;
 	};
 
 	/**
@@ -178,14 +188,10 @@ namespace BuildPatchServices
 		uint64 CalculateDiskSpaceRequirementsWithDeleteDuringInstall(const TArray<FString>& InConstructionStack);
 
 		/**
-		 * Constructs a particular file referenced by the given BuildManifest. The function takes an interface to a class that can provide availability information of chunks so that this
-		 * file construction process can be ran alongside chunk acquisition threads. It will Sleep while waiting for chunks that it needs.
-		 * @param BuildFilename		The relative build filename to use for construction.
-		 * @param FileManifest		The FFileManifest for the file to construct.
-		 * @param bResumeExisting	Whether we should resume from an existing file
-		 * @return	true if no file errors occurred
+		 * Sequentially build the files required, potentially skipping already created files or resuming partial
+		 * files.
 		 */
-		bool ConstructFileFromChunks(const FString& BuildFilename, const FFileManifest& FileManifest, bool bResumeExisting);
+		void ConstructFiles(const FResumeData& ResumeData);
 
 		/**
 		 * Delete all contents of a directory
@@ -213,7 +219,7 @@ namespace BuildPatchServices
 		FThreadSafeBool bShouldAbort;
 
 		// A critical section to protect the flags and variables.
-		// \todo this is another thing where if we keep current index in the stack we don't need a mutex.
+		// \todo remove this entirely and just track the remaining size when we pop a file.
 		FCriticalSection ThreadLock;
 
 		// A stack of filenames for files that need to be constructed.
@@ -228,7 +234,12 @@ namespace BuildPatchServices
 
 		// We always want to know exactly where we think a chunk should be. If it's not there,
 		// we update this list to where it can be found (i.e. cloud)
+		// This is almost always read only after initialization, but in rare situations can be updated
+		// (chunk failures, file resume) and is multi threaded access.
+		FRWLock ChunkLocationsLock;
 		TMap<FGuid, EConstructorChunkLocation> ChunkLocations;
+		// Track how much data we expect to have to download. This is protected by the ChunkLocationsLock since they are in sync.
+		uint64 DownloadRequirement = 0;
 
 		// Pointer to the chunk reference tracker.
 		IChunkReferenceTracker* ChunkReferenceTracker;
@@ -244,12 +255,22 @@ namespace BuildPatchServices
 		// Pointer to the stat class.
 		IFileConstructorStat* FileConstructorStat;
 
+		bool bAllowMultipleFilesInFlight = true;
+
+		// The size we expect for chunks. This should be used for estimation purposes, not anything requiring hard limits.
+		uint32 ExpectedChunkSize = 0;
+
 		// Total job size for tracking progress.
 		int64 TotalJobSize;
 
 		// Byte processed so far for tracking progress.
 		int64 ByteProcessed;
 
+		uint32 MaxWriteBatchSize = 0;
+		uint32 IOBufferSize = 0;
+
+		int32 WriteCount = 0;
+
 		// The amount of disk space requirement that was calculated when beginning the process. 0 if the install process was not started, or no additional space was needed.
 		std::atomic_uint64_t RequiredDiskSpace;
 
@@ -259,16 +280,6 @@ namespace BuildPatchServices
 		// Event executed before deleting an old installation file.
 		FOnBeforeDeleteFile BeforeDeleteFileEvent;
 
-		// Track how much data we expect to have to download
-		uint64 DownloadRequirement = 0;
-		// When we get a failure and switch to the cloud source this will be on any thread. Rather than get fancy
-		// we just lock around it since it's tiny.
-		UE::FMutex DownloadRequirementLock;
-
-		// Instead of reporting file progress on write (which jumps by the write buffer), advance as
-		// we fill the write buffer.
-		uint64 FileProgress = 0;
-		UE::FMutex FileProgressLock;
 
 		TArray<FEvent*> ThreadWakeups;
 		TArray<TArray<IConstructorChunkSource::FRequestProcessFn>> ThreadJobPostings;
@@ -283,9 +294,6 @@ namespace BuildPatchServices
 
 		TUniquePtr<FChunkBackingStore> BackingStore;
 
-		TArray<uint8> WriteBuffers[2];
-		uint32 WriteBufferSize = (16 << 20); // Default write buffer size.
-
 		// Fire this to wake up the main thread to process completed tasks.
 		FEvent* WakeUpDispatchThreadEvent = nullptr;
 
@@ -294,13 +302,23 @@ namespace BuildPatchServices
 		// Where we are in the chunk consumption list after each file.
 		TArray<int32> FileCompletionPositions;
 
-		int32 StartReadBatch(const FFileManifest& FileManifest, int32 StartFileChunkPartIndex, FBatchState& Buffer);
+		void StartReadBatch(FFileConstructionState& CurrentFile, FBatchState& Buffer);
 		void CompleteReadBatch(const FFileManifest& FileManifest, FBatchState& Buffer);
 		void RequestCompletedFn(const FGuid& Guid, bool bAborted, bool bFailedToRead, void* UserPtr);
 
 		std::atomic_int32_t PendingHarvestRequests = 0;
 		void ChunkHarvestCompletedFn(const FGuid& Guid, bool bAborted, bool bFailedToRead, void* UserPtr);
 
+		void StartFile(FFileConstructionState& CurrentFile, const FResumeData& ResumeData);
+		void ResumeFile(FFileConstructionState& FileToResume);
+		void OpenFileToConstruct(FFileConstructionState& CurrentFile);
+		bool HandleInitialDiskSizeCheck(const FFileManifest& CurrentFileManifest, int64 BytesInToCurrentFile);
+		bool HarvestChunksForCompletedFile(const FString& CompletedBuildFileName);
+		void CompleteConstructedFile(FFileConstructionState& CurrentFile);
+		void InitFile(FFileConstructionState& CurrentFile, const FResumeData& ResumeData);
+
+
+
 	public:
 		struct FBackingStoreStats
 		{
diff --git a/Engine/Source/Runtime/Online/BuildPatchServices/Private/BuildPatchInstaller.cpp b/Engine/Source/Runtime/Online/BuildPatchServices/Private/BuildPatchInstaller.cpp
--- a/Engine/Source/Runtime/Online/BuildPatchServices/Private/BuildPatchInstaller.cpp
+++ b/Engine/Source/Runtime/Online/BuildPatchServices/Private/BuildPatchInstaller.cpp
@@ -1411,6 +1411,9 @@ namespace BuildPatchServices
 			FCC.StagingDirectory = InstallStagingDir;
 			FCC.bDeleteChunkDBFilesAfterUse = Configuration.bDeleteChunkDbFilesAfterUse;
 			FCC.BackingStoreDirectory = DataStagingDir;
+			FCC.SpawnAdditionalIOThreads = Configuration.ConstructorSpawnAdditionalIOThreads;
+			FCC.IOBufferSizeMB = Configuration.ConstructorIOBufferSizeMB;
+			FCC.IOBatchSizeMB = Configuration.ConstructorIOBatchSizeMB;
 
 			TUniquePtr<FBuildPatchFileConstructor> FileConstructor(
 				new FBuildPatchFileConstructor(
diff --git a/Engine/Source/Runtime/Online/BuildPatchServices/Private/Installer/ChunkReferenceTracker.h b/Engine/Source/Runtime/Online/BuildPatchServices/Private/Installer/ChunkReferenceTracker.h
--- a/Engine/Source/Runtime/Online/BuildPatchServices/Private/Installer/ChunkReferenceTracker.h
+++ b/Engine/Source/Runtime/Online/BuildPatchServices/Private/Installer/ChunkReferenceTracker.h
@@ -25,6 +25,7 @@ namespace BuildPatchServices
 
 		/**
 		 * Gets a set of all chunks referenced by the installation this tracker refers to.
+		 * This returns the _remaining_ referenced chunks!
 		 * @return set of GUID containing every chunk used.
 		 */
 		virtual TSet<FGuid> GetReferencedChunks() const = 0;
diff --git a/Engine/Source/Runtime/Online/BuildPatchServices/Private/Installer/CloudChunkSource.cpp b/Engine/Source/Runtime/Online/BuildPatchServices/Private/Installer/CloudChunkSource.cpp
--- a/Engine/Source/Runtime/Online/BuildPatchServices/Private/Installer/CloudChunkSource.cpp
+++ b/Engine/Source/Runtime/Online/BuildPatchServices/Private/Installer/CloudChunkSource.cpp
@@ -22,6 +22,7 @@
 #include "Installer/Statistics/DownloadServiceStatistics.h"
 #include "Memory/SharedBuffer.h"
 #include "Tasks/Task.h"
+#include "ProfilingDebugging/CountersTrace.h"
 
 DECLARE_LOG_CATEGORY_EXTERN(LogCloudChunkSource, Log, All);
 DEFINE_LOG_CATEGORY(LogCloudChunkSource);
@@ -867,6 +868,7 @@ namespace BuildPatchServices
 			int32 ExpectedSize = 0;
 			int32 RequestId = 0;
 			double SecondsAtRequested = 0;
+			int32 ReadId = 0;
 
 			// The time when we want to launch this retry, if RetryNum != 0
 			double RetryTime = 0;
@@ -908,6 +910,7 @@ namespace BuildPatchServices
 			Read->Destination = DestinationBuffer;
 			Read->CompleteFn = CompleteFn;
 			Read->UserPtr = UserPtr;
+			Read->ReadId = NextReadId.fetch_add(1, std::memory_order_relaxed);
 
 			ReadQueueCS.Lock();
 			ReadQueue.Add(Read);
@@ -927,7 +930,7 @@ namespace BuildPatchServices
 		}
 
 		virtual void Abort() override;
-		virtual void Tick(bool bStartNewDownloads, uint32& OutTimeToNextTickMs) override;
+		virtual void Tick(bool bStartNewDownloads, uint32& OutTimeToNextTickMs, uint32 InMaxDownloads) override;
 		virtual void SetWakeupFunction(TUniqueFunction<void()>&& InWakeupMainThreadFn) override
 		{
 			WakeupMainThreadFn = MoveTemp(InWakeupMainThreadFn);
@@ -972,6 +975,7 @@ namespace BuildPatchServices
 
 		// Tracking health and connection state
 		std::atomic<int64> CyclesAtLastData = 0;
+		std::atomic_int32_t NextReadId = 0;
 
 		// Communication from download thread to processing thread.
 		FCriticalSection CompletedDownloadsCS;
@@ -1032,17 +1036,20 @@ namespace BuildPatchServices
 	// to have a touchpoint on the main thread to dispatch and manage everything. This thread
 	// should do very little work! It should just be brokering things, not accomplishing anything.
 	//
-	void FConstructorCloudChunkSource::Tick(bool bStartNewDownloads, uint32& OutTimeToNextTickMs)
+	void FConstructorCloudChunkSource::Tick(bool bStartNewDownloads, uint32& OutTimeToNextTickMs, uint32 InMaxDownloads)
 	{
 		int32 StartingDownloadCount = InFlightDownloads.Num();
 
-		if (UE_LOG_ACTIVE(LogCloudChunkSource, Verbose))
 		{
 			int32 PendingCount;
 			{
 				FScopeLock _(&ReadQueueCS);
 				PendingCount = ReadQueue.Num();
 			}
+
+			TRACE_INT_VALUE(TEXT("BPS.Cloud.ActiveReads"), InFlightDownloads.Num());
+			TRACE_INT_VALUE(TEXT("BPS.Cloud.PendingReads"), PendingCount);
+
 			UE_LOG(LogCloudChunkSource, Verbose, TEXT("Cloud: Active: %d, Pending %d"), InFlightDownloads.Num(), PendingCount);
 		}
 
@@ -1195,8 +1202,9 @@ namespace BuildPatchServices
 
 		// Update connection status and health.
 		EBuildPatchDownloadHealth DownloadHealth;
+		bool bAnyDownloadsRetrying = false;
 		{
-			bool bAllDownloadsRetrying = false;
+			bool bAllDownloadsRetrying = true;
 			{
 				FScopeLock _(&ReadQueueCS);
 				bAllDownloadsRetrying = ReadQueue.Num() > 0;
@@ -1206,6 +1214,10 @@ namespace BuildPatchServices
 					{
 						bAllDownloadsRetrying = false;
 					}
+					else
+					{
+						bAnyDownloadsRetrying = true;
+					}
 				}
 			}
 
@@ -1229,15 +1241,34 @@ namespace BuildPatchServices
 		}
 
 		// Start new downloads, if we aren't gated or paused
+		int32 WaitedForOldDownloads = 0;
+		int32 OldestReadId = TNumericLimits<int32>::Max();
+		int32 NewestReadId = 0;
+		for (TPair<int32, FCloudRead*>& InFlight : InFlightDownloads)
+		{
+			if (InFlight.Value->ReadId < OldestReadId)
+			{
+				OldestReadId = InFlight.Value->ReadId;
+			}
+			if (InFlight.Value->ReadId > NewestReadId)
+			{
+				NewestReadId = InFlight.Value->ReadId;
+			}
+		}
+		if (OldestReadId > NewestReadId)
+		{
+			OldestReadId = NewestReadId;
+		}
+
 		while (bStartNewDownloads)
 		{
-			FCloudRead* Read = nullptr;
-			const uint32 MaxDownloads = FMath::Max(1U, DownloadCount->GetAdjustedCount(InFlightDownloads.Num(), DownloadHealth));
+			const uint32 MaxDownloads = FMath::Min(InMaxDownloads, FMath::Max(1U, DownloadCount->GetAdjustedCount(InFlightDownloads.Num(), DownloadHealth)));
 			if ((uint32)InFlightDownloads.Num() >= MaxDownloads)
 			{
 				break;
 			}
 
+			FCloudRead* Read = nullptr;
 			{
 				FScopeLock _(&ReadQueueCS);
 
@@ -1267,9 +1298,10 @@ namespace BuildPatchServices
 			}
 		}
 
+		TRACE_INT_VALUE(TEXT("BPS.Cloud.OldestDownload"), NewestReadId - OldestReadId);
+
 		if (StartingDownloadCount != InFlightDownloads.Num())
 		{
-			//UE_LOG(LogTemp, Warning, TEXT("count: %d"), InFlightDownloads.Num());
 			CloudChunkSourceStat->OnActiveRequestCountUpdated(InFlightDownloads.Num());
 		}
 
diff --git a/Engine/Source/Runtime/Online/BuildPatchServices/Private/Installer/CloudChunkSource.h b/Engine/Source/Runtime/Online/BuildPatchServices/Private/Installer/CloudChunkSource.h
--- a/Engine/Source/Runtime/Online/BuildPatchServices/Private/Installer/CloudChunkSource.h
+++ b/Engine/Source/Runtime/Online/BuildPatchServices/Private/Installer/CloudChunkSource.h
@@ -60,7 +60,7 @@ namespace BuildPatchServices
 
 		virtual void Abort() = 0;
 
-		virtual void Tick(bool bStartNewDownloads, uint32& OutTimeToNextTickMs) = 0;
+		virtual void Tick(bool bStartNewDownloads, uint32& OutTimeToNextTickMs, uint32 InMaxDownloads) = 0;
 
 		// When downloads are complete we need a way to alert the main thread.
 		virtual void SetWakeupFunction(TUniqueFunction<void()>&& WakeupMainThreadFn) = 0;
diff --git a/Engine/Source/Runtime/Online/BuildPatchServices/Public/BuildPatchSettings.h b/Engine/Source/Runtime/Online/BuildPatchServices/Public/BuildPatchSettings.h
--- a/Engine/Source/Runtime/Online/BuildPatchServices/Public/BuildPatchSettings.h
+++ b/Engine/Source/Runtime/Online/BuildPatchServices/Public/BuildPatchSettings.h
@@ -292,6 +292,19 @@ namespace BuildPatchServices
 		bool bAllowConcurrentExecution;
 		// Whether to gather individual file operation statistics during install
 		bool bTrackFileOperations;
+
+		// Whether the file constructor is allowed to spawn multiple threads for IO read/write servicing.
+		// Unset uses Engine.ini value [Portal.BuildPatch]ConstructorSpawnAdditionalIOThreads if present, FFileConstructorConfig::bDefaultSpawnAdditionalIOThreads otherwise.
+		TOptional<bool> ConstructorSpawnAdditionalIOThreads;
+		// File constructor batch size, in MB. File will attempt to fill/write this chunk size to disk. 
+		// Unset uses Engine.ini value [Portal.BuildPatch]ConstructorIOBatchSizeMB if present, FFileConstructorConfig::DefaultIOBatchSizeMB otherwise.
+		TOptional<int32> ConstructorIOBatchSizeMB;
+
+		// File constructor buffer size, in MB. This affects constructor performance heavily and must be at least big enough to fit a batch. 
+		// For purely disk->disk installs, this can be relatively small - even 2x batch size will facilitate multiple GB/s installations. However
+		// for cloud/download installs this needs to be fairly large to mitigate the effect of out of order download completion. 
+		// Unset uses Engine.ini value [Portal.BuildPatch]ConstructorIOBufferSizeMB if present, FFileConstructorConfig::DefaultIOBufferSizeMB otherwise.
+		TOptional<int32> ConstructorIOBufferSizeMB;
 	};
 
 	/**
