Fix for mac cis

#rnx

diff --git a/Engine/Source/Runtime/Online/BuildPatchServices/Private/BuildPatchFileConstructor.cpp b/Engine/Source/Runtime/Online/BuildPatchServices/Private/BuildPatchFileConstructor.cpp
--- a/Engine/Source/Runtime/Online/BuildPatchServices/Private/BuildPatchFileConstructor.cpp
+++ b/Engine/Source/Runtime/Online/BuildPatchServices/Private/BuildPatchFileConstructor.cpp
@@ -2760,7 +2760,7 @@ void FBuildPatchFileConstructor::CompleteConstructedFile(FFileConstructionState&
 			struct stat FileInfo;
 			if (stat(TCHAR_TO_UTF8(*CurrentFile.NewFilename), &FileInfo) == 0)
 			{
-				CurrentFile.bSuccess = chmod(TCHAR_TO_UTF8(*NewFilename), FileInfo.st_mode | S_IXUSR | S_IXGRP | S_IXOTH) == 0;
+				CurrentFile.bSuccess = chmod(TCHAR_TO_UTF8(*CurrentFile.NewFilename), FileInfo.st_mode | S_IXUSR | S_IXGRP | S_IXOTH) == 0;
 				if (!CurrentFile.bSuccess)
 				{
 					UE_LOG(LogBuildPatchServices, Error, TEXT("FBuildPatchFileConstructor: Failed to set exec bit %s"), *CurrentFile.BuildFilename);
