another mac cis fix

diff --git a/Engine/Source/Runtime/Online/BuildPatchServices/Private/BuildPatchFileConstructor.cpp b/Engine/Source/Runtime/Online/BuildPatchServices/Private/BuildPatchFileConstructor.cpp
--- a/Engine/Source/Runtime/Online/BuildPatchServices/Private/BuildPatchFileConstructor.cpp
+++ b/Engine/Source/Runtime/Online/BuildPatchServices/Private/BuildPatchFileConstructor.cpp
@@ -2754,7 +2754,7 @@ void FBuildPatchFileConstructor::CompleteConstructedFile(FFileConstructionState&
 
 #if PLATFORM_MAC
 		if (CurrentFile.bSuccess && 
-			EnumHasAllFlags(CurrentFile.FileManifest.FileMetaFlags, EFileMetaFlags::UnixExecutable))
+			EnumHasAllFlags(CurrentFile.FileManifest->FileMetaFlags, EFileMetaFlags::UnixExecutable))
 		{
 			// Enable executable permission bit
 			struct stat FileInfo;
