BPS: Exclude hashing times from install source read timings.

#rnx

diff --git a/Engine/Source/Runtime/Online/BuildPatchServices/Private/Installer/InstallChunkSource.cpp b/Engine/Source/Runtime/Online/BuildPatchServices/Private/Installer/InstallChunkSource.cpp
--- a/Engine/Source/Runtime/Online/BuildPatchServices/Private/Installer/InstallChunkSource.cpp
+++ b/Engine/Source/Runtime/Online/BuildPatchServices/Private/Installer/InstallChunkSource.cpp
@@ -270,13 +270,13 @@ namespace BuildPatchServices
 				return;
 			}
 
+			// We set this here because it is used to compute the IO speeds, however we can't call OnLoadComplete because we don't know
+			// the hash result yet.
+			ActivityRecord.CyclesEnd = FStatsCollector::GetCycles();
+
 			// Check chunk hash. 
 			UE::Tasks::Launch(TEXT("Install_Hash"), [DataId, UserPtr, DestinationBuffer, FoundInstallManifest, CompleteFn, ActivityRecord, InstallChunkSourceStat=InstallChunkSourceStat]()
 				{
-					// This is here because the compiler is complaining about ActivityRecord being a const object
-					// even though it's passed by value.
-					ISpeedRecorder::FRecord OurActivityRecord = ActivityRecord;
-
 					IInstallChunkSourceStat::ELoadResult Result = IInstallChunkSourceStat::ELoadResult::Success;
 					{
 						TRACE_CPUPROFILER_EVENT_SCOPE(InstallHash);
@@ -303,11 +303,10 @@ namespace BuildPatchServices
 						}
 					}
 
-					OurActivityRecord.CyclesEnd = FStatsCollector::GetCycles();
 					InstallChunkSourceStat->OnLoadComplete(
 						DataId, 
 						Result, 
-						OurActivityRecord);					
+						ActivityRecord);
 
 					CompleteFn.Execute(DataId, false, Result != IInstallChunkSourceStat::ELoadResult::Success, UserPtr);
 				}
