BPS: Fix chunkdb stats including hash speed and include it in normal disk read speeds.

#rnx

diff --git a/Engine/Source/Runtime/Online/BuildPatchServices/Private/BuildPatchInstaller.cpp b/Engine/Source/Runtime/Online/BuildPatchServices/Private/BuildPatchInstaller.cpp
--- a/Engine/Source/Runtime/Online/BuildPatchServices/Private/BuildPatchInstaller.cpp
+++ b/Engine/Source/Runtime/Online/BuildPatchServices/Private/BuildPatchInstaller.cpp
@@ -446,10 +446,9 @@ namespace BuildPatchServices
 		, DownloadSpeedRecorder(FSpeedRecorderFactory::Create())
 		, DiskReadSpeedRecorder(FSpeedRecorderFactory::Create())
 		, DiskWriteSpeedRecorder(FSpeedRecorderFactory::Create())
-		, ChunkDbReadSpeedRecorder(FSpeedRecorderFactory::Create())
 		, ChunkDataSizeProvider(FChunkDataSizeProviderFactory::Create())
 		, DownloadServiceStatistics(FDownloadServiceStatisticsFactory::Create(DownloadSpeedRecorder.Get(), ChunkDataSizeProvider.Get(), InstallerAnalytics.Get()))
-		, ChunkDbChunkSourceStatistics(FChunkDbChunkSourceStatisticsFactory::Create(ChunkDbReadSpeedRecorder.Get(), FileOperationTracker.Get()))
+		, ChunkDbChunkSourceStatistics(FChunkDbChunkSourceStatisticsFactory::Create(DiskReadSpeedRecorder.Get(), FileOperationTracker.Get()))
 		, InstallChunkSourceStatistics(FInstallChunkSourceStatisticsFactory::Create(DiskReadSpeedRecorder.Get(), InstallerAnalytics.Get(), FileOperationTracker.Get()))
 		, CloudChunkSourceStatistics(FCloudChunkSourceStatisticsFactory::Create(InstallerAnalytics.Get(), &BuildProgress, FileOperationTracker.Get()))
 		, FileConstructorStatistics(FFileConstructorStatisticsFactory::Create(DiskReadSpeedRecorder.Get(), DiskWriteSpeedRecorder.Get(), &BuildProgress, FileOperationTracker.Get()))
@@ -550,11 +549,6 @@ namespace BuildPatchServices
 		return DiskReadSpeedRecorder.Get();
 	}
 
-	const ISpeedRecorder* FBuildPatchInstaller::GetChunkDbReadSpeedRecorder() const
-	{
-		return ChunkDbReadSpeedRecorder.Get();
-	}
-
 	const ISpeedRecorder* FBuildPatchInstaller::GetDiskWriteSpeedRecorder() const
 	{
 		return DiskWriteSpeedRecorder.Get();
@@ -580,6 +574,11 @@ namespace BuildPatchServices
 		return FileConstructorStatistics.Get();
 	}
 
+	const IChunkDbChunkSourceStatistics* FBuildPatchInstaller::GetChunkDbChunkSourceStatistics() const
+	{
+		return ChunkDbChunkSourceStatistics.Get();
+	}
+
 	const IVerifierStatistics* FBuildPatchInstaller::GetVerifierStatistics() const
 	{
 		return VerifierStatistics.Get();
diff --git a/Engine/Source/Runtime/Online/BuildPatchServices/Private/BuildPatchInstaller.h b/Engine/Source/Runtime/Online/BuildPatchServices/Private/BuildPatchInstaller.h
--- a/Engine/Source/Runtime/Online/BuildPatchServices/Private/BuildPatchInstaller.h
+++ b/Engine/Source/Runtime/Online/BuildPatchServices/Private/BuildPatchInstaller.h
@@ -165,7 +165,6 @@ namespace BuildPatchServices
 		TUniquePtr<ISpeedRecorder> DownloadSpeedRecorder;
 		TUniquePtr<ISpeedRecorder> DiskReadSpeedRecorder;
 		TUniquePtr<ISpeedRecorder> DiskWriteSpeedRecorder;
-		TUniquePtr<ISpeedRecorder> ChunkDbReadSpeedRecorder;
 		TUniquePtr<IChunkDataSizeProvider> ChunkDataSizeProvider;
 		TUniquePtr<IDownloadServiceStatistics> DownloadServiceStatistics;
 		TUniquePtr<IChunkDbChunkSourceStatistics> ChunkDbChunkSourceStatistics;
@@ -284,11 +283,6 @@ namespace BuildPatchServices
 		 */
 		const ISpeedRecorder* GetDiskReadSpeedRecorder() const;
 
-		/**
-		 * @return the speed recorder for chunkdb read speed.
-		 */
-		const ISpeedRecorder* GetChunkDbReadSpeedRecorder() const;
-
 		/**
 		 * @return the speed recorder for disk write speed.
 		 */
@@ -314,6 +308,8 @@ namespace BuildPatchServices
 		 */
 		const IFileConstructorStatistics* GetFileConstructorStatistics() const;
 
+		const IChunkDbChunkSourceStatistics* GetChunkDbChunkSourceStatistics() const;
+
 		/**
 		 * @return the verifier statistics interface.
 		 */
diff --git a/Engine/Source/Runtime/Online/BuildPatchServices/Private/Installer/BuildStatistics.cpp b/Engine/Source/Runtime/Online/BuildPatchServices/Private/Installer/BuildStatistics.cpp
--- a/Engine/Source/Runtime/Online/BuildPatchServices/Private/Installer/BuildStatistics.cpp
+++ b/Engine/Source/Runtime/Online/BuildPatchServices/Private/Installer/BuildStatistics.cpp
@@ -2,6 +2,7 @@
 
 #include "Installer/BuildStatistics.h"
 #include "Installer/Statistics/FileOperationTracker.h"
+#include "Installer/Statistics/ChunkDbChunkSourceStatistics.h"
 #include "Installer/Statistics/CloudChunkSourceStatistics.h"
 #include "Installer/Statistics/FileConstructorStatistics.h"
 #include "Installer/Statistics/MemoryChunkStoreStatistics.h"
@@ -47,7 +48,6 @@ namespace BuildPatchServices
 		virtual const TArray<FFileOperation>& GetFileOperationStates() const override;
 		virtual double GetDownloadByteSpeed() const override;
 		virtual double GetDiskReadByteSpeed() const override;
-		virtual double GetChunkDbReadByteSpeed() const override;
 		virtual double GetDiskWriteByteSpeed() const override;
 		virtual int32 GetNumVerifyErrors() const override;
 		virtual TMap<EVerifyError, int32> GetVerifyErrorCounts() const override;
@@ -116,7 +116,8 @@ namespace BuildPatchServices
 	{
 		return Installer->GetFileConstructorStatistics()->IsCurrentlyReading()
 			|| Installer->GetInstallChunkSourceStatistics()->IsCurrentlyReading()
-			|| Installer->GetVerifierStatistics()->IsCurrentlyReading();
+			|| Installer->GetVerifierStatistics()->IsCurrentlyReading()
+			|| Installer->GetChunkDbChunkSourceStatistics()->IsCurrentlyReading();
 	}
 
 	const TArray<FFileOperation>& FBuildStatistics::GetFileOperationStates() const
@@ -134,11 +135,6 @@ namespace BuildPatchServices
 		return Installer->GetDiskReadSpeedRecorder()->GetAverageSpeed(10.0f);
 	}
 
-	double FBuildStatistics::GetChunkDbReadByteSpeed() const
-	{
-		return Installer->GetChunkDbReadSpeedRecorder()->GetAverageSpeed(10.0f);
-	}
-
 	double FBuildStatistics::GetDiskWriteByteSpeed() const
 	{
 		return Installer->GetDiskWriteSpeedRecorder()->GetAverageSpeed(10.0f);
diff --git a/Engine/Source/Runtime/Online/BuildPatchServices/Private/Installer/ChunkDbChunkSource.cpp b/Engine/Source/Runtime/Online/BuildPatchServices/Private/Installer/ChunkDbChunkSource.cpp
--- a/Engine/Source/Runtime/Online/BuildPatchServices/Private/Installer/ChunkDbChunkSource.cpp
+++ b/Engine/Source/Runtime/Online/BuildPatchServices/Private/Installer/ChunkDbChunkSource.cpp
@@ -367,15 +367,18 @@ namespace BuildPatchServices
 			FUniqueBuffer CompressedBuffer;
 			bool Result = ChunkDataSerialization->ValidateAndRead(*ChunkDbFile, DestinationBuffer, Header, CompressedBuffer);
 
+			// Save this here so we only include the IO time and not the hash/decompress time.
 			ActivityRecord.Size = ChunkDbFile->Tell() - ChunkLocation.ByteStart;
+			ActivityRecord.CyclesEnd = FStatsCollector::GetCycles();
+			ChunkDbChunkSourceStat->OnReadComplete(ActivityRecord);
 
 			if (!Result)
 			{
 				LoadResult = IChunkDbChunkSourceStat::ELoadResult::SerializationError;
 				ActivityRecord.CyclesEnd = FStatsCollector::GetCycles();
-				ChunkDbChunkSourceStat->OnLoadComplete(DataId, LoadResult, ActivityRecord);
+				ChunkDbChunkSourceStat->OnLoadComplete(DataId, LoadResult);
 
-				// The header or chunk data was bad.				
+				// The header or chunk data was bad.
 				CompleteFn.Execute(DataId, false, true, UserPtr);
 				return;
 			}
@@ -384,16 +387,12 @@ namespace BuildPatchServices
 			// block IO for this.
 			UE::Tasks::Launch(TEXT("ChunkDbDecompressionAndHash"), 
 				[DataId, CompleteFn, UserPtr, Header, DestinationBuffer, ChunkDataSerialization = ChunkDataSerialization, 
-				CompressedBuffer = MoveTemp(CompressedBuffer), ActivityRecord, ChunkDbChunkSourceStat = ChunkDbChunkSourceStat]()
+				CompressedBuffer = MoveTemp(CompressedBuffer), ChunkDbChunkSourceStat = ChunkDbChunkSourceStat]()
 				{
-					ISpeedRecorder::FRecord OurActivityRecord = ActivityRecord; // copy as the capture is const
 					bool bDecompressSucceeded = ChunkDataSerialization->DecompressValidatedRead(Header, DestinationBuffer, CompressedBuffer);
 
-					OurActivityRecord.CyclesEnd = FStatsCollector::GetCycles();
-
 					ChunkDbChunkSourceStat->OnLoadComplete(DataId, 
-						bDecompressSucceeded ? IChunkDbChunkSourceStat::ELoadResult::Success : IChunkDbChunkSourceStat::ELoadResult::CorruptedData,
-						OurActivityRecord);
+						bDecompressSucceeded ? IChunkDbChunkSourceStat::ELoadResult::Success : IChunkDbChunkSourceStat::ELoadResult::CorruptedData);
 
 					CompleteFn.Execute(DataId, false, !bDecompressSucceeded, UserPtr);
 				}
diff --git a/Engine/Source/Runtime/Online/BuildPatchServices/Private/Installer/ChunkDbChunkSource.h b/Engine/Source/Runtime/Online/BuildPatchServices/Private/Installer/ChunkDbChunkSource.h
--- a/Engine/Source/Runtime/Online/BuildPatchServices/Private/Installer/ChunkDbChunkSource.h
+++ b/Engine/Source/Runtime/Online/BuildPatchServices/Private/Installer/ChunkDbChunkSource.h
@@ -115,9 +115,11 @@ namespace BuildPatchServices
 		 * Called each time a chunk load completes.
 		 * @param ChunkId   The id of the chunk.
 		 * @param Result    The load result.
-		 * @param Record    The recording for the received data.
 		 */
-		virtual void OnLoadComplete(const FGuid& ChunkId, ELoadResult Result, const ISpeedRecorder::FRecord& Record) = 0;
+		virtual void OnLoadComplete(const FGuid& ChunkId, ELoadResult Result) = 0;
+
+		// Called when the read for the load is complete and we're moving on to decompression/hashing.
+		virtual void OnReadComplete(const ISpeedRecorder::FRecord& Record) = 0;
 	};
 	
 	/**
diff --git a/Engine/Source/Runtime/Online/BuildPatchServices/Private/Installer/Statistics/ChunkDbChunkSourceStatistics.cpp b/Engine/Source/Runtime/Online/BuildPatchServices/Private/Installer/Statistics/ChunkDbChunkSourceStatistics.cpp
--- a/Engine/Source/Runtime/Online/BuildPatchServices/Private/Installer/Statistics/ChunkDbChunkSourceStatistics.cpp
+++ b/Engine/Source/Runtime/Online/BuildPatchServices/Private/Installer/Statistics/ChunkDbChunkSourceStatistics.cpp
@@ -18,12 +18,14 @@ namespace BuildPatchServices
 	public:
 		// IChunkDbChunkSourceStat interface begin.
 		virtual void OnLoadStarted(const FGuid& ChunkId) override;
-		virtual void OnLoadComplete(const FGuid& ChunkId, ELoadResult Result, const ISpeedRecorder::FRecord& Record) override;
+		virtual void OnLoadComplete(const FGuid& ChunkId, ELoadResult Result) override;
+		virtual void OnReadComplete(const ISpeedRecorder::FRecord& Record) override;
 		// IChunkDbChunkSourceStat interface end.
 
 		// IChunkDbChunkSourceStatistics interface begin.
 		virtual int32 GetNumSuccessfulLoads() const override;
 		virtual int32 GetNumFailedLoads() const override;
+		virtual bool IsCurrentlyReading() const override { return bIsReading.load(std::memory_order_relaxed); }
 		// IChunkDbChunkSourceStatistics interface end.
 
 	private:
@@ -31,6 +33,7 @@ namespace BuildPatchServices
 		IFileOperationTracker* FileOperationTracker;
 		FThreadSafeInt32 NumSuccessfulLoads;
 		FThreadSafeInt32 NumFailedLoads;
+		std::atomic_bool bIsReading;
 	};
 
 	FChunkDbChunkSourceStatistics::FChunkDbChunkSourceStatistics(ISpeedRecorder* InSpeedRecorder, IFileOperationTracker* InFileOperationTracker)
@@ -38,6 +41,7 @@ namespace BuildPatchServices
 		, FileOperationTracker(InFileOperationTracker)
 		, NumSuccessfulLoads(0)
 		, NumFailedLoads(0)
+		, bIsReading(false)
 	{
 	}
 
@@ -48,9 +52,16 @@ namespace BuildPatchServices
 	void FChunkDbChunkSourceStatistics::OnLoadStarted(const FGuid& ChunkId)
 	{
 		FileOperationTracker->OnDataStateUpdate(ChunkId, EFileOperationState::RetrievingLocalChunkDbData);
+		bIsReading.store(true, std::memory_order_relaxed);
 	}
 
-	void FChunkDbChunkSourceStatistics::OnLoadComplete(const FGuid& ChunkId, ELoadResult Result, const ISpeedRecorder::FRecord& Record)
+	void FChunkDbChunkSourceStatistics::OnReadComplete(const ISpeedRecorder::FRecord& Record)
+	{
+		SpeedRecorder->AddRecord(Record);
+		bIsReading.store(false, std::memory_order_relaxed);
+	}
+
+	void FChunkDbChunkSourceStatistics::OnLoadComplete(const FGuid& ChunkId, ELoadResult Result)
 	{
 		if (Result == ELoadResult::Success)
 		{
@@ -60,7 +71,7 @@ namespace BuildPatchServices
 		{
 			NumFailedLoads.Increment();
 		}
-		SpeedRecorder->AddRecord(Record);
+		
 	}
 
 	int32 FChunkDbChunkSourceStatistics::GetNumSuccessfulLoads() const
diff --git a/Engine/Source/Runtime/Online/BuildPatchServices/Private/Installer/Statistics/ChunkDbChunkSourceStatistics.h b/Engine/Source/Runtime/Online/BuildPatchServices/Private/Installer/Statistics/ChunkDbChunkSourceStatistics.h
--- a/Engine/Source/Runtime/Online/BuildPatchServices/Private/Installer/Statistics/ChunkDbChunkSourceStatistics.h
+++ b/Engine/Source/Runtime/Online/BuildPatchServices/Private/Installer/Statistics/ChunkDbChunkSourceStatistics.h
@@ -24,6 +24,8 @@ namespace BuildPatchServices
 		 * @return the number of chunks which failed to load from provided chunkdbs.
 		 */
 		virtual int32 GetNumFailedLoads() const = 0;
+
+		virtual bool IsCurrentlyReading() const = 0;
 	};
 
 	/**
diff --git a/Engine/Source/Runtime/Online/BuildPatchServices/Public/Interfaces/IBuildStatistics.h b/Engine/Source/Runtime/Online/BuildPatchServices/Public/Interfaces/IBuildStatistics.h
--- a/Engine/Source/Runtime/Online/BuildPatchServices/Public/Interfaces/IBuildStatistics.h
+++ b/Engine/Source/Runtime/Online/BuildPatchServices/Public/Interfaces/IBuildStatistics.h
@@ -207,7 +207,8 @@ namespace BuildPatchServices
 		/**
 		 * @return the current chunkdb read speed in bytes per second.
 		 */
-		virtual double GetChunkDbReadByteSpeed() const = 0;
+		UE_DEPRECATED(5.6, "Chunk db read speeds now included in DiskReadByteSpeed")
+		double GetChunkDbReadByteSpeed() const { return 0; }
 
 		/**
 		 * @return the current disk write speed in bytes per second.
