Add some sockets debug commands in http module.

#tests Tried on android

diff --git a/Engine/Source/Runtime/Online/HTTP/HTTP.Build.cs b/Engine/Source/Runtime/Online/HTTP/HTTP.Build.cs
--- a/Engine/Source/Runtime/Online/HTTP/HTTP.Build.cs
+++ b/Engine/Source/Runtime/Online/HTTP/HTTP.Build.cs
@@ -8,6 +8,8 @@ public class HTTP : ModuleRules
 
 	protected virtual bool bPlatformSupportToIncreaseMaxRequestsAtRuntime { get { return true; } }
 
+	protected virtual bool bPlatformSupportsSocketTestCommand { get { return Target.IsInPlatformGroup(UnrealPlatformGroup.Android); } }
+
 	protected virtual bool bPlatformSupportsWinHttp
 	{
 		get
@@ -109,6 +111,7 @@ public class HTTP : ModuleRules
 
 		PrivateDefinitions.Add("UE_HTTP_EVENT_LOOP_ENABLE_CHANCE_BY_DEFAULT=" + (bPlatformEventLoopEnabledByDefault ? "100" : "0"));
 		PrivateDefinitions.Add("UE_HTTP_SUPPORT_TO_INCREASE_MAX_REQUESTS_AT_RUNTIME=" + (bPlatformSupportToIncreaseMaxRequestsAtRuntime ? "1" : "0"));
+		PrivateDefinitions.Add("UE_HTTP_SOCKET_TEST_COMMAND_ENABLED=" + (bPlatformSupportsSocketTestCommand ? "1" : "0"));
 		PrivateDefinitions.Add("WITH_CURL_LIBCURL =" + (bPlatformSupportsLibCurl ? "1" : "0"));
 		PrivateDefinitions.Add("WITH_CURL_XCURL=" + (bPlatformSupportsXCurl ? "1" : "0"));
 		PrivateDefinitions.Add("WITH_CURL_MULTIPOLL=" + (bPlatformSupportsCurlMultiPoll ? "1" : "0"));
diff --git a/Engine/Source/Runtime/Online/HTTP/Private/HttpModule.cpp b/Engine/Source/Runtime/Online/HTTP/Private/HttpModule.cpp
--- a/Engine/Source/Runtime/Online/HTTP/Private/HttpModule.cpp
+++ b/Engine/Source/Runtime/Online/HTTP/Private/HttpModule.cpp
@@ -10,6 +10,17 @@
 #include "Misc/CommandLine.h"
 #include "Misc/CoreDelegates.h"
 
+// TODO: Create a public common header, similar to BSDSocketTypesPrivate.h to support more platforms
+#if UE_HTTP_SOCKET_TEST_COMMAND_ENABLED
+#include "Logging/StructuredLog.h"
+
+#include <unistd.h>
+#include <sys/socket.h>
+#include <netinet/in.h>
+#include <dirent.h>
+#include <arpa/inet.h>
+#endif
+
 #if !UE_BUILD_SHIPPING
 #include "HAL/IConsoleManager.h"
 #endif
@@ -264,6 +275,74 @@ bool FHttpModule::HandleHTTPCommand(const TCHAR* Cmd, FOutputDevice& Ar)
 			}
 		}
 	}
+#if UE_HTTP_SOCKET_TEST_COMMAND_ENABLED
+	else if (FParse::Command(&Cmd, TEXT("DUMPFDS")))
+	{
+		DIR* Dir = opendir("/proc/self/fd");
+		if (!Dir)
+		{
+			UE_LOG(LogHttp, Warning, TEXT("Failed to open directory /proc/self/fd"));
+		}
+		else
+		{
+			struct dirent* Entry = nullptr;
+			while ((Entry = readdir(Dir)) != nullptr)
+			{
+				if (strcmp(Entry->d_name, ".") == 0 || strcmp(Entry->d_name, "..") == 0)
+				{
+					continue;
+				}
+
+				char FullPath[2048];
+				snprintf(FullPath, sizeof(FullPath), "/proc/self/fd/%s", Entry->d_name);
+				FUtf8StringView FullPathStringView(FullPath);
+
+				char Dest[1024] = { 0 };
+
+				int32 Len = readlink(FullPath, Dest, sizeof(Dest));
+
+				if (Len == -1)
+				{
+					UE_LOGFMT(LogHttp, Warning, "{FullPathStringView}", FullPathStringView);
+					continue;
+				}
+
+				FUtf8StringView DestStringView(Dest, Len);
+				UE_LOGFMT(LogHttp, Warning, "{FullPathStringView} -> {DestStringView}", FullPathStringView, DestStringView);
+			}
+
+			closedir(Dir);
+		}
+	}
+	else if (FParse::Command(&Cmd, TEXT("LEAKSOCKET")))
+	{
+		// Leak sockets in purpose to verify high File Descriptor number can be handled properly in different places
+		int NumLeakedSockets = 0;
+		while (true)
+		{
+			NumLeakedSockets += 1;
+			int Socket = socket(AF_INET, SOCK_DGRAM, IPPROTO_UDP);
+			UE_LOG(LogHttp, Log, TEXT("Leaking socket, Socket=%i, NumLeakedSockets=%i"), Socket, NumLeakedSockets);
+			if (Socket > 1024)
+			{
+				break;
+			}
+		}
+	}
+	else if (FParse::Command(&Cmd, TEXT("SHOWLATESTFD")))
+	{
+		int Socket = socket(AF_INET, SOCK_DGRAM, IPPROTO_UDP);
+		if (Socket == 0)
+		{
+			UE_LOG(LogHttp, Warning, TEXT("Failed to open new socket!"));
+		}
+		else
+		{
+			UE_LOG(LogHttp, Warning, TEXT("Latest socket fd is %i"), Socket);
+			close(Socket);
+		}
+	}
+#endif // UE_HTTP_SOCKET_TEST_COMMAND_ENABLED
 
 	return true;
 }
