BPS: Requested memory sizes no longer overridden by headroom calculations, allowing for an enforced minimum.

#rnx

diff --git a/Engine/Source/Runtime/Online/BuildPatchServices/Private/BuildPatchFileConstructor.cpp b/Engine/Source/Runtime/Online/BuildPatchServices/Private/BuildPatchFileConstructor.cpp
--- a/Engine/Source/Runtime/Online/BuildPatchServices/Private/BuildPatchFileConstructor.cpp
+++ b/Engine/Source/Runtime/Online/BuildPatchServices/Private/BuildPatchFileConstructor.cpp
@@ -917,25 +917,38 @@ public:
 
 
 
-
-		uint64 AvailableMem = FPlatformMemory::GetStats().AvailablePhysical;
+		FPlatformMemoryStats MemoryStats = FPlatformMemory::GetStats();
+		uint64 AvailableMem = MemoryStats.AvailablePhysical;
 		if (ChunkStoreMemoryHeadRoomMB >= 0)
 		{
 			uint64 RequestedHeadRoom = ((uint64)ChunkStoreMemoryHeadRoomMB << 20);
 
+			uint64 ProposedChunkStoreMemoryLimit = ChunkStoreMemoryLimit;
+
 			if (RequestedHeadRoom < AvailableMem)
 			{
 				uint64 MemoryStoreMem = AvailableMem - RequestedHeadRoom;
-				ChunkStoreMemoryLimit = MemoryStoreMem;
+				ProposedChunkStoreMemoryLimit = MemoryStoreMem;
 			}
 			else
 			{
 				// Cap at available.
-				if (ChunkStoreMemoryLimit > AvailableMem)
+				if (ProposedChunkStoreMemoryLimit > AvailableMem)
 				{
-					ChunkStoreMemoryLimit = AvailableMem;
+					ProposedChunkStoreMemoryLimit = AvailableMem;
 				}
 			}
+
+			// If there's already a limit requested by the inis, we don't want to make it _smaller_. If there's no
+			// limit specified, then use the proposed limit.
+			if (ChunkStoreMemoryLimit != ChunkStoreMemoryLimitDisabledSentinel)
+			{
+				ChunkStoreMemoryLimit = FMath::Max(ChunkStoreMemoryLimit, ProposedChunkStoreMemoryLimit);
+			}
+			else
+			{
+				ChunkStoreMemoryLimit = ProposedChunkStoreMemoryLimit;
+			}
 		}
 
 		if (ChunkStoreMemoryLimit == ChunkStoreMemoryLimitDisabledSentinel)
@@ -944,10 +957,13 @@ public:
 		}
 		else
 		{
-			UE_LOG(LogBuildPatchServices, Display, TEXT("ChunkStoreMemoryLimits: %s using headroom of %s (%s available memory)"), 
+			UE_LOG(LogBuildPatchServices, Display, TEXT("ChunkStoreMemoryLimits: %s using headroom of %s (%s available memory, %s used physical, %s used virtual)"), 
 				*FormatNumber(ChunkStoreMemoryLimit),
 				ChunkStoreMemoryHeadRoomMB >= 0 ? *FormatNumber((uint64)ChunkStoreMemoryHeadRoomMB << 20) : TEXT("<disabled>"),
-				*FormatNumber(AvailableMem));
+				*FormatNumber(AvailableMem),
+				*FormatNumber(MemoryStats.UsedPhysical),
+				*FormatNumber(MemoryStats.UsedVirtual)
+			);
 		}
 
 		MaxMemoryBytes = ChunkStoreMemoryLimit;
