LwsWebSockets: Service polling is now enforced in single-threaded environments. Blocking service calls will then take effect, if configured, after process forking for multi-threaded child processes.

diff --git a/Engine/Source/Runtime/Online/WebSockets/Private/Lws/LwsWebSocketsManager.cpp b/Engine/Source/Runtime/Online/WebSockets/Private/Lws/LwsWebSocketsManager.cpp
--- a/Engine/Source/Runtime/Online/WebSockets/Private/Lws/LwsWebSocketsManager.cpp
+++ b/Engine/Source/Runtime/Online/WebSockets/Private/Lws/LwsWebSocketsManager.cpp
@@ -75,6 +75,18 @@ void FLwsWebSocketsManager::UpdateConfigs()
 	// Allow -lws-insecure= to override ini flags
 	FParse::Bool(FCommandLine::Get(), TEXT("lws-insecure="), bDisableCertValidation);
 #endif
+
+	// Check blocking service calls are supported if requested
+	if (!bPollService)
+	{
+		// Ensure multi-threading is supported otherwise 
+		const bool bSupportsMultithreading = FPlatformProcess::SupportsMultithreading() || FForkProcessHelper::IsForkedMultithreadInstance();
+		if (!bSupportsMultithreading)
+		{
+			UE_LOG(LogWebSockets, Warning, TEXT("FLwsWebSocketsManager::UpdateConfigs: Service polling disabled by config in single threaded environment. Forcing service polling to avoid blocking service call"));
+			bPollService = true;
+		}
+	}
 	
 	if (bWasPollingService != bPollService)
 	{
@@ -207,17 +219,33 @@ void FLwsWebSocketsManager::InitWebSockets(TArrayView<const FString> Protocols)
 	// Setup our game thread tick
 	FTickerDelegate TickDelegate = FTickerDelegate::CreateRaw(this, &FLwsWebSocketsManager::GameThreadTick);
 	TickHandle = FTSBackgroundableTicker::GetCoreTicker().AddTicker(TickDelegate, 0.0f);
+
+	// If we get forked, we may want to enable / disable service polling
+	if (FForkProcessHelper::IsForkRequested())
+	{
+		OnChildEndFramePostForkHandle = FCoreDelegates::OnChildEndFramePostFork.AddLambda([this]()
+		{
+			UpdateConfigs();
+		});
+	}
 }
 
 void FLwsWebSocketsManager::ShutdownWebSockets()
 {
 	QUICK_SCOPE_CYCLE_COUNTER(STAT_FLwsWebSocketsManager_ShutdownWebSockets);
+	
 	if (TickHandle.IsValid())
 	{
 		FTSBackgroundableTicker::GetCoreTicker().RemoveTicker(TickHandle);
 		TickHandle.Reset();
 	}
 
+	if (OnChildEndFramePostForkHandle.IsValid())
+	{
+		FCoreDelegates::OnChildEndFramePostFork.Remove(OnChildEndFramePostForkHandle);
+		OnChildEndFramePostForkHandle.Reset();
+	}
+
 	if (Thread)
 	{
 		Thread->Kill(true);
diff --git a/Engine/Source/Runtime/Online/WebSockets/Private/Lws/LwsWebSocketsManager.h b/Engine/Source/Runtime/Online/WebSockets/Private/Lws/LwsWebSocketsManager.h
--- a/Engine/Source/Runtime/Online/WebSockets/Private/Lws/LwsWebSocketsManager.h
+++ b/Engine/Source/Runtime/Online/WebSockets/Private/Lws/LwsWebSocketsManager.h
@@ -104,6 +104,9 @@ private:
 	/** Delegate for callbacks to GameThreadTick */
 	FTSTicker::FDelegateHandle TickHandle;
 
+	/** Delegate handle for FCoreDelegates::OnChildEndFramePostFork */
+	FDelegateHandle OnChildEndFramePostForkHandle;
+
 	// Thread variables
 	/** Pointer to Runnable Thread */
 	FRunnableThread* Thread;
