Add support to return utf8 string view from http response.

diff --git a/Engine/Source/Programs/WebTests/Private/TestHttp.cpp b/Engine/Source/Programs/WebTests/Private/TestHttp.cpp
--- a/Engine/Source/Programs/WebTests/Private/TestHttp.cpp
+++ b/Engine/Source/Programs/WebTests/Private/TestHttp.cpp
@@ -1642,6 +1642,17 @@ TEST_CASE_METHOD(FValidateResponseDependencyFixture, "Http query with parameters
 
 		CHECK(HttpRequest->GetURL() == UrlQueryWithParams);
 
+		// Validate GetContentAsUtf8StringView works well
+		FQueryWithParamsResponse QueryWithParamsResponseThroughStringView;
+		REQUIRE(QueryWithParamsResponseThroughStringView.FromJsonStringView(HttpResponse->GetContentAsUtf8StringView()));
+
+		CHECK(FString::FromInt(QueryWithParamsResponseThroughStringView.VarInt) == HttpRequest->GetURLParameter(TEXT("var_int")));
+		CHECK(QueryWithParamsResponseThroughStringView.VarStr == HttpRequest->GetURLParameter(TEXT("var_str")));
+
+		CHECK(FString::FromInt(QueryWithParamsResponseThroughStringView.VarInt) == HttpResponse->GetURLParameter(TEXT("var_int")));
+		CHECK(QueryWithParamsResponseThroughStringView.VarStr == HttpResponse->GetURLParameter(TEXT("var_str")));
+
+		// Validate GetContentAsString works well
 		FQueryWithParamsResponse QueryWithParamsResponse;
 		REQUIRE(QueryWithParamsResponse.FromJson(HttpResponse->GetContentAsString()));
 
diff --git a/Engine/Source/Runtime/Online/BuildPatchServices/Private/Tests/Mock/HttpResponse.mock.h b/Engine/Source/Runtime/Online/BuildPatchServices/Private/Tests/Mock/HttpResponse.mock.h
--- a/Engine/Source/Runtime/Online/BuildPatchServices/Private/Tests/Mock/HttpResponse.mock.h
+++ b/Engine/Source/Runtime/Online/BuildPatchServices/Private/Tests/Mock/HttpResponse.mock.h
@@ -28,6 +28,12 @@ namespace BuildPatchServices
 			return FString();
 		}
 
+		virtual FUtf8StringView GetContentAsUtf8StringView() const override
+		{
+			MOCK_FUNC_NOT_IMPLEMENTED("FMockHttpRequest::GetContentAsUtf8StringView");
+			return FUtf8StringView();
+		}
+
 		virtual const FString& GetURL() const override
 		{
 			MOCK_FUNC_NOT_IMPLEMENTED("FMockHttpRequest::GetURL");
diff --git a/Engine/Source/Runtime/Online/HTTP/Private/GenericPlatform/HttpResponseCommon.cpp b/Engine/Source/Runtime/Online/HTTP/Private/GenericPlatform/HttpResponseCommon.cpp
--- a/Engine/Source/Runtime/Online/HTTP/Private/GenericPlatform/HttpResponseCommon.cpp
+++ b/Engine/Source/Runtime/Online/HTTP/Private/GenericPlatform/HttpResponseCommon.cpp
@@ -67,3 +67,8 @@ void FHttpResponseCommon::SetResponseCode(int32 InResponseCode)
 	ResponseCode = InResponseCode;
 }
 
+FUtf8StringView FHttpResponseCommon::GetContentAsUtf8StringView() const
+{
+	const TArray<uint8>& Payload = GetContent();
+	return FUtf8StringView(reinterpret_cast<const UTF8CHAR*>(Payload.GetData()), Payload.Num());
+}
diff --git a/Engine/Source/Runtime/Online/HTTP/Private/GenericPlatform/HttpResponseCommon.h b/Engine/Source/Runtime/Online/HTTP/Private/GenericPlatform/HttpResponseCommon.h
--- a/Engine/Source/Runtime/Online/HTTP/Private/GenericPlatform/HttpResponseCommon.h
+++ b/Engine/Source/Runtime/Online/HTTP/Private/GenericPlatform/HttpResponseCommon.h
@@ -14,21 +14,22 @@ class FHttpResponseCommon : public IHttpResponse
 	friend FHttpRequestCommon;
 
 public:
-	HTTP_API FHttpResponseCommon(const FHttpRequestCommon& HttpRequest);
+	FHttpResponseCommon(const FHttpRequestCommon& HttpRequest);
 
 	// IHttpBase
-	HTTP_API virtual FString GetURLParameter(const FString& ParameterName) const override;
-	HTTP_API virtual const FString& GetURL() const override;
-	HTTP_API virtual const FString& GetEffectiveURL() const override;
-	HTTP_API virtual EHttpRequestStatus::Type GetStatus() const override;
-	HTTP_API virtual EHttpFailureReason GetFailureReason() const override;
-	HTTP_API virtual int32 GetResponseCode() const override;
+	virtual FString GetURLParameter(const FString& ParameterName) const override;
+	virtual const FString& GetURL() const override;
+	virtual const FString& GetEffectiveURL() const override;
+	virtual EHttpRequestStatus::Type GetStatus() const override;
+	virtual EHttpFailureReason GetFailureReason() const override;
+	virtual int32 GetResponseCode() const override;
+	virtual FUtf8StringView GetContentAsUtf8StringView() const override;
 
 protected:
-	HTTP_API void SetRequestStatus(EHttpRequestStatus::Type InCompletionStatus);
-	HTTP_API void SetRequestFailureReason(EHttpFailureReason InFailureReason);
-	HTTP_API void SetEffectiveURL(const FString& InEffectiveURL);
-	HTTP_API void SetResponseCode(int32 InResponseCode);
+	void SetRequestStatus(EHttpRequestStatus::Type InCompletionStatus);
+	void SetRequestFailureReason(EHttpFailureReason InFailureReason);
+	void SetEffectiveURL(const FString& InEffectiveURL);
+	void SetResponseCode(int32 InResponseCode);
 
 	FString URL;
 	FString EffectiveURL;
diff --git a/Engine/Source/Runtime/Online/HTTP/Private/NullHttp.cpp b/Engine/Source/Runtime/Online/HTTP/Private/NullHttp.cpp
--- a/Engine/Source/Runtime/Online/HTTP/Private/NullHttp.cpp
+++ b/Engine/Source/Runtime/Online/HTTP/Private/NullHttp.cpp
@@ -297,3 +297,8 @@ FString FNullHttpResponse::GetContentAsString() const
 {
 	return FString();
 }
+
+FUtf8StringView FNullHttpResponse::GetContentAsUtf8StringView() const
+{
+	return FUtf8StringView();
+}
diff --git a/Engine/Source/Runtime/Online/HTTP/Private/NullHttp.h b/Engine/Source/Runtime/Online/HTTP/Private/NullHttp.h
--- a/Engine/Source/Runtime/Online/HTTP/Private/NullHttp.h
+++ b/Engine/Source/Runtime/Online/HTTP/Private/NullHttp.h
@@ -90,6 +90,7 @@ class FNullHttpResponse : public IHttpResponse
 	//~ Begin IHttpResponse Interface
 	virtual int32 GetResponseCode() const override;
 	virtual FString GetContentAsString() const override;
+	virtual FUtf8StringView GetContentAsUtf8StringView() const override;
 
 	FNullHttpResponse() {}
 	virtual ~FNullHttpResponse() {}
diff --git a/Engine/Source/Runtime/Online/HTTP/Public/Interfaces/IHttpResponse.h b/Engine/Source/Runtime/Online/HTTP/Public/Interfaces/IHttpResponse.h
--- a/Engine/Source/Runtime/Online/HTTP/Public/Interfaces/IHttpResponse.h
+++ b/Engine/Source/Runtime/Online/HTTP/Public/Interfaces/IHttpResponse.h
@@ -126,6 +126,14 @@ public:
 	 */
 	virtual FString GetContentAsString() const = 0;
 
+	/**
+	 * Returns the payload as a utf8 string view. This does not validate that the response is valid 
+	 * utf8. It is the caller's responsibility, for example, by checking the Content-Type header
+	 *
+	 * @return the payload as a utf8 string view.
+	 */
+	virtual FUtf8StringView GetContentAsUtf8StringView() const = 0;
+
 	/** 
 	 * Destructor for overrides 
 	 */
