BPS: Allow file system throttling to be overridden per installation. Default to true.

#rnx

diff --git a/Engine/Source/Runtime/Online/BuildPatchServices/Private/BuildPatchFileConstructor.cpp b/Engine/Source/Runtime/Online/BuildPatchServices/Private/BuildPatchFileConstructor.cpp
--- a/Engine/Source/Runtime/Online/BuildPatchServices/Private/BuildPatchFileConstructor.cpp
+++ b/Engine/Source/Runtime/Online/BuildPatchServices/Private/BuildPatchFileConstructor.cpp
@@ -33,17 +33,18 @@
 
 using namespace BuildPatchServices;
 
-static int32 SleepTimeWhenFileSystemThrottledSeconds = 15;
+static int32 SleepTimeWhenFileSystemThrottledSeconds = 1;
 static FAutoConsoleVariableRef CVarSleepTimeWhenFileSystemThrottledSeconds(
 	TEXT("BuildPatchFileConstructor.SleepTimeWhenFileSystemThrottledSeconds"),
 	SleepTimeWhenFileSystemThrottledSeconds,
 	TEXT("The amount of time to sleep if the destination filesystem is throttled."),
 	ECVF_Default);
 
-static bool bStallWhenFileSystemThrottled = false;
+// This can be overridden by the installation parameters.
+static bool bCVarStallWhenFileSystemThrottled = true;
 static FAutoConsoleVariableRef CVarStallWhenFileSystemThrottled(
 	TEXT("BuildPatchFileConstructor.bStallWhenFileSystemThrottled"),
-	bStallWhenFileSystemThrottled,
+	bCVarStallWhenFileSystemThrottled,
 	TEXT("Whether to stall if the file system is throttled"),
 	ECVF_Default);
 
@@ -1365,6 +1366,13 @@ FBuildPatchFileConstructor::FBuildPatchFileConstructor(
 {
 	TRACE_CPUPROFILER_EVENT_SCOPE(FileConstructor_ctor);
 
+	bStallWhenFileSystemThrottled = bCVarStallWhenFileSystemThrottled;
+	if (Configuration.StallWhenFileSystemThrottled.IsSet())
+	{
+		bStallWhenFileSystemThrottled = Configuration.StallWhenFileSystemThrottled.GetValue();
+		UE_LOG(LogBuildPatchServices, Display, TEXT("Overridding StallWhenFileSystemThrottled to: %d, was %d"), Configuration.StallWhenFileSystemThrottled.GetValue(), bStallWhenFileSystemThrottled);
+	}
+
 	BackingStore.Reset(new FChunkBackingStore(this, Configuration.InstallDirectory, BackingStoreStats));
 
 	// Count initial job size
diff --git a/Engine/Source/Runtime/Online/BuildPatchServices/Private/BuildPatchFileConstructor.h b/Engine/Source/Runtime/Online/BuildPatchServices/Private/BuildPatchFileConstructor.h
--- a/Engine/Source/Runtime/Online/BuildPatchServices/Private/BuildPatchFileConstructor.h
+++ b/Engine/Source/Runtime/Online/BuildPatchServices/Private/BuildPatchFileConstructor.h
@@ -89,6 +89,7 @@ namespace BuildPatchServices
 		TOptional<bool> SpawnAdditionalIOThreads;
 		TOptional<int32> IOBatchSizeMB;
 		TOptional<int32> IOBufferSizeMB;
+		TOptional<bool> StallWhenFileSystemThrottled;
 	};
 
 	/**
@@ -212,6 +213,9 @@ namespace BuildPatchServices
 		// If true, the chunkdb source has chunks to provide
 		bool bHasChunkDbSource = false;
 
+		// Our local resolved copy of the cvar with overrides applied.
+		bool bStallWhenFileSystemThrottled = false;
+
 		// A flag marking whether we should be paused.
 		FThreadSafeBool bIsPaused;
 
diff --git a/Engine/Source/Runtime/Online/BuildPatchServices/Private/BuildPatchInstaller.cpp b/Engine/Source/Runtime/Online/BuildPatchServices/Private/BuildPatchInstaller.cpp
--- a/Engine/Source/Runtime/Online/BuildPatchServices/Private/BuildPatchInstaller.cpp
+++ b/Engine/Source/Runtime/Online/BuildPatchServices/Private/BuildPatchInstaller.cpp
@@ -1413,6 +1413,7 @@ namespace BuildPatchServices
 			FCC.SpawnAdditionalIOThreads = Configuration.ConstructorSpawnAdditionalIOThreads;
 			FCC.IOBufferSizeMB = Configuration.ConstructorIOBufferSizeMB;
 			FCC.IOBatchSizeMB = Configuration.ConstructorIOBatchSizeMB;
+			FCC.StallWhenFileSystemThrottled = Configuration.ConstructorStallWhenFileSystemThrottled;
 
 			TUniquePtr<FBuildPatchFileConstructor> FileConstructor(
 				new FBuildPatchFileConstructor(
diff --git a/Engine/Source/Runtime/Online/BuildPatchServices/Public/BuildPatchSettings.h b/Engine/Source/Runtime/Online/BuildPatchServices/Public/BuildPatchSettings.h
--- a/Engine/Source/Runtime/Online/BuildPatchServices/Public/BuildPatchSettings.h
+++ b/Engine/Source/Runtime/Online/BuildPatchServices/Public/BuildPatchSettings.h
@@ -305,6 +305,13 @@ namespace BuildPatchServices
 		// for cloud/download installs this needs to be fairly large to mitigate the effect of out of order download completion. 
 		// Unset uses Engine.ini value [Portal.BuildPatch]ConstructorIOBufferSizeMB if present, FFileConstructorConfig::DefaultIOBufferSizeMB otherwise.
 		TOptional<int32> ConstructorIOBufferSizeMB;
+
+		// Whether the file construction should adhere to file system write limits. If a write limit is hit, writing will stall until
+		// sufficient quota exists for the write. Inheritance for this property is:
+		// Default = false
+		// if set, CVar BuildPatchFileConstructor.bStallWhenFileSystemThrottled
+		// if set, this member.
+		TOptional<bool> ConstructorStallWhenFileSystemThrottled;
 	};
 
 	/**
