BPS: Add http headers support to chunk uri response / download service.
BPS: Re-add chunk uri remapping to new cloud download source.
BPS: fix some cases where the original error could get lost in reporting.

#rnx

diff --git a/Engine/Source/Runtime/Online/BuildPatchServices/Private/BuildPatchFileConstructor.cpp b/Engine/Source/Runtime/Online/BuildPatchServices/Private/BuildPatchFileConstructor.cpp
--- a/Engine/Source/Runtime/Online/BuildPatchServices/Private/BuildPatchFileConstructor.cpp
+++ b/Engine/Source/Runtime/Online/BuildPatchServices/Private/BuildPatchFileConstructor.cpp
@@ -2152,6 +2152,7 @@ void FBuildPatchFileConstructor::RequestCompletedFn(const FGuid& Guid, bool bAbo
 	{
 		if (bFailedToRead)
 		{
+			Request->Batch->ErrorContextGuid = Guid;
 			Request->Batch->FailedRequestCount.fetch_add(1, std::memory_order_relaxed);
 		}
 		else if (!bAborted)
@@ -3261,7 +3262,9 @@ void FBuildPatchFileConstructor::ConstructFiles(const FResumeData& ResumeData)
 							Batch->bIsReading = false;
 							Batch->bNeedsWrite = true;
 
-							if (File->ConstructionError != EConstructionError::None ||
+							// Only update the file construction error if it's the first one
+							// since errors can cascade and make it not clear what the original problem was.
+							if (File->ConstructionError == EConstructionError::None &&
 								Batch->Error != EConstructionError::None)
 							{
 								File->ErrorContextGuid = Batch->ErrorContextGuid;
diff --git a/Engine/Source/Runtime/Online/BuildPatchServices/Private/Installer/CloudChunkSource.cpp b/Engine/Source/Runtime/Online/BuildPatchServices/Private/Installer/CloudChunkSource.cpp
--- a/Engine/Source/Runtime/Online/BuildPatchServices/Private/Installer/CloudChunkSource.cpp
+++ b/Engine/Source/Runtime/Online/BuildPatchServices/Private/Installer/CloudChunkSource.cpp
@@ -525,16 +525,34 @@ namespace BuildPatchServices
 				ChunkUriRequest.RelativePath = ManifestSet->GetDataFilename(ChunkUriToRequest);
 				ChunkUriRequest.RelativePath.RemoveFromStart(TEXT("/"));
 
-				MessagePump->SendRequest(ChunkUriRequest, [ChunkUriResponsesRef, ChunkUriToRequest](FChunkUriResponse Response)
+				bool bMessageSent = MessagePump->SendRequest(ChunkUriRequest, [ChunkUriResponsesRef, ChunkUriToRequest](FChunkUriResponse Response)
 				{
 					ChunkUriResponsesRef->Enqueue(FGuidUriResponse(ChunkUriToRequest, MoveTemp(Response)));
 				});
+
+				if (!bMessageSent)
+				{
+					// If no handler is registered SendRequest does nothing - make our own default response
+					FChunkUriResponse Response;
+					Response.Uri = ChunkUriRequest.CloudDirectory / ChunkUriRequest.RelativePath;
+					ChunkUriResponsesRef->Enqueue(FGuidUriResponse(ChunkUriToRequest, MoveTemp(Response)));
+				}
 			}
 			// Process new chunk uri responses.
 			FGuidUriResponse ChunkUriResponse;
 			for (FGuidUriResponseQueue& ChunkUriResponses = ChunkUriResponsesRef.Get(); ChunkUriResponses.Dequeue(ChunkUriResponse);)
 			{
-				ChunkUris.Add(MoveTemp(ChunkUriResponse.Get<0>()), MoveTemp(ChunkUriResponse.Get<1>()));
+				if (ChunkUriResponse.Get<1>().bFailed)
+				{
+					// We couldn't get a valid url for the chunk and so the chunk should be considered a failed
+					// download. This is considered a hard failure for the chunk (i.e. don't try other CDNs).
+					InstallerError->SetError(EBuildPatchInstallError::DownloadError, DownloadErrorCodes::FailedUriRequest);
+					bShouldAbort = true;
+				}
+				else
+				{
+					ChunkUris.Add(MoveTemp(ChunkUriResponse.Get<0>()), MoveTemp(ChunkUriResponse.Get<1>()));
+				}
 			}
 			// Grab incoming requests as a priority.
 			TArray<FGuid> FrameRequestedDownloads;
@@ -714,7 +732,7 @@ namespace BuildPatchServices
 						TaskInfo.UrlUsed = ChunkUri->Uri;
 						TaskInfo.ExpectedSize = ManifestSet->GetDownloadSize(NextTask);
 						TaskInfo.SecondsAtRequested = FStatsCollector::GetSeconds();
-						int32 RequestId = DownloadService->RequestFile(TaskInfo.UrlUsed, OnDownloadCompleteDelegate, OnDownloadProgressDelegate);
+						int32 RequestId = DownloadService->RequestFileWithHeaders(TaskInfo.UrlUsed, ChunkUri->AdditionalHeaders, OnDownloadCompleteDelegate, OnDownloadProgressDelegate);
 						InFlightDownloads.Add(RequestId, NextTask);
 						PriorityRequests.Remove(NextTask);
 						FailedDownloads.Remove(NextTask);
@@ -951,7 +969,7 @@ namespace BuildPatchServices
 		void OnDownloadComplete(int32 RequestId, const FDownloadRef& Download)
 		{
 			{
-				FScopeLock ScopeLock(&CompletedDownloadsCS);
+				FScopeLock ScopeLock(&CompletedRequestsCS);
 				CompletedDownloads.Add(RequestId, Download);
 			}
 			WakeupMainThreadFn();
@@ -978,8 +996,9 @@ namespace BuildPatchServices
 		std::atomic_int32_t NextReadId = 0;
 
 		// Communication from download thread to processing thread.
-		FCriticalSection CompletedDownloadsCS;
-		TMap<int32, FDownloadRef> CompletedDownloads;
+		FCriticalSection CompletedRequestsCS;
+		TMap<int32, FDownloadRef> CompletedDownloads; // Keyed off the download request id
+		TMap<int32, FChunkUriResponse> CompletedRequests; // Keyed off the read id
 
 		// If we start getting failures on our downloads, we track which ones
 		// fail and avoid them until everything goes bad. Initially we just hit the
@@ -991,7 +1010,8 @@ namespace BuildPatchServices
 		// to the validation task.
 		FCriticalSection ReadQueueCS;
 		TArray<FCloudRead*> ReadQueue;
-		TMap<int32, FCloudRead*> InFlightDownloads;
+		TMap<int32, FCloudRead*> InFlightDownloads; // Keyed off the download request id
+		TMap<int32, FCloudRead*> InFlightURLRequests; // Keyed off the read id
 
 		// \todo does this thing have any recency bias? Surely we only care about the last few
 		// seconds of success...
@@ -1019,6 +1039,46 @@ namespace BuildPatchServices
 			ReadQueue.Empty();
 		}
 
+		// We have to wait until ALL our uri resposes come back since we can't
+		// delete the lambda references out from under them.
+		if (InFlightURLRequests.Num())
+		{
+			double LastReport = FPlatformTime::Seconds();
+			double FirstReport = LastReport;
+			UE_LOG(LogCloudChunkSource, Display, TEXT("Draining outstanding url requests on cancel... (%d)"), InFlightURLRequests.Num());
+			while (InFlightURLRequests.Num())
+			{
+				double CurrentTime = FPlatformTime::Seconds();
+				if (CurrentTime - LastReport > 5)
+				{
+					// We want it to be clear to whoever is looking at the logs that we're waiting on
+					// client code if we are hung here.
+					UE_LOG(LogCloudChunkSource, Display, TEXT("Still waiting on outstanding url requests, %.1f seconds, %d outstanding"), CurrentTime - FirstReport, InFlightURLRequests.Num());
+					LastReport = CurrentTime;
+				}
+
+				TMap<int32, FChunkUriResponse> FrameCompletedRequests;
+				CompletedRequestsCS.Lock();
+				FrameCompletedRequests = MoveTemp(CompletedRequests);
+				CompletedRequestsCS.Unlock();
+
+				for (const TPair<int32, FChunkUriResponse>& Pair : FrameCompletedRequests)
+				{
+					FCloudRead* Read = InFlightURLRequests.FindAndRemoveChecked(Pair.Key);
+					Read->CompleteFn.Execute(Read->DataId, true, false, Read->UserPtr);
+					delete Read;
+				}
+
+				// Ideally we get them all more or less instantly, but otherwise we have to wait. We expect
+				// this to be rare - user initiated cancels should have completed these already, and internal
+				// cancels are not common, so we just use a sleep rather than set up an event.
+				if (InFlightURLRequests.Num())
+				{
+					FPlatformProcess::Sleep(.002f);
+				}
+			}
+		}
+
 		// Abort any downloads.
 		for (TPair<int32, FCloudRead*> Pair : InFlightDownloads)
 		{
@@ -1047,18 +1107,54 @@ namespace BuildPatchServices
 				PendingCount = ReadQueue.Num();
 			}
 
+			TRACE_INT_VALUE(TEXT("BPS.Cloud.ActiveURLRequests"), InFlightURLRequests.Num());
 			TRACE_INT_VALUE(TEXT("BPS.Cloud.ActiveReads"), InFlightDownloads.Num());
 			TRACE_INT_VALUE(TEXT("BPS.Cloud.PendingReads"), PendingCount);
 
-			UE_LOG(LogCloudChunkSource, Verbose, TEXT("Cloud: Active: %d, Pending %d"), InFlightDownloads.Num(), PendingCount);
+			UE_LOG(LogCloudChunkSource, Verbose, TEXT("Cloud: Active: %d, Pending %d, URLS: %d"), InFlightDownloads.Num(), PendingCount, InFlightURLRequests.Num());
 		}
 
 		// Process completed downloads.
 		TMap<int32, FDownloadRef> FrameCompletedDownloads;
+		TMap<int32, FChunkUriResponse> FrameCompletedRequests;
 		{
-			CompletedDownloadsCS.Lock();
+			CompletedRequestsCS.Lock();
+			FrameCompletedRequests = MoveTemp(CompletedRequests);
 			FrameCompletedDownloads = MoveTemp(CompletedDownloads);
-			CompletedDownloadsCS.Unlock();
+			CompletedRequestsCS.Unlock();
+		}
+
+		{
+			for (const TPair<int32, FChunkUriResponse>& FrameCompletedRequest : FrameCompletedRequests)
+			{
+				FCloudRead* Read = InFlightURLRequests.FindAndRemoveChecked(FrameCompletedRequest.Key);
+
+				const FChunkUriResponse& Response = FrameCompletedRequest.Value;
+				if (Response.bFailed)
+				{
+					// Failed to get an auth url means failed to read
+					Read->CompleteFn.Execute(Read->DataId, false, true, Read->UserPtr);
+					delete Read;
+				}
+				else if (!bStartNewDownloads)
+				{
+					// if we can't start new downloads it counts as an abort
+					Read->CompleteFn.Execute(Read->DataId, true, false, Read->UserPtr);
+					delete Read;
+				}
+				else
+				{
+					// Only override the URL if set; the default handler does nothing.
+					if (Response.Uri.Len())
+					{
+						Read->UrlUsed = Response.Uri;
+					}
+					Read->RequestId = DownloadService->RequestFileWithHeaders(Read->UrlUsed, Response.AdditionalHeaders, OnDownloadCompleteDelegate, OnDownloadProgressDelegate);
+					InFlightDownloads.Add(Read->RequestId, Read);
+
+					CloudChunkSourceStat->OnDownloadRequested(Read->DataId);
+				}
+			}
 		}
 
 		{
@@ -1263,7 +1359,7 @@ namespace BuildPatchServices
 		while (bStartNewDownloads)
 		{
 			const uint32 MaxDownloads = FMath::Min(InMaxDownloads, FMath::Max(1U, DownloadCount->GetAdjustedCount(InFlightDownloads.Num(), DownloadHealth)));
-			if ((uint32)InFlightDownloads.Num() >= MaxDownloads)
+			if (((uint32)InFlightDownloads.Num() + (uint32)InFlightURLRequests.Num()) >= MaxDownloads)
 			{
 				break;
 			}
@@ -1281,16 +1377,39 @@ namespace BuildPatchServices
 				ReadQueue.RemoveAt(0);
 			}
 
+			// We have to provide clients of the API a way to remap URLs for one reason or another. This is often
+			// due to appending auth keys to the URL itself, or adding auth headers. This can be (very) latent or instantaneous.
+			FChunkUriRequest ChunkUriRequest;
+			ChunkUriRequest.CloudDirectory = Configuration.CloudRoots[CurrentBestCloudDir];
+			ChunkUriRequest.RelativePath = ManifestSet->GetDataFilename(Read->DataId);
+			ChunkUriRequest.RelativePath.RemoveFromStart(TEXT("/"));
+
 			Read->CloudDirUsed = CurrentBestCloudDir;
 			Read->bQueuedForDownload = true;
-			Read->UrlUsed = Configuration.CloudRoots[CurrentBestCloudDir] / ManifestSet->GetDataFilename(Read->DataId);
 			Read->ExpectedSize = ManifestSet->GetDownloadSize(Read->DataId);
 			Read->SecondsAtRequested = FStatsCollector::GetSeconds();
-			Read->RequestId = DownloadService->RequestFile(Read->UrlUsed, OnDownloadCompleteDelegate, OnDownloadProgressDelegate);
+			Read->UrlUsed = ChunkUriRequest.CloudDirectory / ChunkUriRequest.RelativePath;
+
+			InFlightURLRequests.Add(Read->ReadId, Read);
+			bool bMessageSent = MessagePump->SendRequest(ChunkUriRequest, [this, ReadId = Read->ReadId](FChunkUriResponse Response)
+				{
+					CompletedRequestsCS.Lock();
+					CompletedRequests.Add(ReadId, MoveTemp(Response));
+					CompletedRequestsCS.Unlock();
+
+					WakeupMainThreadFn();
+				});
 
-			InFlightDownloads.Add(Read->RequestId, Read);
+			if (!bMessageSent)
+			{
+				// If no message handler was registered to do any modification of the URL, start the download immediately.
+				InFlightURLRequests.Remove(Read->ReadId);
+
+				Read->RequestId = DownloadService->RequestFileWithHeaders(Read->UrlUsed, {}, OnDownloadCompleteDelegate, OnDownloadProgressDelegate);
+				InFlightDownloads.Add(Read->RequestId, Read);
 
-			CloudChunkSourceStat->OnDownloadRequested(Read->DataId);
+				CloudChunkSourceStat->OnDownloadRequested(Read->DataId);
+			}
 
 			if (StartTime == -1)
 			{
diff --git a/Engine/Source/Runtime/Online/BuildPatchServices/Private/Installer/DownloadService.cpp b/Engine/Source/Runtime/Online/BuildPatchServices/Private/Installer/DownloadService.cpp
--- a/Engine/Source/Runtime/Online/BuildPatchServices/Private/Installer/DownloadService.cpp
+++ b/Engine/Source/Runtime/Online/BuildPatchServices/Private/Installer/DownloadService.cpp
@@ -106,7 +106,7 @@ namespace BuildPatchServices
 		void CancelAllRequests();
 
 		// IDownloadService interface begin.
-		int32 RequestFile(const FString& FileUri, const FDownloadCompleteDelegate& OnCompleteDelegate, const FDownloadProgressDelegate& OnProgressDelegate);
+		int32 RequestFileWithHeaders(const FString& FileUri, const TMap<FString, FString>& RequestHeaders, const FDownloadCompleteDelegate& OnCompleteDelegate, const FDownloadProgressDelegate& OnProgressDelegate);
 		void RequestCancel(int32 RequestId);
 		void RequestAbandon(int32 RequestId);
 		// IDownloadService interface end.
@@ -147,7 +147,12 @@ namespace BuildPatchServices
 		// IDownloadService interface begin.
 		virtual int32 RequestFile(const FString& FileUri, const FDownloadCompleteDelegate& OnCompleteDelegate, const FDownloadProgressDelegate& OnProgressDelegate) override
 		{
-			return Impl->RequestFile(FileUri, OnCompleteDelegate, OnProgressDelegate);
+			TMap<FString, FString> RequestHeaders;
+			return Impl->RequestFileWithHeaders(FileUri, RequestHeaders, OnCompleteDelegate, OnProgressDelegate);
+		}
+		int32 RequestFileWithHeaders(const FString& FileUri, const TMap<FString, FString>& RequestHeaders, const FDownloadCompleteDelegate& OnCompleteDelegate, const FDownloadProgressDelegate& OnProgressDelegate)
+		{
+			return Impl->RequestFileWithHeaders(FileUri, RequestHeaders, OnCompleteDelegate, OnProgressDelegate);
 		}
 		virtual void RequestCancel(int32 RequestId) override { Impl->RequestCancel(RequestId); }
 		virtual void RequestAbandon(int32 RequestId) override { Impl->RequestAbandon(RequestId); }
@@ -186,7 +191,7 @@ namespace BuildPatchServices
 		// callbacks and removal handled by completion delegates
 	}
 
-	int32 FDownloadServiceImpl::RequestFile(const FString& FileUri, const FDownloadCompleteDelegate& OnCompleteDelegate, const FDownloadProgressDelegate& OnProgressDelegate)
+	int32 FDownloadServiceImpl::RequestFileWithHeaders(const FString& FileUri, const TMap<FString, FString>& RequestHeaders, const FDownloadCompleteDelegate& OnCompleteDelegate, const FDownloadProgressDelegate& OnProgressDelegate)
 	{
 		const int32 RequestId = ++RequestIdCounter;
 
@@ -209,6 +214,11 @@ namespace BuildPatchServices
 			HttpRequest->SetDelegateThreadPolicy(EHttpRequestDelegateThreadPolicy::CompleteOnHttpThread);
 			HttpRequest->SetURL(FileUri);
 			HttpRequest->SetVerb(TEXT("GET"));
+
+			for (const TPair<FString, FString>& Header : RequestHeaders)
+			{
+				HttpRequest->SetHeader(Header.Key, Header.Value);
+			}
 			
 			HttpRequest->OnRequestProgress64().BindSPLambda(this, [this, RequestId](FHttpRequestPtr Request, uint64 BytesSent, uint64 BytesReceived)
 			{
diff --git a/Engine/Source/Runtime/Online/BuildPatchServices/Private/Installer/DownloadService.h b/Engine/Source/Runtime/Online/BuildPatchServices/Private/Installer/DownloadService.h
--- a/Engine/Source/Runtime/Online/BuildPatchServices/Private/Installer/DownloadService.h
+++ b/Engine/Source/Runtime/Online/BuildPatchServices/Private/Installer/DownloadService.h
@@ -77,6 +77,9 @@ namespace BuildPatchServices
 		 */
 		virtual int32 RequestFile(const FString& FileUri, const FDownloadCompleteDelegate& OnCompleteDelegate, const FDownloadProgressDelegate& OnProgressDelegate) = 0;
 
+		/** As per RequestFile, except adds the given headers to the request as a <Name, Value> pair. */
+		virtual int32 RequestFileWithHeaders(const FString& FileUri, const TMap<FString, FString>& RequestHeaders, const FDownloadCompleteDelegate& OnCompleteDelegate, const FDownloadProgressDelegate& OnProgressDelegate) = 0;
+
 		/**
 		 * Requests the cancellation of a requested file.
 		 * @param RequestId             The id that was returned by the call to RequestFile which should be canceled.
diff --git a/Engine/Source/Runtime/Online/BuildPatchServices/Private/Installer/InstallerError.h b/Engine/Source/Runtime/Online/BuildPatchServices/Private/Installer/InstallerError.h
--- a/Engine/Source/Runtime/Online/BuildPatchServices/Private/Installer/InstallerError.h
+++ b/Engine/Source/Runtime/Online/BuildPatchServices/Private/Installer/InstallerError.h
@@ -48,6 +48,7 @@ namespace BuildPatchServices
 		static const TCHAR* MissingDeltaFile = TEXT("02");
 		static const TCHAR* UnserialisableDeltaFile = TEXT("03");
 		static const TCHAR* RejectedDeltaFile = TEXT("04");
+		static const TCHAR* FailedUriRequest = TEXT("05");
 	}
 
 	/**
diff --git a/Engine/Source/Runtime/Online/BuildPatchServices/Private/Installer/MessagePump.cpp b/Engine/Source/Runtime/Online/BuildPatchServices/Private/Installer/MessagePump.cpp
--- a/Engine/Source/Runtime/Online/BuildPatchServices/Private/Installer/MessagePump.cpp
+++ b/Engine/Source/Runtime/Online/BuildPatchServices/Private/Installer/MessagePump.cpp
@@ -47,7 +47,7 @@ namespace BuildPatchServices
 		virtual void SendMessage(FChunkSourceEvent Message) override;
 		virtual void SendMessage(FInstallationFileAction Message) override;
 		virtual void SendMessage(FGenericMessage Message) override;
-		virtual void SendRequest(FChunkUriRequest Request, TFunction<void(FChunkUriResponse)> OnResponse) override;
+		virtual bool SendRequest(FChunkUriRequest Request, TFunction<void(FChunkUriResponse)> OnResponse) override;
 		virtual void PumpMessages() override;
 		virtual void RegisterMessageHandler(FMessageHandler* MessageHandler) override;
 		virtual void UnregisterMessageHandler(FMessageHandler* MessageHandler) override;
@@ -77,16 +77,14 @@ namespace BuildPatchServices
 		MessageQueue.Enqueue(FMessageUnion(MoveTemp(Message)));
 	}
 
-	void FMessagePump::SendRequest(FChunkUriRequest Request, TFunction<void(FChunkUriResponse)> OnResponse)
+	bool FMessagePump::SendRequest(FChunkUriRequest Request, TFunction<void(FChunkUriResponse)> OnResponse)
 	{
-		if (!EnumHasAllFlags(static_cast<EMessageRequests>(MessageRequests.load()), EMessageRequests::ChunkUriRequest))
-		{
-			DefaultMessageHandler.HandleRequest(Request, OnResponse);
-		}
-		else
+		if (EnumHasAllFlags(static_cast<EMessageRequests>(MessageRequests.load()), EMessageRequests::ChunkUriRequest))
 		{
 			ChunkUriRequestQueue.Enqueue(FChunkUriRequestResponse(MoveTemp(Request), MoveTemp(OnResponse)));
+			return true;
 		}
+		return false;
 	}
 
 	void FMessagePump::PumpMessages()
diff --git a/Engine/Source/Runtime/Online/BuildPatchServices/Private/Installer/MessagePump.h b/Engine/Source/Runtime/Online/BuildPatchServices/Private/Installer/MessagePump.h
--- a/Engine/Source/Runtime/Online/BuildPatchServices/Private/Installer/MessagePump.h
+++ b/Engine/Source/Runtime/Online/BuildPatchServices/Private/Installer/MessagePump.h
@@ -35,9 +35,23 @@ namespace BuildPatchServices
 		/**
 		 * Sends out a request to resolve the uri to the chunk location
 		 * @param Request   Request for the chunk location
-		 * @param OnResponse   A delegate to call with the response to chunk location
+		 * @param OnResponse   A delegate to call with the chunk location and any potential headers to add to the http request.
+		 *
+		 * Note that while this function supports asynchronous processing, the calling code does not support
+		 * cancelation of the requests under abort scenarios, potentially leading to crashes. As as result,
+		 * incomplete shutdown of the installation requires waiting for all outstanding requests to return before
+		 * shutdown can complete. For user initiated cancellations it is recommended that you abort any async
+		 * uri request handling and call OnResponse with FChunkUriResponse::bFailed to true after calling CancelInstall().
+		 * This will prevent the default URL concatenation from occuring and allow the shutdown logic to complete in a timely
+		 * fashion.
+		 *
+		 * For internally initiated cancellations (due to errors) there's nothing to be done; shutdown just waits for
+		 * the requests to all complete.
+		 *
+		 * As a result, it's highly advisable to make this as immediate as possible, e.g. caching auth tokens up front
+		 * before installation launch.
 		 */
-		virtual void SendRequest(FChunkUriRequest Request, TFunction<void(FChunkUriResponse)> OnResponse) = 0;
+		virtual bool SendRequest(FChunkUriRequest Request, TFunction<void(FChunkUriResponse)> OnResponse) = 0;
 
 		/**
 		 * Dequeues received messages, pushing them to the provided handlers.
diff --git a/Engine/Source/Runtime/Online/BuildPatchServices/Private/Tests/Mock/DownloadService.mock.h b/Engine/Source/Runtime/Online/BuildPatchServices/Private/Tests/Mock/DownloadService.mock.h
--- a/Engine/Source/Runtime/Online/BuildPatchServices/Private/Tests/Mock/DownloadService.mock.h
+++ b/Engine/Source/Runtime/Online/BuildPatchServices/Private/Tests/Mock/DownloadService.mock.h
@@ -37,6 +37,11 @@ namespace BuildPatchServices
 			return ReturnId;
 		}
 
+		virtual int32 RequestFileWithHeaders(const FString& FileUri, const TMap<FString, FString>& Headers, const FDownloadCompleteDelegate& OnCompleteDelegate, const FDownloadProgressDelegate& OnProgressDelegate) override
+		{
+			return RequestFile(FileUri, OnCompleteDelegate, OnProgressDelegate);
+		}
+
 		virtual void RequestCancel(int32 RequestId) override
 		{
 			FScopeLock ScopeLock(&ThreadLock);
diff --git a/Engine/Source/Runtime/Online/BuildPatchServices/Private/Tests/Mock/MessagePump.mock.h b/Engine/Source/Runtime/Online/BuildPatchServices/Private/Tests/Mock/MessagePump.mock.h
--- a/Engine/Source/Runtime/Online/BuildPatchServices/Private/Tests/Mock/MessagePump.mock.h
+++ b/Engine/Source/Runtime/Online/BuildPatchServices/Private/Tests/Mock/MessagePump.mock.h
@@ -27,8 +27,9 @@ namespace BuildPatchServices
 		{
 		}
 
-		virtual void SendRequest(FChunkUriRequest Request, TFunction<void(FChunkUriResponse)> OnResponse) override
+		virtual bool SendRequest(FChunkUriRequest Request, TFunction<void(FChunkUriResponse)> OnResponse) override
 		{
+			return false;
 		}
 
 		virtual void RegisterMessageHandler(FMessageHandler* MessageHandler) override
diff --git a/Engine/Source/Runtime/Online/BuildPatchServices/Public/BuildPatchMessage.h b/Engine/Source/Runtime/Online/BuildPatchServices/Public/BuildPatchMessage.h
--- a/Engine/Source/Runtime/Online/BuildPatchServices/Public/BuildPatchMessage.h
+++ b/Engine/Source/Runtime/Online/BuildPatchServices/Public/BuildPatchMessage.h
@@ -1,6 +1,7 @@
 // Copyright Epic Games, Inc. All Rights Reserved.
 #pragma once
 
+#include "Containers/Map.h"
 #include "Containers/UnrealString.h"
 #include "Misc/Guid.h"
 #include "Templates/Function.h"
@@ -63,7 +64,13 @@ namespace BuildPatchServices
 	 */
 	struct FChunkUriResponse
 	{
+		// Set this to true if the response can not be fulfilled. This will fail the chunk requests and subsequently
+		// the installation.
+		bool bFailed = false;
 		FString Uri;
+		
+		// These headers <name, value> will be added to the HTTP request.
+		TMap<FString, FString> AdditionalHeaders;
 	};
 
 	struct FGenericMessage
@@ -158,9 +165,13 @@ namespace BuildPatchServices
 			: FMessageHandler(static_cast<EMessageRequests>(0xFFFFFFFF))
 		{}
 
+		// This gets called if all registered handlers return false to HandleRequest. If there are
+		// _no_ handlers then this is entirely sidestepped in the cloud chunk source.
 		virtual bool HandleRequest(const FChunkUriRequest& Request, TFunction<void(FChunkUriResponse)> OnResponse) override
 		{
-			OnResponse({ Request.CloudDirectory / Request.RelativePath });
+			FChunkUriResponse Response;
+			Response.Uri = Request.CloudDirectory / Request.RelativePath;
+			OnResponse(MoveTemp(Response));
 			return true;
 		}
 	};
