LwsWebSockets: Fix for FLwsWebSocketsManager::ShutdownWebSockets deadlocking on Thread->Kill whilst the thread is blocked on a service call.

During FLwsWebSocketsManager::ShutdownWebSockets, we now cancel any blocking service calls *before* trying to kill the websockets thread.

diff --git a/Engine/Source/Runtime/Online/WebSockets/Private/Lws/LwsWebSocketsManager.cpp b/Engine/Source/Runtime/Online/WebSockets/Private/Lws/LwsWebSocketsManager.cpp
--- a/Engine/Source/Runtime/Online/WebSockets/Private/Lws/LwsWebSocketsManager.cpp
+++ b/Engine/Source/Runtime/Online/WebSockets/Private/Lws/LwsWebSocketsManager.cpp
@@ -246,6 +246,12 @@ void FLwsWebSocketsManager::ShutdownWebSockets()
 		OnChildEndFramePostForkHandle.Reset();
 	}
 
+	if (LwsContext)
+	{
+		// Cancel service which may be blocking the HTTP thread, before we attempt the kill Thread below 
+		lws_cancel_service(LwsContext);
+	}
+
 	if (Thread)
 	{
 		Thread->Kill(true);
@@ -255,11 +261,6 @@ void FLwsWebSocketsManager::ShutdownWebSockets()
 
 	if (LwsContext)
 	{
-		if (!Sockets.IsEmpty())
-		{
-			lws_cancel_service(LwsContext);
-		}
-
 		lws_context_destroy(LwsContext);
 		LwsContext = nullptr;
 	}
