PR #12855: Fixed typo HttListener in http server log.

diff --git a/Engine/Source/Runtime/Online/HTTPServer/Private/HttpListener.cpp b/Engine/Source/Runtime/Online/HTTPServer/Private/HttpListener.cpp
--- a/Engine/Source/Runtime/Online/HTTPServer/Private/HttpListener.cpp
+++ b/Engine/Source/Runtime/Online/HTTPServer/Private/HttpListener.cpp
@@ -136,7 +136,7 @@ void FHttpListener::StopListening()
 	if (ListenSocket)
 	{
 		UE_LOG(LogHttpListener, Log,
-			TEXT("HttListener stopping listening on Port %u"), ListenPort);
+			TEXT("HttpListener stopping listening on Port %u"), ListenPort);
 
 		ListenSocket.Reset();
 	}
@@ -267,4 +267,4 @@ void FHttpListener::RemoveDestroyedConnections()
 			continue;
 		}
 	}
-}
\ No newline at end of file
+}
