BPS: Don't allow multiple file construction when patching as harvesting is not safe.

#rnx

diff --git a/Engine/Source/Runtime/Online/BuildPatchServices/Private/BuildPatchFileConstructor.cpp b/Engine/Source/Runtime/Online/BuildPatchServices/Private/BuildPatchFileConstructor.cpp
--- a/Engine/Source/Runtime/Online/BuildPatchServices/Private/BuildPatchFileConstructor.cpp
+++ b/Engine/Source/Runtime/Online/BuildPatchServices/Private/BuildPatchFileConstructor.cpp
@@ -1411,6 +1411,10 @@ FBuildPatchFileConstructor::FBuildPatchFileConstructor(
 	if (InstallSource)
 	{
 		InstallSource->SetFileRetirementPositions(MoveTemp(FileRetirementPositions));	
+
+		// Harvesting is not safe when multiple files are in flight as a later file can request
+		// a chunk we needed to have harvested and we can be in the middle of that harvest when the read comes in.
+		bAllowMultipleFilesInFlight = false;
 	}
 
 	//
