BPS: Fix for InstallChunkSource test failures.

#rnx

diff --git a/Engine/Source/Runtime/Online/BuildPatchServices/Private/Tests/Mock/Manifest.mock.h b/Engine/Source/Runtime/Online/BuildPatchServices/Private/Tests/Mock/Manifest.mock.h
--- a/Engine/Source/Runtime/Online/BuildPatchServices/Private/Tests/Mock/Manifest.mock.h
+++ b/Engine/Source/Runtime/Online/BuildPatchServices/Private/Tests/Mock/Manifest.mock.h
@@ -372,7 +372,8 @@ namespace BuildPatchServices
 
 		virtual TArray<FFileChunkPart> GetFilePartsForChunk(const FGuid& ChunkId) const override
 		{
-			return FilePartsForChunk.FindRef(ChunkId);
+			checkNoEntry(); // AFAICT this function is deprecated and never used.
+			return TArray<FFileChunkPart>();
 		}
 
 		virtual bool HasFileAttributes() const override
@@ -425,7 +426,6 @@ namespace BuildPatchServices
 		TMap<FGuid, uint64> FilePartHashes;
 		TSet<FGuid> ProducibleChunks;
 		TSet<FString> OutdatedFiles;
-		TMap<FGuid, TArray<FFileChunkPart>> FilePartsForChunk;
 	};
 
 	typedef TSharedPtr<BuildPatchServices::FMockManifest, ESPMode::ThreadSafe> FMockManifestPtr;
diff --git a/Engine/Source/Runtime/Online/BuildPatchServices/Private/Tests/Unit/InstallChunkSource.spec.cpp b/Engine/Source/Runtime/Online/BuildPatchServices/Private/Tests/Unit/InstallChunkSource.spec.cpp
--- a/Engine/Source/Runtime/Online/BuildPatchServices/Private/Tests/Unit/InstallChunkSource.spec.cpp
+++ b/Engine/Source/Runtime/Online/BuildPatchServices/Private/Tests/Unit/InstallChunkSource.spec.cpp
@@ -107,15 +107,19 @@ void FInstallChunkSourceSpec::Define()
 				{
 					It("should fail.", [this]()
 					{
+						FEvent* DoneEvent = FPlatformProcess::GetSynchEventFromPool(true);
 						bool bFailed = false;
 						TUniqueFunction<void(bool)> Request = InstallChunkSource->CreateRequest(SomeChunk, {}, 0, IConstructorChunkSource::FChunkRequestCompleteDelegate::CreateLambda(
-						[&bFailed](const FGuid& DataId, bool bAborted, bool bFailedToRead, void* UserPtr)
+						[&bFailed, DoneEvent](const FGuid& DataId, bool bAborted, bool bFailedToRead, void* UserPtr)
 						{
 							bFailed = bFailedToRead;
+							DoneEvent->Trigger();
 						}
 						));
 
 						Request(false);
+						DoneEvent->Wait();
+						FPlatformProcess::ReturnSynchEventToPool(DoneEvent);
 
 						TEST_TRUE(bFailed);
 					});
@@ -135,18 +139,22 @@ void FInstallChunkSourceSpec::Define()
 				{
 					It("should return some chunk loading from disk.", [this]()
 					{
+						FEvent* DoneEvent = FPlatformProcess::GetSynchEventFromPool(true);
 						bool bFailed = false;
 						TUniqueFunction<void(bool)> Request = InstallChunkSource->CreateRequest(SomeChunk, 
 							FMutableMemoryView(ReadDestination.GetData(), ReadDestination.Num()), 
 							0, 
 							IConstructorChunkSource::FChunkRequestCompleteDelegate::CreateLambda(
-								[&bFailed](const FGuid& DataId, bool bAborted, bool bFailedToRead, void* UserPtr)
+								[&bFailed, DoneEvent](const FGuid& DataId, bool bAborted, bool bFailedToRead, void* UserPtr)
 								{
 									bFailed = bFailedToRead;
+									DoneEvent->Trigger();
 								}
 							));
 
 						Request(false);
+						DoneEvent->Wait();
+						FPlatformProcess::ReturnSynchEventToPool(DoneEvent);
 
 						TEST_FALSE(bFailed);
 						TEST_EQUAL(MockInstallChunkSourceStat->RxLoadStarted.Num(), 1);
@@ -167,18 +175,22 @@ void FInstallChunkSourceSpec::Define()
 
 						It("should not have attempted to load some chunk from disk.", [this]()
 						{
+							FEvent* DoneEvent = FPlatformProcess::GetSynchEventFromPool(true);
 							bool bFailed = false;
 							TUniqueFunction<void(bool)> Request = InstallChunkSource->CreateRequest(SomeChunk, 
 								FMutableMemoryView(ReadDestination.GetData(), ReadDestination.Num()), 
 								0, 
 								IConstructorChunkSource::FChunkRequestCompleteDelegate::CreateLambda(
-									[&bFailed](const FGuid& DataId, bool bAborted, bool bFailedToRead, void* UserPtr)
+									[&bFailed, DoneEvent](const FGuid& DataId, bool bAborted, bool bFailedToRead, void* UserPtr)
 									{
 										bFailed = bFailedToRead;
+										DoneEvent->Trigger();
 									}
 							));
 
 							Request(false);
+							DoneEvent->Wait();
+							FPlatformProcess::ReturnSynchEventToPool(DoneEvent);
 
 							TEST_TRUE(bFailed);
 							TEST_EQUAL(FakeFileSystem->RxCreateFileReader.Num(), 0);
@@ -201,18 +213,22 @@ void FInstallChunkSourceSpec::Define()
 
 						It("should still succeed to load some chunk from disk.", [this]()
 						{
+							FEvent* DoneEvent = FPlatformProcess::GetSynchEventFromPool(true);
 							bool bFailed = false;
 							TUniqueFunction<void(bool)> Request = InstallChunkSource->CreateRequest(SomeChunk, 
 								FMutableMemoryView(ReadDestination.GetData(), ReadDestination.Num()), 
 								0, 
 								IConstructorChunkSource::FChunkRequestCompleteDelegate::CreateLambda(
-									[&bFailed](const FGuid& DataId, bool bAborted, bool bFailedToRead, void* UserPtr)
+									[&bFailed, DoneEvent](const FGuid& DataId, bool bAborted, bool bFailedToRead, void* UserPtr)
 									{
 										bFailed = bFailedToRead;
+										DoneEvent->Trigger();
 									}
 							));
 
 							Request(false);
+							DoneEvent->Wait();
+							FPlatformProcess::ReturnSynchEventToPool(DoneEvent);
 
 							TEST_FALSE(bFailed);
 							TEST_EQUAL(MockInstallChunkSourceStat->RxLoadStarted.Num(), 1);
@@ -229,18 +245,22 @@ void FInstallChunkSourceSpec::Define()
 
 							It("should fail to load some chunk from disk, reporting IInstallChunkSourceStat::ELoadResult::HashCheckFailed.", [this]()
 							{
+								FEvent* DoneEvent = FPlatformProcess::GetSynchEventFromPool(true);
 								bool bFailed = false;
 								TUniqueFunction<void(bool)> Request = InstallChunkSource->CreateRequest(SomeChunk, 
 									FMutableMemoryView(ReadDestination.GetData(), ReadDestination.Num()), 
 									0, 
 									IConstructorChunkSource::FChunkRequestCompleteDelegate::CreateLambda(
-										[&bFailed](const FGuid& DataId, bool bAborted, bool bFailedToRead, void* UserPtr)
+										[&bFailed, DoneEvent](const FGuid& DataId, bool bAborted, bool bFailedToRead, void* UserPtr)
 										{
 											bFailed = bFailedToRead;
+											DoneEvent->Trigger();
 										}
 								));
 
 								Request(false);
+								DoneEvent->Wait();
+								FPlatformProcess::ReturnSynchEventToPool(DoneEvent);
 
 								TEST_TRUE(bFailed);
 
@@ -262,18 +282,22 @@ void FInstallChunkSourceSpec::Define()
 
 						It("should fail to load some chunk from disk, reporting IInstallChunkSourceStat::ELoadResult::HashCheckFailed.", [this]()
 						{
+							FEvent* DoneEvent = FPlatformProcess::GetSynchEventFromPool(true);
 							bool bFailed = false;
 							TUniqueFunction<void(bool)> Request = InstallChunkSource->CreateRequest(SomeChunk, 
 								FMutableMemoryView(ReadDestination.GetData(), ReadDestination.Num()), 
 								0, 
 								IConstructorChunkSource::FChunkRequestCompleteDelegate::CreateLambda(
-									[&bFailed](const FGuid& DataId, bool bAborted, bool bFailedToRead, void* UserPtr)
+									[&bFailed, DoneEvent](const FGuid& DataId, bool bAborted, bool bFailedToRead, void* UserPtr)
 									{
 										bFailed = bFailedToRead;
+										DoneEvent->Trigger();
 									}
 							));
 
 							Request(false);
+							DoneEvent->Wait();
+							FPlatformProcess::ReturnSynchEventToPool(DoneEvent);
 
 							TEST_TRUE(bFailed);
 							TEST_EQUAL(MockInstallChunkSourceStat->RxLoadComplete.Num(), 1);
@@ -316,117 +340,91 @@ void FInstallChunkSourceSpec::InventUsableChunkData()
 {
 	ReadDestination.SetNumUninitialized(TestChunkSize);
 
+	//
+	// Make two manifests to act as installation sources. The chunks in the overall manifest
+	// are used to make a bunch of files in eacxh installation.
+	//
 	using namespace BuildPatchServices;
-	for (int32 Count = 0; Count < 100; ++Count)
+	const int32 ChunkCountPerSource = 50;
+	for (int32 Count = 0; Count < ChunkCountPerSource*2; ++Count)
 	{
 		MockManifest->DataList.Add(FGuid::NewGuid());
 	}
+
+	TArray<uint8> ChunkData;
+	ChunkData.SetNumUninitialized(TestChunkSize);
+	FRandomStream RandomData(0);
+		
+	int32 FileCounter = 0;
+
 	FMockManifest* MockInstallationManifestA = new FMockManifest();
 	FMockManifest* MockInstallationManifestB = new FMockManifest();
+
 	InstallationSources.Add(TEXT("LocationA/"), MakeShareable(MockInstallationManifestA));
 	InstallationSources.Add(TEXT("LocationB/"), MakeShareable(MockInstallationManifestB));
-	for (int32 DataListIdx = 0; DataListIdx < MockManifest->DataList.Num(); ++DataListIdx)
-	{
-		const FGuid& TheChunk = MockManifest->DataList[DataListIdx];
-		switch (DataListIdx % 3)
-		{
-			case 0:
-				MockInstallationManifestA->ProducibleChunks.Add(TheChunk);
-				SomeAvailableChunks.Add(TheChunk);
-				break;
-			case 1:
-				MockInstallationManifestB->ProducibleChunks.Add(TheChunk);
-				SomeAvailableChunks.Add(TheChunk);
-				break;
-		}
-	}
-	uint32 FileCounter = 1;
-	for (FMockManifest* MockInstallationManifest : {MockInstallationManifestA, MockInstallationManifestB})
-	{
-		FString MockInstallLocation;
-		for (const TPair<FString, FBuildPatchAppManifestRef>& InstallationSourcePair : InstallationSources)
-		{
-			if (&InstallationSourcePair.Value.Get() == MockInstallationManifest)
-			{
-				MockInstallLocation = InstallationSourcePair.Key;
-			}
-		}
-		for (const FGuid& ProducibleChunk : MockInstallationManifest->ProducibleChunks)
-		{
-			uint32 ChunkSizeCounter = TestChunkSize;
-			TArray<FFileChunkPart>& FileChunkParts = MockInstallationManifest->FilePartsForChunk.FindOrAdd(ProducibleChunk);
-			for (int32 FileIdx = 0; FileIdx < 3; ++FileIdx)
-			{
-				FileChunkParts.AddDefaulted();
-				FFileChunkPart& FileChunkPart = FileChunkParts.Last();
-				FileChunkPart.Filename = FString::Printf(TEXT("File%d.dat"), FileCounter++);
-				FileChunkPart.FileOffset = 0;
-				FileChunkPart.ChunkPart.Guid = ProducibleChunk;
-				FileChunkPart.ChunkPart.Offset = TestChunkSize - ChunkSizeCounter;
-				FileChunkPart.ChunkPart.Size = TestChunkSize / 4;
-				ChunkSizeCounter -= FileChunkPart.ChunkPart.Size;
-			}
-			FileChunkParts.AddDefaulted();
-			FFileChunkPart& FileChunkPart = FileChunkParts.Last();
-			FileChunkPart.Filename = FString::Printf(TEXT("File%d.dat"), FileCounter++);
-			FileChunkPart.FileOffset = 0;
-			FileChunkPart.ChunkPart.Guid = ProducibleChunk;
-			FileChunkPart.ChunkPart.Offset = TestChunkSize - ChunkSizeCounter;
-			FileChunkPart.ChunkPart.Size = ChunkSizeCounter;
-			ChunkSizeCounter -= FileChunkPart.ChunkPart.Size;
-			check(ChunkSizeCounter == 0);
-			check(FileChunkPart.ChunkPart.Offset + FileChunkPart.ChunkPart.Size == TestChunkSize);
-		}
-		for (const TPair<FGuid, TArray<FFileChunkPart>>& ChunkFileParts : MockInstallationManifest->FilePartsForChunk)
-		{
-			for (const FFileChunkPart& FileChunkPart : ChunkFileParts.Value)
-			{
-				TArray<uint8>& FileData = FakeFileSystem->DiskData.FindOrAdd(MockInstallLocation / FileChunkPart.Filename);
-				FileData.SetNumUninitialized(FMath::Max<int32>(FileData.Num(), FileChunkPart.FileOffset + FileChunkPart.ChunkPart.Size));
-			}
-		}
-	}
-	FRandomStream RandomData(0);
-	for (TPair<FString, TArray<uint8>>& FileData : FakeFileSystem->DiskData)
-	{
-		uint8* Data = FileData.Value.GetData();
-		for (int32 DataIdx = 0; DataIdx <= (FileData.Value.Num()-4); DataIdx += 4)
-		{
-			*((uint32*)(Data + DataIdx)) = RandomData.GetUnsignedInt();
-		}
-	}
-	TArray<uint8> ChunkData;
-	ChunkData.SetNumUninitialized(TestChunkSize);
-	for (FMockManifest* MockInstallationManifest : {MockInstallationManifestA, MockInstallationManifestB})
+
+	for (int32 ManifestIndex = 0; ManifestIndex < 2; ManifestIndex++)
 	{
-		FString MockInstallLocation;
-		for (const TPair<FString, FBuildPatchAppManifestRef>& InstallationSourcePair : InstallationSources)
+		FMockManifest* ThisMockManifest = ManifestIndex ? MockInstallationManifestB : MockInstallationManifestA;
+		FString InstallLocation = ManifestIndex ? TEXT("LocationB/") : TEXT("LocationA/");
+
+		for (int32 ChunkIndex = 0; ChunkIndex < ChunkCountPerSource; ChunkIndex++)
 		{
-			if (&InstallationSourcePair.Value.Get() == MockInstallationManifest)
+			const FGuid& TheChunk = MockManifest->DataList[ChunkIndex + ChunkCountPerSource*ManifestIndex];
+			SomeAvailableChunks.Add(TheChunk);
+			ThisMockManifest->ProducibleChunks.Add(TheChunk);
+
+			// make sure we can evenly divide the chunk in to 4 files without worrying about leftovers,
+			// and that we can fill using unsigned ints.
+			check((TestChunkSize % 4) == 0);
+			
+			uint32 ChunkSizeCounter = 0;
+			for (int32 FileIdx = 0; FileIdx < 4; ++FileIdx)
 			{
-				MockInstallLocation = InstallationSourcePair.Key;
+				FFileManifest FileManifest;
+				FileManifest.Filename = FString::Printf(TEXT("File%d.dat"), FileCounter++);
+				FChunkPart& ChunkPart = FileManifest.ChunkParts.AddDefaulted_GetRef();
+				ChunkPart.Guid = TheChunk;
+				ChunkPart.Offset = ChunkSizeCounter;
+				ChunkPart.Size = TestChunkSize / 4;
+				ChunkSizeCounter += ChunkPart.Size;
+
+				FileManifest.FileSize = TestChunkSize / 4;
+
+				// Put the chunk's data in our VFS.
+				TArray<uint8>& FileData = FakeFileSystem->DiskData.FindOrAdd(InstallLocation / FileManifest.Filename);
+				FileData.SetNumUninitialized(ChunkPart.Size);
+
+				uint8* Data = FileData.GetData();
+				for (int32 DataIdx = 0; DataIdx <= (FileData.Num()-4); DataIdx += 4)
+				{
+					*((uint32*)(Data + DataIdx)) = RandomData.GetUnsignedInt();
+				}
+
+				FSHA1::HashBuffer(Data, FileManifest.FileSize, FileManifest.FileHash.Hash);
+
+				// Also fill the chunk array so we can hash it later.
+				FMemory::Memcpy(&ChunkData[ChunkPart.Offset], Data, ChunkPart.Size);
+
+				ThisMockManifest->BuildFileList.Add(FileManifest.Filename);
+				ThisMockManifest->FileNameToFileSize.Add(FileManifest.Filename, FileManifest.FileSize);
+				ThisMockManifest->FileNameToHashes.Add(FileManifest.Filename,FileManifest.FileHash);
+				ThisMockManifest->FileManifests.Add(FileManifest.Filename, MoveTemp(FileManifest));
 			}
-		}
-		for (const FGuid& ProducibleChunk : MockInstallationManifest->ProducibleChunks)
-		{
+
 			uint64 ChunkPolyHash;
 			FSHAHash ChunkShaHash;
-			for (const FFileChunkPart& FileChunkPart : MockInstallationManifest->FilePartsForChunk[ProducibleChunk])
-			{
-				TArray<uint8>& FileData = FakeFileSystem->DiskData[MockInstallLocation / FileChunkPart.Filename];
-				uint8* DataPtr = &FileData[FileChunkPart.FileOffset];
-				FMemory::Memcpy(&ChunkData[FileChunkPart.ChunkPart.Offset], DataPtr, FileChunkPart.ChunkPart.Size);
-			}
 			ChunkPolyHash = FRollingHash::GetHashForDataSet(ChunkData.GetData(), ChunkData.Num());
 			FSHA1::HashBuffer(ChunkData.GetData(), TestChunkSize, ChunkShaHash.Hash);
+
 			FChunkInfo ChunkInfo;
-			ChunkInfo.Guid = ProducibleChunk;
+			ChunkInfo.Guid = TheChunk;
 			ChunkInfo.Hash = ChunkPolyHash;
 			ChunkInfo.ShaHash = ChunkShaHash;
 			ChunkInfo.GroupNumber = 0;
 			ChunkInfo.WindowSize = TestChunkSize;
 			ChunkInfo.FileSize = TestChunkSize;
-			MockInstallationManifest->ChunkInfos.Add(ProducibleChunk, ChunkInfo);
+			ThisMockManifest->ChunkInfos.Add(TheChunk, ChunkInfo);
 		}
 	}
 }
@@ -443,20 +441,19 @@ void FInstallChunkSourceSpec::SomeChunkUnavailable()
 
 void FInstallChunkSourceSpec::SomeChunkCorrupted()
 {
+	// Find where the data is for SomeChunk and corrupt it.
 	using namespace BuildPatchServices;
 	for (TPair<FString, FBuildPatchAppManifestRef>& InstallationSourcePair : InstallationSources)
 	{
 		FMockManifest* MockInstallationManifest = (FMockManifest*)&InstallationSourcePair.Value.Get();
-		if (MockInstallationManifest->FilePartsForChunk.Contains(SomeChunk))
+		for (TPair<FString, FFileManifest>& File : MockInstallationManifest->FileManifests)
 		{
-			const TArray<FFileChunkPart>& FileChunkParts = MockInstallationManifest->FilePartsForChunk[SomeChunk];
-			if (FileChunkParts.Num() > 0)
+			for (FChunkPart& ChunkPart : File.Value.ChunkParts)
 			{
-				const FFileChunkPart& FileChunkPart = FileChunkParts[0];
-				TArray<uint8>& FileData = FakeFileSystem->DiskData[InstallationSourcePair.Key / FileChunkPart.Filename];
-				for (int32 DataIdx = 0; DataIdx < 10; ++DataIdx)
+				if (ChunkPart.Guid == SomeChunk)
 				{
-					FileData[DataIdx] = FileData[DataIdx + 1];
+					TArray<uint8>& FileData = FakeFileSystem->DiskData[InstallationSourcePair.Key / File.Value.Filename];
+					FMemory::Memzero(FileData.GetData(), FileData.Num());
 				}
 			}
 		}
