Deprecate marco UE_HTTP_CONNECTION_TIMEOUT_SUPPORT_RETRY because it's no longer necessary.

diff --git a/Engine/Source/Runtime/Online/HTTP/Public/Interfaces/IHttpBase.h b/Engine/Source/Runtime/Online/HTTP/Public/Interfaces/IHttpBase.h
--- a/Engine/Source/Runtime/Online/HTTP/Public/Interfaces/IHttpBase.h
+++ b/Engine/Source/Runtime/Online/HTTP/Public/Interfaces/IHttpBase.h
@@ -169,3 +169,6 @@ public:
 	virtual ~IHttpBase() = default;
 };
 
+#if !defined(UE_HTTP_CONNECTION_TIMEOUT_SUPPORT_RETRY)
+#define UE_HTTP_CONNECTION_TIMEOUT_SUPPORT_RETRY UE_DEPRECATED_MACRO(5.6, "UE_HTTP_CONNECTION_TIMEOUT_SUPPORT_RETRY has been deprecated because it is no longer necessary.") 1
+#endif 
