BPS: Add a CVar to control whether multiple files can be in flight during construction.

#rnx

diff --git a/Engine/Source/Runtime/Online/BuildPatchServices/Private/BuildPatchFileConstructor.cpp b/Engine/Source/Runtime/Online/BuildPatchServices/Private/BuildPatchFileConstructor.cpp
--- a/Engine/Source/Runtime/Online/BuildPatchServices/Private/BuildPatchFileConstructor.cpp
+++ b/Engine/Source/Runtime/Online/BuildPatchServices/Private/BuildPatchFileConstructor.cpp
@@ -48,6 +48,13 @@ static FAutoConsoleVariableRef CVarStallWhenFileSystemThrottled(
 	TEXT("Whether to stall if the file system is throttled"),
 	ECVF_Default);
 
+static bool bCVarAllowMultipleFilesInFlight = true;
+static FAutoConsoleVariableRef CVarAllowMultipleFilesInFlight(
+	TEXT("BuildPatchFileConstructor.bCVarAllowMultipleFilesInFlight"),
+	bCVarAllowMultipleFilesInFlight,
+	TEXT("Whether to allow multiple files to be constructed at the same time, though still sequentially."),
+	ECVF_Default);
+
 
 // Helper functions wrapping common code.
 namespace FileConstructorHelpers
@@ -1373,6 +1380,8 @@ FBuildPatchFileConstructor::FBuildPatchFileConstructor(
 		UE_LOG(LogBuildPatchServices, Display, TEXT("Overridding StallWhenFileSystemThrottled to: %d, was %d"), Configuration.StallWhenFileSystemThrottled.GetValue(), bStallWhenFileSystemThrottled);
 	}
 
+	bAllowMultipleFilesInFlight = bCVarAllowMultipleFilesInFlight;
+
 	BackingStore.Reset(new FChunkBackingStore(this, Configuration.InstallDirectory, BackingStoreStats));
 
 	// Count initial job size
