BPS: Delay starting additional files if they depend on chunks that could get harvested from the earlier file to prevent races during patching.

#rnx

diff --git a/Engine/Source/Runtime/Online/BuildPatchServices/Private/BuildPatchFileConstructor.cpp b/Engine/Source/Runtime/Online/BuildPatchServices/Private/BuildPatchFileConstructor.cpp
--- a/Engine/Source/Runtime/Online/BuildPatchServices/Private/BuildPatchFileConstructor.cpp
+++ b/Engine/Source/Runtime/Online/BuildPatchServices/Private/BuildPatchFileConstructor.cpp
@@ -218,11 +218,12 @@ struct FFileConstructionState
 
 	TUniquePtr<FArchive> NewFile;
 
+	int32 ConstructionIndex = -1;
 	const FFileManifest* FileManifest;
-	FString BuildFilename; // name as it is in the manifest
+	const FString& BuildFilename; // name as it is in the manifest, references the ConstructList in the configuration.
 	FString NewFilename; // with output path and such
 
-	FFileConstructionState(const FFileManifest* InFileManifest, FString&& InBuildFilename, FString&& InNewFilename) 
+	FFileConstructionState(const FFileManifest* InFileManifest, const FString& InBuildFilename, FString&& InNewFilename) 
 	: FileManifest(InFileManifest)
 	, BuildFilename(InBuildFilename)
 	, NewFilename(InNewFilename)
@@ -1354,8 +1355,6 @@ FBuildPatchFileConstructor::FBuildPatchFileConstructor(
 	, bInitialDiskSizeCheck(false)
 	, bIsPaused(false)
 	, bShouldAbort(false)
-	, ThreadLock()
-	, ConstructionStack()
 	, FileSystem(InFileSystem)
 	, ChunkDbSource(InChunkDbChunkSource)
 	, InstallSource(InInstallChunkSource)
@@ -1385,21 +1384,29 @@ FBuildPatchFileConstructor::FBuildPatchFileConstructor(
 	BackingStore.Reset(new FChunkBackingStore(this, Configuration.InstallDirectory, BackingStoreStats));
 
 	// Count initial job size
-	const int32 ConstructListNum = Configuration.ConstructList.Num();
-	ConstructionStack.Reserve(ConstructListNum);
-	ConstructionStack.AddDefaulted(ConstructListNum);
-
+	ConstructionList.Reserve(Configuration.ConstructList.Num());
+	
 	// Track when we will complete files in the reference chain.
 	int32 CurrentPosition = 0;
-	FileCompletionPositions.Reserve(ConstructListNum);
+	FileCompletionPositions.Reserve(Configuration.ConstructList.Num());
 
 	// The first index after the file is complete.
+	// \todo with the construction list now stable across the install, we could
+	// key this off a StringView and save the allocs. For another CL...
 	TMap<FString, int32> FileRetirementPositions;
 
-	for (int32 ConstructListIdx = 0; ConstructListIdx < ConstructListNum ; ++ConstructListIdx)
+	for (const FString& FileToConstructName : Configuration.ConstructList)
 	{
-		const FString& ConstructListElem = Configuration.ConstructList[ConstructListIdx];
-		const FFileManifest* FileManifest = Configuration.ManifestSet->GetNewFileManifest(ConstructListElem);
+		const FFileManifest* FileManifest = Configuration.ManifestSet->GetNewFileManifest(FileToConstructName);
+
+		FFileToConstruct FileToConstruct;
+		FileToConstruct.FileManifest = FileManifest;
+
+		// If we are missing the file manifest, we will fail to install when we get to the file. However,
+		// we guarantee a 1:1 mapping with the arrays we are filling here so we use invalid data for those
+		// slots (which won't get used).
+		// Maybe we should fail immediately? Need to review whether we can fail in the constructor or we need
+		// to delay until Run().		
 		if (FileManifest)
 		{
 			TotalJobSize += FileManifest->FileSize;
@@ -1407,29 +1414,76 @@ FBuildPatchFileConstructor::FBuildPatchFileConstructor(
 			// We will be advancing the chunk reference tracker by this many chunks.
 			int32 AdvanceCount = FileManifest->ChunkParts.Num();
 			CurrentPosition += AdvanceCount;
-
-			FileCompletionPositions.Add(CurrentPosition);
-
-			FileRetirementPositions.Add(ConstructListElem, CurrentPosition);
 		}
 
-		ConstructionStack[(ConstructListNum - 1) - ConstructListIdx] = ConstructListElem;
+		FileCompletionPositions.Add(CurrentPosition);
+		FileRetirementPositions.Add(FileToConstructName, CurrentPosition);
+		ConstructionList.Add(MoveTemp(FileToConstruct));
 	}
 
 	// Let the install source know when we're going to be deleting their sources.
+	bool bHasInstallSource = false;
 	if (InstallSource)
 	{
 		InstallSource->SetFileRetirementPositions(MoveTemp(FileRetirementPositions));	
 
-		// Harvesting is not safe when multiple files are in flight as a later file can request
-		// a chunk we needed to have harvested and we can be in the middle of that harvest when the read comes in.
-		bAllowMultipleFilesInFlight = false;
+		if (InstallSource->GetAvailableChunks().Num() != 0)
+		{
+			bHasInstallSource = true;
+
+			// We need to set up a dependency chain so that files can know when they can start.			
+			TMap<FStringView, int32> FileToIndexMap;
+			FileToIndexMap.Reserve(Configuration.ConstructList.Num());
+			for (int32 FileConstructIndex = 0; FileConstructIndex < Configuration.ConstructList.Num(); FileConstructIndex++)
+			{
+				// The construct list has filenames that match the manifest, and afaict manifest filenames are already normalized.
+				FileToIndexMap.Add(Configuration.ConstructList[FileConstructIndex], FileConstructIndex);
+			}
+
+			for (int32 FileConstructIndex = 0; FileConstructIndex < Configuration.ConstructList.Num(); FileConstructIndex++)
+			{
+				FFileToConstruct& FileToConstruct = ConstructionList[FileConstructIndex];				
+				if (!FileToConstruct.FileManifest)
+				{
+					continue;
+				}
+
+				for (const FChunkPart& ChunkPart : FileToConstruct.FileManifest->ChunkParts)
+				{
+					EConstructorChunkLocation ChunkLocation = ChunkLocations[ChunkPart.Guid];
+					if (ChunkLocation == EConstructorChunkLocation::Install)
+					{
+						InstallSource->EnumerateFilesForChunk(ChunkPart.Guid, [FileConstructIndex, &FileToIndexMap, &FileToConstruct, NormalizedInstallDirectory = &Configuration.InstallDirectory](const FString& ChunkNormalizedInstallDirectory, const FString& NormalizedFilenameContainingChunk)
+							{
+								if (*NormalizedInstallDirectory != ChunkNormalizedInstallDirectory)
+								{
+									// We aren't affecting that install source so it's not a dependency.
+									return;
+								}
+
+								const int32* DependentFileIndexPtr = FileToIndexMap.Find(NormalizedFilenameContainingChunk);
+								if (DependentFileIndexPtr)
+								{
+									int32 DependentFileIndex = *DependentFileIndexPtr;
+
+									// If the file is constructed before us then we can't start until it's ready.
+									// We only care about the latest file.
+									if (DependentFileIndex < FileConstructIndex)
+									{
+										FileToConstruct.LatestDependentInstallSource = FMath::Max(FileToConstruct.LatestDependentInstallSource, DependentFileIndex);
+									}
+								}
+							}
+						);
+					}
+				}
+			}			
+		}
 	}
 
 	//
 	// Create the threads we are allowed to.
 	//
-	bool bHasInstallSource = InstallSource && InstallSource->GetAvailableChunks().Num() != 0;
 	bHasChunkDbSource = ChunkDbSource->GetAvailableChunks().Num() != 0;
 	
 	// Default everything to running synchronously.
@@ -1910,21 +1964,25 @@ void FBuildPatchFileConstructor::CountBytesProcessed( const int64& ByteCount )
 	FileConstructorStat->OnProcessedDataUpdated(ByteProcessed);
 }
 
-bool FBuildPatchFileConstructor::GetFileToConstruct(FString& Filename)
+int64 FBuildPatchFileConstructor::GetRemainingBytes()
 {
-	FScopeLock Lock(&ThreadLock);
-	const bool bFileAvailable = ConstructionStack.Num() > 0;
-	if (bFileAvailable)
+	// Need the sum of the output sizes of files not yet started.
+	// Since this gets called from any thread, construction will continue
+	// as we calculate this, but all the structures are stable as long as
+	// the constructor is still valid memory.
+
+	int32 LocalNextIndexToConstruct = NextIndexToConstruct.load(std::memory_order_acquire);
+
+	uint64 RemainingBytes = 0;
+	for (int32 ConstructIndex = LocalNextIndexToConstruct; ConstructIndex < ConstructionList.Num(); ConstructIndex++)
 	{
-		Filename = ConstructionStack.Pop(EAllowShrinking::No);
+		if (ConstructionList[ConstructIndex].FileManifest)
+		{
+			RemainingBytes += ConstructionList[ConstructIndex].FileManifest->FileSize;
+		}
 	}
-	return bFileAvailable;
-}
 
-int64 FBuildPatchFileConstructor::GetRemainingBytes()
-{
-	FScopeLock Lock(&ThreadLock);
-	return Configuration.ManifestSet->GetTotalNewFileSize(ConstructionStack);
+	return RemainingBytes;
 }
 
 uint64 FBuildPatchFileConstructor::CalculateInProgressDiskSpaceRequired(const FFileManifest& InProgressFileManifest, uint64 InProgressFileAmountWritten)
@@ -1952,11 +2010,17 @@ uint64 FBuildPatchFileConstructor::CalculateInProgressDiskSpaceRequired(const FF
 		}
 
 		// Loop through all files to be made next, in order.
-		for (int32 ConstructionStackIdx = ConstructionStack.Num() - 1; ConstructionStackIdx >= 0; --ConstructionStackIdx)
+		int32 LocalNextIndexToConstruct = NextIndexToConstruct.load(std::memory_order_acquire);
+		for (int32 ConstructionIndex = LocalNextIndexToConstruct; ConstructionIndex < ConstructionList.Num(); ConstructionIndex++)
 		{
-			const FString& FileToConstruct = ConstructionStack[ConstructionStackIdx];
-			const FFileManifest* NewFileManifest = Configuration.ManifestSet->GetNewFileManifest(FileToConstruct);
-			const FFileManifest* OldFileManifest = Configuration.ManifestSet->GetCurrentFileManifest(FileToConstruct);
+			const FFileManifest* NewFileManifest = ConstructionList[ConstructionIndex].FileManifest;
+			if (!NewFileManifest)
+			{
+				continue;
+			}
+
+			const FFileManifest* OldFileManifest = Configuration.ManifestSet->GetCurrentFileManifest(Configuration.ConstructList[ConstructionIndex]);
+
 			// First we would need to make the new file.
 			DiskSpaceDelta += NewFileManifest->FileSize;
 			if (DiskSpaceDeltaPeak < DiskSpaceDelta)
@@ -1974,13 +2038,13 @@ uint64 FBuildPatchFileConstructor::CalculateInProgressDiskSpaceRequired(const FF
 	else
 	{
 		// When not destructive, we always stage all new and changed files.
-		uint64 RemainingFilesSpace = Configuration.ManifestSet->GetTotalNewFileSize(ConstructionStack);
+		uint64 RemainingFilesSpace = GetRemainingBytes();
 		uint64 RemainingThisFileSpace = InProgressFileManifest.FileSize - InProgressFileAmountWritten;
 		return RemainingFilesSpace + RemainingThisFileSpace;
 	}
 }
 
-uint64 FBuildPatchFileConstructor::CalculateDiskSpaceRequirementsWithDeleteDuringInstall(const TArray<FString>& InBackwardsFilesLeftToConstruct)
+uint64 FBuildPatchFileConstructor::CalculateDiskSpaceRequirementsWithDeleteDuringInstall()
 {
 	if (ChunkDbSource == nullptr)
 	{
@@ -1994,10 +2058,14 @@ uint64 FBuildPatchFileConstructor::CalculateDiskSpaceRequirementsWithDeleteDurin
 	uint64 TotalChunkDbSize = ChunkDbSource->GetChunkDbSizesAtIndexes(FileCompletionPositions, ChunkDbSizesAtPosition);
 
 	// Strip off the files we've completed.
-	int32 CompletedFileCount = Configuration.ConstructList.Num() - InBackwardsFilesLeftToConstruct.Num();
+	int32 CompletedFileCount = NextIndexToConstruct.load(std::memory_order_acquire);
 
 	// Since we are called after the first file is popped (but before it's actually done), we have one less completed.
-	CompletedFileCount--;
+	check(CompletedFileCount > 0); // should ALWAYS be at least 1!
+	if (CompletedFileCount > 0)
+	{
+		CompletedFileCount--;
+	}
 
 	uint64 MaxDiskSize = FBuildPatchUtils::CalculateDiskSpaceRequirementsWithDeleteDuringInstall(
 		Configuration.ConstructList, CompletedFileCount, Configuration.ManifestSet, ChunkDbSizesAtPosition, TotalChunkDbSize);
@@ -2587,7 +2655,7 @@ bool FBuildPatchFileConstructor::HandleInitialDiskSizeCheck(const FFileManifest&
 		// space as we add.
 		if (Configuration.bDeleteChunkDBFilesAfterUse)
 		{
-			LocalDiskSpaceRequired = CalculateDiskSpaceRequirementsWithDeleteDuringInstall(ConstructionStack);
+			LocalDiskSpaceRequired = CalculateDiskSpaceRequirementsWithDeleteDuringInstall();
 		}
 
 		uint64 LocalDiskSpaceAvailable = 0;
@@ -3105,16 +3173,16 @@ void FBuildPatchFileConstructor::ConstructFiles(const FResumeData& ResumeData)
 	// Ensure that our batch size can always make progress (all chunks can fit)
 	uint32 LargestChunkSize = 0;
 	{
-		FScopeLock Lock(&ThreadLock);
-
-		for (const FString& FileToConstruct : ConstructionStack)
+		for (const FFileToConstruct& FileToConstruct : ConstructionList)
 		{
-			const FFileManifest* FileManifest = Configuration.ManifestSet->GetNewFileManifest(FileToConstruct);
-			for (const FChunkPart& ChunkPart : FileManifest->ChunkParts)
+			if (FileToConstruct.FileManifest)
 			{
-				if (ChunkPart.Size > LargestChunkSize)
+				for (const FChunkPart& ChunkPart : FileToConstruct.FileManifest->ChunkParts)
 				{
-					LargestChunkSize = ChunkPart.Size;
+					if (ChunkPart.Size > LargestChunkSize)
+					{
+						LargestChunkSize = ChunkPart.Size;
+					}
 				}
 			}
 		}
@@ -3360,42 +3428,74 @@ void FBuildPatchFileConstructor::ConstructFiles(const FResumeData& ResumeData)
 						{
 							bool bAllowAnotherFile = bAllowMultipleFilesInFlight || (ActiveFiles.Num() == 0);
 
-							FString AnotherFileToConstruct;
-							if (bAllowAnotherFile && GetFileToConstruct(AnotherFileToConstruct))
+							// This is not a race because the only place we ever increment is this thread, this function right
+							// below.
+							int32 IndexToConstruct = NextIndexToConstruct.load(std::memory_order_acquire);
+							bool bAnotherFileExists = IndexToConstruct < ConstructionList.Num();
+							
+							if (bAllowAnotherFile && bAnotherFileExists)
 							{
-								FString StagingFileName = Configuration.StagingDirectory / AnotherFileToConstruct;
-								TUniquePtr<FFileConstructionState> AnotherFile = MakeUnique<FFileConstructionState>(
-									Configuration.ManifestSet->GetNewFileManifest(AnotherFileToConstruct), 
-									MoveTemp(AnotherFileToConstruct), 
-									MoveTemp(StagingFileName));
+								// Even though another file exists, we might not be able to start it if its dependent
+								// files aren't done.
+								bool bDelayForDependencies = false;
+								if (ConstructionList[IndexToConstruct].LatestDependentInstallSource != -1)
+								{
+									// Files are in construct order - if the first one is after our last dependency then we know we
+									// are safe.
+									if (ActiveFileCount && ActiveFiles[0]->ConstructionIndex <= ConstructionList[IndexToConstruct].LatestDependentInstallSource)
+									{
+										bDelayForDependencies = true;
+									}
+								}
 
-								TRACE_BOOKMARK(TEXT("Starting File: %s [%u]"), *AnotherFile->BuildFilename, AnotherFile->FileManifest->ChunkParts.Num());
+								if (!bDelayForDependencies)
+								{
+									// Now commit this file since we are starting it.
+									NextIndexToConstruct.fetch_add(1, std::memory_order_acq_rel);
 
-								UE_LOG(LogBuildPatchServices, Verbose, TEXT("Starting File: %s [%s bytes, %d chunks], New ActiveFileCount = %d"), 
-									*AnotherFile->BuildFilename, 
-									*FormatNumber(AnotherFile->FileManifest->FileSize),
-									AnotherFile->FileManifest->ChunkParts.Num(),
-									ActiveFiles.Num() + 1);
+									FString StagingFileName = Configuration.StagingDirectory / Configuration.ConstructList[IndexToConstruct];
+									TUniquePtr<FFileConstructionState> AnotherFile = MakeUnique<FFileConstructionState>(
+										ConstructionList[IndexToConstruct].FileManifest, 
+										Configuration.ConstructList[IndexToConstruct], 
+										MoveTemp(StagingFileName));
 
-								StartFile(*AnotherFile.Get(), ResumeData);
+									AnotherFile->ConstructionIndex = IndexToConstruct;
 
-								if (AnotherFile->bSkippedConstruction)
-								{
-									// Nothing else needs to happen with this file - we'll loop around to try for another.
-									bCheckForAnotherRead = true;
-								}
-								else
-								{
-									// Only start the file if there's not currently an active file - otherwise we
-									// do it when the current one finishes.
-									if (ActiveFileCount == 0)
+									TRACE_BOOKMARK(TEXT("Starting File: %s [%u]"), *AnotherFile->BuildFilename, AnotherFile->FileManifest->ChunkParts.Num());
+
+									UE_LOG(LogBuildPatchServices, Verbose, TEXT("Starting File: %s [%s bytes, %d chunks], New ActiveFileCount = %d"), 
+										   *AnotherFile->BuildFilename, 
+										   *FormatNumber(AnotherFile->FileManifest->FileSize),
+										   AnotherFile->FileManifest->ChunkParts.Num(),
+										   ActiveFiles.Num() + 1);
+
+									StartFile(*AnotherFile.Get(), ResumeData);
+
+									if (AnotherFile->bSkippedConstruction)
 									{
-										FileConstructorStat->OnFileStarted(AnotherFile->BuildFilename, AnotherFile->FileManifest->FileSize);
+										// Nothing else needs to happen with this file - we'll loop around to try for another.
+										bCheckForAnotherRead = true;
 									}
-
-									FileToStart = AnotherFile.Get();
-									ActiveFiles.Push(MoveTemp(AnotherFile));
+									else
+									{
+										// Only start the file if there's not currently an active file - otherwise we
+										// do it when the current one finishes.
+										if (ActiveFileCount == 0)
+										{
+											FileConstructorStat->OnFileStarted(AnotherFile->BuildFilename, AnotherFile->FileManifest->FileSize);
+										}
+
+										FileToStart = AnotherFile.Get();
+										ActiveFiles.Push(MoveTemp(AnotherFile));
+									}
+								} // end if not delaying for dependencies
+								else
+								{
+									UE_LOG(LogBuildPatchServices, Verbose, TEXT("Delaying %s due to incomplete dependencies"), *Configuration.ConstructList[IndexToConstruct]);
 								}
+
+								// If we are delaying for dependencies we fall through here with FileToStart = nullptr
+								// and nothing happens until we check again when the construct thread is woken up again.
 							}
 							else
 							{
diff --git a/Engine/Source/Runtime/Online/BuildPatchServices/Private/BuildPatchFileConstructor.h b/Engine/Source/Runtime/Online/BuildPatchServices/Private/BuildPatchFileConstructor.h
--- a/Engine/Source/Runtime/Online/BuildPatchServices/Private/BuildPatchFileConstructor.h
+++ b/Engine/Source/Runtime/Online/BuildPatchServices/Private/BuildPatchFileConstructor.h
@@ -165,13 +165,6 @@ namespace BuildPatchServices
 		 */
 		void CountBytesProcessed(const int64& ByteCount);
 
-		/**
-		 * Function to fetch a chunk from the download list
-		 * @param Filename		Receives the filename for the file to construct from the manifest
-		 * @return true if there was a chunk guid in the list
-		 */
-		bool GetFileToConstruct(FString& Filename);
-
 		/**
 		 * @return the total bytes size of files not yet started construction
 		 */
@@ -185,8 +178,8 @@ namespace BuildPatchServices
 		 */
 		uint64 CalculateInProgressDiskSpaceRequired(const FFileManifest& InProgressFileManifest, uint64 InProgressFileSize);
 
-		// Calculates the amount of disk space we need to finish the install, needs to be called on file boundaries.
-		uint64 CalculateDiskSpaceRequirementsWithDeleteDuringInstall(const TArray<FString>& InConstructionStack);
+		// Calculates the amount of disk space we need to finish the install, needs to be called on file boundaries on the construct thread.
+		uint64 CalculateDiskSpaceRequirementsWithDeleteDuringInstall();
 
 		/**
 		 * Sequentially build the files required, potentially skipping already created files or resuming partial
@@ -222,12 +215,23 @@ namespace BuildPatchServices
 		// A flag marking whether we should abort operations and exit. Always call Abort() to set this.
 		FThreadSafeBool bShouldAbort;
 
-		// A critical section to protect the flags and variables.
-		// \todo remove this entirely and just track the remaining size when we pop a file.
-		FCriticalSection ThreadLock;
+		// Indexes in to ConstructionList and associated parallel arrays. This is the next file that will
+		// start construction when dependencies are met.
+		std::atomic_int32_t NextIndexToConstruct;
+
+		struct FFileToConstruct
+		{
+			const FFileManifest* FileManifest = nullptr;
+
+			// When using an install source with multiple files in flight, we can't start this 
+			// file until all of the install sources it needs have been harvested. Since files are
+			// constructed in order, we only track the file that will be harvested last. If this
+			// file has no dependencies this is -1
+			int32 LatestDependentInstallSource = -1;
+		};
 
-		// A stack of filenames for files that need to be constructed.
-		TArray<FString> ConstructionStack;
+		// The in-oder list of files to construct. The array is parallel with Configuration.ConstructList.
+		TArray<FFileToConstruct> ConstructionList;
 
 		// Pointer to the file system.
 		IFileSystem* FileSystem;
diff --git a/Engine/Source/Runtime/Online/BuildPatchServices/Private/BuildPatchInstaller.cpp b/Engine/Source/Runtime/Online/BuildPatchServices/Private/BuildPatchInstaller.cpp
--- a/Engine/Source/Runtime/Online/BuildPatchServices/Private/BuildPatchInstaller.cpp
+++ b/Engine/Source/Runtime/Online/BuildPatchServices/Private/BuildPatchInstaller.cpp
@@ -1403,6 +1403,7 @@ namespace BuildPatchServices
 			FFileConstructorConfig FCC;
 			FCC.ConstructList = FilesToConstruct.Array();
 			FCC.InstallDirectory = Configuration.InstallDirectory;
+			FPaths::NormalizeDirectoryName(FCC.InstallDirectory);
 			FCC.InstallMode = Configuration.InstallMode;
 			FCC.ManifestSet = ManifestSet.Get();
 			FCC.MetaDirectory = MetaStagingDir;
diff --git a/Engine/Source/Runtime/Online/BuildPatchServices/Private/Installer/InstallChunkSource.cpp b/Engine/Source/Runtime/Online/BuildPatchServices/Private/Installer/InstallChunkSource.cpp
--- a/Engine/Source/Runtime/Online/BuildPatchServices/Private/Installer/InstallChunkSource.cpp
+++ b/Engine/Source/Runtime/Online/BuildPatchServices/Private/Installer/InstallChunkSource.cpp
@@ -47,6 +47,34 @@ namespace BuildPatchServices
 
 		virtual void GetChunksForFile(const FString& FilePath, TSet<FGuid>& OutChunks) const override;
 
+		virtual void EnumerateFilesForChunk(const FGuid& DataId, TUniqueFunction<void(const FString& NormalizedInstallDirectory, const FString& NormalizedFileName)>&& Callback) const override
+		{
+			const FString* FoundInstallDirectory;
+			const FBuildPatchAppManifest* FoundInstallManifest;
+			FindChunkLocation(DataId, &FoundInstallDirectory, &FoundInstallManifest);
+
+			const TArray<FChunkSourceDetails>* ChunkSource = ChunkSources.Find(DataId);
+
+			if (FoundInstallDirectory == nullptr || FoundInstallManifest == nullptr || ChunkSource == nullptr)
+			{
+				return;
+			}
+
+			// The installation directory starts off normalized by then appends a directory which
+			// might be empty, leaving a trailing slash. Rather than chase down all possibilities we
+			// just re normalize.
+			FString NormalizedInstallDirectory = *FoundInstallDirectory;
+			FPaths::NormalizeDirectoryName(NormalizedInstallDirectory);
+
+			for (const FChunkSourceDetails& ChunkDetails : *ChunkSource)
+			{
+				// afaict the file manifest filename is normalized because in the manifest builder it
+				// generates it from file spans, which are created in directorybuildstreamer which makes them
+				// relative, and internal to that function they are normalized.
+				Callback(NormalizedInstallDirectory, ChunkDetails.FileManifest->Filename);
+			}
+		}
+
 	private:
 		void FindChunkLocation(const FGuid& DataId, const FString** FoundInstallDirectory, const FBuildPatchAppManifest** FoundInstallManifest) const;
 
diff --git a/Engine/Source/Runtime/Online/BuildPatchServices/Private/Installer/InstallChunkSource.h b/Engine/Source/Runtime/Online/BuildPatchServices/Private/Installer/InstallChunkSource.h
--- a/Engine/Source/Runtime/Online/BuildPatchServices/Private/Installer/InstallChunkSource.h
+++ b/Engine/Source/Runtime/Online/BuildPatchServices/Private/Installer/InstallChunkSource.h
@@ -36,6 +36,9 @@ namespace BuildPatchServices
 		// is deleted
 		virtual void GetChunksForFile(const FString& FilePath, TSet<FGuid>& OutChunks) const = 0;
 
+		// Call the lambda on each file the chunk needs to load bits from.
+		virtual void EnumerateFilesForChunk(const FGuid& DataId, TUniqueFunction<void(const FString& NormalizedInstallDirectory, const FString& NormalizedFileName)>&& Callback) const = 0;
+
 		// Tell the install source when files will be going away so we can report when chunks are unavailable.
 		// This should be the ChunkReferenceTracker->GetCurrentUsageIndex at which the files will get deleted.
 		virtual void SetFileRetirementPositions(TMap<FString, int32>&& FileRetirementPositions) = 0;
