Fixing missing initialization of NextIndexToConstruct

diff --git a/Engine/Source/Runtime/Online/BuildPatchServices/Private/BuildPatchFileConstructor.h b/Engine/Source/Runtime/Online/BuildPatchServices/Private/BuildPatchFileConstructor.h
--- a/Engine/Source/Runtime/Online/BuildPatchServices/Private/BuildPatchFileConstructor.h
+++ b/Engine/Source/Runtime/Online/BuildPatchServices/Private/BuildPatchFileConstructor.h
@@ -217,7 +217,7 @@ namespace BuildPatchServices
 
 		// Indexes in to ConstructionList and associated parallel arrays. This is the next file that will
 		// start construction when dependencies are met.
-		std::atomic_int32_t NextIndexToConstruct;
+		std::atomic_int32_t NextIndexToConstruct = 0;
 
 		struct FFileToConstruct
 		{
