Add http constant strings for commonly used headers' key/value and verbs;

diff --git a/Engine/Source/Runtime/Online/HTTP/Private/HttpConstants.cpp b/Engine/Source/Runtime/Online/HTTP/Private/HttpConstants.cpp
new file mode 100644
--- /dev/null
+++ b/Engine/Source/Runtime/Online/HTTP/Private/HttpConstants.cpp
@@ -0,0 +1,28 @@
+// Copyright Epic Games, Inc. All Rights Reserved.
+
+#include "HttpConstants.h"
+
+const TCHAR* const FHttpConstants::AUTH_BASIC{ TEXT("Basic") };
+const TCHAR* const FHttpConstants::AUTH_BEARER{ TEXT("Bearer") };
+
+const TCHAR* const FHttpConstants::VERB_DELETE{ TEXT("DELETE") };
+const TCHAR* const FHttpConstants::VERB_GET{ TEXT("GET") };
+const TCHAR* const FHttpConstants::VERB_HEAD{ TEXT("HEAD") };
+const TCHAR* const FHttpConstants::VERB_PATCH{ TEXT("PATCH") };
+const TCHAR* const FHttpConstants::VERB_POST{ TEXT("POST") };
+const TCHAR* const FHttpConstants::VERB_PUT{ TEXT("PUT") };
+
+const TCHAR* const FHttpConstants::HEADER_ACCEPT{ TEXT("Accept") };
+const TCHAR* const FHttpConstants::HEADER_ACCEPT_ENCODING{ TEXT("Accept-Encoding") };
+const TCHAR* const FHttpConstants::HEADER_AUTHORIZATION{ TEXT("Authorization") };
+const TCHAR* const FHttpConstants::HEADER_CONTENT_LENGTH{ TEXT("Content-Length") };
+const TCHAR* const FHttpConstants::HEADER_CONTENT_TYPE{ TEXT("Content-Type") };
+const TCHAR* const FHttpConstants::HEADER_DATE{ TEXT("Date") };
+const TCHAR* const FHttpConstants::HEADER_USER_AGENT{ TEXT("User-Agent") };
+
+const TCHAR* const FHttpConstants::MEDIATYPE_JSON{ TEXT("application/json") };
+const TCHAR* const FHttpConstants::MEDIATYPE_FORM_URLENCODED{ TEXT("application/x-www-form-urlencoded") };
+const TCHAR* const FHttpConstants::MEDIATYPE_OCTET_STREAM{ TEXT("application/octet-stream") };
+
+const TCHAR* const FHttpConstants::MEDIAENCODING_GZIP{ TEXT("gzip") };
+
diff --git a/Engine/Source/Runtime/Online/HTTP/Public/HttpConstants.h b/Engine/Source/Runtime/Online/HTTP/Public/HttpConstants.h
new file mode 100644
--- /dev/null
+++ b/Engine/Source/Runtime/Online/HTTP/Public/HttpConstants.h
@@ -0,0 +1,55 @@
+// Copyright Epic Games, Inc. All Rights Reserved.
+
+#pragma once
+
+#include "HAL/Platform.h"
+
+/**
+ * Universal HTTP Constants
+ */
+struct HTTP_API FHttpConstants
+{
+	/** Basic auth */
+	static const TCHAR* const AUTH_BASIC;
+	/** Bearer auth */
+	static const TCHAR* const AUTH_BEARER;
+
+	/** DELETE */
+	static const TCHAR* const VERB_DELETE;
+	/** GET */
+	static const TCHAR* const VERB_GET;
+	/** HEAD */
+	static const TCHAR* const VERB_HEAD;
+	/** PATCH */
+	static const TCHAR* const VERB_PATCH;
+	/** POST */
+	static const TCHAR* const VERB_POST;
+	/** PUT */
+	static const TCHAR* const VERB_PUT;
+
+	/** Accept */
+	static const TCHAR* const HEADER_ACCEPT;
+	/** Accept-Encoding */
+	static const TCHAR* const HEADER_ACCEPT_ENCODING;
+	/** Authorization */
+	static const TCHAR* const HEADER_AUTHORIZATION;
+	/** Content-Length */
+	static const TCHAR* const HEADER_CONTENT_LENGTH;
+	/** Content-Type */
+	static const TCHAR* const HEADER_CONTENT_TYPE;
+	/** Server Date */
+	static const TCHAR* const HEADER_DATE;
+	/** User Agent */
+	static const TCHAR* const HEADER_USER_AGENT;
+
+	/** application/json */
+	static const TCHAR* const MEDIATYPE_JSON;
+	/** application/x-www-form-urlencoded */
+	static const TCHAR* const MEDIATYPE_FORM_URLENCODED;
+	/** application/octet-stream */
+	static const TCHAR* const MEDIATYPE_OCTET_STREAM;
+
+	/** gzip */
+	static const TCHAR* const MEDIAENCODING_GZIP;
+};
+
diff --git a/Engine/Source/Runtime/Online/HTTP/Public/HttpFwd.h b/Engine/Source/Runtime/Online/HTTP/Public/HttpFwd.h
--- a/Engine/Source/Runtime/Online/HTTP/Public/HttpFwd.h
+++ b/Engine/Source/Runtime/Online/HTTP/Public/HttpFwd.h
@@ -4,6 +4,7 @@
 
 #include "Templates/SharedPointer.h"
 #include "Misc/Optional.h"
+#include "HttpConstants.h"
 
 class IHttpRequest;
 class IHttpResponse;
