Fix the illegal instruction error in EventLoopUnitTests when using clang.

diff --git a/Engine/Source/Runtime/Online/Experimental/EventLoopTests/Tests/EventLoop/EventLoopManagedStorageTests.cpp b/Engine/Source/Runtime/Online/Experimental/EventLoopTests/Tests/EventLoop/EventLoopManagedStorageTests.cpp
--- a/Engine/Source/Runtime/Online/Experimental/EventLoopTests/Tests/EventLoop/EventLoopManagedStorageTests.cpp
+++ b/Engine/Source/Runtime/Online/Experimental/EventLoopTests/Tests/EventLoop/EventLoopManagedStorageTests.cpp
@@ -600,13 +600,14 @@ namespace UE::EventLoop
 			}
 
 			FManagedData(FManagedData&& Other)
-				: IntValue(Other.IntValue)
+				: IntValue(MoveTemp(Other.IntValue))
 			{
 			}
 
 			FManagedData& operator=(FManagedData&& Other)
 			{
-				IntValue = Other.IntValue;
+				IntValue = MoveTemp(Other.IntValue);
+				return *this;
 			}
 
 			int32 IntValue;
