[AutoRTFM] Defer HTTP network cancellation while inside a transaction.

This can be caused by player backfilling, which can occur inside a Verse transaction. When AutoRTFM is off, this change should be a no-op.

diff --git a/Engine/Source/Runtime/Online/HTTP/Private/HttpRetrySystem.cpp b/Engine/Source/Runtime/Online/HTTP/Private/HttpRetrySystem.cpp
--- a/Engine/Source/Runtime/Online/HTTP/Private/HttpRetrySystem.cpp
+++ b/Engine/Source/Runtime/Online/HTTP/Private/HttpRetrySystem.cpp
@@ -1,6 +1,7 @@
 // Copyright Epic Games, Inc. All Rights Reserved.
 
 #include "HttpRetrySystem.h"
+#include "AutoRTFM/AutoRTFM.h"
 #include "GenericPlatform/HttpRequestCommon.h"
 #include "HAL/ConsoleManager.h"
 #include "HAL/PlatformTime.h"
@@ -531,7 +532,7 @@ void FHttpRetrySystem::FManager::RetryHttpRequestWithDelay(FManager::FHttpRetryR
 				if (EntryIndex != INDEX_NONE)
 				{
 					FManager::FHttpRetryRequestEntry* HttpRetryRequestEntry = &RetryManagerPtr->RequestList[EntryIndex];
-					// TODO: Move this into RetryHttpRequest after stablizing the flow with CVarHttpRetrySystemNonGameThreadSupportEnabled on
+					// TODO: Move this into RetryHttpRequest after stabilizing the flow with CVarHttpRetrySystemNonGameThreadSupportEnabled on
 					HttpRetryRequestEntry->Request->LastResponse = HttpRetryRequestEntry->Request->GetResponse();
 					HttpRetryRequestEntry->Request->bLastSucceeded = bWasSucceeded;
 					RetryManagerPtr->RetryHttpRequest(*HttpRetryRequestEntry);
@@ -636,33 +637,36 @@ bool FHttpRetrySystem::FManager::ProcessRequest(TSharedRef<FHttpRetrySystem::FRe
 
 void FHttpRetrySystem::FManager::CancelRequest(TSharedRef<FHttpRetrySystem::FRequest, ESPMode::ThreadSafe>& HttpRetryRequest)
 {
-	QUICK_SCOPE_CYCLE_COUNTER(STAT_FHttpRetrySystem_FManager_CancelRequest);
+	UE_AUTORTFM_ONCOMMIT(this, HttpRetryRequest)
+	{
+		QUICK_SCOPE_CYCLE_COUNTER(STAT_FHttpRetrySystem_FManager_CancelRequest);
 
-	FScopeLock ScopeLock(&RequestListLock);
+		FScopeLock ScopeLock(&RequestListLock);
 
-	// Find the existing request entry if is was previously processed.
-	bool bFound = false;
-	for (int32 i = 0; i < RequestList.Num(); ++i)
-	{
-		FHttpRetryRequestEntry& EntryRef = RequestList[i];
+		// Find the existing request entry if is was previously processed.
+		bool bFound = false;
+		for (int32 i = 0; i < RequestList.Num(); ++i)
+		{
+			FHttpRetryRequestEntry& EntryRef = RequestList[i];
 
-		if (EntryRef.Request == HttpRetryRequest)
+			if (EntryRef.Request == HttpRetryRequest)
+			{
+				EntryRef.bShouldCancel = true;
+				bFound = true;
+			}
+		}
+		// If we did not find the entry, likely auth failed for the request, in which case ProcessRequest does not get called.
+		// Adding it to the list and flagging as cancel will process it on next tick.
+		if (!bFound)
 		{
-			EntryRef.bShouldCancel = true;
-			bFound = true;
+			FHttpRetryRequestEntry RetryRequestEntry(HttpRetryRequest);
+			RetryRequestEntry.bShouldCancel = true;
+			RequestList.Add(RetryRequestEntry);
 		}
-	}
-	// If we did not find the entry, likely auth failed for the request, in which case ProcessRequest does not get called.
-	// Adding it to the list and flagging as cancel will process it on next tick.
-	if (!bFound)
-	{
-		FHttpRetryRequestEntry RetryRequestEntry(HttpRetryRequest);
-		RetryRequestEntry.bShouldCancel = true;
-		RequestList.Add(RetryRequestEntry);
-	}
 
-	HttpRetryRequest->HttpRequest->CancelRequest();
-	HttpRetryRequest->RetryStatus = FHttpRetrySystem::FRequest::EStatus::Cancelled;
+		HttpRetryRequest->HttpRequest->CancelRequest();
+		HttpRetryRequest->RetryStatus = FHttpRetrySystem::FRequest::EStatus::Cancelled;
+	};
 }
 
 /* This should only be used when shutting down or suspending, to make sure 
