BPS: Fixing silly logging/comment mistakes

#rnx

diff --git a/Engine/Source/Runtime/Online/BuildPatchServices/Private/BuildPatchFileConstructor.cpp b/Engine/Source/Runtime/Online/BuildPatchServices/Private/BuildPatchFileConstructor.cpp
--- a/Engine/Source/Runtime/Online/BuildPatchServices/Private/BuildPatchFileConstructor.cpp
+++ b/Engine/Source/Runtime/Online/BuildPatchServices/Private/BuildPatchFileConstructor.cpp
@@ -1375,8 +1375,8 @@ FBuildPatchFileConstructor::FBuildPatchFileConstructor(
 	bStallWhenFileSystemThrottled = bCVarStallWhenFileSystemThrottled;
 	if (Configuration.StallWhenFileSystemThrottled.IsSet())
 	{
+		UE_LOG(LogBuildPatchServices, Display, TEXT("Overridding StallWhenFileSystemThrottled to: %d, cvar was %d"), Configuration.StallWhenFileSystemThrottled.GetValue(), bStallWhenFileSystemThrottled);
 		bStallWhenFileSystemThrottled = Configuration.StallWhenFileSystemThrottled.GetValue();
-		UE_LOG(LogBuildPatchServices, Display, TEXT("Overridding StallWhenFileSystemThrottled to: %d, was %d"), Configuration.StallWhenFileSystemThrottled.GetValue(), bStallWhenFileSystemThrottled);
 	}
 
 	bAllowMultipleFilesInFlight = bCVarAllowMultipleFilesInFlight;
diff --git a/Engine/Source/Runtime/Online/BuildPatchServices/Public/BuildPatchSettings.h b/Engine/Source/Runtime/Online/BuildPatchServices/Public/BuildPatchSettings.h
--- a/Engine/Source/Runtime/Online/BuildPatchServices/Public/BuildPatchSettings.h
+++ b/Engine/Source/Runtime/Online/BuildPatchServices/Public/BuildPatchSettings.h
@@ -308,8 +308,7 @@ namespace BuildPatchServices
 
 		// Whether the file construction should adhere to file system write limits. If a write limit is hit, writing will stall until
 		// sufficient quota exists for the write. Inheritance for this property is:
-		// Default = false
-		// if set, CVar BuildPatchFileConstructor.bStallWhenFileSystemThrottled
+		// CVar BuildPatchFileConstructor.bStallWhenFileSystemThrottled, which defaults to true.
 		// if set, this member.
 		TOptional<bool> ConstructorStallWhenFileSystemThrottled;
 	};
