- Add stat scope around RequestAddedDelegate delegate

#rnx

diff --git a/Engine/Source/Runtime/Online/HTTP/Private/HttpManager.cpp b/Engine/Source/Runtime/Online/HTTP/Private/HttpManager.cpp
--- a/Engine/Source/Runtime/Online/HTTP/Private/HttpManager.cpp
+++ b/Engine/Source/Runtime/Online/HTTP/Private/HttpManager.cpp
@@ -573,6 +573,8 @@ void FHttpManager::AddRequest(const FHttpRequestRef& Request)
 		UE_CLOG(bFlushing, LogHttp, Warning, TEXT("Adding request %s to http manager while flushing"), *Request->GetURL());
 		Requests.Add(Request);
 	}
+
+	QUICK_SCOPE_CYCLE_COUNTER(STAT_FHttpManager_RequestAddedDelegate);
 	RequestAddedDelegate.ExecuteIfBound(Request);
 }
 
