[WebSockets] Fixed an issue where the WebSockets weren't working correctly after the LwsWebSocketsManager forked
- Also made the LWS socket manager thread priority configurable (though it still defaults to TPri_BelowNormal)

diff --git a/Engine/Source/Runtime/Online/WebSockets/Private/Lws/LwsWebSocketsManager.cpp b/Engine/Source/Runtime/Online/WebSockets/Private/Lws/LwsWebSocketsManager.cpp
--- a/Engine/Source/Runtime/Online/WebSockets/Private/Lws/LwsWebSocketsManager.cpp
+++ b/Engine/Source/Runtime/Online/WebSockets/Private/Lws/LwsWebSocketsManager.cpp
@@ -40,6 +40,19 @@ namespace {
 	};
 }
 
+namespace UE::LwsWebSocketsManager
+{
+	// We declare these explicitly rather than just casting in case the enum changes in future
+	const int32 ThreadPriorities[] =
+	{
+		EThreadPriority::TPri_Lowest,
+		EThreadPriority::TPri_BelowNormal,
+		EThreadPriority::TPri_SlightlyBelowNormal,
+		EThreadPriority::TPri_Normal,
+		EThreadPriority::TPri_AboveNormal
+	};
+}
+
 static void LwsLog(int Level, const char* LogLine);
 
 // FLwsWebSocketsManager
@@ -64,13 +77,14 @@ void FLwsWebSocketsManager::UpdateConfigs()
 {
 	const bool bWasPollingService = bPollService;
 	GConfig->GetBool(TEXT("WebSockets.LibWebSockets"), TEXT("bPollService"), bPollService, GEngineIni);
+	GConfig->GetInt(TEXT("WebSockets.LibWebSockets"), TEXT("ServiceTimeoutMs"), ServiceTimeoutMs, GEngineIni);
 	GConfig->GetDouble(TEXT("WebSockets.LibWebSockets"), TEXT("ThreadTargetFrameTimeInSeconds"), ThreadTargetFrameTimeInSeconds, GEngineIni);
 	GConfig->GetDouble(TEXT("WebSockets.LibWebSockets"), TEXT("ThreadMinimumSleepTimeInSeconds"), ThreadMinimumSleepTimeInSeconds, GEngineIni);
 	
 	GConfig->GetBool(TEXT("LwsWebSocket"), TEXT("bDisableDomainAllowlist"), bDisableDomainAllowlist, GEngineIni);
 	
 	GConfig->GetBool(TEXT("LwsWebSocket"), TEXT("bDisableCertValidation"), bDisableCertValidation, GEngineIni);
-	
+
 #if !UE_BUILD_SHIPPING
 	// Allow -lws-insecure= to override ini flags
 	FParse::Bool(FCommandLine::Get(), TEXT("lws-insecure="), bDisableCertValidation);
@@ -205,9 +219,16 @@ void FLwsWebSocketsManager::InitWebSockets(TArrayView<const FString> Protocols)
 		return;
 	}
 	
+	ExitRequest.Set(false);
 	int32 ThreadStackSize = 128 * 1024;
 	GConfig->GetInt(TEXT("WebSockets.LibWebSockets"), TEXT("ThreadStackSize"), ThreadStackSize, GEngineIni);
-	Thread = FForkProcessHelper::CreateForkableThread(this, TEXT("LibwebsocketsThread"), 128 * 1024, TPri_BelowNormal, FPlatformAffinity::GetNoAffinityMask());
+
+	int32 ThreadPriorityIndex = 1; // TPri_BelowNormal
+	GConfig->GetInt(TEXT("WebSockets.LibWebSockets"), TEXT("ThreadPriority"), ThreadPriorityIndex, GEngineIni);
+
+	const EThreadPriority ThreadPriority = static_cast<EThreadPriority>(UE::LwsWebSocketsManager::ThreadPriorities[FMath::Clamp(ThreadPriorityIndex, 0, UE_ARRAY_COUNT(UE::LwsWebSocketsManager::ThreadPriorities) - 1)]);
+
+	Thread = FForkProcessHelper::CreateForkableThread(this, TEXT("LibwebsocketsThread"), 128 * 1024, ThreadPriority, FPlatformAffinity::GetNoAffinityMask());
 	if (!Thread)
 	{
 		UE_LOG(LogWebSockets, Error, TEXT("FLwsWebSocketsManager failed to initialize thread!"));
@@ -223,9 +244,17 @@ void FLwsWebSocketsManager::InitWebSockets(TArrayView<const FString> Protocols)
 	// If we get forked, we may want to enable / disable service polling
 	if (FForkProcessHelper::IsForkRequested())
 	{
-		OnChildEndFramePostForkHandle = FCoreDelegates::OnChildEndFramePostFork.AddLambda([this]()
+		TArray<FString> ProtocolsCopy;
+		ProtocolsCopy.Append(Protocols);
+
+		OnChildEndFramePostForkHandle = FCoreDelegates::OnChildEndFramePostFork.AddLambda([this, Protocols = ProtocolsCopy]()
 		{
-			UpdateConfigs();
+			// We have to make a local copy of the array because, at least on Windows, the captured variables become invalidated
+			// when the lambda gets unregistered by ShutdownWebSockets()
+			TArray<FString> ProtocolsCopy = Protocols;
+
+			ShutdownWebSockets();
+			InitWebSockets(ProtocolsCopy);
 		});
 	}
 }
@@ -458,7 +487,7 @@ void FLwsWebSocketsManager::Tick()
 	}
 	if (LwsContext)
 	{
-		lws_service(LwsContext, bPollService ? 0 : MAX_int32);
+		lws_service(LwsContext, bPollService ? 0 : ServiceTimeoutMs);
 	}
 	for (FLwsWebSocket* Socket : SocketsDestroyedDuringService)
 	{
diff --git a/Engine/Source/Runtime/Online/WebSockets/Private/Lws/LwsWebSocketsManager.h b/Engine/Source/Runtime/Online/WebSockets/Private/Lws/LwsWebSocketsManager.h
--- a/Engine/Source/Runtime/Online/WebSockets/Private/Lws/LwsWebSocketsManager.h
+++ b/Engine/Source/Runtime/Online/WebSockets/Private/Lws/LwsWebSocketsManager.h
@@ -118,9 +118,13 @@ private:
 	 * If true (default), our thread will explicitly sleep (according to ThreadTargetFrameTimeInSeconds & ThreadMinimumSleepTimeInSeconds) in
 	 * between calls to lws_service with 0 timeout.
 	 * 
-	 * If false, lws_service will use a MAX_int32 timeout, allowing lws to block and handle wake / sleep on activity internally itself.
+	 * If false, lws_service will use a ServiceTimoutMs timeout, allowing lws to block and handle wake / sleep on activity internally itself.
 	 */
 	bool bPollService = true;
+	
+	/** Timeout value in milliseconds to be used by lws_service when bPollService is false */
+	int32 ServiceTimeoutMs = MAX_int32;
+
 	/** Target frame time for our thread's tick */
 	double ThreadTargetFrameTimeInSeconds = 1.0f / 30.0f; // 30Hz;
 	/** Minimum time to sleep in our thread's tick, even if the sleep makes us exceed our target frame time */
