Switch BuildPatchFileConstructor.bStallWhenFileSystemThrottled CVar to false until we get file system level write throttling working and decide to use it.
Should resolve #sh-03-24-rr-switch-unable-to-download-ext.

diff --git a/Engine/Source/Runtime/Online/BuildPatchServices/Private/BuildPatchFileConstructor.cpp b/Engine/Source/Runtime/Online/BuildPatchServices/Private/BuildPatchFileConstructor.cpp
--- a/Engine/Source/Runtime/Online/BuildPatchServices/Private/BuildPatchFileConstructor.cpp
+++ b/Engine/Source/Runtime/Online/BuildPatchServices/Private/BuildPatchFileConstructor.cpp
@@ -41,7 +41,7 @@ static FAutoConsoleVariableRef CVarSleepTimeWhenFileSystemThrottledSeconds(
 	ECVF_Default);
 
 // This can be overridden by the installation parameters.
-static bool bCVarStallWhenFileSystemThrottled = true;
+static bool bCVarStallWhenFileSystemThrottled = false;
 static FAutoConsoleVariableRef CVarStallWhenFileSystemThrottled(
 	TEXT("BuildPatchFileConstructor.bStallWhenFileSystemThrottled"),
 	bCVarStallWhenFileSystemThrottled,
