Clean up threading for verifier
Add some sanity checks.

#rnx

diff --git a/Engine/Source/Runtime/Online/BuildPatchServices/Private/Installer/Verifier.cpp b/Engine/Source/Runtime/Online/BuildPatchServices/Private/Installer/Verifier.cpp
--- a/Engine/Source/Runtime/Online/BuildPatchServices/Private/Installer/Verifier.cpp
+++ b/Engine/Source/Runtime/Online/BuildPatchServices/Private/Installer/Verifier.cpp
@@ -9,8 +9,7 @@
 #include "IBuildManifestSet.h"
 #include "Misc/ConfigCacheIni.h"
 
-DECLARE_LOG_CATEGORY_EXTERN(LogVerifier, Warning, All);
-DEFINE_LOG_CATEGORY(LogVerifier);
+DEFINE_LOG_CATEGORY_STATIC(LogVerifier, Log, All);
 
 namespace BuildPatchServices
 {
@@ -81,11 +80,12 @@ namespace BuildPatchServices
 		};
 
 		FEvent* AllThreadsDoneEvent;
-		std::atomic_int32_t ThreadsCreated;
-		std::atomic_int32_t ThreadsCompleted;
+		std::atomic_int32_t ThreadsRemaining;
 		std::atomic_int64_t ThreadProcessedBytes;
-		std::atomic_int32_t ThreadPullIndex;
 		TArray<EVerifyResult> ThreadJobResults;
+
+		UE::FMutex ThreadJobListLock;
+		int32 NextThreadJobListIndex = 0;		
 		TArray<FThreadVerifyJob> ThreadJobList;
 
 		TArray<IBuildInstallerThread*> SecondaryProcessThreads;
@@ -111,7 +111,6 @@ namespace BuildPatchServices
 	void FVerifier::ProcessVerifyJobs()
 	{
 		TArray<uint8> FileReadBuffer;
-		FileReadBuffer.SetNumUninitialized(4 << 20); // 4MB chunked reads.
 
 		for (;;)
 		{
@@ -121,20 +120,25 @@ namespace BuildPatchServices
 			}
 
 			// Get a thread index.
-			int32 PossiblePullIndex = ThreadPullIndex.load();
-			if (!ThreadPullIndex.compare_exchange_strong(PossiblePullIndex, PossiblePullIndex +1))
+			int32 OurJobIndex = -1;
 			{
-				// Couldn't claim this job, try again
-				continue;
+				UE::TUniqueLock _(ThreadJobListLock);
+				if (NextThreadJobListIndex < ThreadJobList.Num())
+				{
+					OurJobIndex = NextThreadJobListIndex;
+					NextThreadJobListIndex++;
+				}
 			}
 
-			int32 OurJobIndex = PossiblePullIndex;
-			if (OurJobIndex >= ThreadJobList.Num())
+			if (OurJobIndex == -1)
 			{
-				// Ran out of jobs to do.
+				// No work to do.
 				break;
 			}
 
+			// Only allocate our buffer when we get a job
+			FileReadBuffer.SetNumUninitialized(4 << 20); // 4MB chunked reads.
+
 			FThreadVerifyJob& OurJob = ThreadJobList[OurJobIndex];
 
 			VerifierStat->OnFileStarted(OurJob.FileName, OurJob.BuildFileManifest->FileSize);
@@ -146,10 +150,13 @@ namespace BuildPatchServices
 			ThreadJobResults[OurJobIndex] = FileVerifyResult;
 		}
 
-		if ((ThreadsCompleted.fetch_add(1) + 1) == ThreadsCreated.load(std::memory_order_relaxed))
+		if (ThreadsRemaining.fetch_add(-1) == 1)
 		{
 			// We're the last one.
 			AllThreadsDoneEvent->Trigger();
+
+			// We can't touch _anything_ after here because the outer class holding the lambda could
+			// be destroyed.
 		}
 	}
 
@@ -188,7 +195,7 @@ namespace BuildPatchServices
 		ThreadJobList.Reset();
 		ThreadJobResults.Reset();
 		ThreadProcessedBytes = 0;
-		ThreadPullIndex = 0;
+		NextThreadJobListIndex = 0;
 
 		for (const FString& BuildFile : VerifyList)
 		{
@@ -201,9 +208,13 @@ namespace BuildPatchServices
 			ThreadJobResults.Add(EVerifyResult::Aborted);
 		}
 
-		ThreadsCompleted = 0;
 		AllThreadsDoneEvent = FPlatformProcess::GetSynchEventFromPool();
 
+		// Events get reset when they are returned to the pool but if there are any outstanding
+		// triggers due to races in the previous user we would have a triggered event, so we Reset here. 
+		// This doesn't really remove the race but it prevents it from sitting latent for long periods of time.
+		AllThreadsDoneEvent->Reset();
+
 		if (SharedContext)
 		{
 			// 3 cores should saturate any modern drive, and on anything that isn't we just sit in the wait
@@ -211,7 +222,7 @@ namespace BuildPatchServices
 			int32 ThreadsToCreate = 1;
 			GConfig->GetInt(TEXT("BuildPatchTool"), TEXT("VerificationThreadCount"), ThreadsToCreate, GEngineIni);
 
-			ThreadsCreated = ThreadsToCreate;
+			ThreadsRemaining = ThreadsToCreate;			
 
 			// We wait to create our threads so that the shared threads from file construction are returned.
 			for (int i = 0; i < ThreadsToCreate; i++)
@@ -240,6 +251,13 @@ namespace BuildPatchServices
 				VerifierStat->OnProcessedDataUpdated(ThreadProcessedBytes.load(std::memory_order_acquire));
 			}
 
+			// At this point we should not have any threads still doing work!
+			int32 LocalThreadsRemaining = ThreadsRemaining;
+			if (LocalThreadsRemaining != 0)
+			{
+				UE_LOG(LogVerifier, Error, TEXT("Verifier passed our Wait() call with threads remaining = %d!"), LocalThreadsRemaining);
+			}
+
 			for (int i = 0; i < ThreadsToCreate; i++)
 			{
 				SharedContext->ReleaseThread(SecondaryProcessThreads[i]);
@@ -249,7 +267,6 @@ namespace BuildPatchServices
 		else
 		{
 			// Can't create threads, just do the work here.
-			ThreadsCreated = 1;
 			ProcessVerifyJobs();
 		}
 
