Add logfiles (.log) as a text/plain mime type

diff --git a/Engine/Source/Runtime/Online/HTTP/Private/GenericPlatform/GenericPlatformHttp.cpp b/Engine/Source/Runtime/Online/HTTP/Private/GenericPlatform/GenericPlatformHttp.cpp
--- a/Engine/Source/Runtime/Online/HTTP/Private/GenericPlatform/GenericPlatformHttp.cpp
+++ b/Engine/Source/Runtime/Online/HTTP/Private/GenericPlatform/GenericPlatformHttp.cpp
@@ -568,6 +568,7 @@ FString FGenericPlatformHttp::GetMimeType(const FString& FilePath)
 		// Documents
 		ExtensionMimeTypeMap.Add(TEXT(".xml"), TEXT("application/xml"));
 		ExtensionMimeTypeMap.Add(TEXT(".txt"), TEXT("text/plain"));
+		ExtensionMimeTypeMap.Add(TEXT(".log"), TEXT("text/plain"));
 		ExtensionMimeTypeMap.Add(TEXT(".tsv"), TEXT("text/tab-separated-values"));
 		ExtensionMimeTypeMap.Add(TEXT(".csv"), TEXT("text/csv"));
 		ExtensionMimeTypeMap.Add(TEXT(".json"), TEXT("application/json"));
