Add option to disable resume for smaller BPS  installs.

#rnx

diff --git a/Engine/Source/Runtime/Online/BuildPatchServices/Private/BuildPatchFileConstructor.cpp b/Engine/Source/Runtime/Online/BuildPatchServices/Private/BuildPatchFileConstructor.cpp
--- a/Engine/Source/Runtime/Online/BuildPatchServices/Private/BuildPatchFileConstructor.cpp
+++ b/Engine/Source/Runtime/Online/BuildPatchServices/Private/BuildPatchFileConstructor.cpp
@@ -55,6 +55,12 @@ static FAutoConsoleVariableRef CVarAllowMultipleFilesInFlight(
 	TEXT("Whether to allow multiple files to be constructed at the same time, though still sequentially."),
 	ECVF_Default);
 
+static int32 CVarDisableResumeBelowMB = 0;
+static FAutoConsoleVariableRef CVarRefDisableResumeBelowMB(
+	TEXT("BuildPatchFileConstructor.DisableResumeBelowMB"),
+	CVarDisableResumeBelowMB,
+	TEXT("If nonzero, installs (not patches) below this size will not create or check any resume data."),
+	ECVF_Default);
 
 // Helper functions wrapping common code.
 namespace FileConstructorHelpers
@@ -269,7 +275,10 @@ public:
 	TSet<FString> FilesIncompatible;
 
 	// Whether we have any resume data for this install
-	bool bHasResumeData;
+	bool bHasResumeData = false;
+
+	// For small installs we may disable resume entirely to mitigate the number of file operations.
+	bool bResumeEnabled = false;
 
 public:
 
@@ -278,8 +287,14 @@ public:
 		, ManifestSet(InManifestSet)
 		, StagingDir(InStagingDir)
 		, ResumeDataFilename(InResumeDataFilename)
-		, bHasResumeData(false)
 	{
+		// Leave resume disabled until initialized;
+	}
+
+	void InitResume()
+	{
+		bResumeEnabled = true;
+
 		// Load data from previous resume file
 		bHasResumeData = FileSystem->FileExists(*ResumeDataFilename);
 		GLog->Logf(TEXT("BuildPatchResumeData file found: %s"), bHasResumeData ? TEXT("true") : TEXT("false"));
@@ -316,7 +331,10 @@ public:
 	void SaveOut(const TSet<FString>& ResumeIds)
 	{
 		// Save out the patch versions
-		FileSystem->SaveStringToFile(*ResumeDataFilename, FString::Join(ResumeIds, TEXT("\n")));
+		if (bResumeEnabled)
+		{
+			FileSystem->SaveStringToFile(*ResumeDataFilename, FString::Join(ResumeIds, TEXT("\n")));
+		}
 	}
 
 	/**
@@ -843,6 +861,7 @@ public:
 
 		if (bUseDiskOverflowStore)
 		{
+			IFileManager::Get().MakeDirectory(*InParentConstructor->Configuration.BackingStoreDirectory);
 			BackingStoreFileHandle.Reset(FPlatformFileManager::Get().GetPlatformFile().OpenWrite(*BackingStoreFileName, false, true));
 		}
 
@@ -1887,47 +1906,68 @@ void FBuildPatchFileConstructor::Run()
 
 		if (ExpectedChunkSize)
 		{
-			UE_LOG(LogBuildPatchServices, Display, TEXT("Expected chunk size: %d count %d"), ExpectedChunkSize, LargestWindowCount);
+			UE_LOG(LogBuildPatchServices, Log, TEXT("Expected chunk size: %d count %d"), ExpectedChunkSize, LargestWindowCount);
 		}
 		else
 		{
-			UE_LOG(LogBuildPatchServices, Display, TEXT("Can't find largest chunk size, using 1MB"));
+			UE_LOG(LogBuildPatchServices, Log, TEXT("Can't find largest chunk size, using 1MB"));
 			ExpectedChunkSize = 1 << 20;
 		}
 	}
 
-	// Check for resume data, we need to also look for a legacy resume file to use instead in case we are resuming from an install of previous code version.
-	const FString LegacyResumeDataFilename = Configuration.StagingDirectory / TEXT("$resumeData");
-	const FString ResumeDataFilename = Configuration.MetaDirectory / TEXT("$resumeData");
-	const bool bHasLegacyResumeData = FileSystem->FileExists(*LegacyResumeDataFilename);
-	// If we find a legacy resume data file, lets move it first.
-	if (bHasLegacyResumeData)
+	// We disable resume if we are a fresh install that's below a certain threshold in order to minimize
+	// io operations.
+	bool bResumeEnabled = true;
+	if (!Configuration.ManifestSet->ContainsUpdate())
 	{
-		FileSystem->MoveFile(*ResumeDataFilename, *LegacyResumeDataFilename);
+		int64 DisableResumeBelowBytes = ((int64)Configuration.DisableResumeBelowMB.Get(CVarDisableResumeBelowMB)) << 20;
+		if (DisableResumeBelowBytes > TotalJobSize)
+		{
+			UE_LOG(LogBuildPatchServices, Log, TEXT("Disabling resume: JobSize = %s, Disable = %s, from %s"), *FormatNumber(TotalJobSize), *FormatNumber(DisableResumeBelowBytes), Configuration.DisableResumeBelowMB.IsSet() ? TEXT("config") : TEXT("cvar"));
+			bResumeEnabled = false;
+		}
 	}
+
+	const FString ResumeDataFilename = Configuration.MetaDirectory / TEXT("$resumeData");
 	FResumeData ResumeData(FileSystem, Configuration.ManifestSet, Configuration.StagingDirectory, ResumeDataFilename);
 
-	// Remove incompatible files
-	if (ResumeData.bHasResumeData)
+	if (bResumeEnabled)
 	{
-		for (const FString& FileToConstruct : Configuration.ConstructList)
+		IFileManager::Get().MakeDirectory(*Configuration.MetaDirectory);
+
+		// Check for resume data, we need to also look for a legacy resume file to use instead in case we are resuming from an install of previous code version.
+		const FString LegacyResumeDataFilename = Configuration.StagingDirectory / TEXT("$resumeData");		
+		const bool bHasLegacyResumeData = FileSystem->FileExists(*LegacyResumeDataFilename);
+		// If we find a legacy resume data file, lets move it first.
+		if (bHasLegacyResumeData)
 		{
-			ResumeData.CheckFile(FileToConstruct);
-			const bool bFileIncompatible = ResumeData.FilesIncompatible.Contains(FileToConstruct);
-			if (bFileIncompatible)
+			FileSystem->MoveFile(*ResumeDataFilename, *LegacyResumeDataFilename);
+		}
+
+		ResumeData.InitResume();
+
+		// Remove incompatible files
+		if (ResumeData.bHasResumeData)
+		{
+			for (const FString& FileToConstruct : Configuration.ConstructList)
 			{
-				GLog->Logf(TEXT("FBuildPatchFileConstructor: Deleting incompatible stage file %s"), *FileToConstruct);
-				FileSystem->DeleteFile(*(Configuration.StagingDirectory / FileToConstruct));
+				ResumeData.CheckFile(FileToConstruct);
+				const bool bFileIncompatible = ResumeData.FilesIncompatible.Contains(FileToConstruct);
+				if (bFileIncompatible)
+				{
+					GLog->Logf(TEXT("FBuildPatchFileConstructor: Deleting incompatible stage file %s"), *FileToConstruct);
+					FileSystem->DeleteFile(*(Configuration.StagingDirectory / FileToConstruct));
+				}
 			}
 		}
-	}
 
-	// Save for started versions
-	TSet<FString> ResumeIds;
-	const bool bCheckLegacyIds = false;
+		// Save for started versions
+		TSet<FString> ResumeIds;
+		const bool bCheckLegacyIds = false;
 
-	Configuration.ManifestSet->GetInstallResumeIds(ResumeIds, bCheckLegacyIds);
-	ResumeData.SaveOut(ResumeIds);
+		Configuration.ManifestSet->GetInstallResumeIds(ResumeIds, bCheckLegacyIds);
+		ResumeData.SaveOut(ResumeIds);
+	}
 
 	// Start resume progress at zero or one.
 	FileConstructorStat->OnResumeStarted();
diff --git a/Engine/Source/Runtime/Online/BuildPatchServices/Private/BuildPatchFileConstructor.h b/Engine/Source/Runtime/Online/BuildPatchServices/Private/BuildPatchFileConstructor.h
--- a/Engine/Source/Runtime/Online/BuildPatchServices/Private/BuildPatchFileConstructor.h
+++ b/Engine/Source/Runtime/Online/BuildPatchServices/Private/BuildPatchFileConstructor.h
@@ -90,6 +90,7 @@ namespace BuildPatchServices
 		TOptional<int32> IOBatchSizeMB;
 		TOptional<int32> IOBufferSizeMB;
 		TOptional<bool> StallWhenFileSystemThrottled;
+		TOptional<int32> DisableResumeBelowMB;
 	};
 
 	/**
diff --git a/Engine/Source/Runtime/Online/BuildPatchServices/Private/BuildPatchInstaller.cpp b/Engine/Source/Runtime/Online/BuildPatchServices/Private/BuildPatchInstaller.cpp
--- a/Engine/Source/Runtime/Online/BuildPatchServices/Private/BuildPatchInstaller.cpp
+++ b/Engine/Source/Runtime/Online/BuildPatchServices/Private/BuildPatchInstaller.cpp
@@ -1233,10 +1233,13 @@ namespace BuildPatchServices
 		FPaths::NormalizeDirectoryName(InstallStagingDir);
 		FPaths::NormalizeDirectoryName(MetaStagingDir);
 
-		// Make sure staging directories exist
-		IFileManager::Get().MakeDirectory(*DataStagingDir, true);
+		// Make sure staging directories exist		
 		IFileManager::Get().MakeDirectory(*InstallStagingDir, true);
-		IFileManager::Get().MakeDirectory(*MetaStagingDir, true);
+
+		// These we don't make here - we make them only if we need them because for basic small installs
+		// on some platforms we want to minimize io operations.
+		//IFileManager::Get().MakeDirectory(*DataStagingDir, true);
+		//IFileManager::Get().MakeDirectory(*MetaStagingDir, true);
 
 		// Reset our error and build progress
 		InstallerError.Reset(FInstallerErrorFactory::Create());
@@ -1415,6 +1418,7 @@ namespace BuildPatchServices
 			FCC.IOBufferSizeMB = Configuration.ConstructorIOBufferSizeMB;
 			FCC.IOBatchSizeMB = Configuration.ConstructorIOBatchSizeMB;
 			FCC.StallWhenFileSystemThrottled = Configuration.ConstructorStallWhenFileSystemThrottled;
+			FCC.DisableResumeBelowMB = Configuration.ConstructorDisableResumeBelowMB;
 
 			TUniquePtr<FBuildPatchFileConstructor> FileConstructor(
 				new FBuildPatchFileConstructor(
diff --git a/Engine/Source/Runtime/Online/BuildPatchServices/Public/BuildPatchSettings.h b/Engine/Source/Runtime/Online/BuildPatchServices/Public/BuildPatchSettings.h
--- a/Engine/Source/Runtime/Online/BuildPatchServices/Public/BuildPatchSettings.h
+++ b/Engine/Source/Runtime/Online/BuildPatchServices/Public/BuildPatchSettings.h
@@ -311,6 +311,12 @@ namespace BuildPatchServices
 		// CVar BuildPatchFileConstructor.bStallWhenFileSystemThrottled, which defaults to true.
 		// if set, this member.
 		TOptional<bool> ConstructorStallWhenFileSystemThrottled;
+
+		// Whether to disable resume on clean installs below a given threshold, in megabytes. 0 means "never disable".
+		// Inheritance:
+		// CVar BuildPatchFileConstructor.DisableResumeBelowMB, which defaults to 0.
+		// if set, this member
+		TOptional<int32> ConstructorDisableResumeBelowMB;
 	};
 
 	/**
