Limit the number of concurrent HTTP requests by default. This helps avoid connection timeouts when queueing up very large number of HTTP requests

diff --git a/Engine/Source/Runtime/Online/HTTP/HTTP.Build.cs b/Engine/Source/Runtime/Online/HTTP/HTTP.Build.cs
--- a/Engine/Source/Runtime/Online/HTTP/HTTP.Build.cs
+++ b/Engine/Source/Runtime/Online/HTTP/HTTP.Build.cs
@@ -45,7 +45,7 @@ public class HTTP : ModuleRules
 		} 
 	}
 
-	protected virtual int DefaultMaxConcurrentRequests { get { return int.MaxValue; } }
+	protected virtual int DefaultMaxConcurrentRequests { get { return 256; } }
 
 	private bool bPlatformSupportsCurl { get { return bPlatformSupportsLibCurl || bPlatformSupportsXCurl; } }
 
