Fix invalid monolithic check for SSL initialization guards under merged modules

#rnx

diff --git a/Engine/Source/Runtime/Online/SSL/Private/SslManager.cpp b/Engine/Source/Runtime/Online/SSL/Private/SslManager.cpp
--- a/Engine/Source/Runtime/Online/SSL/Private/SslManager.cpp
+++ b/Engine/Source/Runtime/Online/SSL/Private/SslManager.cpp
@@ -30,7 +30,7 @@ bool FSslManager::InitializeSsl()
 	// while we are statically linking OpenSSL in various libraries (such as libcurl), when the
 	// SSL module is non-monolithic, OpenSSL would only get initialized in the SSL module's 
 	// scope (not the caller)
-#if IS_MONOLITHIC
+#if IS_MONOLITHIC || UE_MERGED_MODULES
 	++ InitCount;
 	if (InitCount == 1)
 	{
@@ -50,13 +50,13 @@ bool FSslManager::InitializeSsl()
 	{
 		bSuccessful = true;
 	}
-#endif
+#endif // IS_MONOLITHIC || UE_MERGED_MODULES
 	return bSuccessful;
 }
 
 void FSslManager::ShutdownSsl()
 {
-#if IS_MONOLITHIC
+#if IS_MONOLITHIC || UE_MERGED_MODULES
 	-- InitCount;
 	check(InitCount >= 0);
 	if (InitCount == 0)
@@ -67,16 +67,16 @@ void FSslManager::ShutdownSsl()
 		ERR_free_strings();
 #if OPENSSL_VERSION_NUMBER < 0x10100000L
 		ERR_remove_thread_state(NULL);
-#endif
+#endif // OPENSSL_VERSION_NUMBER < 0x10100000L
 		CONF_modules_free();
 	}
-#endif
+#endif // IS_MONOLITHIC || UE_MERGED_MODULES
 }
 
 SSL_CTX* FSslManager::CreateSslContext(const FSslContextCreateOptions& CreateOptions)
 {
 	SSL_CTX* SslContext = nullptr;
-#if IS_MONOLITHIC
+#if IS_MONOLITHIC || UE_MERGED_MODULES
 	check(InitCount > 0);
 	if (InitCount > 0)
 	{
@@ -126,19 +126,19 @@ SSL_CTX* FSslManager::CreateSslContext(const FSslContextCreateOptions& CreateOpt
 			UE_LOG(LogSsl, Warning, TEXT("FSslManager::CreateSslContext: Failed to create method SSLv23_client_method: %s"), *GetSslErrorString());
 		}
 	}
-#endif
+#endif // IS_MONOLITHIC || UE_MERGED_MODULES
 	return SslContext;
 }
 
 void FSslManager::DestroySslContext(SSL_CTX* SslContext)
 {
-#if IS_MONOLITHIC
+#if IS_MONOLITHIC || UE_MERGED_MODULES
 	check(InitCount > 0);
 	if (InitCount > 0)
 	{
 		SSL_CTX_free(SslContext);
 	}
-#endif
+#endif // IS_MONOLITHIC || UE_MERGED_MODULES
 }
 
 #endif
