Add a version of FHttpServerResponse::Create that takes a UTF8 string.  This means the response string can be used without conversion.

diff --git a/Engine/Source/Runtime/Online/HTTPServer/Private/HttpServerResponse.cpp b/Engine/Source/Runtime/Online/HTTPServer/Private/HttpServerResponse.cpp
--- a/Engine/Source/Runtime/Online/HTTPServer/Private/HttpServerResponse.cpp
+++ b/Engine/Source/Runtime/Online/HTTPServer/Private/HttpServerResponse.cpp
@@ -1,10 +1,12 @@
 // Copyright Epic Games, Inc. All Rights Reserved.
 
 #include "HttpServerResponse.h"
+
+#include "Containers/Utf8String.h"
 #include "HttpServerConstants.h"
 #include "HttpServerConstantsPrivate.h"
 
-TUniquePtr<FHttpServerResponse> FHttpServerResponse::Create(const FString& Text, FString ContentType)
+TUniquePtr<FHttpServerResponse> FHttpServerResponse::Create(const FString& Text, const FString& ContentType)
 {
 	TUniquePtr<FHttpServerResponse> Response = MakeUnique<FHttpServerResponse>();
 	Response->Code = EHttpServerResponseCodes::Ok;
@@ -20,6 +22,42 @@ TUniquePtr<FHttpServerResponse> FHttpServerResponse::Create(const FString& Text,
 	return Response;
 }
 
+TUniquePtr<FHttpServerResponse> FHttpServerResponse::Create(const FUtf8String& Text, const FString& ContentType)
+{
+	TUniquePtr<FHttpServerResponse> Response = MakeUnique<FHttpServerResponse>();
+	Response->Code = EHttpServerResponseCodes::Ok;
+
+	const uint8* Utf8Bytes = (reinterpret_cast<const uint8*>(Text.GetCharArray().GetData()));
+	Response->Body.Append(Utf8Bytes, Text.Len()); // exclude null terminator
+
+	FString Utf8CharsetContentType = FString::Printf(TEXT("%s;charset=utf-8"), *ContentType);
+	TArray<FString> ContentTypeValue = { MoveTemp(Utf8CharsetContentType) };
+	Response->Headers.Add(UE_HTTP_SERVER_HEADER_KEYS_CONTENT_TYPE, MoveTemp(ContentTypeValue));
+
+	return Response;
+}
+
+TUniquePtr<FHttpServerResponse> FHttpServerResponse::Create(FUtf8String&& Text, const FString& ContentType)
+{
+	TUniquePtr<FHttpServerResponse> Response = MakeUnique<FHttpServerResponse>();
+	Response->Code = EHttpServerResponseCodes::Ok;
+
+	Response->Body = TArray<uint8>(MoveTemp(Text.GetCharArray()));
+	Text.Empty();
+
+	// exclude null terminator
+	if (!Response->Body.IsEmpty() && Response->Body.Last() == 0)
+	{
+		Response->Body.Pop(EAllowShrinking::No);
+	}
+
+	FString Utf8CharsetContentType = FString::Printf(TEXT("%s;charset=utf-8"), *ContentType);
+	TArray<FString> ContentTypeValue = { MoveTemp(Utf8CharsetContentType) };
+	Response->Headers.Add(UE_HTTP_SERVER_HEADER_KEYS_CONTENT_TYPE, MoveTemp(ContentTypeValue));
+
+	return Response;
+}
+
 TUniquePtr<FHttpServerResponse> FHttpServerResponse::Create(TArray<uint8>&& RawBytes, FString ContentType)
 {
 	TUniquePtr<FHttpServerResponse> Response = MakeUnique<FHttpServerResponse>(MoveTemp(RawBytes));
diff --git a/Engine/Source/Runtime/Online/HTTPServer/Public/HttpServerResponse.h b/Engine/Source/Runtime/Online/HTTPServer/Public/HttpServerResponse.h
--- a/Engine/Source/Runtime/Online/HTTPServer/Public/HttpServerResponse.h
+++ b/Engine/Source/Runtime/Online/HTTPServer/Public/HttpServerResponse.h
@@ -11,6 +11,8 @@
 #include "Templates/UniquePtr.h"
 #include "Templates/UnrealTemplate.h"
 
+class FUtf8String;
+
 struct FHttpServerResponse final
 {
 public:
@@ -51,7 +53,25 @@ public:
 	 * @param  ContentType  The HTTP response content type
 	 * @return              A unique pointer to an initialized response object
 	 */
-	HTTPSERVER_API static TUniquePtr<FHttpServerResponse> Create(const FString& Text, FString ContentType);
+	HTTPSERVER_API static TUniquePtr<FHttpServerResponse> Create(const FString& Text, const FString& ContentType);
+
+	/**
+	 * Creates an FHttpServerResponse from a string
+	 *
+	 * @param  Text         The text to serialize
+	 * @param  ContentType  The HTTP response content type
+	 * @return              A unique pointer to an initialized response object
+	 */
+	HTTPSERVER_API static TUniquePtr<FHttpServerResponse> Create(const FUtf8String& Text, const FString& ContentType);
+
+	/**
+	 * Creates an FHttpServerResponse from a string
+	 *
+	 * @param  Text         The text to serialize
+	 * @param  ContentType  The HTTP response content type
+	 * @return              A unique pointer to an initialized response object
+	 */
+	HTTPSERVER_API static TUniquePtr<FHttpServerResponse> Create(FUtf8String&& Text, const FString& ContentType);
 
 	/**
 	 * Creates an FHttpServerResponse from a raw byte buffer
