Add option to enable different http versions for http request.

diff --git a/Engine/Source/Programs/WebTests/Private/TestHttp.cpp b/Engine/Source/Programs/WebTests/Private/TestHttp.cpp
--- a/Engine/Source/Programs/WebTests/Private/TestHttp.cpp
+++ b/Engine/Source/Programs/WebTests/Private/TestHttp.cpp
@@ -2988,3 +2988,16 @@ TEST_CASE_METHOD(FWaitUntilCompleteHttpFixture, "Http Methods over Unix Domain S
 }
 
 #endif //UE_HTTP_SUPPORT_UNIX_SOCKET
+
+TEST_CASE_METHOD(FWaitUntilCompleteHttpFixture, "Http Methods over http2", HTTP_TAG)
+{
+	TSharedRef<IHttpRequest> HttpRequest = CreateRequest();
+	HttpRequest->SetURL(UrlToTestMethods());
+	HttpRequest->SetOption(HttpRequestOptions::HttpVersion, TEXT("2"));
+	HttpRequest->OnProcessRequestComplete().BindLambda([](FHttpRequestPtr HttpRequest, FHttpResponsePtr HttpResponse, bool bSucceeded) {
+		CHECK(bSucceeded);
+		REQUIRE(HttpResponse != nullptr);
+		CHECK(HttpResponse->GetResponseCode() == 200);
+	});
+	HttpRequest->ProcessRequest();
+}
diff --git a/Engine/Source/Runtime/Online/HTTP/Private/Curl/CurlHttp.cpp b/Engine/Source/Runtime/Online/HTTP/Private/Curl/CurlHttp.cpp
--- a/Engine/Source/Runtime/Online/HTTP/Private/Curl/CurlHttp.cpp
+++ b/Engine/Source/Runtime/Online/HTTP/Private/Curl/CurlHttp.cpp
@@ -869,6 +869,23 @@ bool FCurlHttpRequest::SetupRequestHttpThread()
 		}
 #endif //UE_HTTP_SUPPORT_UNIX_SOCKET
 
+		const FString HttpVersion = GetOption(HttpRequestOptions::HttpVersion);
+		if (!HttpVersion.IsEmpty())
+		{
+			if (HttpVersion == TEXT("2"))
+			{
+				curl_easy_setopt(EasyHandle, CURLOPT_HTTP_VERSION, CURL_HTTP_VERSION_2TLS);
+			}
+			else if (HttpVersion == TEXT("1.1"))
+			{
+				curl_easy_setopt(EasyHandle, CURLOPT_HTTP_VERSION, CURL_HTTP_VERSION_1_1);
+			}
+			else
+			{
+				UE_LOG(LogHttp, Warning, TEXT("Unsupported http version %s"), *HttpVersion);
+			}
+		}
+
 		if (!FCurlHttpManager::CurlRequestOptions.LocalHostAddr.IsEmpty())
 		{
 			// Set the local address to use for making these requests
diff --git a/Engine/Source/Runtime/Online/HTTP/Public/Interfaces/IHttpRequest.h b/Engine/Source/Runtime/Online/HTTP/Public/Interfaces/IHttpRequest.h
--- a/Engine/Source/Runtime/Online/HTTP/Public/Interfaces/IHttpRequest.h
+++ b/Engine/Source/Runtime/Online/HTTP/Public/Interfaces/IHttpRequest.h
@@ -130,6 +130,7 @@ private:
  */
 namespace HttpRequestOptions
 {
+	static const FName HttpVersion("HttpVersion");
 #if UE_HTTP_SUPPORT_UNIX_SOCKET
 	static const FName UnixSocketPath("UnixSocketPath");
 #endif //UE_HTTP_SUPPORT_UNIX_SOCKET
