Add define UE_WEBSOCETS_PLATFORM_SUPPORT_EVENT_LOOP to decide if libWebSockets support event loop or not on different platforms.

diff --git a/Engine/Source/Runtime/Online/WebSockets/Private/Lws/LwsWebSocketsManager.cpp b/Engine/Source/Runtime/Online/WebSockets/Private/Lws/LwsWebSocketsManager.cpp
--- a/Engine/Source/Runtime/Online/WebSockets/Private/Lws/LwsWebSocketsManager.cpp
+++ b/Engine/Source/Runtime/Online/WebSockets/Private/Lws/LwsWebSocketsManager.cpp
@@ -90,6 +90,15 @@ void FLwsWebSocketsManager::UpdateConfigs()
 	FParse::Bool(FCommandLine::Get(), TEXT("lws-insecure="), bDisableCertValidation);
 #endif
 
+	// Check if it's supported on platform
+#if !UE_WEBSOCKETS_PLATFORM_SUPPORT_EVENT_LOOP
+	if (!bPollService)
+	{
+		UE_LOG(LogWebSockets, Warning, TEXT("FLwsWebSocketsManager::UpdateConfigs: Service polling disabled by config while it's not supported on platform. Forcing service polling to avoid blocking service call"));
+		bPollService = true;
+	}
+#endif
+
 	// Check blocking service calls are supported if requested
 	if (!bPollService)
 	{
diff --git a/Engine/Source/Runtime/Online/WebSockets/WebSockets.Build.cs b/Engine/Source/Runtime/Online/WebSockets/WebSockets.Build.cs
--- a/Engine/Source/Runtime/Online/WebSockets/WebSockets.Build.cs
+++ b/Engine/Source/Runtime/Online/WebSockets/WebSockets.Build.cs
@@ -17,6 +17,16 @@ public class WebSockets : ModuleRules
 		}
 	}
 
+	protected virtual bool bPlatformSupportEventLoop
+	{
+		get
+		{
+			return
+				Target.Platform == UnrealTargetPlatform.Win64 ||
+				Target.IsInPlatformGroup(UnrealPlatformGroup.Unix);
+		}
+	}
+
 	protected virtual bool bPlatformSupportsWinHttpWebSockets
 	{
 		get
@@ -123,6 +133,8 @@ public class WebSockets : ModuleRules
 			}
 		}
 
+		PrivateDefinitions.Add("UE_WEBSOCKETS_PLATFORM_SUPPORT_EVENT_LOOP=" + (bPlatformSupportEventLoop ? "1" : "0"));
+
 		PublicDefinitions.Add("WEBSOCKETS_PACKAGE=1");
 		PublicDefinitions.Add("WITH_WEBSOCKETS=" + (bWithWebSockets ? "1" : "0"));
 		PublicDefinitions.Add("WITH_LIBWEBSOCKETS=" + (bWithLibWebSockets ? "1" : "0"));
