BPS: Adding InstallToMemory and ConstructInMemory
BPS: Fixing a bug where cloud chunks would always route through the memory store.

diff --git a/Engine/Source/Runtime/Online/BuildPatchServices/Private/BuildPatchFileConstructor.cpp b/Engine/Source/Runtime/Online/BuildPatchServices/Private/BuildPatchFileConstructor.cpp
--- a/Engine/Source/Runtime/Online/BuildPatchServices/Private/BuildPatchFileConstructor.cpp
+++ b/Engine/Source/Runtime/Online/BuildPatchServices/Private/BuildPatchFileConstructor.cpp
@@ -222,8 +222,12 @@ struct FFileConstructionState
 	// Can be nonzero in the first batch due to resume.
 	int32 NextChunkPartToRead = 0;
 
+	// This is null if we are constructing/install in memory.
 	TUniquePtr<FArchive> NewFile;
 
+	// Which index in the chunk reference tracker this file starts at.
+	int32 BaseReferenceIndex = 0;
+
 	int32 ConstructionIndex = -1;
 	const FFileManifest* FileManifest;
 	const FString& BuildFilename; // name as it is in the manifest, references the ConstructList in the configuration.
@@ -240,6 +244,33 @@ struct FFileConstructionState
 			ConstructionError = EConstructionError::MissingFileInfo;
 		}
 	}
+
+	void SetAttributes()
+	{
+#if PLATFORM_MAC
+		if (bSuccess && 
+			EnumHasAllFlags(FileManifest->FileMetaFlags, EFileMetaFlags::UnixExecutable))
+		{
+			// Enable executable permission bit
+			struct stat FileInfo;
+			if (stat(TCHAR_TO_UTF8(*NewFilename), &FileInfo) == 0)
+			{
+				bSuccess = chmod(TCHAR_TO_UTF8(*NewFilename), FileInfo.st_mode | S_IXUSR | S_IXGRP | S_IXOTH) == 0;
+				if (!bSuccess)
+				{
+					UE_LOG(LogBuildPatchServices, Error, TEXT("FBuildPatchFileConstructor: Failed to set exec bit %s"), *BuildFilename);
+				}
+			}
+		}
+#endif
+
+#if PLATFORM_ANDROID
+		if (bSuccess)
+		{
+			IFileManager::Get().SetTimeStamp(*NewFilename, FDateTime::UtcNow());
+		}
+#endif
+	}
 };
 
 
@@ -832,6 +863,12 @@ public:
 
 		bool bUseDiskOverflowStore = true;
 		GConfig->GetBool(TEXT("BuildPatchServices"), TEXT("bEnableDiskOverflowStore"), bUseDiskOverflowStore, GEngineIni);
+
+		if (ParentConstructor->Configuration.bInstallToMemory)
+		{
+			UE_LOG(LogBuildPatchServices, Display, TEXT("Disabling backing store due to InstallToMemory"));
+			bUseDiskOverflowStore = false;
+		}
 		
 		// Is there a hard limit on how much disk space we can use?
 		// negative = no
@@ -868,7 +905,7 @@ public:
 		if (!BackingStoreFileHandle)
 		{
 			// Prevent any pageouts.
-			UE_LOG(LogBuildPatchServices, Warning, TEXT("Unable to open disk backing store at %s"), *BackingStoreFileName);
+			UE_CLOG(bUseDiskOverflowStore, LogBuildPatchServices, Warning, TEXT("Unable to open disk backing store at %s"), *BackingStoreFileName);
 			UE_LOG(LogBuildPatchServices, Warning, TEXT("Disk backing store will be disabled"));
 			MaxDiskSpaceBytes = 0;
 		}
@@ -1000,13 +1037,15 @@ public:
 
 	~FChunkBackingStore()
 	{
-		BackingStoreFileHandle.Reset(); // have to close the file before we try to delete it.
-		if (!FPlatformFileManager::Get().GetPlatformFile().DeleteFile(*BackingStoreFileName))
+		if (BackingStoreFileHandle.IsValid())
 		{
-			UE_LOG(LogBuildPatchServices, Warning, TEXT("Unable to delete disk backing store: %s"), *BackingStoreFileName);
+			BackingStoreFileHandle.Reset(); // have to close the file before we try to delete it.
+			if (!FPlatformFileManager::Get().GetPlatformFile().DeleteFile(*BackingStoreFileName))
+			{
+				UE_LOG(LogBuildPatchServices, Warning, TEXT("Unable to delete disk backing store: %s"), *BackingStoreFileName);
+			}
+			UE_LOG(LogBuildPatchServices, Verbose, TEXT("BackingStore Reads: %d Writes: %d"), ReadCount, WriteCount);
 		}
-
-		UE_LOG(LogBuildPatchServices, Verbose, TEXT("BackingStore Reads: %d Writes: %d"), ReadCount, WriteCount);
 	}
 
 	// Set the amount of disk space the installation needs to we can ensure that we don't 
@@ -1531,8 +1570,11 @@ FBuildPatchFileConstructor::FBuildPatchFileConstructor(
 	int32 ThreadCount = 0;
 	if (bSpawnAdditionalIOThreads)
 	{
-		WriteThreadIndex = 0;
-		ThreadCount++;
+		if (!Configuration.bInstallToMemory && !Configuration.bConstructInMemory)
+		{
+			WriteThreadIndex = 0;
+			ThreadCount++;
+		}
 
 		if (bHasInstallSource)
 		{
@@ -1928,6 +1970,12 @@ void FBuildPatchFileConstructor::Run()
 		}
 	}
 
+	if (Configuration.bInstallToMemory)
+	{
+		UE_LOG(LogBuildPatchServices, Log, TEXT("Disabling resume: install to memory"));
+		bResumeEnabled = false;
+	}
+
 	const FString ResumeDataFilename = Configuration.MetaDirectory / TEXT("$resumeData");
 	FResumeData ResumeData(FileSystem, Configuration.ManifestSet, Configuration.StagingDirectory, ResumeDataFilename);
 
@@ -2162,6 +2210,8 @@ struct FRequestInfo
 	bool bFailedToRead = false;
 	bool bLaunchedFallback = false;
 
+	// Which index in the chunk reference tracker we are requesting.
+	int32 ChunkIndexInFile;
 	int32 ChunkUnavailableAt = TNumericLimits<int32>::Max();
 };
 
@@ -2485,6 +2535,7 @@ void FBuildPatchFileConstructor::StartReadBatch(FFileConstructionState& CurrentF
 			Request.Batch = &Batch;
 			Request.ChunkSize = Configuration.ManifestSet->GetChunkInfo(ChunkPart.Guid)->WindowSize;
 			Request.File = &CurrentFile;
+			Request.ChunkIndexInFile = EndChunkIdx;
 			
 			FRequestInfo::FRequestSplat& Splat = Request.Splats.AddDefaulted_GetRef();
 			Splat.DestinationOffset = BufferFillLevel;
@@ -2590,7 +2641,12 @@ void FBuildPatchFileConstructor::StartReadBatch(FFileConstructionState& CurrentF
 			Request.ChunkUnavailableAt = ThisChunkSource->GetChunkUnavailableAt(Request.Guid);
 				
 			const bool bNeedsEntireChunk = Request.Splats.Num() == 1 && Request.Splats[0].BytesToCopy == Request.ChunkSize && Request.Splats[0].OffsetInChunk == 0;
-			const bool bNeededAfterRetirement = LastUsageIndex >= Request.ChunkUnavailableAt;
+
+			int32 ThisUsageIndex = CurrentFile.BaseReferenceIndex + Request.ChunkIndexInFile;
+
+			// Unavailable might be reported before the current usage which would force this to route through
+			// memory store even though it could read direct, so only check needed after retirement if we need it again.
+			const bool bNeededAfterRetirement = (LastUsageIndex > ThisUsageIndex) && LastUsageIndex >= Request.ChunkUnavailableAt;
 
 			if (bNeedsEntireChunk && !bNeededAfterRetirement)
 			{
@@ -2671,7 +2727,11 @@ void FBuildPatchFileConstructor::InitFile(FFileConstructionState& CurrentFile, c
 		!CurrentFile.FileManifest->SymlinkTarget.IsEmpty())
 	{
 #if PLATFORM_MAC
-		CurrentFile.bSuccess = symlink(TCHAR_TO_UTF8(*CurrentFile.FileManifest->SymlinkTarget), TCHAR_TO_UTF8(*CurrentFile.NewFilename)) == 0;
+		CurrentFile.bSuccess = true;
+		if (!Configuration.bInstallToMemory && !Configuration.bConstructInMemory)
+		{
+			CurrentFile.bSuccess = symlink(TCHAR_TO_UTF8(*CurrentFile.FileManifest->SymlinkTarget), TCHAR_TO_UTF8(*CurrentFile.NewFilename)) == 0;
+		}
 		CurrentFile.bSkippedConstruction = true;
 #else
 		const bool bSymlinkNotImplemented = false;
@@ -2679,10 +2739,26 @@ void FBuildPatchFileConstructor::InitFile(FFileConstructionState& CurrentFile, c
 		CurrentFile.bSuccess = false;
 #endif
 	}
+
+	if (Configuration.bInstallToMemory || Configuration.bConstructInMemory)
+	{
+		// Allocate the destination buffer. This is most likely where we would run out of memory, but UE doesn't
+		// really have any way to gracefully handle OOM the way we might handle running out of disk space.
+		TArray64<uint8> OutputFile;
+		OutputFile.SetNumUninitialized(CurrentFile.FileManifest->FileSize);
+
+		MemoryOutputFiles.Add(CurrentFile.BuildFilename, MoveTemp(OutputFile));
+	}
 }
 
 bool FBuildPatchFileConstructor::HandleInitialDiskSizeCheck(const FFileManifest& FileManifest, int64 StartPosition)
 {
+	if (Configuration.bInstallToMemory)
+	{
+		bInitialDiskSizeCheck = true;
+		return true;
+	}
+
 	// Returns false if we failed the disk space check.
 	if (!bInitialDiskSizeCheck)
 	{
@@ -2740,6 +2816,9 @@ void FBuildPatchFileConstructor::ResumeFile(FFileConstructionState& FileToResume
 		return;
 	}
 
+	check(!Configuration.bInstallToMemory);
+	check(!Configuration.bConstructInMemory);
+
 	// We have to read in the existing file so that the hash check can still be done.
 	TUniquePtr<FArchive> NewFileReader(IFileManager::Get().CreateFileReader(*FileToResume.NewFilename));
 	if (!NewFileReader.IsValid())
@@ -2807,7 +2886,7 @@ void FBuildPatchFileConstructor::ResumeFile(FFileConstructionState& FileToResume
 
 void FBuildPatchFileConstructor::OpenFileToConstruct(FFileConstructionState& CurrentFile)
 {
-	if (!CurrentFile.bSuccess || CurrentFile.bSkippedConstruction)
+	if (!CurrentFile.bSuccess || CurrentFile.bSkippedConstruction || Configuration.bInstallToMemory || Configuration.bConstructInMemory)
 	{
 		return;
 	}
@@ -2847,30 +2926,31 @@ void FBuildPatchFileConstructor::CompleteConstructedFile(FFileConstructionState&
 	TRACE_CPUPROFILER_EVENT_SCOPE(Constructor_CompleteFile);
 	if (!CurrentFile.bSkippedConstruction)
 	{
-		if (CurrentFile.bSuccess && CurrentFile.NewFile->IsError())
+		if (CurrentFile.NewFile) // we don't have a file when constructing/installing in memory.
 		{
-			CurrentFile.ConstructionError = EConstructionError::SerializeError;
-			CurrentFile.bSuccess = false;
-		}
+			if (CurrentFile.bSuccess && CurrentFile.NewFile->IsError())
+			{
+				CurrentFile.ConstructionError = EConstructionError::SerializeError;
+				CurrentFile.bSuccess = false;
+			}
 
-		// Close the file writer
-		bool bArchiveSuccess = false;
-		if (CurrentFile.NewFile)
-		{
-			FAdministrationScope _(FileConstructorStat);
-			bArchiveSuccess = CurrentFile.NewFile->Close();
-			CurrentFile.NewFile.Reset();
-		}
-		else
-		{
-			bArchiveSuccess = true;
-		}
+			if (CurrentFile.bSuccess)
+			{
+				// Close the file writer
+				bool bArchiveSuccess = false;
+				{
+					FAdministrationScope _(FileConstructorStat);
+					bArchiveSuccess = CurrentFile.NewFile->Close();
+					CurrentFile.NewFile.Reset();
+				}
 
-		// Check for final success
-		if (CurrentFile.bSuccess && !bArchiveSuccess)
-		{
-			CurrentFile.ConstructionError = EConstructionError::SerializeError;
-			CurrentFile.bSuccess = false;
+				// Check for final success
+				if (CurrentFile.bSuccess && !bArchiveSuccess)
+				{
+					CurrentFile.ConstructionError = EConstructionError::SerializeError;
+					CurrentFile.bSuccess = false;
+				}
+			}
 		}
 
 		// We can't check for zero locks if we have multiple files in flight because the other
@@ -2898,39 +2978,17 @@ void FBuildPatchFileConstructor::CompleteConstructedFile(FFileConstructionState&
 			}
 		}
 
-#if PLATFORM_MAC
-		if (CurrentFile.bSuccess && 
-			EnumHasAllFlags(CurrentFile.FileManifest->FileMetaFlags, EFileMetaFlags::UnixExecutable))
+		if (!Configuration.bInstallToMemory && !Configuration.bConstructInMemory)
 		{
-			// Enable executable permission bit
-			struct stat FileInfo;
-			if (stat(TCHAR_TO_UTF8(*CurrentFile.NewFilename), &FileInfo) == 0)
-			{
-				CurrentFile.bSuccess = chmod(TCHAR_TO_UTF8(*CurrentFile.NewFilename), FileInfo.st_mode | S_IXUSR | S_IXGRP | S_IXOTH) == 0;
-				if (!CurrentFile.bSuccess)
-				{
-					UE_LOG(LogBuildPatchServices, Error, TEXT("FBuildPatchFileConstructor: Failed to set exec bit %s"), *CurrentFile.BuildFilename);
-				}
-			}
+			CurrentFile.SetAttributes();
 		}
-#endif
 
-#if PLATFORM_ANDROID
-		if (CurrentFile.bSuccess)
-		{
-			IFileManager::Get().SetTimeStamp(*CurrentFile.NewFilename, FDateTime::UtcNow());
-		}
-#endif
 	} // end if we did actual construction work for this file.
 
-	if (CurrentFile.bSuccess)
-	{
-		ChunkDbSource->ReportFileCompletion(ChunkReferenceTracker->GetRemainingChunkCount());
-	}
-
 	// If we are destructive, remove the old file.
 	if (CurrentFile.bSuccess &&
-		Configuration.InstallMode == EInstallMode::DestructiveInstall)
+		Configuration.InstallMode == EInstallMode::DestructiveInstall &&
+		!Configuration.bInstallToMemory) // Install to memory never patches/has an existing installation dir.
 	{
 		FString FileToDelete = Configuration.InstallDirectory / CurrentFile.BuildFilename;
 		FPaths::NormalizeFilename(FileToDelete);
@@ -2955,12 +3013,102 @@ void FBuildPatchFileConstructor::CompleteConstructedFile(FFileConstructionState&
 		}
 	}
 
+	if (Configuration.bConstructInMemory)
+	{
+		// Now that we have the entire file ready, we write it in one big pass. This
+		// is synchronous to avoid complexity, but we do eat some performance since we
+		// aren't overlapping work with these writes.
+		
+		// This is specifically after the deletion of the old installation file in order to
+		// manage disk space requirements - we never want to have two copies of the same file (independent
+		// of version) on the disk at the same time - avoiding that is the entire reason we added bConstructInMemory!
+
+		{
+			FAdministrationScope _(FileConstructorStat);
+			CurrentFile.NewFile = FileSystem->CreateFileWriter(*CurrentFile.NewFilename, EWriteFlags::None);
+			CurrentFile.CreateFilePlatformLastError = FPlatformMisc::GetLastError();
+		}
+		
+		CurrentFile.bSuccess = CurrentFile.NewFile != nullptr;
+		if (!CurrentFile.bSuccess)
+		{
+			CurrentFile.ConstructionError = EConstructionError::CannotCreateFile;
+		}
+		else
+		{
+			TArray64<uint8>& Data = MemoryOutputFiles[CurrentFile.BuildFilename];
+
+			
+			FileConstructorStat->OnBeforeWrite();
+			ISpeedRecorder::FRecord ActivityRecord;
+			ActivityRecord.CyclesStart = FStatsCollector::GetCycles();
+		
+			{
+				TRACE_CPUPROFILER_EVENT_SCOPE(MemoryConstruction_Write)
+				int64 WritePosition = 0;
+				for (;;)
+				{
+					if (WritePosition == Data.Num())
+					{
+						break;
+					}
+					int64 WriteAmount = 4 << 20; // 4mb write chunks.
+					if (WritePosition + WriteAmount > Data.Num())
+					{
+						WriteAmount = Data.Num() - WritePosition;
+					}
+
+					CurrentFile.NewFile->Serialize(Data.GetData() + WritePosition, WriteAmount);
+
+					WritePosition += WriteAmount;
+				}
+			}
+
+			ActivityRecord.Size = Data.Num();
+			ActivityRecord.CyclesEnd = FStatsCollector::GetCycles();
+			FileConstructorStat->OnAfterWrite(ActivityRecord);
+
+			{
+				TRACE_CPUPROFILER_EVENT_SCOPE(MemoryConstruction_WriteFree)
+				Data.Empty();
+			}
+			
+			if (CurrentFile.NewFile->IsError())
+			{
+				CurrentFile.bSuccess = false;
+				CurrentFile.ConstructionError = EConstructionError::SerializeError;
+			}
+
+			bool bArchiveSuccess = true;
+			{
+				FAdministrationScope _(FileConstructorStat);
+				bArchiveSuccess = CurrentFile.NewFile->Close();
+				CurrentFile.NewFile.Reset();
+			}
+
+			// Check for final success
+			if (CurrentFile.bSuccess && !bArchiveSuccess)
+			{
+				CurrentFile.ConstructionError = EConstructionError::SerializeError;
+				CurrentFile.bSuccess = false;
+			}
+		}
+
+		CurrentFile.SetAttributes();
+	}
+
+	if (CurrentFile.bSuccess)
+	{
+		ChunkDbSource->ReportFileCompletion(ChunkReferenceTracker->GetRemainingChunkCount());
+	}
+
 	FileConstructorStat->OnFileCompleted(CurrentFile.BuildFilename, CurrentFile.bSuccess);
 	
 	// Report errors.
 	if (!CurrentFile.bSuccess)
 	{
-		if (CurrentFile.ConstructionError == EConstructionError::SerializeError)
+		if (!Configuration.bInstallToMemory &&
+			CurrentFile.ConstructionError == EConstructionError::SerializeError)
 		{
 			// Serialize error is our catchall file error right now. This should probably get
 			// migrated such that it's when we fail to load an existing chunk (i.e. corruption)
@@ -3193,7 +3341,7 @@ void FBuildPatchFileConstructor::ConstructFiles(const FResumeData& ResumeData)
 	}
 	if (Configuration.IOBatchSizeMB.IsSet())
 	{
-		MaxWriteBatchSize = Configuration.IOBatchSizeMB.Get(MaxWriteBatchSize);
+		MaxWriteBatchSize = FMath::Max(1, Configuration.IOBatchSizeMB.Get(MaxWriteBatchSize));
 		UE_LOG(LogBuildPatchServices, Verbose, TEXT("Got override ConstructorIOBatchSizeMB = %d"), MaxWriteBatchSize);
 	}
 	MaxWriteBatchSize <<= 20; // to MB;
@@ -3205,14 +3353,20 @@ void FBuildPatchFileConstructor::ConstructFiles(const FResumeData& ResumeData)
 	}
 	if (Configuration.IOBufferSizeMB.IsSet())
 	{
-		IOBufferSize = Configuration.IOBufferSizeMB.Get(IOBufferSize);
+		IOBufferSize = FMath::Max(1, Configuration.IOBufferSizeMB.Get(IOBufferSize));
 		UE_LOG(LogBuildPatchServices, Verbose, TEXT("Got override ConstructorIOBufferSizeMB = %d"), IOBufferSize);
 	}
 	IOBufferSize <<= 20; // to MB;
 
+	// This is the buffer we issue reads into and writes out of. We segment it up
+	// based on batch sizing.	
+	TArray<uint8> IOBuffer;
+
 	// Ensure that our batch size can always make progress (all chunks can fit)
-	uint32 LargestChunkSize = 0;
+	if (!Configuration.bInstallToMemory &&
+		!Configuration.bConstructInMemory)
 	{
+		uint32 LargestChunkSize = 0;
 		for (const FFileToConstruct& FileToConstruct : ConstructionList)
 		{
 			if (FileToConstruct.FileManifest)
@@ -3238,6 +3392,8 @@ void FBuildPatchFileConstructor::ConstructFiles(const FResumeData& ResumeData)
 				UE_LOG(LogBuildPatchServices, Display, TEXT("Increasing IO buffer size to fit batch size: %u bytes"), IOBufferSize);
 			}
 		}
+
+		IOBuffer.SetNumUninitialized(IOBufferSize);
 	}
 
 	uint64 ConstructStartCycles = FPlatformTime::Cycles64();
@@ -3256,17 +3412,12 @@ void FBuildPatchFileConstructor::ConstructFiles(const FResumeData& ResumeData)
 	// List of batches in flight. These must be dispatched in order.
 	TArray<TUniquePtr<FBatchState>> Batches;
 
-	// This is the buffer we issue reads into and writes out of. We segment it up
-	// based on batch sizing.
-	TArray<uint8> IOBUffer;
-	IOBUffer.SetNumUninitialized(IOBufferSize);
-
 	struct Range
 	{
 		uint64 Start=0, End=0;
 	};
 	TArray<Range, TInlineAllocator<3>> IOBufferFreeBlockList;
-	IOBufferFreeBlockList.Add({0, IOBufferSize});
+	IOBufferFreeBlockList.Add({0, (uint64)IOBuffer.Num()});
 
 	auto SortAndCoalesceFreeList = [](TArray<Range, TInlineAllocator<3>>& FreeList)
 	{
@@ -3300,8 +3451,6 @@ void FBuildPatchFileConstructor::ConstructFiles(const FResumeData& ResumeData)
 	bool bStillProcessingFiles = true;
 	for (;bStillProcessingFiles;)
 	{		
-		
-
 		bool bHasAnyFileFailed = false;
 		// Check state. We can continue to run this after a failure has been encountered in order
 		// to drain any async tasks so we know we can shut down safely.
@@ -3312,30 +3461,34 @@ void FBuildPatchFileConstructor::ConstructFiles(const FResumeData& ResumeData)
 			{
 				uint64 WriteDoneStartCycles = FPlatformTime::Cycles64();
 
-				// Could be empty buffer if sentinel
-				if (Batches[0]->BatchBuffer.GetSize())
-				{
-					uint64 BufferBase = (uint64)((uint8*)Batches[0]->BatchBuffer.GetData() - (uint8*)IOBUffer.GetData());
-
-					IOBufferFreeBlockList.Add({BufferBase, BufferBase + Batches[0]->BatchBuffer.GetSize()});
-					SortAndCoalesceFreeList(IOBufferFreeBlockList);
-				}
-
 				FFileConstructionState* File = Batches[0]->OwningFile;
 				File->OutstandingBatches--;
 
-				Batches.RemoveAt(0);
-				
-				if (File->NewFile->IsError())
+				if (!Configuration.bInstallToMemory && 
+					!Configuration.bConstructInMemory)
 				{
-					File->bSuccess = false;
+					// Could be empty buffer if sentinel
+					if (Batches[0]->BatchBuffer.GetSize())
+					{
+						uint64 BufferBase = (uint64)((uint8*)Batches[0]->BatchBuffer.GetData() - (uint8*)IOBuffer.GetData());
 
-					// SerializeError signals the error management code to check for low disk space
-					// since we currently have no way of routing out of disk errors from the UE low level
-					// writing code.
-					File->ConstructionError = EConstructionError::SerializeError;
+						IOBufferFreeBlockList.Add({BufferBase, BufferBase + Batches[0]->BatchBuffer.GetSize()});
+						SortAndCoalesceFreeList(IOBufferFreeBlockList);
+					}
+				
+					if (File->NewFile->IsError())
+					{
+						File->bSuccess = false;
+
+						// SerializeError signals the error management code to check for low disk space
+						// since we currently have no way of routing out of disk errors from the UE low level
+						// writing code.
+						File->ConstructionError = EConstructionError::SerializeError;
+					}
 				}
-					
+
+				Batches.RemoveAt(0);
+				
 				// If we have nothing further to read and all the file's reads are done,
 				// since we only ever have 1 write batch we know it must be done.
 				if (File->NextChunkPartToRead == File->FileManifest->ChunkParts.Num() &&
@@ -3442,18 +3595,23 @@ void FBuildPatchFileConstructor::ConstructFiles(const FResumeData& ResumeData)
 					// We only check for another if we got one queued so we fill up the buffer space asap.
 					bCheckForAnotherRead = false;
 
+					bool bNeedsIoBufferSpace = !Configuration.bInstallToMemory && !Configuration.bConstructInMemory;
+					
 					int32 BiggestFreeBlockSize = 0;
 					int32 BiggestFreeBlockSlotIndex = 0;
-					for (int32 Slot = 0; Slot < IOBufferFreeBlockList.Num(); Slot++)
+					if (bNeedsIoBufferSpace)
 					{
-						if (IOBufferFreeBlockList[Slot].End - IOBufferFreeBlockList[Slot].Start > BiggestFreeBlockSize)
+						for (int32 Slot = 0; Slot < IOBufferFreeBlockList.Num(); Slot++)
 						{
-							BiggestFreeBlockSize = IOBufferFreeBlockList[Slot].End - IOBufferFreeBlockList[Slot].Start;
-							BiggestFreeBlockSlotIndex = Slot;
+							if (IOBufferFreeBlockList[Slot].End - IOBufferFreeBlockList[Slot].Start > BiggestFreeBlockSize)
+							{
+								BiggestFreeBlockSize = IOBufferFreeBlockList[Slot].End - IOBufferFreeBlockList[Slot].Start;
+								BiggestFreeBlockSlotIndex = Slot;
+							}
 						}
 					}
 
-					if (BiggestFreeBlockSize > 0)
+					if (!bNeedsIoBufferSpace || BiggestFreeBlockSize > 0)
 					{
 						// Default to continue to work on the last file.
 						FFileConstructionState* FileToStart = nullptr;
@@ -3466,7 +3624,37 @@ void FBuildPatchFileConstructor::ConstructFiles(const FResumeData& ResumeData)
 						// No more work do to on active files - is there another to start?
 						if (!FileToStart)
 						{
-							bool bAllowAnotherFile = bAllowMultipleFilesInFlight || (ActiveFiles.Num() == 0);
+							bool bAllowAnotherFile = ActiveFiles.Num() == 0;
+							if (!bAllowAnotherFile)
+							{
+								if (bAllowMultipleFilesInFlight)
+								{
+									bAllowAnotherFile = true;
+
+									// The IOBuffer size (i.e. memory) is what limits how many files can be in flight. 
+									// For normal operation, this happens implicitly below when we look for a read destination,
+									// but for situations where we are writing directly to the destination, we need to explictly
+									// govern this.
+									if (Configuration.bInstallToMemory || Configuration.bConstructInMemory)
+									{
+										// For install to memory, we are going to eventually have anything in memory _anyway_,
+										// but this also limits how many requests get queued up so we still limit.
+										// For construct in memory, if we don't limit here it'll freely allocate the entire
+										// installation when we almost certainly don't want that.
+										//
+										// Note that the IOBuffer isn't actually getting allocated in these situations - but
+										// we can use the user provided limitation as a limit on how much to have in flight.
+										uint64 TotalActiveData = 0;
+										for (const TUniquePtr<FFileConstructionState>& ActiveFile : ActiveFiles)
+										{
+											const TArray64<uint8>& Data = MemoryOutputFiles[ActiveFile->BuildFilename];
+											TotalActiveData += Data.Num();
+										}
+
+										bAllowAnotherFile = TotalActiveData < IOBufferSize;
+									}
+								}
+							}
 
 							// This is not a race because the only place we ever increment is this thread, this function right
 							// below.
@@ -3499,6 +3687,12 @@ void FBuildPatchFileConstructor::ConstructFiles(const FResumeData& ResumeData)
 										Configuration.ConstructList[IndexToConstruct], 
 										MoveTemp(StagingFileName));
 
+									AnotherFile->BaseReferenceIndex = 0;
+									if (IndexToConstruct)
+									{
+										AnotherFile->BaseReferenceIndex = FileCompletionPositions[IndexToConstruct-1];
+									}
+
 									AnotherFile->ConstructionIndex = IndexToConstruct;
 
 									TRACE_BOOKMARK(TEXT("Starting File: %s [%u]"), *AnotherFile->BuildFilename, AnotherFile->FileManifest->ChunkParts.Num());
@@ -3574,9 +3768,44 @@ void FBuildPatchFileConstructor::ConstructFiles(const FResumeData& ResumeData)
 
 								bCheckForAnotherRead = true;
 							}
-							else
+							else if (Configuration.bInstallToMemory || Configuration.bConstructInMemory)
 							{
+								// When installing to memory we can start a batch for the entirety of the file.
+								TArray64<uint8>& Output = MemoryOutputFiles[FileToStart->BuildFilename];
+
+								// \todo break up the batches a bit so that the hash overlaps with other
+								// work. Right now its a giant hit on the constructor thread.
+
+								TUniquePtr<FBatchState> Batch = MakeUnique<FBatchState>();
+								Batch->bNeedsWrite = true;
+								Batch->bIsReading = true;
+								Batch->OwningFile = FileToStart;
+								Batch->bIsFinished.store(false, std::memory_order_release);
+
+								Batch->BatchBuffer = FMutableMemoryView(Output.GetData(), FileToStart->FileManifest->FileSize);
 
+								StartReadBatch(*FileToStart, *Batch.Get());
+	
+								// This should use the entire buffer. If it doesn't we've violated some assumptions here.
+								if (Batch->BatchBuffer.GetSize() != FileToStart->FileManifest->FileSize)
+								{
+									FileToStart->bSuccess = false;
+									FileToStart->ConstructionError = EConstructionError::InternalConsistencyError;
+									bHasAnyFileFailed = true;
+																	
+									UE_LOG(LogBuildPatchServices, Error, TEXT("Memory construction setup failed - batch buffer didn't fill entirely: Expected %llu, got %llu. File %s"),
+										FileToStart->FileManifest->FileSize,
+										Batch->BatchBuffer.GetSize(),
+										*FileToStart->BuildFilename);
+								}
+								else
+								{
+									bCheckForAnotherRead = true;								
+									Batches.Push(MoveTemp(Batch));
+								}
+							}
+							else
+							{
 								uint32 MaxBufferSize = BiggestFreeBlockSize;
 								if (MaxBufferSize > MaxWriteBatchSize)
 								{
@@ -3624,7 +3853,7 @@ void FBuildPatchFileConstructor::ConstructFiles(const FResumeData& ResumeData)
 									Batch->bNeedsWrite = true;
 									Batch->bIsReading = true;
 									Batch->OwningFile = FileToStart;
-									Batch->BatchBuffer = FMutableMemoryView(IOBUffer.GetData() + IOBufferFreeBlockList[BiggestFreeBlockSlotIndex].Start, MaxBufferSize);
+									Batch->BatchBuffer = FMutableMemoryView(IOBuffer.GetData() + IOBufferFreeBlockList[BiggestFreeBlockSlotIndex].Start, MaxBufferSize);
 									Batch->bIsFinished.store(false, std::memory_order_release);
 									StartReadBatch(*FileToStart, *Batch.Get());
 
@@ -3664,21 +3893,33 @@ void FBuildPatchFileConstructor::ConstructFiles(const FResumeData& ResumeData)
 					uint64 WriteStartStartCycles = FPlatformTime::Cycles64();
 
 					FFileConstructionState* FirstFile = ActiveFiles[0].Get();
-
-					Batches[0]->bNeedsWrite = false;
-					Batches[0]->bIsWriting = true;
-					Batches[0]->bIsFinished.store(false, std::memory_order_relaxed);
-
-					// Launch the write and hash the buffer on this thread.
-					UE_LOG(LogBuildPatchServices, VeryVerbose, TEXT("Writing Batch: %d, file %s [%d - %d]"),
-						Batches[0]->BatchId, 
-						*FirstFile->BuildFilename,
-						Batches[0]->StartChunkPartIndex,
-						Batches[0]->StartChunkPartIndex + Batches[0]->ChunkCount
+					if (Configuration.bInstallToMemory || Configuration.bConstructInMemory)
+					{
+						// We read directly into the output - nothing needs to be done. However we need to
+						// catch the write completion logic so we mark bIsFinished and signal our wakeup
+						// event so we immediately do another pass
+						Batches[0]->bNeedsWrite = false;
+						Batches[0]->bIsWriting = true;
+						Batches[0]->bIsFinished.store(true, std::memory_order_release);
+						WakeUpDispatchThreadEvent->Trigger();
+					}
+					else
+					{
+						Batches[0]->bNeedsWrite = false;
+						Batches[0]->bIsWriting = true;
+						Batches[0]->bIsFinished.store(false, std::memory_order_release);
+
+						// Launch the write and hash the buffer on this thread.
+						UE_LOG(LogBuildPatchServices, VeryVerbose, TEXT("Writing Batch: %d, file %s [%d - %d]"),
+							Batches[0]->BatchId, 
+							*FirstFile->BuildFilename,
+							Batches[0]->StartChunkPartIndex,
+							Batches[0]->StartChunkPartIndex + Batches[0]->ChunkCount
 						);
 
-					IConstructorChunkSource::FRequestProcessFn WriteFn = CreateWriteRequest(FirstFile->NewFile.Get(), *Batches[0].Get());
-					QueueGenericThreadTask(WriteThreadIndex, MoveTemp(WriteFn));
+						IConstructorChunkSource::FRequestProcessFn WriteFn = CreateWriteRequest(FirstFile->NewFile.Get(), *Batches[0].Get());
+						QueueGenericThreadTask(WriteThreadIndex, MoveTemp(WriteFn));
+					}
 					
 					uint64 HashStartCycles = FPlatformTime::Cycles64();
 					WriteStartCycles += HashStartCycles - WriteStartStartCycles;
@@ -3738,7 +3979,16 @@ void FBuildPatchFileConstructor::ConstructFiles(const FResumeData& ResumeData)
 			// 16 issues old - 32 downloads and you'll wait on one 32 issues old. We try to bound this by capping the issued
 			// downloads here, and in the cloud source we prevent queues if we get too far ahead of the last download.
 			// Note this also gets adjusted by the connection health stuff internal to the cloud source.
-			const uint32 MaxDownloads = FMath::Max(1U, FMath::DivideAndRoundUp(MaxWriteBatchSize, ExpectedChunkSize));
+			uint32 MaxDownloads = FMath::Max(1U, FMath::DivideAndRoundUp(MaxWriteBatchSize, ExpectedChunkSize));
+
+			// If we are reading directly, then we don't have to worry about individual batches anymore - but we don't want to wait
+			// too long at the end of each file either. We expose IoBuffer as a conceptual limit to how much the user wants going
+			// on in these cases, so we use it here as well.
+			if (Configuration.bInstallToMemory || Configuration.bConstructInMemory)
+			{
+				MaxDownloads = FMath::Max(1U, FMath::DivideAndRoundUp(IOBufferSize, ExpectedChunkSize));
+			}
+
 			
 			// WaitTimeMs is an OUT param
 			CloudSource->Tick(!bIsPaused && !bHasAnyFileFailed, WaitTimeMs, MaxDownloads);
@@ -3835,20 +4085,3 @@ void FBuildPatchFileConstructor::ConstructFiles(const FResumeData& ResumeData)
 		100.0 * UnaccountedForCycles / ConstructCycles
 		);
 }
-
-void FBuildPatchFileConstructor::DeleteDirectoryContents(const FString& RootDirectory)
-{
-	TArray<FString> SubDirNames;
-	IFileManager::Get().FindFiles(SubDirNames, *(RootDirectory / TEXT("*")), false, true);
-	for (const FString& DirName : SubDirNames)
-	{
-		IFileManager::Get().DeleteDirectory(*(RootDirectory / DirName), false, true);
-	}
-
-	TArray<FString> SubFileNames;
-	IFileManager::Get().FindFiles(SubFileNames, *(RootDirectory / TEXT("*")), true, false);
-	for (const FString& FileName : SubFileNames)
-	{
-		IFileManager::Get().Delete(*(RootDirectory / FileName), false, true);
-	}
-}
diff --git a/Engine/Source/Runtime/Online/BuildPatchServices/Private/BuildPatchFileConstructor.h b/Engine/Source/Runtime/Online/BuildPatchServices/Private/BuildPatchFileConstructor.h
--- a/Engine/Source/Runtime/Online/BuildPatchServices/Private/BuildPatchFileConstructor.h
+++ b/Engine/Source/Runtime/Online/BuildPatchServices/Private/BuildPatchFileConstructor.h
@@ -80,6 +80,10 @@ namespace BuildPatchServices
 
 		IBuildInstallerSharedContext* SharedContext;
 
+		// See comments in installer config
+		bool bInstallToMemory = false;
+		bool bConstructInMemory = false;
+
 		bool bDeleteChunkDBFilesAfterUse = false;
 
 		static const bool bDefaultSpawnAdditionalIOThreads = true;
@@ -156,6 +160,11 @@ namespace BuildPatchServices
 		DECLARE_EVENT_OneParam(FBuildPatchFileConstructor, FOnBeforeDeleteFile, const FString& /*BuildFile*/);
 		FOnBeforeDeleteFile& OnBeforeDeleteFile();
 
+		void GrabFilesInstalledToMemory(TMap<FString, TArray64<uint8>>& OutFilesInstalledToMemory)
+		{
+			OutFilesInstalledToMemory = MoveTemp(MemoryOutputFiles);
+		}
+
 	private:
 
 		void SetChunkLocation(const FGuid& InGuid, EConstructorChunkLocation InNewLocation);
@@ -188,12 +197,6 @@ namespace BuildPatchServices
 		 */
 		void ConstructFiles(const FResumeData& ResumeData);
 
-		/**
-		 * Delete all contents of a directory
-		 * @param RootDirectory	 	Directory to make empty
-		 */
-		void DeleteDirectoryContents(const FString& RootDirectory);
-
 	private:
 		// The configuration for the constructor.
 		const FFileConstructorConfig Configuration;
@@ -241,6 +244,9 @@ namespace BuildPatchServices
 		IConstructorInstallChunkSource* InstallSource;
 		IConstructorCloudChunkSource* CloudSource;
 
+		// Keyed off of the filename relative to the install directory.
+		TMap<FString, TArray64<uint8>> MemoryOutputFiles;
+
 		// We always want to know exactly where we think a chunk should be. If it's not there,
 		// we update this list to where it can be found (i.e. cloud)
 		// This is almost always read only after initialization, but in rare situations can be updated
diff --git a/Engine/Source/Runtime/Online/BuildPatchServices/Private/BuildPatchInstaller.cpp b/Engine/Source/Runtime/Online/BuildPatchServices/Private/BuildPatchInstaller.cpp
--- a/Engine/Source/Runtime/Online/BuildPatchServices/Private/BuildPatchInstaller.cpp
+++ b/Engine/Source/Runtime/Online/BuildPatchServices/Private/BuildPatchInstaller.cpp
@@ -644,6 +644,30 @@ namespace BuildPatchServices
 			}
 		}
 
+		if (Configuration.bInstallToMemory && Configuration.bConstructFilesInMemory)
+		{
+			UE_LOG(LogBuildPatchServices, Error, TEXT("Installer configuration: Can't construct in memory and install to memory."));
+			InstallerError->SetError(EBuildPatchInstallError::InitializationError, InitializationErrorCodes::InvalidMemoryConstructionConfig, 0,
+				NSLOCTEXT("BuildPatchInstallError", "InvalidMemoryConstructionConfig", "This installation could not continue due to a configuration issue. Please contact support."));
+			bInstallerInitSuccess = false;
+		}
+
+		// We can't be patching if we are installing to memory.
+		if (Configuration.bInstallToMemory)
+		{
+			for (const FBuildPatchInstallerAction& InstallerAction : InstallerActions)
+			{
+				if (!InstallerAction.IsInstall())
+				{
+					UE_LOG(LogBuildPatchServices, Error, TEXT("Installer configuration: Installation to memory only supports installs (e.g. no updates/repairs)."));
+					InstallerError->SetError(EBuildPatchInstallError::InitializationError, InitializationErrorCodes::InvalidInstallToMemory, 0,
+						NSLOCTEXT("BuildPatchInstallError", "InvalidInstallToMemory", "This installation could not continue due to a configuration issue. Please contact support."));
+					bInstallerInitSuccess = false;
+					break;
+				}
+			}
+		}
+
 		// Check that we were provided with a bound delegate.
 		if (!CompleteDelegate.IsBound())
 		{
@@ -653,23 +677,26 @@ namespace BuildPatchServices
 		}
 
 		// Make sure we have install directory access.
-		IFileManager::Get().MakeDirectory(*Configuration.InstallDirectory, true);
-		if (!IFileManager::Get().DirectoryExists(*Configuration.InstallDirectory))
+		if (!Configuration.bInstallToMemory)
 		{
-			UE_LOG(LogBuildPatchServices, Error, TEXT("Installer setup: Inability to create InstallDirectory %s."), *Configuration.InstallDirectory);
-			InstallerError->SetError(EBuildPatchInstallError::InitializationError, InitializationErrorCodes::MissingInstallDirectory, 0,
-				FText::Format(NSLOCTEXT("BuildPatchInstallError", "MissingInstallDirectory", "The installation directory could not be created.\n{0}"), FText::FromString(Configuration.InstallDirectory)));
-			bInstallerInitSuccess = false;
-		}
+			IFileManager::Get().MakeDirectory(*Configuration.InstallDirectory, true);
+			if (!IFileManager::Get().DirectoryExists(*Configuration.InstallDirectory))
+			{
+				UE_LOG(LogBuildPatchServices, Error, TEXT("Installer setup: Inability to create InstallDirectory %s."), *Configuration.InstallDirectory);
+				InstallerError->SetError(EBuildPatchInstallError::InitializationError, InitializationErrorCodes::MissingInstallDirectory, 0,
+					FText::Format(NSLOCTEXT("BuildPatchInstallError", "MissingInstallDirectory", "The installation directory could not be created.\n{0}"), FText::FromString(Configuration.InstallDirectory)));
+				bInstallerInitSuccess = false;
+			}
 
-		// Make sure we have staging directory access.
-		IFileManager::Get().MakeDirectory(*Configuration.StagingDirectory, true);
-		if (!IFileManager::Get().DirectoryExists(*Configuration.StagingDirectory))
-		{
-			UE_LOG(LogBuildPatchServices, Error, TEXT("Installer setup: Inability to create StagingDirectory %s."), *Configuration.StagingDirectory);
-			InstallerError->SetError(EBuildPatchInstallError::InitializationError, InitializationErrorCodes::MissingStageDirectory, 0,
-				FText::Format(NSLOCTEXT("BuildPatchInstallError", "MissingStageDirectory", "The following directory could not be created.\n{0}"), FText::FromString(Configuration.StagingDirectory)));
-			bInstallerInitSuccess = false;
+			// Make sure we have staging directory access.
+			IFileManager::Get().MakeDirectory(*Configuration.StagingDirectory, true);
+			if (!IFileManager::Get().DirectoryExists(*Configuration.StagingDirectory))
+			{
+				UE_LOG(LogBuildPatchServices, Error, TEXT("Installer setup: Inability to create StagingDirectory %s."), *Configuration.StagingDirectory);
+				InstallerError->SetError(EBuildPatchInstallError::InitializationError, InitializationErrorCodes::MissingStageDirectory, 0,
+					FText::Format(NSLOCTEXT("BuildPatchInstallError", "MissingStageDirectory", "The following directory could not be created.\n{0}"), FText::FromString(Configuration.StagingDirectory)));
+				bInstallerInitSuccess = false;
+			}
 		}
 
 		// Make sure that we have a prereq if we've specified a prereq only install.
@@ -854,7 +881,8 @@ namespace BuildPatchServices
 					bProcessSuccess = bInstallSuccess && RunVerification(CorruptFiles);
 
 					// Clean staging if INSTALL success, we still do cleanup if we failed at the verify stage.
-					if (bInstallSuccess)
+					if (bInstallSuccess && 
+						!Configuration.bInstallToMemory)
 					{
 						CleanUpTimer.Start();
 						if (Configuration.InstallMode == EInstallMode::StageFiles)
@@ -878,10 +906,13 @@ namespace BuildPatchServices
 					const bool bWillRetry = !bProcessSuccess && bCanRetry;
 
 					// If successful or we will retry, remove the moved files marker
-					if (bProcessSuccess || bCanRetry)
+					if (!Configuration.bInstallToMemory)
 					{
-						UE_LOG(LogBuildPatchServices, Log, TEXT("Reset MM."));
-						IFileManager::Get().Delete(*PreviousMoveMarker, false, true);
+						if (bProcessSuccess || bCanRetry)
+						{
+							UE_LOG(LogBuildPatchServices, Log, TEXT("Reset MM."));
+							IFileManager::Get().Delete(*PreviousMoveMarker, false, true);
+						}
 					}
 
 					// Setup end of attempt stats
@@ -1233,11 +1264,9 @@ namespace BuildPatchServices
 		FPaths::NormalizeDirectoryName(InstallStagingDir);
 		FPaths::NormalizeDirectoryName(MetaStagingDir);
 
-		// Make sure staging directories exist		
-		IFileManager::Get().MakeDirectory(*InstallStagingDir, true);
-
-		// These we don't make here - we make them only if we need them because for basic small installs
-		// on some platforms we want to minimize io operations.
+		// UE file system abstractions create all necessary directories when creating files.
+		// By not doing this up front we avoid any unnecessary file operations as some platforms are sensitive to that.
+		//IFileManager::Get().MakeDirectory(*InstallStagingDir, true);
 		//IFileManager::Get().MakeDirectory(*DataStagingDir, true);
 		//IFileManager::Get().MakeDirectory(*MetaStagingDir, true);
 
@@ -1249,7 +1278,17 @@ namespace BuildPatchServices
 		TaggedFiles.Empty();
 		ManifestSet->GetExpectedFiles(TaggedFiles);
 		OutdatedFiles.Empty();
-		ManifestSet->GetOutdatedFiles(Configuration.InstallDirectory, OutdatedFiles);
+		
+		if (Configuration.bInstallToMemory)
+		{
+			// When installing to memory we know we aren't patching and we know we don't have an installation
+			// directory, so we are always installing everything.
+			OutdatedFiles = TaggedFiles;
+		}
+		else
+		{
+			ManifestSet->GetOutdatedFiles(Configuration.InstallDirectory, OutdatedFiles);
+		}
 
 		const bool bIsPrereqOnly = Configuration.InstallMode == EInstallMode::PrereqOnly;
 		const bool bHasCorruptFiles = CorruptFiles.Num() > 0;
@@ -1259,7 +1298,9 @@ namespace BuildPatchServices
 
 		// Check if we should skip out of this process due to existing installation,
 		// that will mean we start with the verification stage
-		if (!bHasCorruptFiles && (bIsPrereqOnly || CheckForExternallyInstalledFiles(FilesToConstruct)))
+		if (!Configuration.bInstallToMemory &&
+			!bHasCorruptFiles && 
+			(bIsPrereqOnly || CheckForExternallyInstalledFiles(FilesToConstruct)))
 		{
 			UE_LOG(LogBuildPatchServices, Log, TEXT("Detected previous staging completed, or existing files in target directory"));
 			// Add required files to the verifier as 'touched' since we do not know their state.
@@ -1287,21 +1328,24 @@ namespace BuildPatchServices
 		}
 
 		// Make sure all the files won't exceed the maximum path length
-		for (const FString& FileToConstruct : FilesToConstruct)
+		if (!Configuration.bInstallToMemory)
 		{
-			const FString InstallConstructionFile = Configuration.InstallDirectory / FileToConstruct;
-			const FString StagedConstructionFile = InstallStagingDir / FileToConstruct;
-			if (InstallConstructionFile.Len() >= FPlatformMisc::GetMaxPathLength())
+			for (const FString& FileToConstruct : FilesToConstruct)
 			{
-				UE_LOG(LogBuildPatchServices, Error, TEXT("Could not create new file due to exceeding maximum path length %s"), *InstallConstructionFile);
-				InstallerError->SetError(EBuildPatchInstallError::PathLengthExceeded, PathLengthErrorCodes::InstallDirectory);
-				return false;
-			}
-			if (StagedConstructionFile.Len() >= FPlatformMisc::GetMaxPathLength())
-			{
-				UE_LOG(LogBuildPatchServices, Error, TEXT("Could not create new file due to exceeding maximum path length %s"), *StagedConstructionFile);
-				InstallerError->SetError(EBuildPatchInstallError::PathLengthExceeded, PathLengthErrorCodes::StagingDirectory);
-				return false;
+				const FString InstallConstructionFile = Configuration.InstallDirectory / FileToConstruct;
+				const FString StagedConstructionFile = InstallStagingDir / FileToConstruct;
+				if (InstallConstructionFile.Len() >= FPlatformMisc::GetMaxPathLength())
+				{
+					UE_LOG(LogBuildPatchServices, Error, TEXT("Could not create new file due to exceeding maximum path length %s"), *InstallConstructionFile);
+					InstallerError->SetError(EBuildPatchInstallError::PathLengthExceeded, PathLengthErrorCodes::InstallDirectory);
+					return false;
+				}
+				if (StagedConstructionFile.Len() >= FPlatformMisc::GetMaxPathLength())
+				{
+					UE_LOG(LogBuildPatchServices, Error, TEXT("Could not create new file due to exceeding maximum path length %s"), *StagedConstructionFile);
+					InstallerError->SetError(EBuildPatchInstallError::PathLengthExceeded, PathLengthErrorCodes::StagingDirectory);
+					return false;
+				}
 			}
 		}
 
@@ -1345,15 +1389,21 @@ namespace BuildPatchServices
 
 			// Add a source for pulling from the existing directory (for patching). This uses the manifest for the
 			// currently deployed files to create a list of available chunks to read from.
-			TUniquePtr<IConstructorInstallChunkSource> InstallChunkSource(IConstructorInstallChunkSource::CreateInstallSource(
-				FileSystem.Get(),
-				InstallChunkSourceStatistics.Get(),
-				InstallationInfo,
-				ReferencedChunks));
+			TUniquePtr<IConstructorInstallChunkSource> InstallChunkSource;
+			
+			TSet<FGuid> EmptySet;
+			if (!Configuration.bInstallToMemory)
+			{
+				InstallChunkSource.Reset(IConstructorInstallChunkSource::CreateInstallSource(
+					FileSystem.Get(),
+					InstallChunkSourceStatistics.Get(),
+					InstallationInfo,
+					ReferencedChunks));
+			}
 
 			// Sub out the chunks we have available on disk - anything we don't we have to download.
 			const TSet<FGuid>& ChunkDbChunksAvailable = ChunkDbChunkSource->GetAvailableChunks();
-			const TSet<FGuid>& InstallChunksAvailable = InstallChunkSource->GetAvailableChunks();
+			const TSet<FGuid>& InstallChunksAvailable = InstallChunkSource.IsValid() ? InstallChunkSource->GetAvailableChunks() : EmptySet;
 			
 			TSet<FGuid> CloudChunks = ReferencedChunks.Difference(InstallChunksAvailable).Difference(ChunkDbChunksAvailable);
 			TSet<FGuid> ChunkDbChunks = ReferencedChunks.Intersect(ChunkDbChunksAvailable);
@@ -1404,6 +1454,8 @@ namespace BuildPatchServices
 
 			// Set up the class that actually requests chunks and writes them to disk.
 			FFileConstructorConfig FCC;
+			FCC.bInstallToMemory = Configuration.bInstallToMemory;
+			FCC.bConstructInMemory = Configuration.bConstructFilesInMemory;
 			FCC.ConstructList = FilesToConstruct.Array();
 			FCC.InstallDirectory = Configuration.InstallDirectory;
 			FPaths::NormalizeDirectoryName(FCC.InstallDirectory);
@@ -1503,6 +1555,11 @@ namespace BuildPatchServices
 			}
 			UE_LOG(LogBuildPatchServices, Log, TEXT("File construction complete"));
 
+			if (Configuration.bInstallToMemory)
+			{
+				FileConstructor->GrabFilesInstalledToMemory(FilesInstalledToMemory);
+			}
+
 			// Set any final stats before system destruction.
 			{
 				FScopeLock Lock(&ThreadLock);
@@ -1669,7 +1726,9 @@ namespace BuildPatchServices
 		BuildProgress.SetStateProgress(EBuildPatchState::MovingToInstall, 0.0f);
 		// We skip this step if performing stage only
 		bool bRunBackupAndMoveSuccess = true;
-		if (Configuration.InstallMode == EInstallMode::StageFiles || (bFirstInstallIteration && Configuration.InstallMode == EInstallMode::PrereqOnly))
+		if (Configuration.InstallMode == EInstallMode::StageFiles || 
+			(bFirstInstallIteration && Configuration.InstallMode == EInstallMode::PrereqOnly) ||
+			Configuration.bInstallToMemory)
 		{
 			UE_LOG(LogBuildPatchServices, Log, TEXT("Skipping backup and stage relocation"));
 			BuildProgress.SetStateProgress(EBuildPatchState::MovingToInstall, 1.0f);
@@ -1783,6 +1842,11 @@ namespace BuildPatchServices
 
 	bool FBuildPatchInstaller::RunFileAttributes()
 	{
+		if (Configuration.bInstallToMemory)
+		{
+			return true;
+		}
+
 		// Only provide stage directory if stage-only mode
 		FString EmptyString;
 		FString& OptionalStageDirectory = Configuration.InstallMode == EInstallMode::StageFiles ? InstallStagingDir : EmptyString;
@@ -1801,6 +1865,13 @@ namespace BuildPatchServices
 
 	bool FBuildPatchInstaller::RunVerification(TArray< FString >& CorruptFiles)
 	{
+		if (Configuration.bInstallToMemory)
+		{
+			// We hash all the files during construction and can only install (not repair) so this
+			// doesn't do anything other than rehash files in memory.
+			return true;
+		}
+
 		// Make sure this function can never be parallelized
 		static FCriticalSection SingletonFunctionLockCS;
 		const bool bShouldLock = !Configuration.bAllowConcurrentExecution;
@@ -2017,6 +2088,11 @@ namespace BuildPatchServices
 		return BuildStats.ErrorCode;
 	}
 
+	TMap<FString, TArray64<uint8>>& FBuildPatchInstaller::GetFilesInstalledToMemory()
+	{
+		return FilesInstalledToMemory;
+	}
+
 	//@todo this is deprecated and shouldn't be used anymore [6/4/2014 justin.sargent]
 	FText FBuildPatchInstaller::GetPercentageText() const
 	{
diff --git a/Engine/Source/Runtime/Online/BuildPatchServices/Private/BuildPatchInstaller.h b/Engine/Source/Runtime/Online/BuildPatchServices/Private/BuildPatchInstaller.h
--- a/Engine/Source/Runtime/Online/BuildPatchServices/Private/BuildPatchInstaller.h
+++ b/Engine/Source/Runtime/Online/BuildPatchServices/Private/BuildPatchInstaller.h
@@ -136,6 +136,9 @@ namespace BuildPatchServices
 		// The files that were removed during installation due to destructive patch behavior.
 		TSet<FString> OldFilesRemovedBySystem;
 
+		// With bInstallToMemory, files end up here after construction.
+		TMap<FString, TArray64<uint8>> FilesInstalledToMemory;
+
 		// Map of registered installations.
 		TMultiMap<FString, FBuildPatchAppManifestRef> InstallationInfo;
 
@@ -256,6 +259,7 @@ namespace BuildPatchServices
 #if !UE_BUILD_SHIPPING
 		virtual void GetDebugText(TArray<FString>& Output) override;
 #endif
+		virtual TMap<FString, TArray64<uint8>>& GetFilesInstalledToMemory() override;
 		// IBuildInstaller interface end.
 
 		/**
diff --git a/Engine/Source/Runtime/Online/BuildPatchServices/Private/Installer/InstallerError.h b/Engine/Source/Runtime/Online/BuildPatchServices/Private/Installer/InstallerError.h
--- a/Engine/Source/Runtime/Online/BuildPatchServices/Private/Installer/InstallerError.h
+++ b/Engine/Source/Runtime/Online/BuildPatchServices/Private/Installer/InstallerError.h
@@ -19,6 +19,8 @@ namespace BuildPatchServices
 		static const TCHAR* ChunkReferenceTracking = TEXT("05");
 		static const TCHAR* MissingPrereqForPrereqOnlyInstall = TEXT("06");
 		static const TCHAR* InvalidDataInManifest = TEXT("07");
+		static const TCHAR* InvalidInstallToMemory = TEXT("08");
+		static const TCHAR* InvalidMemoryConstructionConfig = TEXT("09");
 	}
 
 	/**
diff --git a/Engine/Source/Runtime/Online/BuildPatchServices/Public/BuildPatchSettings.h b/Engine/Source/Runtime/Online/BuildPatchServices/Public/BuildPatchSettings.h
--- a/Engine/Source/Runtime/Online/BuildPatchServices/Public/BuildPatchSettings.h
+++ b/Engine/Source/Runtime/Online/BuildPatchServices/Public/BuildPatchSettings.h
@@ -303,6 +303,7 @@ namespace BuildPatchServices
 		// File constructor buffer size, in MB. This affects constructor performance heavily and must be at least big enough to fit a batch. 
 		// For purely disk->disk installs, this can be relatively small - even 2x batch size will facilitate multiple GB/s installations. However
 		// for cloud/download installs this needs to be fairly large to mitigate the effect of out of order download completion. 
+		// When bConstructFilesInMemory or bInstallToMemory is set, this is used to determine how many active files can be in flight.
 		// Unset uses Engine.ini value [Portal.BuildPatch]ConstructorIOBufferSizeMB if present, FFileConstructorConfig::DefaultIOBufferSizeMB otherwise.
 		TOptional<int32> ConstructorIOBufferSizeMB;
 
@@ -317,6 +318,20 @@ namespace BuildPatchServices
 		// CVar BuildPatchFileConstructor.DisableResumeBelowMB, which defaults to 0.
 		// if set, this member
 		TOptional<int32> ConstructorDisableResumeBelowMB;
+
+		// If set, files will be built entirely in memory - requiring at least the largest file's size worth of memory.
+		// IO batch and buffer sizes are ignored as operations operate directly on the output memory buffer. If allowing
+		// multiple files in flight, this causes the size of the entire installation to be allocated as there's no IO buffer
+		// size limiting the number of in flight requests.
+		// Once the file completes construction it is written to disk in the normal place defined by InstallMode. This occurs
+		// after the old file is deleted for Destructive installations in order to manage peak disk space usage.
+		//
+		// This mode does not support resuming mid-file.
+		bool bConstructFilesInMemory = false;
+
+		// As bConstructFilesInMemory, except the file is never written to disk and is instead made available after via 
+		// GetFilesInstalledToMemory 
+		bool bInstallToMemory = false;
 	};
 
 	/**
diff --git a/Engine/Source/Runtime/Online/BuildPatchServices/Public/Interfaces/IBuildInstaller.h b/Engine/Source/Runtime/Online/BuildPatchServices/Public/Interfaces/IBuildInstaller.h
--- a/Engine/Source/Runtime/Online/BuildPatchServices/Public/Interfaces/IBuildInstaller.h
+++ b/Engine/Source/Runtime/Online/BuildPatchServices/Public/Interfaces/IBuildInstaller.h
@@ -466,6 +466,13 @@ public:
 	 */
 	virtual void GetDebugText(TArray<FString>& Output) {}
 #endif
+
+	/**
+	 * Get the completed files when installing to memory. This must not be called before installation
+	 * completes, and is non-const so that clients can Move the memory as needed. Once installation is complete
+	 * the installer does not touch this other than to delete it on destruction.
+	 */
+	virtual TMap<FString, TArray64<uint8>>& GetFilesInstalledToMemory() = 0; 
 };
 
 static_assert((uint32)EBuildPatchInstallError::NumInstallErrors == 12, "Please add support for the extra values to the Lex functions below.");
