BPS: Fix for a race condition with patching and AllowMultipleFilesInFlight.

#rnx

diff --git a/Engine/Source/Runtime/Online/BuildPatchServices/Private/Installer/InstallChunkSource.cpp b/Engine/Source/Runtime/Online/BuildPatchServices/Private/Installer/InstallChunkSource.cpp
--- a/Engine/Source/Runtime/Online/BuildPatchServices/Private/Installer/InstallChunkSource.cpp
+++ b/Engine/Source/Runtime/Online/BuildPatchServices/Private/Installer/InstallChunkSource.cpp
@@ -3,6 +3,7 @@
 #include "Installer/InstallChunkSource.h"
 
 #include "Algo/Transform.h"
+#include "Async/Mutex.h"
 #include "BuildPatchHash.h"
 #include "Common/FileSystem.h"
 #include "Common/StatsCollector.h"
@@ -37,6 +38,9 @@ namespace BuildPatchServices
 		virtual void OnBeforeDeleteFile(const FString& FilePath) override
 		{
 			// Make sure we close our handle before the deletion occurs.
+			// With multiple files in flight in the constructor we can be deleting a file at the same time as we are
+			// reading chunks for other files, which means we can hit OpenedFileHandles from multiple threads.
+			UE::TUniqueLock _(FileHandleLock);
 			OpenedFileHandles.Remove(FilePath);
 		}
 		virtual int32 GetChunkUnavailableAt(const FGuid& DataId) const override;
@@ -87,6 +91,7 @@ namespace BuildPatchServices
 		TSet<FGuid> AvailableInBuilds;
 		TArray<TPair<FString, FBuildPatchAppManifestRef>> InstallationSources;
 
+		UE::FMutex FileHandleLock; // Protects OpenedFileHandles
 		TMap<FString, TUniquePtr<FArchive>> OpenedFileHandles;
 
 		// The index (ChunkReferenceTracker->GetCurrentUsageIndex) at which our files will get deleted due to destructive install to make room for
@@ -239,16 +244,33 @@ namespace BuildPatchServices
 
 				FString FullFilename = FoundInstallDirectory / FileChunkPart.FileManifest->Filename;
 
-				TUniquePtr<FArchive>* FileArchive = OpenedFileHandles.Find(FullFilename);
-				if (FileArchive == nullptr)
+				// We use the internal FArchive pointer so that we don't have to hold the file handle
+				// lock over the read - the uniqueptr might move around, but the managed pointer will not.
+				// We do know that our specific file won't get deleted until we are done.
+				FArchive* FileArchive = nullptr;
 				{
-					TRACE_CPUPROFILER_EVENT_SCOPE(Install_OpenSource);
-					TUniquePtr<FArchive> NewReader = FileSystem->CreateFileReader(*FullFilename);
-					if (NewReader.IsValid())
+					FileHandleLock.Lock();
+					TUniquePtr<FArchive>* FileArchivePtr = OpenedFileHandles.Find(FullFilename);
+					if (FileArchivePtr == nullptr)
+					{
+						TRACE_CPUPROFILER_EVENT_SCOPE(Install_OpenSource);
+
+						// Don't lock over the file open since that could take a while.
+						FileHandleLock.Unlock();
+						TUniquePtr<FArchive> NewReader = FileSystem->CreateFileReader(*FullFilename);
+						FileHandleLock.Lock();
+
+						if (NewReader.IsValid())
+						{
+							OpenedFileHandles.Add(FullFilename, MoveTemp(NewReader));
+							FileArchivePtr = OpenedFileHandles.Find(FullFilename);
+						}
+					}
+					if (FileArchivePtr)
 					{
-						OpenedFileHandles.Add(FullFilename, MoveTemp(NewReader));
-						FileArchive = OpenedFileHandles.Find(FullFilename);
+						FileArchive = FileArchivePtr->Get();
 					}
+					FileHandleLock.Unlock();
 				}
 
 				if (FileArchive == nullptr)
@@ -259,8 +281,8 @@ namespace BuildPatchServices
 
 				{
 					TRACE_CPUPROFILER_EVENT_SCOPE(Install_Serialize);
-					(*FileArchive)->Seek(FileChunkPart.FileOffset);
-					(*FileArchive)->Serialize((uint8*)DestinationBuffer.GetData() + FileChunkPart.ChunkPart->Offset, FileChunkPart.ChunkPart->Size);
+					FileArchive->Seek(FileChunkPart.FileOffset);
+					FileArchive->Serialize((uint8*)DestinationBuffer.GetData() + FileChunkPart.ChunkPart->Offset, FileChunkPart.ChunkPart->Size);
 
 					ActivityRecord.Size += FileChunkPart.ChunkPart->Size;
 
