BPS: Try and work around the spurious event wakeup issue.

#rnx

diff --git a/Engine/Source/Runtime/Online/BuildPatchServices/Private/Installer/Verifier.cpp b/Engine/Source/Runtime/Online/BuildPatchServices/Private/Installer/Verifier.cpp
--- a/Engine/Source/Runtime/Online/BuildPatchServices/Private/Installer/Verifier.cpp
+++ b/Engine/Source/Runtime/Online/BuildPatchServices/Private/Installer/Verifier.cpp
@@ -6,6 +6,7 @@
 #include "Common/FileSystem.h"
 #include "BuildPatchVerify.h"
 #include "BuildPatchUtil.h"
+#include "HAL/UESemaphore.h"
 #include "IBuildManifestSet.h"
 #include "Misc/ConfigCacheIni.h"
 
@@ -80,7 +81,6 @@ namespace BuildPatchServices
 		};
 
 		FEvent* AllThreadsDoneEvent;
-		std::atomic_int32_t ThreadsRemaining;
 		std::atomic_int64_t ThreadProcessedBytes;
 		TArray<EVerifyResult> ThreadJobResults;
 
@@ -90,7 +90,7 @@ namespace BuildPatchServices
 
 		TArray<IBuildInstallerThread*> SecondaryProcessThreads;
 
-		void ProcessVerifyJobs();
+		void ProcessVerifyJobs(FSemaphore* ThreadDoneSem);
 	};
 
 	FVerifier::FVerifier(IFileSystem* InFileSystem, IVerifierStat* InVerificationStat, EVerifyMode InVerifyMode, IBuildInstallerSharedContextPtr InSharedContext, IBuildManifestSet* InManifestSet, FString InVerifyDirectory, FString InStagedFileDirectory)
@@ -108,7 +108,7 @@ namespace BuildPatchServices
 
 	}
 
-	void FVerifier::ProcessVerifyJobs()
+	void FVerifier::ProcessVerifyJobs(FSemaphore* ThreadDoneSem)
 	{
 		TArray<uint8> FileReadBuffer;
 
@@ -150,14 +150,13 @@ namespace BuildPatchServices
 			ThreadJobResults[OurJobIndex] = FileVerifyResult;
 		}
 
-		if (ThreadsRemaining.fetch_add(-1) == 1)
+		if (ThreadDoneSem)
 		{
-			// We're the last one.
-			AllThreadsDoneEvent->Trigger();
-
-			// We can't touch _anything_ after here because the outer class holding the lambda could
-			// be destroyed.
+			ThreadDoneSem->Release(1);
 		}
+		
+		// We can't touch _anything_ after here because the outer class holding the lambda could
+		// be destroyed.
 	}
 
 	void FVerifier::SetPaused(bool bInIsPaused)
@@ -222,40 +221,39 @@ namespace BuildPatchServices
 			int32 ThreadsToCreate = 1;
 			GConfig->GetInt(TEXT("BuildPatchTool"), TEXT("VerificationThreadCount"), ThreadsToCreate, GEngineIni);
 
-			ThreadsRemaining = ThreadsToCreate;			
+			int32 ThreadsRemaining = ThreadsToCreate;
 
+			FSemaphore ThreadDoneSem(0, ThreadsToCreate);
+			
 			// We wait to create our threads so that the shared threads from file construction are returned.
 			for (int i = 0; i < ThreadsToCreate; i++)
 			{
 				SecondaryProcessThreads.Add(SharedContext->CreateThread());
-				SecondaryProcessThreads[i]->RunTask([this]() { ProcessVerifyJobs(); });
+				SecondaryProcessThreads[i]->RunTask([this, &ThreadDoneSem]() { ProcessVerifyJobs(&ThreadDoneSem); });
 			}
 
 			for (;;)
 			{
-				if (AllThreadsDoneEvent->Wait(100))
-				{
-					// triggered which means we're all done.
-					break;
-				}
-
-				if (bShouldAbort)
+				if (ThreadDoneSem.TryAcquire(100))
 				{
-					// if we aborted without a trigger, then the threads could still be working and
-					// we need to wait until they are done.
-					AllThreadsDoneEvent->Wait();
-					break;
+					ThreadsRemaining--;
+					if (ThreadsRemaining == 0)
+					{
+						// We acquired all the counts so we're done.
+						break;
+					}
 				}
 
 				// Update our overall progress tracker.
 				VerifierStat->OnProcessedDataUpdated(ThreadProcessedBytes.load(std::memory_order_acquire));
 			}
 
-			// At this point we should not have any threads still doing work!
-			int32 LocalThreadsRemaining = ThreadsRemaining;
-			if (LocalThreadsRemaining != 0)
+			// If our event is triggered then something has gone very wrong since we could not have triggered it.
+			// This is still here to try and figure out what's been going wrong with it before we switched to a semaphore.
+			bool bEventIsTriggered = AllThreadsDoneEvent->Wait(0);
+			if (!ensure(!bEventIsTriggered))
 			{
-				UE_LOG(LogVerifier, Error, TEXT("Verifier passed our Wait() call with threads remaining = %d!"), LocalThreadsRemaining);
+				UE_LOG(LogVerifier, Error, TEXT("Verifier has a triggered event!"));
 			}
 
 			for (int i = 0; i < ThreadsToCreate; i++)
@@ -267,7 +265,7 @@ namespace BuildPatchServices
 		else
 		{
 			// Can't create threads, just do the work here.
-			ProcessVerifyJobs();
+			ProcessVerifyJobs(nullptr);
 		}
 
 		FPlatformProcess::ReturnSynchEventToPool(AllThreadsDoneEvent);
