BPS: Fix for overflows when constructing in memory files >4gb.
BPS: Fix for construct in memory limits not checking if there's space for the next file.

#rnx

diff --git a/Engine/Source/Runtime/Online/BuildPatchServices/Private/BuildPatchFileConstructor.cpp b/Engine/Source/Runtime/Online/BuildPatchServices/Private/BuildPatchFileConstructor.cpp
--- a/Engine/Source/Runtime/Online/BuildPatchServices/Private/BuildPatchFileConstructor.cpp
+++ b/Engine/Source/Runtime/Online/BuildPatchServices/Private/BuildPatchFileConstructor.cpp
@@ -2173,15 +2173,15 @@ struct FRequestInfo
 {
 	struct FRequestSplat
 	{
-		int32 DestinationOffset;
-		int32 OffsetInChunk;
-		int32 BytesToCopy;
+		uint64 DestinationOffset;
+		uint32 OffsetInChunk;
+		uint32 BytesToCopy;
 	};
 
 	FGuid Guid;
 
 	// We can do a lot of shortcuts if we are working with the entire chunk
-	int32 ChunkSize = 0;
+	uint32 ChunkSize = 0;
 
 	// We could request the same guid multiple times
 	// for the same buffer... in this case we want one
@@ -2521,7 +2521,7 @@ void FBuildPatchFileConstructor::StartReadBatch(FFileConstructionState& CurrentF
 	int32 EndChunkIdx = Batch.StartChunkPartIndex;
 	{
 		TRACE_CPUPROFILER_EVENT_SCOPE(CFFC_Count);
-		uint32 BufferFillLevel = 0;
+		uint64 BufferFillLevel = 0;
 		for (; EndChunkIdx < CurrentFile.FileManifest->ChunkParts.Num(); ++EndChunkIdx)
 		{
 			const FChunkPart& ChunkPart = CurrentFile.FileManifest->ChunkParts[EndChunkIdx];
@@ -3624,8 +3624,13 @@ void FBuildPatchFileConstructor::ConstructFiles(const FResumeData& ResumeData)
 						// No more work do to on active files - is there another to start?
 						if (!FileToStart)
 						{
+							// This is not a race because the only place we ever increment is this thread, this function right
+							// below.
+							int32 IndexToConstruct = NextIndexToConstruct.load(std::memory_order_acquire);
+							bool bAnotherFileExists = IndexToConstruct < ConstructionList.Num();
+
 							bool bAllowAnotherFile = ActiveFiles.Num() == 0;
-							if (!bAllowAnotherFile)
+							if (!bAllowAnotherFile && bAnotherFileExists)
 							{
 								if (bAllowMultipleFilesInFlight)
 								{
@@ -3651,15 +3656,15 @@ void FBuildPatchFileConstructor::ConstructFiles(const FResumeData& ResumeData)
 											TotalActiveData += Data.Num();
 										}
 
-										bAllowAnotherFile = TotalActiveData < IOBufferSize;
+										// Be sure to include the size of the next file. We don't get here at all
+										// if there are no active files.
+										uint64 NeededData = ConstructionList[IndexToConstruct].FileManifest->FileSize;
+
+										bAllowAnotherFile = (TotalActiveData + NeededData) < IOBufferSize;
 									}
 								}
 							}
 
-							// This is not a race because the only place we ever increment is this thread, this function right
-							// below.
-							int32 IndexToConstruct = NextIndexToConstruct.load(std::memory_order_acquire);
-							bool bAnotherFileExists = IndexToConstruct < ConstructionList.Num();
 							
 							if (bAllowAnotherFile && bAnotherFileExists)
 							{
