Use constants for http version option;
Get rid of CVarHttpCurlSetResponseCodeWhenMarkAsCompleted and CVarHttpCurlReadContentLengthWhenFinish.

diff --git a/Engine/Source/Programs/WebTests/Private/TestHttp.cpp b/Engine/Source/Programs/WebTests/Private/TestHttp.cpp
--- a/Engine/Source/Programs/WebTests/Private/TestHttp.cpp
+++ b/Engine/Source/Programs/WebTests/Private/TestHttp.cpp
@@ -6,6 +6,7 @@
 #include "HAL/PlatformTime.h"
 #include "HAL/Runnable.h"
 #include "HAL/RunnableThread.h"
+#include "HttpConstants.h"
 #include "HttpManager.h"
 #include "HttpRetrySystem.h"
 #include "Http.h"
@@ -2993,7 +2994,7 @@ TEST_CASE_METHOD(FWaitUntilCompleteHttpFixture, "Http Methods over http2", HTTP_
 {
 	TSharedRef<IHttpRequest> HttpRequest = CreateRequest();
 	HttpRequest->SetURL(UrlToTestMethods());
-	HttpRequest->SetOption(HttpRequestOptions::HttpVersion, TEXT("2"));
+	HttpRequest->SetOption(HttpRequestOptions::HttpVersion, FHttpConstants::VERSION_2TLS);
 	HttpRequest->OnProcessRequestComplete().BindLambda([](FHttpRequestPtr HttpRequest, FHttpResponsePtr HttpResponse, bool bSucceeded) {
 		CHECK(bSucceeded);
 		REQUIRE(HttpResponse != nullptr);
diff --git a/Engine/Source/Runtime/Online/HTTP/Private/Curl/CurlHttp.cpp b/Engine/Source/Runtime/Online/HTTP/Private/Curl/CurlHttp.cpp
--- a/Engine/Source/Runtime/Online/HTTP/Private/Curl/CurlHttp.cpp
+++ b/Engine/Source/Runtime/Online/HTTP/Private/Curl/CurlHttp.cpp
@@ -5,6 +5,7 @@
 #include "Curl/CurlHttp.h"
 #include "Stats/Stats.h"
 #include "Misc/App.h"
+#include "HttpConstants.h"
 #include "HttpModule.h"
 #include "Http.h"
 #include "Misc/EngineVersion.h"
@@ -19,20 +20,6 @@
 #include <openssl/ssl.h>
 #endif
 
-TAutoConsoleVariable<bool> CVarHttpCurlSetResponseCodeWhenMarkAsCompleted(
-	TEXT("http.CurlSetResponseCodeWhenMarkAsCompleted"),
-	true,
-	TEXT("Set response code in when mark as completed, instead of setting it in FinishRequest."),
-	ECVF_SaveForNextBoot
-);
-
-TAutoConsoleVariable<bool> CVarHttpCurlReadContentLengthWhenFinish(
-	TEXT("http.CurlReadContentLengthWhenFinish"),
-	false,
-	TEXT("Whether to read content length from curl handle when finish request"),
-	ECVF_SaveForNextBoot
-);
-
 TAutoConsoleVariable<bool> CVarHttpCurlAllowHTTP2(
 	TEXT("http.CurlAllowHTTP2"),
 	false,
@@ -861,30 +848,8 @@ bool FCurlHttpRequest::SetupRequestHttpThread()
 
 		curl_easy_setopt(EasyHandle, CURLOPT_URL, TCHAR_TO_ANSI(*GetURL()));
 
-#if UE_HTTP_SUPPORT_UNIX_SOCKET
-		FString UnixSocketPath = GetOption(HttpRequestOptions::UnixSocketPath);
-		if (UnixSocketPath.Len() > 0)
-		{
-			curl_easy_setopt(EasyHandle, CURLOPT_UNIX_SOCKET_PATH, TCHAR_TO_ANSI(*UnixSocketPath));
-		}
-#endif //UE_HTTP_SUPPORT_UNIX_SOCKET
-
-		const FString HttpVersion = GetOption(HttpRequestOptions::HttpVersion);
-		if (!HttpVersion.IsEmpty())
-		{
-			if (HttpVersion == TEXT("2"))
-			{
-				curl_easy_setopt(EasyHandle, CURLOPT_HTTP_VERSION, CURL_HTTP_VERSION_2TLS);
-			}
-			else if (HttpVersion == TEXT("1.1"))
-			{
-				curl_easy_setopt(EasyHandle, CURLOPT_HTTP_VERSION, CURL_HTTP_VERSION_1_1);
-			}
-			else
-			{
-				UE_LOG(LogHttp, Warning, TEXT("Unsupported http version %s"), *HttpVersion);
-			}
-		}
+		SetupOptionUnixSocketPath();
+		SetupOptionHttpVersion();
 
 		if (!FCurlHttpManager::CurlRequestOptions.LocalHostAddr.IsEmpty())
 		{
@@ -1046,6 +1011,37 @@ bool FCurlHttpRequest::SetupRequestHttpThread()
 	return true;
 }
 
+void FCurlHttpRequest::SetupOptionUnixSocketPath()
+{
+#if UE_HTTP_SUPPORT_UNIX_SOCKET
+		FString UnixSocketPath = GetOption(HttpRequestOptions::UnixSocketPath);
+		if (UnixSocketPath.Len() > 0)
+		{
+			curl_easy_setopt(EasyHandle, CURLOPT_UNIX_SOCKET_PATH, TCHAR_TO_ANSI(*UnixSocketPath));
+		}
+#endif //UE_HTTP_SUPPORT_UNIX_SOCKET
+}
+
+void FCurlHttpRequest::SetupOptionHttpVersion()
+{
+	const FString HttpVersion = GetOption(HttpRequestOptions::HttpVersion);
+	if (!HttpVersion.IsEmpty())
+	{
+		if (HttpVersion == FHttpConstants::VERSION_2TLS)
+		{
+			curl_easy_setopt(EasyHandle, CURLOPT_HTTP_VERSION, CURL_HTTP_VERSION_2TLS);
+		}
+		else if (HttpVersion == FHttpConstants::VERSION_1_1)
+		{
+			curl_easy_setopt(EasyHandle, CURLOPT_HTTP_VERSION, CURL_HTTP_VERSION_1_1);
+		}
+		else
+		{
+			UE_LOG(LogHttp, Warning, TEXT("Unsupported http version %s"), *HttpVersion);
+		}
+	}
+}
+
 void FCurlHttpRequest::CleanupRequestHttpThread()
 {
 	curl_easy_setopt(EasyHandle, CURLOPT_SHARE, nullptr);
@@ -1190,16 +1186,13 @@ void FCurlHttpRequest::MarkAsCompleted(CURLcode InCurlCompletionResult)
 	CurlCompletionResult = InCurlCompletionResult;
 	bCurlRequestCompleted = true;
 
-	if (CVarHttpCurlSetResponseCodeWhenMarkAsCompleted.GetValueOnAnyThread())
+	if (TSharedPtr<FCurlHttpResponse> Response = StaticCastSharedPtr<FCurlHttpResponse>(ResponseCommon))
 	{
-		if (TSharedPtr<FCurlHttpResponse> Response = StaticCastSharedPtr<FCurlHttpResponse>(ResponseCommon))
+		// get the information
+		long HttpCode = 0;
+		if (CURLE_OK == curl_easy_getinfo(EasyHandle, CURLINFO_RESPONSE_CODE, &HttpCode))
 		{
-			// get the information
-			long HttpCode = 0;
-			if (CURLE_OK == curl_easy_getinfo(EasyHandle, CURLINFO_RESPONSE_CODE, &HttpCode))
-			{
-				Response->SetResponseCode(HttpCode);
-			}
+			Response->SetResponseCode(HttpCode);
 		}
 	}
 
@@ -1223,31 +1216,10 @@ void FCurlHttpRequest::FinishRequest()
 		{
 			Response->bSucceeded = (CURLE_OK == CurlCompletionResult);
 
-			if (!CVarHttpCurlSetResponseCodeWhenMarkAsCompleted.GetValueOnAnyThread())
-			{
-				// get the information
-				long HttpCode = 0;
-				if (CURLE_OK == curl_easy_getinfo(EasyHandle, CURLINFO_RESPONSE_CODE, &HttpCode))
-				{
-					Response->SetResponseCode(HttpCode);
-				}
-			}
-
 			// If content length wasn't received through response header 
 			if (Response->ContentLength == 0)
 			{
-				curl_off_t ContentLengthDownload = 0;
-				if (CVarHttpCurlReadContentLengthWhenFinish.GetValueOnAnyThread() && 
-					CURLE_OK == curl_easy_getinfo(EasyHandle, CURLINFO_CONTENT_LENGTH_DOWNLOAD_T, &ContentLengthDownload) && 
-					ContentLengthDownload > 0.0)
-				{
-					Response->ContentLength = static_cast< uint64 >(ContentLengthDownload);
-				}
-				else
-				{
-					// If curl did not know how much we downloaded, or we were missing a Content-Length header (Chunked request), set our ContentLength as the amount we downloaded
-					Response->ContentLength = TotalBytesRead;
-				}
+				Response->ContentLength = TotalBytesRead;
 			}
 
 			if (Response->GetResponseCode() <= 0 && GetURL().StartsWith(TEXT("Http"), ESearchCase::IgnoreCase))
diff --git a/Engine/Source/Runtime/Online/HTTP/Private/Curl/CurlHttp.h b/Engine/Source/Runtime/Online/HTTP/Private/Curl/CurlHttp.h
--- a/Engine/Source/Runtime/Online/HTTP/Private/Curl/CurlHttp.h
+++ b/Engine/Source/Runtime/Online/HTTP/Private/Curl/CurlHttp.h
@@ -360,6 +360,9 @@ private:
 	virtual FHttpResponsePtr CreateResponse() override;
 	virtual void MockResponseData() override;
 
+	void SetupOptionUnixSocketPath();
+	void SetupOptionHttpVersion();
+
 private:
 
 	/** Pointer to an easy handle specific to this request */
diff --git a/Engine/Source/Runtime/Online/HTTP/Private/HttpConstants.cpp b/Engine/Source/Runtime/Online/HTTP/Private/HttpConstants.cpp
--- a/Engine/Source/Runtime/Online/HTTP/Private/HttpConstants.cpp
+++ b/Engine/Source/Runtime/Online/HTTP/Private/HttpConstants.cpp
@@ -26,3 +26,7 @@ const TCHAR* const FHttpConstants::MEDIATYPE_OCTET_STREAM{ TEXT("application/oct
 
 const TCHAR* const FHttpConstants::MEDIAENCODING_GZIP{ TEXT("gzip") };
 
+
+const TCHAR* const FHttpConstants::VERSION_2TLS{ TEXT("2TLS") };
+const TCHAR* const FHttpConstants::VERSION_1_1{ TEXT("1_1") };
+
diff --git a/Engine/Source/Runtime/Online/HTTP/Public/HttpConstants.h b/Engine/Source/Runtime/Online/HTTP/Public/HttpConstants.h
--- a/Engine/Source/Runtime/Online/HTTP/Public/HttpConstants.h
+++ b/Engine/Source/Runtime/Online/HTTP/Public/HttpConstants.h
@@ -53,6 +53,10 @@ struct FHttpConstants
 
 	/** gzip */
 	static UE_API const TCHAR* const MEDIAENCODING_GZIP;
+
+	/** version */
+	static UE_API const TCHAR* const VERSION_2TLS;
+	static UE_API const TCHAR* const VERSION_1_1;
 };
 
 #undef UE_API
