Update serialization error handling to be more correct until we have actual errors from low level UE.

#rnx

diff --git a/Engine/Source/Runtime/Online/BuildPatchServices/Private/BuildPatchFileConstructor.cpp b/Engine/Source/Runtime/Online/BuildPatchServices/Private/BuildPatchFileConstructor.cpp
--- a/Engine/Source/Runtime/Online/BuildPatchServices/Private/BuildPatchFileConstructor.cpp
+++ b/Engine/Source/Runtime/Online/BuildPatchServices/Private/BuildPatchFileConstructor.cpp
@@ -190,7 +190,8 @@ enum class EConstructionError : uint8
 	OutboundDataError,
 	InternalConsistencyError,
 	Aborted,
-	MissingFileInfo
+	MissingFileInfo,
+	CloseError
 };
 
 
@@ -2930,6 +2931,8 @@ void FBuildPatchFileConstructor::CompleteConstructedFile(FFileConstructionState&
 		{
 			if (CurrentFile.bSuccess && CurrentFile.NewFile->IsError())
 			{
+				// This should aleady have been caught during the write!
+				UE_LOG(LogBuildPatchServices, Warning, TEXT("Got serialize error during CompleteConstructedFile! Should have already been caught."));
 				CurrentFile.ConstructionError = EConstructionError::SerializeError;
 				CurrentFile.bSuccess = false;
 			}
@@ -2947,7 +2950,7 @@ void FBuildPatchFileConstructor::CompleteConstructedFile(FFileConstructionState&
 				// Check for final success
 				if (CurrentFile.bSuccess && !bArchiveSuccess)
 				{
-					CurrentFile.ConstructionError = EConstructionError::SerializeError;
+					CurrentFile.ConstructionError = EConstructionError::CloseError;
 					CurrentFile.bSuccess = false;
 				}
 			}
@@ -3107,45 +3110,6 @@ void FBuildPatchFileConstructor::CompleteConstructedFile(FFileConstructionState&
 	// Report errors.
 	if (!CurrentFile.bSuccess)
 	{
-		if (!Configuration.bInstallToMemory &&
-			CurrentFile.ConstructionError == EConstructionError::SerializeError)
-		{
-			// Serialize error is our catchall file error right now. This should probably get
-			// migrated such that it's when we fail to load an existing chunk (i.e. corruption)
-			// but instead that shows up as a missing chunk.
-
-			// Right now however we get this whenever we fail to write due to our disk space disappearing,
-			// and we'd like to report that, so we check and see if we're low on space if it looks like
-			// that's the problem.
-
-			uint64 TotalSize = 0;
-			uint64 FreeSize = 0;
-			if (FPlatformMisc::GetDiskTotalAndFreeSpace(Configuration.InstallDirectory, TotalSize, FreeSize))
-			{
-				// Don't bother recalculating our disk space requirements - our original requirement is still
-				// valid, the only thing that's changed is someone took our space.
-				AvailableDiskSpace.store(FreeSize, std::memory_order_release);
-
-				// We're responding to an actual failure, which would have happened because we literally weren't able
-				// to write our write butter. Because of transient stuff this might not be correct so we double our
-				// write buffer size for this check.
-				if (FreeSize < ((uint64)MaxWriteBatchSize << 1))
-				{
-					CurrentFile.ConstructionError = EConstructionError::OutOfDiskSpace;
-				}
-				else
-				{
-					// If it looks like we had enough disk space to write the last buffer, then 
-					// leave it as serialize.
-				}
-			}
-			else
-			{
-				// If we can't get the free space then likely the disk has disconnected or otherwise had a Bad Error, leave
-				// as serialize.
-			}
-		}
-
 		const bool bReportAnalytic = InstallerError->HasError() == false;
 		switch (CurrentFile.ConstructionError)
 		{
@@ -3182,6 +3146,16 @@ void FBuildPatchFileConstructor::CompleteConstructedFile(FFileConstructionState&
 				InstallerError->SetError(EBuildPatchInstallError::FileConstructionFail, ConstructionErrorCodes::FileCreateFail, CurrentFile.CreateFilePlatformLastError);
 				break;
 			}
+		case EConstructionError::CloseError:
+			{
+				if (bReportAnalytic)
+				{
+					InstallerAnalytics->RecordConstructionError(CurrentFile.BuildFilename, INDEX_NONE, TEXT("Could Not Close File"));
+				}
+				UE_LOG(LogBuildPatchServices, Error, TEXT("FBuildPatchFileConstructor: Could not close %s"), *CurrentFile.BuildFilename);
+				InstallerError->SetError(EBuildPatchInstallError::FileConstructionFail, ConstructionErrorCodes::CloseFileError);
+				break;
+			}
 		case EConstructionError::MissingChunk:
 			{
 				if (bReportAnalytic)
@@ -3258,6 +3232,7 @@ void FBuildPatchFileConstructor::CompleteConstructedFile(FFileConstructionState&
 
 			// Errors we expect there to be issues with the outbound file:
 			case EConstructionError::CannotCreateFile:
+			case EConstructionError::CloseError:
 			case EConstructionError::SerializeError:
 			case EConstructionError::TrackingError:
 			case EConstructionError::OutboundDataError:
@@ -3480,10 +3455,11 @@ void FBuildPatchFileConstructor::ConstructFiles(const FResumeData& ResumeData)
 					{
 						File->bSuccess = false;
 
-						// SerializeError signals the error management code to check for low disk space
-						// since we currently have no way of routing out of disk errors from the UE low level
-						// writing code.
-						File->ConstructionError = EConstructionError::SerializeError;
+						// This is not great, but right now the FArchive can't report what the error was. Without any
+						// way to tease it apart, we run with the most likely issue: we ran out of disk space. Unfortunately
+						// it's unreliable to ask the file system what is available and check if it's low and use that for the 
+						// conversion.
+						File->ConstructionError = EConstructionError::OutOfDiskSpace;
 					}
 				}
 
diff --git a/Engine/Source/Runtime/Online/BuildPatchServices/Private/Installer/InstallerError.h b/Engine/Source/Runtime/Online/BuildPatchServices/Private/Installer/InstallerError.h
--- a/Engine/Source/Runtime/Online/BuildPatchServices/Private/Installer/InstallerError.h
+++ b/Engine/Source/Runtime/Online/BuildPatchServices/Private/Installer/InstallerError.h
@@ -66,6 +66,7 @@ namespace BuildPatchServices
 		static const TCHAR* SerializationError = TEXT("06");
 		static const TCHAR* TrackingError = TEXT("07");
 		static const TCHAR* InternalConsistencyFailure = TEXT("08");
+		static const TCHAR* CloseFileError = TEXT("09");
 	}
 
 	/**
