Use event loop for websockets by default if it's supported on the platform.

diff --git a/Engine/Source/Runtime/Online/WebSockets/Private/Lws/LwsWebSocketsManager.h b/Engine/Source/Runtime/Online/WebSockets/Private/Lws/LwsWebSocketsManager.h
--- a/Engine/Source/Runtime/Online/WebSockets/Private/Lws/LwsWebSocketsManager.h
+++ b/Engine/Source/Runtime/Online/WebSockets/Private/Lws/LwsWebSocketsManager.h
@@ -120,7 +120,11 @@ private:
 	 * 
 	 * If false, lws_service will use a ServiceTimoutMs timeout, allowing lws to block and handle wake / sleep on activity internally itself.
 	 */
+#if UE_WEBSOCETS_PLATFORM_SUPPORT_EVENT_LOOP
+	bool bPollService = false;
+#else
 	bool bPollService = true;
+#endif
 	
 	/** Timeout value in milliseconds to be used by lws_service when bPollService is false */
 	int32 ServiceTimeoutMs = MAX_int32;
diff --git a/Engine/Source/Runtime/Online/WebSockets/WebSockets.Build.cs b/Engine/Source/Runtime/Online/WebSockets/WebSockets.Build.cs
--- a/Engine/Source/Runtime/Online/WebSockets/WebSockets.Build.cs
+++ b/Engine/Source/Runtime/Online/WebSockets/WebSockets.Build.cs
@@ -19,12 +19,7 @@ public class WebSockets : ModuleRules
 
 	protected virtual bool bPlatformSupportEventLoop
 	{
-		get
-		{
-			return
-				Target.Platform == UnrealTargetPlatform.Win64 ||
-				Target.IsInPlatformGroup(UnrealPlatformGroup.Unix);
-		}
+		get => true;
 	}
 
 	protected virtual bool bPlatformSupportsWinHttpWebSockets
