[HTTP]
Extract member Payload and Headers from response impls to FHttpResponseCommon;
Moved member function GetHeader, GetContent and GetContentAsString into FHttpResponseCommon;
Renamed CVarHttpUrlPatternsToMockFailure to CVarHttpUrlPatternsToMockResponse;
Add support of mocking response payload.

diff --git a/Engine/Source/Programs/WebTests/Private/TestHttp.cpp b/Engine/Source/Programs/WebTests/Private/TestHttp.cpp
--- a/Engine/Source/Programs/WebTests/Private/TestHttp.cpp
+++ b/Engine/Source/Programs/WebTests/Private/TestHttp.cpp
@@ -42,7 +42,7 @@
 extern TAutoConsoleVariable<bool> CVarHttpInsecureProtocolEnabled;
 extern TAutoConsoleVariable<int32> CVarHttpMaxConcurrentRequests;
 extern TAutoConsoleVariable<int32> CVarHttpEventLoopEnableChance;
-extern TAutoConsoleVariable<FString> CVarHttpUrlPatternsToMockFailure;
+extern TAutoConsoleVariable<FString> CVarHttpUrlPatternsToMockResponse;
 extern TAutoConsoleVariable<FString> CVarHttpUrlPatternsToDisableFailedLog;
 
 class FMockHttpModule : public FHttpModule
@@ -485,7 +485,7 @@ TEST_CASE_METHOD(FWaitUntilCompleteHttpFixture, "Can mock connection error throu
 {
 	DisableWarningsInThisTest();
 
-	CVarHttpUrlPatternsToMockFailure->Set(TEXT("epicgames.com->0 unrealengine.com->503"));
+	CVarHttpUrlPatternsToMockResponse->Set(TEXT("epicgames.com->0 unrealengine.com->503"));
 
 	float ExpectedTimeoutDuration = 2.0f;
 	HttpModule->HttpConnectionTimeout = ExpectedTimeoutDuration;
@@ -499,14 +499,14 @@ TEST_CASE_METHOD(FWaitUntilCompleteHttpFixture, "Can mock connection error throu
 		CHECK(HttpRequest->GetFailureReason() == EHttpFailureReason::ConnectionError);
 		const double DurationInSeconds  = FPlatformTime::Seconds() - StartTime;
 		CHECK(FMath::IsNearlyEqual(DurationInSeconds, ExpectedTimeoutDuration, HTTP_TIME_DIFF_TOLERANCE_OF_REQUEST));
-		CVarHttpUrlPatternsToMockFailure->Set(TEXT(""));
+		CVarHttpUrlPatternsToMockResponse->Set(TEXT(""));
 	});
 	HttpRequest->ProcessRequest();
 }
 
 TEST_CASE_METHOD(FWaitUntilCompleteHttpFixture, "Can mock response failure through CVar", HTTP_TAG)
 {
-	CVarHttpUrlPatternsToMockFailure->Set(TEXT("epicgames.com->0 unrealengine.com->503"));
+	CVarHttpUrlPatternsToMockResponse->Set(TEXT("epicgames.com->0 unrealengine.com->503"));
 
 	TSharedRef<IHttpRequest> HttpRequest = CreateRequest();
 	HttpRequest->SetVerb(TEXT("GET"));
@@ -515,7 +515,25 @@ TEST_CASE_METHOD(FWaitUntilCompleteHttpFixture, "Can mock response failure throu
 		CHECK(bSucceeded);
 		REQUIRE(HttpResponse != nullptr);
 		CHECK(HttpResponse->GetResponseCode() == 503);
-		CVarHttpUrlPatternsToMockFailure->Set(TEXT(""));
+		CVarHttpUrlPatternsToMockResponse->Set(TEXT(""));
+	});
+	HttpRequest->ProcessRequest();
+}
+
+TEST_CASE_METHOD(FWaitUntilCompleteHttpFixture, "Can mock successful response through CVar", HTTP_TAG)
+{
+	CVarHttpUrlPatternsToMockResponse->Set(TEXT("epicgames.com->200 unrealengine_wrong_domain.com->200++{\"var_str\":\"abcd\",\"var_int\":20}++Content-Type:application/json"));
+
+	TSharedRef<IHttpRequest> HttpRequest = CreateRequest();
+	HttpRequest->SetVerb(FHttpConstants::VERB_GET);
+	HttpRequest->SetURL(TEXT("https://www.unrealengine_wrong_domain.com/"));
+	HttpRequest->OnProcessRequestComplete().BindLambda([](FHttpRequestPtr HttpRequest, FHttpResponsePtr HttpResponse, bool bSucceeded) {
+		CHECK(bSucceeded);
+		REQUIRE(HttpResponse != nullptr);
+		CHECK(HttpResponse->GetResponseCode() == 200);
+		CHECK(HttpResponse->GetHeader(FHttpConstants::HEADER_CONTENT_TYPE) == FHttpConstants::MEDIATYPE_JSON);
+		CHECK(HttpResponse->GetContentAsUtf8StringView() == "{\"var_str\":\"abcd\",\"var_int\":20}");
+		CVarHttpUrlPatternsToMockResponse->Set(TEXT(""));
 	});
 	HttpRequest->ProcessRequest();
 }
@@ -577,7 +595,7 @@ TEST_CASE_METHOD(FEarlyCompleteHttpFixture, "Can delete failed request during co
 
 	DisableWarningsInThisTest();
 
-	CVarHttpUrlPatternsToMockFailure->Set(TEXT("epicgames.com->0 unrealengine.com->503"));
+	CVarHttpUrlPatternsToMockResponse->Set(TEXT("epicgames.com->0 unrealengine.com->503"));
 
 	HttpRequest = CreateRequest();
 
@@ -592,7 +610,7 @@ TEST_CASE_METHOD(FEarlyCompleteHttpFixture, "Can delete failed request during co
 	}
 
 	HttpRequest->OnProcessRequestComplete().BindLambda([this](FHttpRequestPtr InHttpRequest, FHttpResponsePtr HttpResponse, bool bSucceeded) {
-		CVarHttpUrlPatternsToMockFailure->Set(TEXT(""));
+		CVarHttpUrlPatternsToMockResponse->Set(TEXT(""));
 		HttpRequest.Reset();
 	});
 
@@ -2413,7 +2431,7 @@ TEST_CASE_METHOD(FWaitUntilCompleteHttpFixture, "Retry request with CompleteOnHt
 	}
 	SECTION("With mock connection failure for valid domain")
 	{
-		CVarHttpUrlPatternsToMockFailure->Set(TEXT("www.epicgames.com->0"));
+		CVarHttpUrlPatternsToMockResponse->Set(TEXT("www.epicgames.com->0"));
 		OriginalUrl = UrlToTestMethods();
 		OriginalUrl = OriginalUrl.Replace(*UrlDomain(), TEXT("www.epicgames.com"));
 	}
@@ -2443,14 +2461,14 @@ TEST_CASE_METHOD(FWaitUntilCompleteHttpFixture, "Retry request with CompleteOnHt
 	});
 	HttpRequest->ProcessRequest();
 
-	CVarHttpUrlPatternsToMockFailure->Set(TEXT(""));
+	CVarHttpUrlPatternsToMockResponse->Set(TEXT(""));
 }
 
 TEST_CASE_METHOD(FWaitUntilQuitFromTestFixture, "Will keep request reference with mock connection failure which using CompleteOnHttpThread policy until mocked connection timeout", HTTP_TAG)
 {
 	DisableWarningsInThisTest();
 
-	CVarHttpUrlPatternsToMockFailure->Set(TEXT("www.epicgames.com->0"));
+	CVarHttpUrlPatternsToMockResponse->Set(TEXT("www.epicgames.com->0"));
 	HttpModule->HttpConnectionTimeout = 1.0f;
 
 	TSharedRef<IHttpRequest> HttpRequest = CreateRequest();
@@ -2464,7 +2482,7 @@ TEST_CASE_METHOD(FWaitUntilQuitFromTestFixture, "Will keep request reference wit
 		bQuitRequested = true;
 	});
 	HttpRequest->ProcessRequest();
-	CVarHttpUrlPatternsToMockFailure->Set(TEXT(""));
+	CVarHttpUrlPatternsToMockResponse->Set(TEXT(""));
 }
 
 TEST_CASE_METHOD(FWaitUntilCompleteHttpFixture, "Retry will add default domain to alt domains", HTTP_TAG)
@@ -2476,7 +2494,7 @@ TEST_CASE_METHOD(FWaitUntilCompleteHttpFixture, "Retry will add default domain t
 
 	DisableWarningsInThisTest();
 
-	CVarHttpUrlPatternsToMockFailure->Set(TEXT("fakedomain->0"));
+	CVarHttpUrlPatternsToMockResponse->Set(TEXT("fakedomain->0"));
 
 	HttpModule->HttpConnectionTimeout = 0.1f;
 
@@ -2535,7 +2553,7 @@ TEST_CASE_METHOD(FWaitUntilCompleteHttpFixture, "Retry will add default domain t
 	HttpRequest->OnProcessRequestComplete().BindLambda([ExpectedDomainsOrderInRetry](FHttpRequestPtr HttpRequest, FHttpResponsePtr HttpResponse, bool bSucceeded) {
 		CHECK(ExpectedDomainsOrderInRetry->IsEmpty());
 		delete ExpectedDomainsOrderInRetry;
-		CVarHttpUrlPatternsToMockFailure->Set(TEXT(""));
+		CVarHttpUrlPatternsToMockResponse->Set(TEXT(""));
 
 		CHECK(!bSucceeded);
 		REQUIRE(HttpResponse == nullptr);
diff --git a/Engine/Source/Runtime/Online/HTTP/Private/Apple/AppleHttp.cpp b/Engine/Source/Runtime/Online/HTTP/Private/Apple/AppleHttp.cpp
--- a/Engine/Source/Runtime/Online/HTTP/Private/Apple/AppleHttp.cpp
+++ b/Engine/Source/Runtime/Online/HTTP/Private/Apple/AppleHttp.cpp
@@ -25,9 +25,6 @@ TAutoConsoleVariable<bool> CVarAppleHttpStartTaskInHttpThreadEnabled(
 
 @interface FAppleHttpResponseDelegate : NSObject<NSURLSessionDataDelegate>
 {
-	/** Holds the payload as we receive it. */
-	@public TArray<uint8> Payload;
-
 	// Flag to indicate the request was initialized with stream. In that case even if stream was set to 
 	// null later on internally, the request itself won't cache received data anymore
 	@public BOOL bInitializedWithValidStream;
@@ -150,7 +147,7 @@ TAutoConsoleVariable<bool> CVarAppleHttpStartTaskInHttpThreadEnabled(
 {
 	if (TSharedPtr<FAppleHttpRequest> Request = SourceRequest.Pin())
 	{
-		return Request->PassReceivedDataToStream(Ptr, InSize);
+		return Request->HandleResponseBodyDataReceived(reinterpret_cast<uint8*>(Ptr), InSize);
 	}
 	return false;
 }
@@ -163,21 +160,22 @@ TAutoConsoleVariable<bool> CVarAppleHttpStartTaskInHttpThreadEnabled(
 	}
 }
 
--(void) BroadcastResponseHeadersReceived
+- (void) HandleResponseHeadersReceived
 {
 	if (TSharedPtr<FAppleHttpRequest> Request = SourceRequest.Pin())
 	{
-		if (Request->GetDelegateThreadPolicy() == EHttpRequestDelegateThreadPolicy::CompleteOnHttpThread)
-		{
-			Request->BroadcastResponseHeadersReceived();
-		}
-		else if (Request->OnHeaderReceived().IsBound())
+		TMap<FString, FString> ResponseHeaders;
+		if (NSDictionary* Headers = [self GetResponseHeaders])
 		{
-			FHttpModule::Get().GetHttpManager().AddGameThreadTask([Request]()
+			for (NSString* Key in [Headers allKeys])
 			{
-				Request->BroadcastResponseHeadersReceived();
-			});
+				FString ConvertedValue([Headers objectForKey:Key]);
+				FString ConvertedKey(Key);
+				ResponseHeaders.Emplace(MoveTemp(ConvertedKey), MoveTemp(ConvertedValue));
+			}
 		}
+
+		Request->HandleResponseHeadersReceived(MoveTemp(ResponseHeaders));
 	}
 }
 
@@ -229,14 +227,10 @@ TAutoConsoleVariable<bool> CVarAppleHttpStartTaskInHttpThreadEnabled(
 	FString EffectiveURL([Url absoluteString]);
 	[self SaveEffectiveURL: EffectiveURL];
 
-	[self BroadcastResponseHeadersReceived];
+	[self HandleResponseHeadersReceived];
 
 	uint64 ExpectedResponseLength = response.expectedContentLength;
-	if(!bInitializedWithValidStream && ExpectedResponseLength != NSURLResponseUnknownLength)
-	{
-		Payload.Empty(ExpectedResponseLength);
-	}
-	UE_LOG(LogHttp, VeryVerbose, TEXT("URLSession:dataTask:didReceiveResponse:completionHandler: expectedContentLength = %lld. Length = %llu: %p"), ExpectedResponseLength, Payload.Max(), self);
+	UE_LOG(LogHttp, VeryVerbose, TEXT("URLSession:dataTask:didReceiveResponse:completionHandler: expectedContentLength = %lld"), ExpectedResponseLength);
 	completionHandler(NSURLSessionResponseAllow);
 }
 
@@ -248,26 +242,16 @@ TAutoConsoleVariable<bool> CVarAppleHttpStartTaskInHttpThreadEnabled(
 	}
 	
 	__block int64 NewBytesReceived = 0;
-	if (bInitializedWithValidStream)
-	{
-		__block bool bSerializeSucceed = false;
-		[data enumerateByteRangesUsingBlock:^(const void *bytes, NSRange byteRange, BOOL *stop) {
-			NewBytesReceived += byteRange.length;
-			bSerializeSucceed = [self HandleBodyDataReceived : const_cast<void*>(bytes) Size : byteRange.length];
-			*stop = bSerializeSucceed? NO : YES;
-		}];
-		
-		if (!bSerializeSucceed)
-		{
-			[dataTask cancel];
-		}
-	}
-	else
+	__block bool bSerializeSucceed = false;
+	[data enumerateByteRangesUsingBlock:^(const void *bytes, NSRange byteRange, BOOL *stop) {
+		NewBytesReceived += byteRange.length;
+		bSerializeSucceed = [self HandleBodyDataReceived : const_cast<void*>(bytes) Size : byteRange.length];
+		*stop = bSerializeSucceed? NO : YES;
+	}];
+	
+	if (!bSerializeSucceed)
 	{
-		[data enumerateByteRangesUsingBlock:^(const void *bytes, NSRange byteRange, BOOL *stop) {
-			NewBytesReceived += byteRange.length;
-			Payload.Append((const uint8*)bytes, byteRange.length);
-		}];
+		[dataTask cancel];
 	}
 	// Keep BytesReceived as a separated value to avoid concurrent accesses to Payload
 	self.BytesReceived += NewBytesReceived;
@@ -936,10 +920,44 @@ void FAppleHttpRequest::Tick(float DeltaSeconds)
 }
 
 bool FAppleHttpRequest::IsInitializedWithValidStream() const
-{ 
+{
 	return bInitializedWithValidStream;
 }
 
+void FAppleHttpRequest::HandleResponseHeadersReceived(TMap<FString, FString>&& ResponseHeaders)
+{
+	FHttpResponsePtr HttpResponse = GetResponse();
+	TSharedPtr<FAppleHttpResponse> AppHttpResponse = StaticCastSharedPtr<FAppleHttpResponse>(HttpResponse);
+	AppHttpResponse->SetHeaders(MoveTemp(ResponseHeaders));
+
+	if (GetDelegateThreadPolicy() == EHttpRequestDelegateThreadPolicy::CompleteOnHttpThread)
+	{
+		BroadcastResponseHeadersReceived();
+	}
+	else if (OnHeaderReceived().IsBound())
+	{
+		FHttpModule::Get().GetHttpManager().AddGameThreadTask([this, StrongThis = AsShared()]()
+		{
+			BroadcastResponseHeadersReceived();
+		});
+	}
+}
+
+bool FAppleHttpRequest::HandleResponseBodyDataReceived(uint8* Ptr, uint64 Size)
+{
+	if (bInitializedWithValidStream)
+	{
+		return PassReceivedDataToStream(Ptr, Size);
+	}
+	else
+	{
+		FHttpResponsePtr HttpResponse = GetResponse();
+		TSharedPtr<FAppleHttpResponse> AppleHttpResponse = StaticCastSharedPtr<FAppleHttpResponse>(HttpResponse);
+		AppleHttpResponse->AppendToPayload(reinterpret_cast<uint8*>(Ptr), Size);
+		return true;
+	}
+}
+
 void FAppleHttpRequest::CheckProgressDelegate()
 {
 	TSharedPtr<FAppleHttpResponse> Response = StaticCastSharedPtr<FAppleHttpResponse>(ResponseCommon);
@@ -1029,25 +1047,14 @@ void FAppleHttpResponse::SetNewAppleHttpEventDelegate(FNewAppleHttpEventDelegate
 	ResponseDelegate->NewAppleHttpEventDelegate = MoveTemp(Delegate);
 }
 
-void FAppleHttpResponse::CleanSharedObjects()
+void FAppleHttpResponse::SetHeaders(TMap<FString, FString>&& InHeaders)
 {
-	[ResponseDelegate CleanSharedObjects];
+	Headers = MoveTemp(InHeaders);
 }
 
-FString FAppleHttpResponse::GetHeader(const FString& HeaderName) const
+void FAppleHttpResponse::CleanSharedObjects()
 {
-	SCOPED_AUTORELEASE_POOL;
-	if(NSDictionary* Headers = [ResponseDelegate GetResponseHeaders])
-	{
-		UE_LOG(LogHttp, Verbose, TEXT("FAppleHttpResponse::GetHeader()"));
-		NSHTTPURLResponse* Response = (NSHTTPURLResponse*)ResponseDelegate.Response;
-		NSString* ConvertedHeaderName = HeaderName.GetNSString();
-		return FString([Response.allHeaderFields objectForKey:ConvertedHeaderName]);
-	}
-	else
-	{
-		return FString();
-	}
+	[ResponseDelegate CleanSharedObjects];
 }
 
 TArray<FString> FAppleHttpResponse::GetAllHeaders() const
@@ -1077,27 +1084,6 @@ uint64 FAppleHttpResponse::GetContentLength() const
 	return ResponseDelegate.Response.expectedContentLength;
 }
 
-const TArray<uint8>& FAppleHttpResponse::GetContent() const
-{
-	if( !IsReady() )
-	{
-		UE_LOG(LogHttp, Warning, TEXT("Payload is incomplete. Response still processing. %s"), *GetURL());
-	}
-	return ResponseDelegate->Payload;
-}
-
-FString FAppleHttpResponse::GetContentAsString() const
-{
-	// Fill in our data.
-	const TArray<uint8>& Payload = GetContent();
-
-	TArray<uint8> ZeroTerminatedPayload;
-	ZeroTerminatedPayload.AddZeroed( Payload.Num() + 1 );
-	FMemory::Memcpy( ZeroTerminatedPayload.GetData(), Payload.GetData(), Payload.Num() );
-
-	return UTF8_TO_TCHAR( ZeroTerminatedPayload.GetData() );
-}
-
 bool FAppleHttpResponse::IsReady() const
 {
 	return EHttpRequestStatus::IsFinished(ResponseDelegate.RequestStatus);
diff --git a/Engine/Source/Runtime/Online/HTTP/Private/Apple/AppleHttp.h b/Engine/Source/Runtime/Online/HTTP/Private/Apple/AppleHttp.h
--- a/Engine/Source/Runtime/Online/HTTP/Private/Apple/AppleHttp.h
+++ b/Engine/Source/Runtime/Online/HTTP/Private/Apple/AppleHttp.h
@@ -78,6 +78,9 @@ PACKAGE_SCOPE:
 
 	bool IsInitializedWithValidStream() const;
 
+	void HandleResponseHeadersReceived(TMap<FString, FString>&& ResponseHeaders);
+	bool HandleResponseBodyDataReceived(uint8* Ptr, uint64 Size);
+
 private:
 	/**
 	 * Trigger the request progress delegate if progress has changed
@@ -146,17 +149,11 @@ class FAppleHttpResponse : public FHttpResponseCommon
 {
 public:
 	//~ Begin IHttpBase Interface
-	virtual FString GetHeader(const FString& HeaderName) const override;
 	virtual TArray<FString> GetAllHeaders() const override;	
 	virtual FString GetContentType() const override;
 	virtual uint64 GetContentLength() const override;
-	virtual const TArray<uint8>& GetContent() const override;
 	//~ End IHttpBase Interface
 
-	//~ Begin IHttpResponse Interface
-	virtual FString GetContentAsString() const override;
-	//~ End IHttpResponse Interface
-
 	/**
 	 * Check whether a response is ready or not.
 	 */
@@ -183,6 +180,10 @@ public:
 	*/
 	void SetNewAppleHttpEventDelegate(FNewAppleHttpEventDelegate&& Delegate);
 
+	void SetHeaders(TMap<FString, FString>&& InHeaders);
+
+	using FHttpResponseCommon::AppendToPayload;
+
 	/**
 	 * Constructor
 	 *
diff --git a/Engine/Source/Runtime/Online/HTTP/Private/Curl/CurlHttp.cpp b/Engine/Source/Runtime/Online/HTTP/Private/Curl/CurlHttp.cpp
--- a/Engine/Source/Runtime/Online/HTTP/Private/Curl/CurlHttp.cpp
+++ b/Engine/Source/Runtime/Online/HTTP/Private/Curl/CurlHttp.cpp
@@ -1250,7 +1250,6 @@ void FCurlHttpRequest::FinishRequest()
 		{
 			BroadcastNewlyReceivedHeaders();
 		}
-		Response->bIsReady = true;
 	}
 
 	bool bSucceeded = Response.IsValid() && Response->bSucceeded;
@@ -1326,29 +1325,10 @@ void FCurlHttpRequest::MockResponseData()
 FCurlHttpResponse::FCurlHttpResponse(const FCurlHttpRequest& InRequest)
 	: FHttpResponseCommon(InRequest)
 	, ContentLength(0)
-	, bIsReady(0)
 	, bSucceeded(0)
 {
 }
 
-FString FCurlHttpResponse::GetHeader(const FString& HeaderName) const
-{
-	FString Result;
-	if (!bIsReady)
-	{
-		UE_LOG(LogHttp, Warning, TEXT("Can't get cached header [%s]. Response still processing. %s"), *HeaderName, *GetURL());
-	}
-	else
-	{
-		const FString* Header = Headers.Find(HeaderName);
-		if (Header != NULL)
-		{
-			Result = *Header;
-		}
-	}
-	return Result;
-}
-
 TArray<FString> FCurlHttpResponse::GetAllHeaders() const
 {
 	TArray<FString> Result;
@@ -1377,21 +1357,4 @@ uint64 FCurlHttpResponse::GetContentLength() const
 	return ContentLength;
 }
 
-const TArray<uint8>& FCurlHttpResponse::GetContent() const
-{
-	if (!bIsReady)
-	{
-		UE_LOG(LogHttp, Warning, TEXT("Payload is incomplete. Response still processing. %s"), *GetURL());
-	}
-	return Payload;
-}
-
-FString FCurlHttpResponse::GetContentAsString() const
-{
-	// Content is NOT null-terminated; we need to specify lengths here
-	FUTF8ToTCHAR TCHARData(reinterpret_cast<const ANSICHAR*>(Payload.GetData()), Payload.Num());
-	return FString(TCHARData.Length(), TCHARData.Get());
-}
-
-
 #endif //WITH_CURL
\ No newline at end of file
diff --git a/Engine/Source/Runtime/Online/HTTP/Private/Curl/CurlHttp.h b/Engine/Source/Runtime/Online/HTTP/Private/Curl/CurlHttp.h
--- a/Engine/Source/Runtime/Online/HTTP/Private/Curl/CurlHttp.h
+++ b/Engine/Source/Runtime/Online/HTTP/Private/Curl/CurlHttp.h
@@ -417,17 +417,11 @@ public:
 	friend class FCurlHttpRequest;
 
 	//~ Begin IHttpBase Interface
-	virtual FString GetHeader(const FString& HeaderName) const override;
 	virtual TArray<FString> GetAllHeaders() const override;	
 	virtual FString GetContentType() const override;
 	virtual uint64 GetContentLength() const override;
-	virtual const TArray<uint8>& GetContent() const override;
 	//~ End IHttpBase Interface
 
-	//~ Begin IHttpResponse Interface
-	virtual FString GetContentAsString() const override;
-	//~ End IHttpResponse Interface
-
 	/**
 	 * Constructor
 	 *
@@ -436,16 +430,8 @@ public:
 	FCurlHttpResponse(const FCurlHttpRequest& InRequest);
 
 private:
-	/** BYTE array to fill in as the response is read via didReceiveData */
-	TArray<uint8> Payload;
-	/** The stream to receive response body */
-	TSharedPtr<FArchive> ResponseBodyReceiveStream;
-	/** Cached key/value header pairs. Parsed once request completes. Only accessible on the game thread. */
-	TMap<FString, FString> Headers;
 	/** Cached content length from completed response */
 	uint64 ContentLength;
-	/** True when the response has finished async processing */
-	int32 volatile bIsReady;
 	/** True if the response was successfully received/processed */
 	int32 volatile bSucceeded;
 };
diff --git a/Engine/Source/Runtime/Online/HTTP/Private/GenericPlatform/HttpRequestCommon.cpp b/Engine/Source/Runtime/Online/HTTP/Private/GenericPlatform/HttpRequestCommon.cpp
--- a/Engine/Source/Runtime/Online/HTTP/Private/GenericPlatform/HttpRequestCommon.cpp
+++ b/Engine/Source/Runtime/Online/HTTP/Private/GenericPlatform/HttpRequestCommon.cpp
@@ -101,12 +101,12 @@ bool FHttpRequestCommon::PreCheck() const
 	return true;
 }
 
-bool FHttpRequestCommon::TriggerMockFailure()
+bool FHttpRequestCommon::TriggerMockResponse()
 {
-	TOptional<int32> MockResponseCode = FHttpModule::Get().GetHttpManager().GetMockFailure(GetURL());
-	if (MockResponseCode.IsSet())
+	TOptional<FHttpManager::FMockResponse> MockResponse = FHttpModule::Get().GetHttpManager().GetMockResponse(GetURL());
+	if (MockResponse.IsSet())
 	{
-		if (MockResponseCode.GetValue() == EHttpResponseCodes::Unknown)
+		if (MockResponse.GetValue().StatusCode == EHttpResponseCodes::Unknown)
 		{
 			int32 HttpConnectionTimeout = FHttpModule::Get().GetHttpConnectionTimeout();
 			TSharedPtr<FHttpRequestCommon> RequestPtr(SharedThis(this));
@@ -123,8 +123,14 @@ bool FHttpRequestCommon::TriggerMockFailure()
 		else
 		{
 			InitResponse();
-			ResponseCommon->SetResponseCode(MockResponseCode.GetValue());
+			ResponseCommon->SetResponseCode(MockResponse.GetValue().StatusCode);
 			MockResponseData();
+			const FUtf8String& ResponsePayload = MockResponse.GetValue().ResponsePayload;
+			if (!ResponsePayload.IsEmpty())
+			{
+				ResponseCommon->AppendToPayload(reinterpret_cast<const uint8*>(*ResponsePayload), ResponsePayload.Len());
+			}
+			ResponseCommon->Headers = MockResponse.GetValue().ResponseHeaders;
 			FinishRequestNotInHttpManager();
 		}
 
@@ -161,7 +167,7 @@ bool FHttpRequestCommon::PreProcess()
 		return false;
 	}
 
-	if (TriggerMockFailure())
+	if (TriggerMockResponse())
 	{
 		return false;
 	}
@@ -787,6 +793,11 @@ void FHttpRequestCommon::OnFinishRequest(bool bSucceeded)
 {
 	// TODO: Move more code from impl into this common function
 
+	if (ResponseCommon)
+	{
+		ResponseCommon->bIsReady = true;
+	}
+
 	if (bSucceeded)
 	{
 		HandleRequestSucceed();
diff --git a/Engine/Source/Runtime/Online/HTTP/Private/GenericPlatform/HttpRequestCommon.h b/Engine/Source/Runtime/Online/HTTP/Private/GenericPlatform/HttpRequestCommon.h
--- a/Engine/Source/Runtime/Online/HTTP/Private/GenericPlatform/HttpRequestCommon.h
+++ b/Engine/Source/Runtime/Online/HTTP/Private/GenericPlatform/HttpRequestCommon.h
@@ -112,7 +112,7 @@ protected:
 
 	HTTP_API void LogFailure() const;
 	HTTP_API void LogResponse(const TSharedPtr<IHttpResponse>& InResponse);
-	HTTP_API bool TriggerMockFailure();
+	HTTP_API bool TriggerMockResponse();
 
 	HTTP_API void InitResponse();
 
diff --git a/Engine/Source/Runtime/Online/HTTP/Private/GenericPlatform/HttpResponseCommon.cpp b/Engine/Source/Runtime/Online/HTTP/Private/GenericPlatform/HttpResponseCommon.cpp
--- a/Engine/Source/Runtime/Online/HTTP/Private/GenericPlatform/HttpResponseCommon.cpp
+++ b/Engine/Source/Runtime/Online/HTTP/Private/GenericPlatform/HttpResponseCommon.cpp
@@ -3,6 +3,7 @@
 #include "GenericPlatform/HttpResponseCommon.h"
 #include "GenericPlatform/HttpRequestCommon.h"
 #include "GenericPlatform/GenericPlatformHttp.h"
+#include "Http.h"
 
 FHttpResponseCommon::FHttpResponseCommon(const FHttpRequestCommon& HttpRequest)
 	: URL(HttpRequest.GetURL())
@@ -69,6 +70,45 @@ void FHttpResponseCommon::SetResponseCode(int32 InResponseCode)
 
 FUtf8StringView FHttpResponseCommon::GetContentAsUtf8StringView() const
 {
-	const TArray<uint8>& Payload = GetContent();
 	return FUtf8StringView(reinterpret_cast<const UTF8CHAR*>(Payload.GetData()), Payload.Num());
 }
+
+FString FHttpResponseCommon::GetHeader(const FString& HeaderName) const
+{
+	FString Result;
+	if (!bIsReady)
+	{
+		UE_LOG(LogHttp, Warning, TEXT("Can't get cached header [%s]. Response still processing. %s"), *HeaderName, *GetURL());
+	}
+	else
+	{
+		const FString* Header = Headers.Find(HeaderName);
+		if (Header != NULL)
+		{
+			Result = *Header;
+		}
+	}
+	return Result;
+}
+
+const TArray<uint8>& FHttpResponseCommon::GetContent() const
+{
+	if (!bIsReady)
+	{
+		UE_LOG(LogHttp, Warning, TEXT("Payload is incomplete. Response still processing for %s"), *GetURL());
+	}
+	return Payload;
+}
+
+FString FHttpResponseCommon::GetContentAsString() const
+{
+	// Content is NOT null-terminated; we need to specify lengths here
+	FUTF8ToTCHAR TCHARData(reinterpret_cast<const ANSICHAR*>(Payload.GetData()), Payload.Num());
+	return FString(TCHARData.Length(), TCHARData.Get());
+}
+
+void FHttpResponseCommon::AppendToPayload(const uint8* Ptr, int64 Size)
+{
+	Payload.Append(Ptr, Size);
+}
+
diff --git a/Engine/Source/Runtime/Online/HTTP/Private/GenericPlatform/HttpResponseCommon.h b/Engine/Source/Runtime/Online/HTTP/Private/GenericPlatform/HttpResponseCommon.h
--- a/Engine/Source/Runtime/Online/HTTP/Private/GenericPlatform/HttpResponseCommon.h
+++ b/Engine/Source/Runtime/Online/HTTP/Private/GenericPlatform/HttpResponseCommon.h
@@ -24,16 +24,23 @@ public:
 	virtual EHttpFailureReason GetFailureReason() const override;
 	virtual int32 GetResponseCode() const override;
 	virtual FUtf8StringView GetContentAsUtf8StringView() const override;
+	virtual FString GetHeader(const FString& HeaderName) const override;
+	virtual const TArray<uint8>& GetContent() const override;
+	virtual FString GetContentAsString() const override;
 
 protected:
 	void SetRequestStatus(EHttpRequestStatus::Type InCompletionStatus);
 	void SetRequestFailureReason(EHttpFailureReason InFailureReason);
 	void SetEffectiveURL(const FString& InEffectiveURL);
 	void SetResponseCode(int32 InResponseCode);
+	void AppendToPayload(const uint8* Ptr, int64 Size);
 
 	FString URL;
 	FString EffectiveURL;
 	EHttpRequestStatus::Type CompletionStatus;
 	EHttpFailureReason FailureReason;
 	int32 ResponseCode = EHttpResponseCodes::Unknown;
+	TArray<uint8> Payload;
+	TMap<FString, FString> Headers;
+	bool bIsReady = false;
 };
diff --git a/Engine/Source/Runtime/Online/HTTP/Private/HttpManager.cpp b/Engine/Source/Runtime/Online/HTTP/Private/HttpManager.cpp
--- a/Engine/Source/Runtime/Online/HTTP/Private/HttpManager.cpp
+++ b/Engine/Source/Runtime/Online/HTTP/Private/HttpManager.cpp
@@ -39,8 +39,8 @@ TAutoConsoleVariable<FString> CVarHttpUrlPatternsToDisableFailedLog(
 	TEXT("List of url patterns to disable failed log: \"epicgames.com unrealengine.com ...\"")
 );
 
-TAutoConsoleVariable<FString> CVarHttpUrlPatternsToMockFailure(
-	TEXT("http.UrlPatternsToMockFailure"),
+TAutoConsoleVariable<FString> CVarHttpUrlPatternsToMockResponse(
+	TEXT("http.UrlPatternsToMockResponse"),
 	TEXT(""),
 	TEXT("List of url patterns to mock failure with response code, 0 indicates ConnectionError: \"epicgames.com->0 unrealengine.com->503 ...\""),
 	ECVF_SaveForNextBoot
@@ -127,15 +127,15 @@ void FHttpManager::Initialize()
 	UpdateUrlPatternsToDisableFailedLog(CVarHttpUrlPatternsToDisableFailedLog.AsVariable());
 	CVarHttpUrlPatternsToDisableFailedLog.AsVariable()->OnChangedDelegate().AddRaw(this, &FHttpManager::UpdateUrlPatternsToDisableFailedLog);
 
-	UpdateUrlPatternsToMockFailure(CVarHttpUrlPatternsToMockFailure.AsVariable());
-	CVarHttpUrlPatternsToMockFailure.AsVariable()->OnChangedDelegate().AddRaw(this, &FHttpManager::UpdateUrlPatternsToMockFailure);
+	UpdateUrlPatternsToMockResponse(CVarHttpUrlPatternsToMockResponse.AsVariable());
+	CVarHttpUrlPatternsToMockResponse.AsVariable()->OnChangedDelegate().AddRaw(this, &FHttpManager::UpdateUrlPatternsToMockResponse);
 }
 
 void FHttpManager::Shutdown()
 {
 	CVarHttpUrlPatternsToLogResponse.AsVariable()->OnChangedDelegate().Clear();
 	CVarHttpUrlPatternsToDisableFailedLog.AsVariable()->OnChangedDelegate().Clear();
-	CVarHttpUrlPatternsToMockFailure.AsVariable()->OnChangedDelegate().Clear();
+	CVarHttpUrlPatternsToMockResponse.AsVariable()->OnChangedDelegate().Clear();
 
 	{
 		FScopeLock ScopeLock(&RequestLock);
@@ -744,34 +744,70 @@ bool FHttpManager::ShouldDisableFailedLog(FStringView Url) const
 	return false;
 }
 
-void FHttpManager::UpdateUrlPatternsToMockFailure(IConsoleVariable* CVar)
+void FHttpManager::UpdateUrlPatternsToMockResponse(IConsoleVariable* CVar)
 {
-	const FScopeLock CacheLock(&UrlPatternsToMockFailureCriticalSection);
-	const FString UrlPatternsToMockFailureStr = CVar->AsVariable()->GetString();
+	const FScopeLock CacheLock(&UrlPatternsToMockResponseCriticalSection);
+	const FString UrlPatternsToMockResponseStr = CVar->AsVariable()->GetString();
 
-	TArray<FString> UrlPatternsToMockFailureStrings;
-	UrlPatternsToMockFailureStr.ParseIntoArray(UrlPatternsToMockFailureStrings, TEXT(" "));
+	TArray<FString> UrlPatternsToMockResponseStrings;
+	UrlPatternsToMockResponseStr.ParseIntoArray(UrlPatternsToMockResponseStrings, TEXT(" "));
 
-	UrlPatternsToMockFailure.Empty();
+	UrlPatternsToMockResponse.Empty();
 
-	for (const FString& UrlPatternToMockFailureString : UrlPatternsToMockFailureStrings)
+	for (const FString& UrlPatternToMockResponseString : UrlPatternsToMockResponseStrings)
 	{
+		// Example runtime args(NOTE: It can't contain space in body of each url pattern): 
+		//  -ini:Engine:[ConsoleVariables]:http.UrlPatternsToMockResponse="www.mydomain.com/mypath->403++{\"friendly_error\":\"online.presence.restricted\",\"error_code\":3456}++Content-Type:application/json"
 		TArray<FString> UrlPattern;
-		UrlPatternToMockFailureString.ParseIntoArray(UrlPattern, TEXT("->"));
+		UrlPatternToMockResponseString.ParseIntoArray(UrlPattern, TEXT("->"));
 		if (UrlPattern.Num() == 2)
 		{
-			int32 ResponseCode = FCString::Atoi(*UrlPattern[1]);
-			UrlPatternsToMockFailure.Emplace(UrlPattern[0], ResponseCode);
+			FMockResponse MockResponse;
+
+			FString MockResponseStr = UrlPattern[1];
+			const FStringView PayloadSeparator = TEXT("++");
+			const int32 PayloadSeparatorIndex = UE::String::FindFirst(MockResponseStr, PayloadSeparator);
+			if (PayloadSeparatorIndex != INDEX_NONE)
+			{
+				MockResponse.StatusCode = FCString::Atoi(*MockResponseStr.Left(PayloadSeparatorIndex));
+				FString ResponsePayload = MockResponseStr.RightChop(PayloadSeparatorIndex + PayloadSeparator.Len());
+				const int32 HeaderSeparatorIndex = UE::String::FindFirst(ResponsePayload, PayloadSeparator);
+				if (HeaderSeparatorIndex != INDEX_NONE)
+				{
+					FString HeadersStr = ResponsePayload.RightChop(HeaderSeparatorIndex + PayloadSeparator.Len());
+					ResponsePayload = ResponsePayload.Left(HeaderSeparatorIndex);
+
+					TArray<FString> ResponseHeaderStrings;
+					HeadersStr.ParseIntoArray(ResponseHeaderStrings, TEXT(","));
+					for (const FString& ResponseHeaderString : ResponseHeaderStrings)
+					{
+						FString HeaderKey;
+						FString HeaderValue;
+						if (ResponseHeaderString.Split(TEXT(":"), &HeaderKey, &HeaderValue))
+						{
+							MockResponse.ResponseHeaders.Emplace(MoveTemp(HeaderKey), MoveTemp(HeaderValue));
+						}
+					}
+				}
+
+				MockResponse.ResponsePayload = StringCast<UTF8CHAR>(*ResponsePayload);
+			}
+			else
+			{
+				MockResponse.StatusCode = FCString::Atoi(*MockResponseStr);
+			}
+
+			UrlPatternsToMockResponse.Emplace(UrlPattern[0], MoveTemp(MockResponse));
 		}
 	}
 }
 
-TOptional<int32> FHttpManager::GetMockFailure(FStringView Url) const
+TOptional<FHttpManager::FMockResponse> FHttpManager::GetMockResponse(FStringView Url) const
 {
-	TOptional<int32> Result;
-	const FScopeLock CacheLock(&UrlPatternsToMockFailureCriticalSection);
+	TOptional<FMockResponse> Result;
+	const FScopeLock CacheLock(&UrlPatternsToMockResponseCriticalSection);
 
-	for (const TPair<FString, int32>& UrlPattern : UrlPatternsToMockFailure)
+	for (const TPair<FString, FMockResponse>& UrlPattern : UrlPatternsToMockResponse)
 	{
 		if (Url.Contains(UrlPattern.Key))
 		{
diff --git a/Engine/Source/Runtime/Online/HTTP/Public/HttpManager.h b/Engine/Source/Runtime/Online/HTTP/Public/HttpManager.h
--- a/Engine/Source/Runtime/Online/HTTP/Public/HttpManager.h
+++ b/Engine/Source/Runtime/Online/HTTP/Public/HttpManager.h
@@ -8,6 +8,7 @@
 #include "Containers/Ticker.h"
 #include "Containers/Ticker.h"
 #include "Containers/UnrealString.h"
+#include "Containers/Utf8String.h"
 #include "CoreMinimal.h"
 #include "CoreTypes.h"
 #include "HAL/CriticalSection.h"
@@ -318,7 +319,7 @@ protected:
 
 	HTTP_API void UpdateUrlPatternsToLogResponse(IConsoleVariable* CVar);
 	HTTP_API void UpdateUrlPatternsToDisableFailedLog(IConsoleVariable* CVar);
-	HTTP_API void UpdateUrlPatternsToMockFailure(IConsoleVariable* CVar);
+	HTTP_API void UpdateUrlPatternsToMockResponse(IConsoleVariable* CVar);
 
 	HTTP_API void AddRequest(const FHttpRequestRef& Request);
 
@@ -396,15 +397,18 @@ protected:
 	TArray<FString> UrlPatternsToDisableFailedLog;
 	mutable FCriticalSection UrlPatternsToDisableFailedLogCriticalSection;
 
-	TMap<FString, int32> UrlPatternsToMockFailure;
-	mutable FCriticalSection UrlPatternsToMockFailureCriticalSection;
-
 #if !UE_BUILD_SHIPPING
 	mutable FCriticalSection ThreadsCompletingRequestCriticalSection;
 	TMap<uint32, bool> ThreadsCompletingRequest;
 #endif
 
 PACKAGE_SCOPE:
+	struct FMockResponse
+	{
+		int32 StatusCode = 0;
+		FUtf8String ResponsePayload;
+		TMap<FString, FString> ResponseHeaders;
+	};
 
 	/** Used to lock access to add/remove/find requests */
 	mutable FCriticalSection RequestLock;
@@ -442,8 +446,12 @@ PACKAGE_SCOPE:
 	HTTP_API bool ShouldLogResponse(FStringView Url) const;
 	HTTP_API bool ShouldDisableFailedLog(FStringView Url) const;
 
-	HTTP_API TOptional<int32> GetMockFailure(FStringView Url) const;
+	HTTP_API TOptional<FMockResponse> GetMockResponse(FStringView Url) const;
 
 	HTTP_API void MarkCurrentThreadCompletingRequest(bool bCompleting);
 	HTTP_API bool IsCurrentThreadCompletingRequest() const;
+
+protected:
+	TMap<FString, FMockResponse> UrlPatternsToMockResponse;
+	mutable FCriticalSection UrlPatternsToMockResponseCriticalSection;
 };
