Shutdown platform custom ssl as well, besides OpenSSL.

diff --git a/Engine/Source/Runtime/Online/HTTP/Private/Curl/CurlHttpManager.cpp b/Engine/Source/Runtime/Online/HTTP/Private/Curl/CurlHttpManager.cpp
--- a/Engine/Source/Runtime/Online/HTTP/Private/Curl/CurlHttpManager.cpp
+++ b/Engine/Source/Runtime/Online/HTTP/Private/Curl/CurlHttpManager.cpp
@@ -353,6 +353,8 @@ void FCurlHttpManager::ShutdownCurl()
 		GMultiHandle = nullptr;
 	}
 
+	FPlatformHttp::ShutdownPlatformCustomSsl();
+
 	curl_global_cleanup();
 
 	LibCryptoMemHooks::UnsetMemoryHooks();
diff --git a/Engine/Source/Runtime/Online/HTTP/Public/GenericPlatform/GenericPlatformHttp.h b/Engine/Source/Runtime/Online/HTTP/Public/GenericPlatform/GenericPlatformHttp.h
--- a/Engine/Source/Runtime/Online/HTTP/Public/GenericPlatform/GenericPlatformHttp.h
+++ b/Engine/Source/Runtime/Online/HTTP/Public/GenericPlatform/GenericPlatformHttp.h
@@ -78,6 +78,11 @@ public:
 	 */
 	static HTTP_API void Shutdown();
 
+	/**
+	 * Platform shutdown custom ssl if supported
+	 */
+	static HTTP_API void ShutdownPlatformCustomSsl() {}
+
 	/**
 	 * Creates a new Http request instance for the current platform
 	 *
