[DedicatedServer] Retry hotfix check during server initialization.

diff --git a/Engine/Source/Runtime/Online/HTTP/Public/HttpRetrySystem.h b/Engine/Source/Runtime/Online/HTTP/Public/HttpRetrySystem.h
--- a/Engine/Source/Runtime/Online/HTTP/Public/HttpRetrySystem.h
+++ b/Engine/Source/Runtime/Online/HTTP/Public/HttpRetrySystem.h
@@ -80,8 +80,8 @@ namespace FHttpRetrySystem
 		/** Max back off seconds */
 		float MaxBackoffSeconds = 60.0f;
 
-		bool IsValid() const;
-		float Compute(uint32 RetryNumber) const;
+		HTTP_API bool IsValid() const;
+		HTTP_API float Compute(uint32 RetryNumber) const;
 	};
 
 	/**
