BPS: Fix for user cancels not actually cancelling installs.

#rnx

diff --git a/Engine/Source/Runtime/Online/BuildPatchServices/Private/BuildPatchInstaller.cpp b/Engine/Source/Runtime/Online/BuildPatchServices/Private/BuildPatchInstaller.cpp
--- a/Engine/Source/Runtime/Online/BuildPatchServices/Private/BuildPatchInstaller.cpp
+++ b/Engine/Source/Runtime/Online/BuildPatchServices/Private/BuildPatchInstaller.cpp
@@ -1270,9 +1270,15 @@ namespace BuildPatchServices
 		//IFileManager::Get().MakeDirectory(*DataStagingDir, true);
 		//IFileManager::Get().MakeDirectory(*MetaStagingDir, true);
 
-		// Reset our error and build progress
-		InstallerError.Reset(FInstallerErrorFactory::Create());
-		InstallerError->RegisterForErrors([this]() { CancelInstall(); });
+		
+		// We want to reset errors because we are retrying, however we ONLY want to do this if the error was caused by us!
+		// This is a bit delicate because this pointer gets hit by any thread, but we know that installer errors are latching.
+		// (except for retry).
+		if (!InstallerError->IsCancelled())
+		{
+			InstallerError->Reset();
+			InstallerError->RegisterForErrors([this]() { CancelInstall(); });
+		}
 
 		// Get the list of required files, by the tags
 		TaggedFiles.Empty();
