Replace some usages of FORCEINLINE with inline in Foundation modules.

For context about this change, see this doc: https://docs.google.com/document/d/1LjpT7NpzTecfXdFNyv5KM6sCu67827BRljNi1iaWVSE/edit?tab=t.0#heading=h.fvtz8gjukz7x

Includes modules: HTTPChunkInstaller, PS4Download1VFS, InstallBundleManager, ImageDownload

Reviewers were automatically determined. Please comment if anyone else should be added.

#rnx

diff --git a/Engine/Plugins/Runtime/HTTPChunkInstaller/Source/Public/ChunkInstall.h b/Engine/Plugins/Runtime/HTTPChunkInstaller/Source/Public/ChunkInstall.h
--- a/Engine/Plugins/Runtime/HTTPChunkInstaller/Source/Public/ChunkInstall.h
+++ b/Engine/Plugins/Runtime/HTTPChunkInstaller/Source/Public/ChunkInstall.h
@@ -154,7 +154,7 @@ public:
 		return TEXT("FChunkDescovery");
 	}
 
-	FORCEINLINE TStatId GetStatId() const
+	inline TStatId GetStatId() const
 	{
 		RETURN_QUICK_DECLARE_CYCLE_STAT(FChunkInstallTask, STATGROUP_ThreadPoolAsyncTasks);
 	}
diff --git a/Engine/Plugins/Runtime/HTTPChunkInstaller/Source/Public/ChunkSetup.h b/Engine/Plugins/Runtime/HTTPChunkInstaller/Source/Public/ChunkSetup.h
--- a/Engine/Plugins/Runtime/HTTPChunkInstaller/Source/Public/ChunkSetup.h
+++ b/Engine/Plugins/Runtime/HTTPChunkInstaller/Source/Public/ChunkSetup.h
@@ -179,7 +179,7 @@ public:
 		return TEXT("FChunkSetup");
 	}
 
-	FORCEINLINE TStatId GetStatId() const
+	inline TStatId GetStatId() const
 	{
 		RETURN_QUICK_DECLARE_CYCLE_STAT(FChunkSetupTask, STATGROUP_ThreadPoolAsyncTasks);
 	}
@@ -332,7 +332,7 @@ public:
 		return TEXT("FChunkSetup");
 	}
 
-	FORCEINLINE TStatId GetStatId() const
+	inline TStatId GetStatId() const
 	{
 		RETURN_QUICK_DECLARE_CYCLE_STAT(FChunkMountTask, STATGROUP_ThreadPoolAsyncTasks);
 	}
diff --git a/Engine/Source/Runtime/InstallBundleManager/Public/InstallBundleUtils.h b/Engine/Source/Runtime/InstallBundleManager/Public/InstallBundleUtils.h
--- a/Engine/Source/Runtime/InstallBundleManager/Public/InstallBundleUtils.h
+++ b/Engine/Source/Runtime/InstallBundleManager/Public/InstallBundleUtils.h
@@ -197,7 +197,7 @@ namespace InstallBundleUtil
 			}
 		}
 
-		FORCEINLINE TStatId GetStatId() const
+		inline TStatId GetStatId() const
 		{
 			RETURN_QUICK_DECLARE_CYCLE_STAT(FInstallBundleWork, STATGROUP_ThreadPoolAsyncTasks);
 		}
diff --git a/Engine/Source/Runtime/Online/ImageDownload/Public/WebImage.h b/Engine/Source/Runtime/Online/ImageDownload/Public/WebImage.h
--- a/Engine/Source/Runtime/Online/ImageDownload/Public/WebImage.h
+++ b/Engine/Source/Runtime/Online/ImageDownload/Public/WebImage.h
@@ -53,13 +53,13 @@ public:
 	DECLARE_DELEGATE_OneParam(FOnImageDownloaded, bool);
 
 	/** Set the brush that is currently being returned (this will be overridden when any async download completes) */
-	FORCEINLINE FWebImage& SetStandInBrush(TAttribute<const FSlateBrush*> StandInBrushIn) { StandInBrush = StandInBrushIn; DownloadedBrush.Reset(); return *this; }
+	inline FWebImage& SetStandInBrush(TAttribute<const FSlateBrush*> StandInBrushIn) { StandInBrush = StandInBrushIn; DownloadedBrush.Reset(); return *this; }
 
 	/** Begin downloading an image. This will automatically set the current brush to the downloaded image when it completes (if successful) */
 	IMAGEDOWNLOAD_API bool BeginDownload(const FString& InUrl, const TOptional<FString>& StandInETag = TOptional<FString>(), const FOnImageDownloaded& DownloadCallback = FOnImageDownloaded());
 
 	/** Begin downloading an image. */
-	FORCEINLINE bool BeginDownload(const FString& InUrl, const FOnImageDownloaded& DownloadCallback)
+	inline bool BeginDownload(const FString& InUrl, const FOnImageDownloaded& DownloadCallback)
 	{
 		return BeginDownload(InUrl, TOptional<FString>(), DownloadCallback);
 	}
@@ -73,25 +73,25 @@ public:
 	IMAGEDOWNLOAD_API TAttribute< const FSlateBrush* > Attr() const;
 
 	/** Get the current brush displayed (will automatically change when download completes) */
-	FORCEINLINE const FSlateBrush* GetBrush() const { return DownloadedBrush.IsValid() ? DownloadedBrush.Get() : StandInBrush.Get(); }
+	inline const FSlateBrush* GetBrush() const { return DownloadedBrush.IsValid() ? DownloadedBrush.Get() : StandInBrush.Get(); }
 
 	/** Only returns the downloaded brush. May be null if the download hasn't finished or was unsuccessful */
-	FORCEINLINE const FSlateBrush* GetDownloadedBrush() const { return DownloadedBrush.IsValid() ? DownloadedBrush.Get() : nullptr; }
+	inline const FSlateBrush* GetDownloadedBrush() const { return DownloadedBrush.IsValid() ? DownloadedBrush.Get() : nullptr; }
 
 	/** Is there a pending HTTP request */
-	FORCEINLINE bool IsDownloadPending() const { return PendingRequest.IsValid(); }
+	inline bool IsDownloadPending() const { return PendingRequest.IsValid(); }
 
 	/** Has the download finished AND was it successful */
-	FORCEINLINE bool DidDownloadSucceed() const { return bDownloadSuccess; }
+	inline bool DidDownloadSucceed() const { return bDownloadSuccess; }
 
 	/** Has the download finished AND did it fail */
-	FORCEINLINE bool DidDownloadFail() const { return !IsDownloadPending() && !DidDownloadSucceed(); }
+	inline bool DidDownloadFail() const { return !IsDownloadPending() && !DidDownloadSucceed(); }
 
 	/** What URL was requested */
-	FORCEINLINE const FString& GetUrl() const { return Url; }
+	inline const FString& GetUrl() const { return Url; }
 
 	/** What is the ETag of the downloaded resource */
-	FORCEINLINE const TOptional<FString>& GetETag() const { return ETag; }
+	inline const TOptional<FString>& GetETag() const { return ETag; }
 
 private:
 	/** request complete callback */
diff --git a/Engine/Source/Runtime/Online/ImageDownload/Public/WebImageCache.h b/Engine/Source/Runtime/Online/ImageDownload/Public/WebImageCache.h
--- a/Engine/Source/Runtime/Online/ImageDownload/Public/WebImageCache.h
+++ b/Engine/Source/Runtime/Online/ImageDownload/Public/WebImageCache.h
@@ -29,7 +29,7 @@ public:
 	IMAGEDOWNLOAD_API TSharedRef<const FWebImage> Download(const FString& Url, const TOptional<FString>& DefaultImageUrl = TOptional<FString>());
 
 	/** Set the brush that will be returned until the download completes (only affects future downloads). */
-	FORCEINLINE void SetDefaultStandInBrush(TAttribute<const FSlateBrush*> StandInBrushIn) { DefaultStandInBrush = StandInBrushIn; }
+	inline void SetDefaultStandInBrush(TAttribute<const FSlateBrush*> StandInBrushIn) { DefaultStandInBrush = StandInBrushIn; }
 
 	/* This function causes the web image cache to stop holding on to strong references to images. Normally
 	 * once downloaded, an image is cached forever. This allows us to release images that are not currently
