Remove VirtualFileCache and BPS Install Streamer

diff --git a/Engine/Source/Runtime/Online/BuildPatchServices/BuildPatchServices.Build.cs b/Engine/Source/Runtime/Online/BuildPatchServices/BuildPatchServices.Build.cs
--- a/Engine/Source/Runtime/Online/BuildPatchServices/BuildPatchServices.Build.cs
+++ b/Engine/Source/Runtime/Online/BuildPatchServices/BuildPatchServices.Build.cs
@@ -29,7 +29,6 @@ public class BuildPatchServices : ModuleRules
 				"AnalyticsET",
 				"HTTP",
 				"Json",
-				"VirtualFileCache",
 			}
 		);
 	}
diff --git a/Engine/Source/Runtime/Online/BuildPatchServices/Private/BuildPatchServicesModule.cpp b/Engine/Source/Runtime/Online/BuildPatchServices/Private/BuildPatchServicesModule.cpp
--- a/Engine/Source/Runtime/Online/BuildPatchServices/Private/BuildPatchServicesModule.cpp
+++ b/Engine/Source/Runtime/Online/BuildPatchServices/Private/BuildPatchServicesModule.cpp
@@ -114,14 +114,6 @@ void FBuildPatchServicesModule::ShutdownModule()
 	GLog->Log(ELogVerbosity::VeryVerbose, TEXT( "BuildPatchServicesModule: Finished shutting down" ) );
 }
 
-IBuildInstallStreamerRef FBuildPatchServicesModule::CreateBuildInstallStreamer(BuildPatchServices::FBuildInstallStreamerConfiguration Configuration)
-{
-	FBuildInstallStreamerRef Streamer = MakeShareable(FBuildInstallStreamerFactory::Create(MoveTemp(Configuration)));
-	FBuildInstallStreamerWeakPtr WeakStreamer = Streamer;
-	AsyncHelpers::ExecuteOnGameThread<void>([this, WeakStreamer = MoveTemp(WeakStreamer)] { WeakBuildInstallStreamers.Add(WeakStreamer); });
-	return Streamer;
-}
-
 IBuildInstallerRef FBuildPatchServicesModule::CreateBuildInstaller(BuildPatchServices::FBuildInstallerConfiguration Configuration, FBuildPatchInstallerDelegate CompleteDelegate) const
 {
 	// Override prereq install using the config/commandline value to force skip them.
@@ -228,16 +220,6 @@ bool FBuildPatchServicesModule::Tick(float Delta)
 		}
 	}
 
-	// Tick running streamers.
-	for (auto StreamerIter = WeakBuildInstallStreamers.CreateIterator(); StreamerIter; ++StreamerIter)
-	{
-		const FBuildInstallStreamerPtr& Streamer = StreamerIter->Pin();
-		if (!Streamer.IsValid() || !Streamer->Tick())
-		{
-			StreamerIter.RemoveCurrent();
-		}
-	}
-
 	// More ticks.
 	return true;
 }
@@ -403,15 +385,6 @@ void FBuildPatchServicesModule::PreExit()
 		const FBuildPatchInstallerRef& Installer = *InstallerIter;
 		Installer->PreExit();
 	}
-	// Inform streamers
-	for (auto StreamerIter = WeakBuildInstallStreamers.CreateIterator(); StreamerIter; ++StreamerIter)
-	{
-		const FBuildInstallStreamerPtr& Streamer = StreamerIter->Pin();
-		if (Streamer.IsValid())
-		{
-			Streamer->PreExit();
-		}
-	}
 
 	// Release our ptr to analytics
 	Analytics.Reset();
diff --git a/Engine/Source/Runtime/Online/BuildPatchServices/Private/BuildPatchServicesModule.h b/Engine/Source/Runtime/Online/BuildPatchServices/Private/BuildPatchServicesModule.h
--- a/Engine/Source/Runtime/Online/BuildPatchServices/Private/BuildPatchServicesModule.h
+++ b/Engine/Source/Runtime/Online/BuildPatchServices/Private/BuildPatchServicesModule.h
@@ -4,7 +4,6 @@
 
 #include "CoreMinimal.h"
 #include "Misc/Guid.h"
-#include "Installer/BuildInstallStreamer.h"
 #include "Interfaces/IBuildPatchServicesModule.h"
 #include "BuildPatchInstaller.h"
 #include "Containers/Ticker.h"
@@ -36,7 +35,6 @@ public:
 	// IModuleInterface interface end.
 
 	// IBuildPatchServicesModule interface begin.
-	virtual IBuildInstallStreamerRef CreateBuildInstallStreamer(BuildPatchServices::FBuildInstallStreamerConfiguration Configuration) override;
 	virtual IBuildInstallerRef CreateBuildInstaller(BuildPatchServices::FBuildInstallerConfiguration Configuration, FBuildPatchInstallerDelegate OnComplete) const override;
 	virtual IBuildInstallerSharedContextRef CreateBuildInstallerSharedContext(const TCHAR* DebugName) const override;
 	virtual BuildPatchServices::IBuildStatisticsRef CreateBuildStatistics(const IBuildInstallerRef& Installer) const override;
@@ -142,9 +140,6 @@ private:
 	// Holds available installations used for recycling install data
 	TMultiMap<FString, FBuildPatchAppManifestRef> AvailableInstallations;
 
-	// Array of running streamers
-	TArray<FBuildInstallStreamerWeakPtr> WeakBuildInstallStreamers;
-
 	// Handle to the registered Tick delegate
 	FTSTicker::FDelegateHandle TickDelegateHandle;
 
diff --git a/Engine/Source/Runtime/Online/BuildPatchServices/Private/Installer/BuildInstallStreamer.cpp b/Engine/Source/Runtime/Online/BuildPatchServices/Private/Installer/BuildInstallStreamer.cpp
deleted file mode 100644
diff --git a/Engine/Source/Runtime/Online/BuildPatchServices/Private/Installer/BuildInstallStreamer.h b/Engine/Source/Runtime/Online/BuildPatchServices/Private/Installer/BuildInstallStreamer.h
deleted file mode 100644
diff --git a/Engine/Source/Runtime/Online/BuildPatchServices/Private/Installer/VirtualFileConstructor.cpp b/Engine/Source/Runtime/Online/BuildPatchServices/Private/Installer/VirtualFileConstructor.cpp
deleted file mode 100644
diff --git a/Engine/Source/Runtime/Online/BuildPatchServices/Private/Installer/VirtualFileConstructor.h b/Engine/Source/Runtime/Online/BuildPatchServices/Private/Installer/VirtualFileConstructor.h
deleted file mode 100644
diff --git a/Engine/Source/Runtime/Online/BuildPatchServices/Public/BuildPatchSettings.h b/Engine/Source/Runtime/Online/BuildPatchServices/Public/BuildPatchSettings.h
--- a/Engine/Source/Runtime/Online/BuildPatchServices/Public/BuildPatchSettings.h
+++ b/Engine/Source/Runtime/Online/BuildPatchServices/Public/BuildPatchSettings.h
@@ -334,48 +334,6 @@ namespace BuildPatchServices
 		bool bInstallToMemory = false;
 	};
 
-	/**
-	 * Defines a list of all the options of a build streamer class.
-	 */
-	struct FBuildInstallStreamerConfiguration
-	{
-	public:
-		IBuildManifestPtr Manifest;
-		// The list of chunk database filenames that will be used to pull patch data from.
-		TArray<FString> ChunkDatabaseFiles;
-		// The list of cloud directory roots that will be used to pull patch data from. Empty array will use module's global setting..
-		TArray<FString> CloudDirectories;
-		// Whether the streamer should batch all requests into each cycle.
-		bool bShouldBatch = true;
-		// Whether completion delegates should be called on main thread, or the streamer's thread,
-		bool bMainThreadDelegates = true;
-	};
-
-	/**
-	 * The collection of statistics gathered by the build streamer class.
-	 */
-	struct FBuildInstallStreamerStats
-	{
-		//float FileMegaBytesDownloaded;
-		uint64 FileRequestsCompleted;
-		uint64 FileRequestsMade;
-		uint64 FileRequestsCancelled;
-		
-		//float BundleMegaBytesDownloaded;
-		uint64 BundleRequestsCancelled;
-		uint64 BundleRequestsCompleted;
-		uint64 BundleRequestsMade;
-		
-		float TotalMegaBytesDownloaded;
-		double MaxRequestTime;
-		double AverageRequestTime;
-		
-		float VFCCachedTotalSize;
-		float VFCCachedUsedSize;
-		float VFCRequestedFileWrite;
-		float VFCActualFileWrite;
-	};
-
 	/**
 	 * Defines a list of all options for the build chunking task.
 	 */
diff --git a/Engine/Source/Runtime/Online/BuildPatchServices/Public/Interfaces/IBuildInstallStreamer.h b/Engine/Source/Runtime/Online/BuildPatchServices/Public/Interfaces/IBuildInstallStreamer.h
deleted file mode 100644
diff --git a/Engine/Source/Runtime/Online/BuildPatchServices/Public/Interfaces/IBuildPatchServicesModule.h b/Engine/Source/Runtime/Online/BuildPatchServices/Public/Interfaces/IBuildPatchServicesModule.h
--- a/Engine/Source/Runtime/Online/BuildPatchServices/Public/Interfaces/IBuildPatchServicesModule.h
+++ b/Engine/Source/Runtime/Online/BuildPatchServices/Public/Interfaces/IBuildPatchServicesModule.h
@@ -5,7 +5,6 @@
 #include "Modules/ModuleInterface.h"
 #include "Interfaces/IBuildInstaller.h"
 #include "Interfaces/IBuildStatistics.h"
-#include "Interfaces/IBuildInstallStreamer.h"
 #include "Interfaces/IPatchDataEnumeration.h"
 #include "BuildPatchSettings.h"
 
@@ -38,13 +37,6 @@ public:
 	 */
 	virtual ~IBuildPatchServicesModule() { }
 
-	/**
-	 * Factory providing construction of a build streamer class.
-	 * @param Configuration     The configuration for the installer.
-	 * @return an instance of an IBuildInstallStreamer implementation.
-	 */
-	virtual IBuildInstallStreamerRef CreateBuildInstallStreamer(BuildPatchServices::FBuildInstallStreamerConfiguration Configuration) = 0;
-
 	/**
 	 * Factory providing construction of a build installer class.
 	 * @param Configuration     The configuration for the installer.
diff --git a/Engine/Source/Runtime/VirtualFileCache/Private/VirtualFileCache.cpp b/Engine/Source/Runtime/VirtualFileCache/Private/VirtualFileCache.cpp
deleted file mode 100644
diff --git a/Engine/Source/Runtime/VirtualFileCache/Private/VirtualFileCacheInternal.cpp b/Engine/Source/Runtime/VirtualFileCache/Private/VirtualFileCacheInternal.cpp
deleted file mode 100644
diff --git a/Engine/Source/Runtime/VirtualFileCache/Private/VirtualFileCacheInternal.h b/Engine/Source/Runtime/VirtualFileCache/Private/VirtualFileCacheInternal.h
deleted file mode 100644
diff --git a/Engine/Source/Runtime/VirtualFileCache/Public/VirtualFileCache.h b/Engine/Source/Runtime/VirtualFileCache/Public/VirtualFileCache.h
deleted file mode 100644
diff --git a/Engine/Source/Runtime/VirtualFileCache/VirtualFileCache.Build.cs b/Engine/Source/Runtime/VirtualFileCache/VirtualFileCache.Build.cs
deleted file mode 100644
