BPS: Remove the ensure used for tracking the event pooling bug.

#rnx

diff --git a/Engine/Source/Runtime/Online/BuildPatchServices/Private/Installer/Verifier.cpp b/Engine/Source/Runtime/Online/BuildPatchServices/Private/Installer/Verifier.cpp
--- a/Engine/Source/Runtime/Online/BuildPatchServices/Private/Installer/Verifier.cpp
+++ b/Engine/Source/Runtime/Online/BuildPatchServices/Private/Installer/Verifier.cpp
@@ -80,7 +80,6 @@ namespace BuildPatchServices
 			FString FileName;
 		};
 
-		FEvent* AllThreadsDoneEvent;
 		std::atomic_int64_t ThreadProcessedBytes;
 		TArray<EVerifyResult> ThreadJobResults;
 
@@ -207,13 +206,6 @@ namespace BuildPatchServices
 			ThreadJobResults.Add(EVerifyResult::Aborted);
 		}
 
-		AllThreadsDoneEvent = FPlatformProcess::GetSynchEventFromPool();
-
-		// Events get reset when they are returned to the pool but if there are any outstanding
-		// triggers due to races in the previous user we would have a triggered event, so we Reset here. 
-		// This doesn't really remove the race but it prevents it from sitting latent for long periods of time.
-		AllThreadsDoneEvent->Reset();
-
 		if (SharedContext)
 		{
 			// 3 cores should saturate any modern drive, and on anything that isn't we just sit in the wait
@@ -248,14 +240,6 @@ namespace BuildPatchServices
 				VerifierStat->OnProcessedDataUpdated(ThreadProcessedBytes.load(std::memory_order_acquire));
 			}
 
-			// If our event is triggered then something has gone very wrong since we could not have triggered it.
-			// This is still here to try and figure out what's been going wrong with it before we switched to a semaphore.
-			bool bEventIsTriggered = AllThreadsDoneEvent->Wait(0);
-			if (!ensure(!bEventIsTriggered))
-			{
-				UE_LOG(LogVerifier, Error, TEXT("Verifier has a triggered event!"));
-			}
-
 			for (int i = 0; i < ThreadsToCreate; i++)
 			{
 				SharedContext->ReleaseThread(SecondaryProcessThreads[i]);
@@ -268,9 +252,6 @@ namespace BuildPatchServices
 			ProcessVerifyJobs(nullptr);
 		}
 
-		FPlatformProcess::ReturnSynchEventToPool(AllThreadsDoneEvent);
-		AllThreadsDoneEvent = nullptr;
-
 		// move results over.
 		for (int32 i=0; i<ThreadJobList.Num(); i++)
 		{
