Remove resume for partial files when constructing in memory: if somehow we had a valid resumedata present, we would hit first check(!Configuration.bConstructInMemory)  in FBuildPatchFileConstructor::ResumeFile and crash at Request->Batch->PendingRequestCount in FBuildPatchFileConstructor::RequestCompletedFn where Request->Batch is an invalid pointer

diff --git a/Engine/Source/Runtime/Online/BuildPatchServices/Private/BuildPatchFileConstructor.cpp b/Engine/Source/Runtime/Online/BuildPatchServices/Private/BuildPatchFileConstructor.cpp
--- a/Engine/Source/Runtime/Online/BuildPatchServices/Private/BuildPatchFileConstructor.cpp
+++ b/Engine/Source/Runtime/Online/BuildPatchServices/Private/BuildPatchFileConstructor.cpp
@@ -2008,6 +2008,13 @@ void FBuildPatchFileConstructor::Run()
 					FileSystem->DeleteFile(*(Configuration.StagingDirectory / FileToConstruct));
 				}
 			}
+
+			if (Configuration.bConstructInMemory)
+			{
+				UE_LOG(LogBuildPatchServices, Log, TEXT("Emptying Resume Data FilesStarted: construct in memory"));
+				ResumeData.FilesStarted.Empty();
+			}
+
 		}
 
 		// Save for started versions
