Add bSkipInitialDiskSizeCheck to FBuildInstallerConfiguration so that we may skip initial size check prior to starting the installation

diff --git a/Engine/Source/Runtime/Online/BuildPatchServices/Private/BuildPatchFileConstructor.cpp b/Engine/Source/Runtime/Online/BuildPatchServices/Private/BuildPatchFileConstructor.cpp
--- a/Engine/Source/Runtime/Online/BuildPatchServices/Private/BuildPatchFileConstructor.cpp
+++ b/Engine/Source/Runtime/Online/BuildPatchServices/Private/BuildPatchFileConstructor.cpp
@@ -2761,7 +2761,7 @@ void FBuildPatchFileConstructor::InitFile(FFileConstructionState& CurrentFile, c
 
 bool FBuildPatchFileConstructor::HandleInitialDiskSizeCheck(const FFileManifest& FileManifest, int64 StartPosition)
 {
-	if (Configuration.bInstallToMemory)
+	if (Configuration.bInstallToMemory || Configuration.bSkipInitialDiskSizeCheck)
 	{
 		bInitialDiskSizeCheck = true;
 		return true;
diff --git a/Engine/Source/Runtime/Online/BuildPatchServices/Private/BuildPatchFileConstructor.h b/Engine/Source/Runtime/Online/BuildPatchServices/Private/BuildPatchFileConstructor.h
--- a/Engine/Source/Runtime/Online/BuildPatchServices/Private/BuildPatchFileConstructor.h
+++ b/Engine/Source/Runtime/Online/BuildPatchServices/Private/BuildPatchFileConstructor.h
@@ -83,6 +83,7 @@ namespace BuildPatchServices
 		// See comments in installer config
 		bool bInstallToMemory = false;
 		bool bConstructInMemory = false;
+		bool bSkipInitialDiskSizeCheck = false;
 
 		bool bDeleteChunkDBFilesAfterUse = false;
 
diff --git a/Engine/Source/Runtime/Online/BuildPatchServices/Private/BuildPatchInstaller.cpp b/Engine/Source/Runtime/Online/BuildPatchServices/Private/BuildPatchInstaller.cpp
--- a/Engine/Source/Runtime/Online/BuildPatchServices/Private/BuildPatchInstaller.cpp
+++ b/Engine/Source/Runtime/Online/BuildPatchServices/Private/BuildPatchInstaller.cpp
@@ -1462,6 +1462,7 @@ namespace BuildPatchServices
 			FFileConstructorConfig FCC;
 			FCC.bInstallToMemory = Configuration.bInstallToMemory;
 			FCC.bConstructInMemory = Configuration.bConstructFilesInMemory;
+			FCC.bSkipInitialDiskSizeCheck = Configuration.bSkipInitialDiskSizeCheck;
 			FCC.ConstructList = FilesToConstruct.Array();
 			FCC.InstallDirectory = Configuration.InstallDirectory;
 			FPaths::NormalizeDirectoryName(FCC.InstallDirectory);
diff --git a/Engine/Source/Runtime/Online/BuildPatchServices/Public/BuildPatchSettings.h b/Engine/Source/Runtime/Online/BuildPatchServices/Public/BuildPatchSettings.h
--- a/Engine/Source/Runtime/Online/BuildPatchServices/Public/BuildPatchSettings.h
+++ b/Engine/Source/Runtime/Online/BuildPatchServices/Public/BuildPatchSettings.h
@@ -332,6 +332,9 @@ namespace BuildPatchServices
 		// As bConstructFilesInMemory, except the file is never written to disk and is instead made available after via 
 		// GetFilesInstalledToMemory 
 		bool bInstallToMemory = false;
+
+		// If set, no initial size check will be performed prior to starting the installation. Running out of disk space during installation is unchanged.
+		bool bSkipInitialDiskSizeCheck = false;
 	};
 
 	/**
