Prevent files constructing in memory to create a file on disk when installation failed

diff --git a/Engine/Source/Runtime/Online/BuildPatchServices/Private/BuildPatchFileConstructor.cpp b/Engine/Source/Runtime/Online/BuildPatchServices/Private/BuildPatchFileConstructor.cpp
--- a/Engine/Source/Runtime/Online/BuildPatchServices/Private/BuildPatchFileConstructor.cpp
+++ b/Engine/Source/Runtime/Online/BuildPatchServices/Private/BuildPatchFileConstructor.cpp
@@ -3023,7 +3023,7 @@ void FBuildPatchFileConstructor::CompleteConstructedFile(FFileConstructionState&
 		}
 	}
 
-	if (Configuration.bConstructInMemory)
+	if (CurrentFile.bSuccess && Configuration.bConstructInMemory)
 	{
 		// Now that we have the entire file ready, we write it in one big pass. This
 		// is synchronous to avoid complexity, but we do eat some performance since we
