Fix websocket event loop not being enabled on platforms where it is supported
Fix websocket event loop not being enabled by default in forked multithreaded servers

diff --git a/Engine/Source/Runtime/Online/WebSockets/Private/Lws/LwsWebSocketsManager.cpp b/Engine/Source/Runtime/Online/WebSockets/Private/Lws/LwsWebSocketsManager.cpp
--- a/Engine/Source/Runtime/Online/WebSockets/Private/Lws/LwsWebSocketsManager.cpp
+++ b/Engine/Source/Runtime/Online/WebSockets/Private/Lws/LwsWebSocketsManager.cpp
@@ -76,6 +76,13 @@ FLwsWebSocketsManager& FLwsWebSocketsManager::Get()
 void FLwsWebSocketsManager::UpdateConfigs()
 {
 	const bool bWasPollingService = bPollService;
+	// For forked dedicated servers, the parent process may be running in single threaded mode, which will force bPollService to true below
+	// Restore bPollService to its code defined default to avoid the code specified default not applying to a post-fork multithreaded server
+#if UE_WEBSOCKETS_PLATFORM_SUPPORT_EVENT_LOOP
+	bPollService = false;
+#else
+	bPollService = true;
+#endif
 	GConfig->GetBool(TEXT("WebSockets.LibWebSockets"), TEXT("bPollService"), bPollService, GEngineIni);
 	GConfig->GetInt(TEXT("WebSockets.LibWebSockets"), TEXT("ServiceTimeoutMs"), ServiceTimeoutMs, GEngineIni);
 	GConfig->GetDouble(TEXT("WebSockets.LibWebSockets"), TEXT("ThreadTargetFrameTimeInSeconds"), ThreadTargetFrameTimeInSeconds, GEngineIni);
diff --git a/Engine/Source/Runtime/Online/WebSockets/Private/Lws/LwsWebSocketsManager.h b/Engine/Source/Runtime/Online/WebSockets/Private/Lws/LwsWebSocketsManager.h
--- a/Engine/Source/Runtime/Online/WebSockets/Private/Lws/LwsWebSocketsManager.h
+++ b/Engine/Source/Runtime/Online/WebSockets/Private/Lws/LwsWebSocketsManager.h
@@ -120,7 +120,7 @@ private:
 	 * 
 	 * If false, lws_service will use a ServiceTimoutMs timeout, allowing lws to block and handle wake / sleep on activity internally itself.
 	 */
-#if UE_WEBSOCETS_PLATFORM_SUPPORT_EVENT_LOOP
+#if UE_WEBSOCKETS_PLATFORM_SUPPORT_EVENT_LOOP
 	bool bPollService = false;
 #else
 	bool bPollService = true;
