- Add error logging to unary calls

diff --git a/Engine/Source/Runtime/Online/Experimental/GrpcClient/Public/GrpcClient.h b/Engine/Source/Runtime/Online/Experimental/GrpcClient/Public/GrpcClient.h
--- a/Engine/Source/Runtime/Online/Experimental/GrpcClient/Public/GrpcClient.h
+++ b/Engine/Source/Runtime/Online/Experimental/GrpcClient/Public/GrpcClient.h
@@ -272,8 +272,11 @@ public:
 protected:
 	virtual void HandleNext(void* Tag, bool bOk) override
 	{
-		ResponseDelegate.ExecuteIfBound(this->MapStatus(Status), MoveTemp(Response));
-		this->UnregisterCall();
+		UE_CLOG(Status.error_code() != grpc::StatusCode::OK, LogGrpcClient, Warning, TEXT("Status=%s ErrorMessage=[%hs] ErrorDetails=[%hs]"),
+			LexToString(MapStatus(Status)), Status.error_message().c_str(), Status.error_details().c_str());
+
+		ResponseDelegate.ExecuteIfBound(MapStatus(Status), MoveTemp(Response));
+		UnregisterCall();
 	}
 
 	FGrpcTag FinishTag{this};
