- Add error logging to non unary operations

diff --git a/Engine/Source/Runtime/Online/Experimental/GrpcClient/Public/GrpcClient.h b/Engine/Source/Runtime/Online/Experimental/GrpcClient/Public/GrpcClient.h
--- a/Engine/Source/Runtime/Online/Experimental/GrpcClient/Public/GrpcClient.h
+++ b/Engine/Source/Runtime/Online/Experimental/GrpcClient/Public/GrpcClient.h
@@ -352,6 +352,9 @@ protected:
 
 		if (Tag == (void*)&FinishTag)
 		{
+			UE_CLOG(Status.error_code() != grpc::StatusCode::OK, LogGrpcClient, Warning, TEXT("Status=%s ErrorMessage=[%hs] ErrorDetails=[%hs]"),
+				LexToString(MapStatus(Status)), Status.error_message().c_str(), Status.error_details().c_str());
+
 			DoneDelegate.ExecuteIfBound(this->MapStatus(Status));
 
 			this->UnregisterCall();
@@ -459,6 +462,9 @@ protected:
 		// Complete the call and call Done delegate
 		if (Tag == (void*)&FinishTag)
 		{
+			UE_CLOG(Status.error_code() != grpc::StatusCode::OK, LogGrpcClient, Warning, TEXT("Status=%s ErrorMessage=[%hs] ErrorDetails=[%hs]"),
+				LexToString(MapStatus(Status)), Status.error_message().c_str(), Status.error_details().c_str());
+
 			DoneDelegate.ExecuteIfBound(this->MapStatus(Status));
 
 			this->UnregisterCall();
