BPS: Added logging to try and hunt down a live crash.

diff --git a/Engine/Source/Runtime/Online/BuildPatchServices/Private/BuildPatchFileConstructor.cpp b/Engine/Source/Runtime/Online/BuildPatchServices/Private/BuildPatchFileConstructor.cpp
--- a/Engine/Source/Runtime/Online/BuildPatchServices/Private/BuildPatchFileConstructor.cpp
+++ b/Engine/Source/Runtime/Online/BuildPatchServices/Private/BuildPatchFileConstructor.cpp
@@ -1490,7 +1490,8 @@ FBuildPatchFileConstructor::FBuildPatchFileConstructor(
 		{
 			bHasInstallSource = true;
 
-			// We need to set up a dependency chain so that files can know when they can start.			
+			// We need to set up a dependency chain so that files can know when they can start. We can't start constructing a file
+			// that needs parts from a file that isn't done as the patch authoring tool expects the file to be done.
 			TMap<FStringView, int32> FileToIndexMap;
 			FileToIndexMap.Reserve(Configuration.ConstructList.Num());
 			for (int32 FileConstructIndex = 0; FileConstructIndex < Configuration.ConstructList.Num(); FileConstructIndex++)
@@ -1501,7 +1502,7 @@ FBuildPatchFileConstructor::FBuildPatchFileConstructor(
 
 			for (int32 FileConstructIndex = 0; FileConstructIndex < Configuration.ConstructList.Num(); FileConstructIndex++)
 			{
-				FFileToConstruct& FileToConstruct = ConstructionList[FileConstructIndex];				
+				FFileToConstruct& FileToConstruct = ConstructionList[FileConstructIndex];
 				if (!FileToConstruct.FileManifest)
 				{
 					continue;
@@ -1536,7 +1537,14 @@ FBuildPatchFileConstructor::FBuildPatchFileConstructor(
 						);
 					}
 				}
-			}			
+
+				if (FileToConstruct.LatestDependentInstallSource >= 0)
+				{
+					UE_LOG(LogBuildPatchServices, Display, TEXT("File: %s can't start until %s finishes"), 
+						*Configuration.ConstructList[FileConstructIndex],
+						*Configuration.ConstructList[FileToConstruct.LatestDependentInstallSource]);
+				}
+			}
 		}
 	}
 
@@ -2931,6 +2939,9 @@ void FBuildPatchFileConstructor::OpenFileToConstruct(FFileConstructionState& Cur
 
 void FBuildPatchFileConstructor::CompleteConstructedFile(FFileConstructionState& CurrentFile)
 {
+	UE_LOG(LogBuildPatchServices, Display, TEXT("Completing: %s"), *CurrentFile.BuildFilename);
+	bShouldLogNextDependencyWait = true;
+
 	TRACE_CPUPROFILER_EVENT_SCOPE(Constructor_CompleteFile);
 	if (!CurrentFile.bSkippedConstruction)
 	{
@@ -3660,6 +3671,19 @@ void FBuildPatchFileConstructor::ConstructFiles(const FResumeData& ResumeData)
 									// are safe.
 									if (ActiveFileCount && ActiveFiles[0]->ConstructionIndex <= ConstructionList[IndexToConstruct].LatestDependentInstallSource)
 									{
+										if (bShouldLogNextDependencyWait)
+										{
+											bShouldLogNextDependencyWait = false;
+
+											UE_LOG(LogBuildPatchServices, Display, TEXT("Delaying %s, active head is %s, waiting for %s %d/%d"), 
+												*Configuration.ConstructList[IndexToConstruct],
+												*Configuration.ConstructList[ActiveFiles[0]->ConstructionIndex],
+												*Configuration.ConstructList[ConstructionList[IndexToConstruct].LatestDependentInstallSource],
+												ActiveFiles[0]->ConstructionIndex,
+												ConstructionList[IndexToConstruct].LatestDependentInstallSource
+											);
+										}
+
 										bDelayForDependencies = true;
 									}
 								}
@@ -3685,7 +3709,7 @@ void FBuildPatchFileConstructor::ConstructFiles(const FResumeData& ResumeData)
 
 									TRACE_BOOKMARK(TEXT("Starting File: %s [%u]"), *AnotherFile->BuildFilename, AnotherFile->FileManifest->ChunkParts.Num());
 
-									UE_LOG(LogBuildPatchServices, Verbose, TEXT("Starting File: %s [%s bytes, %d chunks], New ActiveFileCount = %d"), 
+									UE_LOG(LogBuildPatchServices, Display, TEXT("Starting File: %s [%s bytes, %d chunks], New ActiveFileCount = %d"), 
 										   *AnotherFile->BuildFilename, 
 										   *FormatNumber(AnotherFile->FileManifest->FileSize),
 										   AnotherFile->FileManifest->ChunkParts.Num(),
@@ -3711,10 +3735,6 @@ void FBuildPatchFileConstructor::ConstructFiles(const FResumeData& ResumeData)
 										ActiveFiles.Push(MoveTemp(AnotherFile));
 									}
 								} // end if not delaying for dependencies
-								else
-								{
-									UE_LOG(LogBuildPatchServices, Verbose, TEXT("Delaying %s due to incomplete dependencies"), *Configuration.ConstructList[IndexToConstruct]);
-								}
 
 								// If we are delaying for dependencies we fall through here with FileToStart = nullptr
 								// and nothing happens until we check again when the construct thread is woken up again.
diff --git a/Engine/Source/Runtime/Online/BuildPatchServices/Private/BuildPatchFileConstructor.h b/Engine/Source/Runtime/Online/BuildPatchServices/Private/BuildPatchFileConstructor.h
--- a/Engine/Source/Runtime/Online/BuildPatchServices/Private/BuildPatchFileConstructor.h
+++ b/Engine/Source/Runtime/Online/BuildPatchServices/Private/BuildPatchFileConstructor.h
@@ -273,6 +273,8 @@ namespace BuildPatchServices
 
 		bool bAllowMultipleFilesInFlight = true;
 
+		bool bShouldLogNextDependencyWait = true;
+
 		// The size we expect for chunks. This should be used for estimation purposes, not anything requiring hard limits.
 		uint32 ExpectedChunkSize = 0;
 
diff --git a/Engine/Source/Runtime/Online/BuildPatchServices/Private/Installer/InstallChunkSource.cpp b/Engine/Source/Runtime/Online/BuildPatchServices/Private/Installer/InstallChunkSource.cpp
--- a/Engine/Source/Runtime/Online/BuildPatchServices/Private/Installer/InstallChunkSource.cpp
+++ b/Engine/Source/Runtime/Online/BuildPatchServices/Private/Installer/InstallChunkSource.cpp
@@ -41,6 +41,11 @@ namespace BuildPatchServices
 			// With multiple files in flight in the constructor we can be deleting a file at the same time as we are
 			// reading chunks for other files, which means we can hit OpenedFileHandles from multiple threads.
 			UE::TUniqueLock _(FileHandleLock);
+			TUniquePtr<FOpenedFile>* File = OpenedFileHandles.Find(FilePath);
+			if (File && (*File)->UsageCount != 0)
+			{
+				UE_LOG(LogInstallChunkSource, Error, TEXT("Tried to delete a file we're reading from! %s"), *FilePath);
+			}
 			OpenedFileHandles.Remove(FilePath);
 		}
 		virtual int32 GetChunkUnavailableAt(const FGuid& DataId) const override;
@@ -91,8 +96,17 @@ namespace BuildPatchServices
 		TSet<FGuid> AvailableInBuilds;
 		TArray<TPair<FString, FBuildPatchAppManifestRef>> InstallationSources;
 
+		// This is here to try and track down an unexplained crash where it looks like maybe we're hitting the
+		// same file during delete as we are during read (which shouldn't happen).
+		struct FOpenedFile
+		{
+			FOpenedFile(TUniquePtr<FArchive>&& InFileHandle) : FileHandle(MoveTemp(InFileHandle)) {}
+
+			TUniquePtr<FArchive> FileHandle;
+			std::atomic_int32_t UsageCount;
+		};
 		UE::FMutex FileHandleLock; // Protects OpenedFileHandles
-		TMap<FString, TUniquePtr<FArchive>> OpenedFileHandles;
+		TMap<FString, TUniquePtr<FOpenedFile>> OpenedFileHandles;
 
 		// The index (ChunkReferenceTracker->GetCurrentUsageIndex) at which our files will get deleted due to destructive install to make room for
 		// the new file. 
@@ -187,7 +201,8 @@ namespace BuildPatchServices
 			return [](bool) {return;};
 		}
 
-		// Must be run sequentially! No threaded protection provided.
+		// Must be run sequentially! No threaded protection provided. The expectation is that all process functions
+		// returned here must be run sequentially, and that the underlying file must not be deleted.
 		return [this, FoundInstallManifest, FoundInstallDirectory = *FoundInstallDirectory, DataId, DestinationBuffer, UserPtr, CompleteFn](bool bIsAborted)
 		{
 			ISpeedRecorder::FRecord ActivityRecord;
@@ -247,11 +262,11 @@ namespace BuildPatchServices
 				// We use the internal FArchive pointer so that we don't have to hold the file handle
 				// lock over the read - the uniqueptr might move around, but the managed pointer will not.
 				// We do know that our specific file won't get deleted until we are done.
-				FArchive* FileArchive = nullptr;
+				FOpenedFile* OpenedFile = nullptr;
 				{
 					FileHandleLock.Lock();
-					TUniquePtr<FArchive>* FileArchivePtr = OpenedFileHandles.Find(FullFilename);
-					if (FileArchivePtr == nullptr)
+					TUniquePtr<FOpenedFile>* OpenedFilePtr = OpenedFileHandles.Find(FullFilename);
+					if (OpenedFilePtr == nullptr)
 					{
 						TRACE_CPUPROFILER_EVENT_SCOPE(Install_OpenSource);
 
@@ -262,27 +277,43 @@ namespace BuildPatchServices
 
 						if (NewReader.IsValid())
 						{
-							OpenedFileHandles.Add(FullFilename, MoveTemp(NewReader));
-							FileArchivePtr = OpenedFileHandles.Find(FullFilename);
+							TUniquePtr<FOpenedFile> NewOpenedFile = MakeUnique<FOpenedFile>(MoveTemp(NewReader));
+
+							OpenedFileHandles.Add(FullFilename, MoveTemp(NewOpenedFile));
+							OpenedFilePtr = OpenedFileHandles.Find(FullFilename);
 						}
 					}
-					if (FileArchivePtr)
+					if (OpenedFilePtr)
 					{
-						FileArchive = FileArchivePtr->Get();
+						OpenedFile = OpenedFilePtr->Get();
 					}
 					FileHandleLock.Unlock();
 				}
 
-				if (FileArchive == nullptr)
+				if (OpenedFile == nullptr)
 				{
 					Result = IInstallChunkSourceStat::ELoadResult::OpenFileFail;
 					break;
 				}
 
+				// Should only ever be one thread hitting files.
+				if (OpenedFile->UsageCount.fetch_add(1, std::memory_order_acq_rel) != 0)
+				{
+					UE_LOG(LogInstallChunkSource, Warning, TEXT("Invalid usage count in install source! %s"), *FileChunkPart.FileManifest->Filename);
+				}
+
 				{
 					TRACE_CPUPROFILER_EVENT_SCOPE(Install_Serialize);
-					FileArchive->Seek(FileChunkPart.FileOffset);
-					FileArchive->Serialize((uint8*)DestinationBuffer.GetData() + FileChunkPart.ChunkPart->Offset, FileChunkPart.ChunkPart->Size);
+					OpenedFile->FileHandle->Seek(FileChunkPart.FileOffset);
+
+					// We're hunting down some weird crashes where we can't tell if it's because our dest pointer is invalid or not,
+					// so do a silly check here to try and catch it if so.
+					{
+						uint8* BasePtr = (uint8*)DestinationBuffer.GetData() + FileChunkPart.ChunkPart->Offset;
+						BasePtr[0] = 0;
+						BasePtr[FileChunkPart.ChunkPart->Size-1] = 0;
+					}
+					OpenedFile->FileHandle->Serialize((uint8*)DestinationBuffer.GetData() + FileChunkPart.ChunkPart->Offset, FileChunkPart.ChunkPart->Size);
 
 					ActivityRecord.Size += FileChunkPart.ChunkPart->Size;
 
@@ -299,6 +330,12 @@ namespace BuildPatchServices
 					ChunkBlocks.Add(FileChunkPart.ChunkPart->Offset, FileChunkPart.ChunkPart->Size);
 				}
 
+				// Should only ever be one thread hitting files.
+				if (OpenedFile->UsageCount.fetch_add(-1, std::memory_order_acq_rel) != 1)
+				{
+					UE_LOG(LogInstallChunkSource, Warning, TEXT("Invalid usage count in install source post read! %s"), *FileChunkPart.FileManifest->Filename);
+				}
+
 				// The expectation is that we get the full chunk only once we've assembled all of the parts
 				// provided by the manifest, so the last iteration should set this to true. If it isn't the last
 				// iteration, then we'll hit the faux-error case at the top of the loop.
