BPS: Avoid statics to clean up compiler generated thread synchronization code, decrease exe size, and prevent construction of every error string when we only need one.

diff --git a/Engine/Source/Runtime/Online/BuildPatchServices/Private/Installer/InstallerError.cpp b/Engine/Source/Runtime/Online/BuildPatchServices/Private/Installer/InstallerError.cpp
--- a/Engine/Source/Runtime/Online/BuildPatchServices/Private/Installer/InstallerError.cpp
+++ b/Engine/Source/Runtime/Online/BuildPatchServices/Private/Installer/InstallerError.cpp
@@ -10,47 +10,45 @@ namespace BuildPatchServices
 {
 	FText GetStandardErrorText(EBuildPatchInstallError ErrorType)
 	{
-		// The generic texts, for when not specified by the system reporting the error.
-		static const FText NoError(NSLOCTEXT("BuildPatchInstallError", "BuildPatchInstallShortError_NoError", "The operation was successful."));
-		static const FText DownloadError(NSLOCTEXT("BuildPatchInstallError", "BuildPatchInstallShortError_DownloadError", "Could not download patch data. Please check your internet connection, and try again."));
-		static const FText FileConstructionFail(NSLOCTEXT("BuildPatchInstallError", "BuildPatchInstallShortError_FileConstructionFail", "A file corruption has occurred. Please try again."));
-		static const FText MoveFileToInstall(NSLOCTEXT("BuildPatchInstallError", "BuildPatchInstallShortError_MoveFileToInstall", "A file access error has occurred. Please check your running processes."));
-		static const FText BuildVerifyFail(NSLOCTEXT("BuildPatchInstallError", "BuildPatchInstallShortError_BuildCorrupt", "The installation is corrupt. Please contact support."));
-		static const FText ApplicationClosing(NSLOCTEXT("BuildPatchInstallError", "BuildPatchInstallShortError_ApplicationClosing", "The application is closing."));
-		static const FText ApplicationError(NSLOCTEXT("BuildPatchInstallError", "BuildPatchInstallShortError_ApplicationError", "Patching service could not start. Please contact support."));
-		static const FText UserCanceled(NSLOCTEXT("BuildPatchInstallError", "BuildPatchInstallShortError_UserCanceled", "User cancelled."));
-		static const FText PrerequisiteError(NSLOCTEXT("BuildPatchInstallError", "BuildPatchInstallShortError_PrerequisiteError", "The necessary prerequisites have failed to install. Please contact support."));
-		static const FText InitializationError(NSLOCTEXT("BuildPatchInstallError", "BuildPatchInstallShortError_InitializationError", "The installer failed to initialize. Please contact support."));
-		static const FText PathLengthExceeded(NSLOCTEXT("BuildPatchInstallError", "BuildPatchInstallShortError_PathLengthExceeded", "Maximum path length exceeded. Please specify a shorter install location."));
-		static const FText OutOfDiskSpace(NSLOCTEXT("BuildPatchInstallError", "BuildPatchInstallShortError_OutOfDiskSpace", "Not enough disk space available. Please free up some disk space and try again."));
-		static const FText InvalidOrMax(NSLOCTEXT("BuildPatchInstallError", "BuildPatchInstallShortError_InvalidOrMax", "An unknown error ocurred. Please contact support."));
-
 		switch (ErrorType)
 		{
-			case EBuildPatchInstallError::NoError: return NoError;
-			case EBuildPatchInstallError::DownloadError: return DownloadError;
-			case EBuildPatchInstallError::FileConstructionFail: return FileConstructionFail;
-			case EBuildPatchInstallError::MoveFileToInstall: return MoveFileToInstall;
-			case EBuildPatchInstallError::BuildVerifyFail: return BuildVerifyFail;
-			case EBuildPatchInstallError::ApplicationClosing: return ApplicationClosing;
-			case EBuildPatchInstallError::ApplicationError: return ApplicationError;
-			case EBuildPatchInstallError::UserCanceled: return UserCanceled;
-			case EBuildPatchInstallError::PrerequisiteError: return PrerequisiteError;
-			case EBuildPatchInstallError::InitializationError: return InitializationError;
-			case EBuildPatchInstallError::PathLengthExceeded: return PathLengthExceeded;
-			case EBuildPatchInstallError::OutOfDiskSpace: return OutOfDiskSpace;
-			default: return InvalidOrMax;
+			case EBuildPatchInstallError::NoError: 
+				return NSLOCTEXT("BuildPatchInstallError", "BuildPatchInstallShortError_NoError", "The operation was successful.");
+			case EBuildPatchInstallError::DownloadError: 
+				return NSLOCTEXT("BuildPatchInstallError", "BuildPatchInstallShortError_DownloadError", "Could not download patch data. Please check your internet connection, and try again.");
+			case EBuildPatchInstallError::FileConstructionFail: 
+				return NSLOCTEXT("BuildPatchInstallError", "BuildPatchInstallShortError_FileConstructionFail", "A file corruption has occurred. Please try again.");
+			case EBuildPatchInstallError::MoveFileToInstall: 
+				return NSLOCTEXT("BuildPatchInstallError", "BuildPatchInstallShortError_MoveFileToInstall", "A file access error has occurred. Please check your running processes.");
+			case EBuildPatchInstallError::BuildVerifyFail: 
+				return NSLOCTEXT("BuildPatchInstallError", "BuildPatchInstallShortError_BuildCorrupt", "The installation is corrupt. Please contact support.");
+			case EBuildPatchInstallError::ApplicationClosing: 
+				return NSLOCTEXT("BuildPatchInstallError", "BuildPatchInstallShortError_ApplicationClosing", "The application is closing.");
+			case EBuildPatchInstallError::ApplicationError: 
+				return NSLOCTEXT("BuildPatchInstallError", "BuildPatchInstallShortError_ApplicationError", "Patching service could not start. Please contact support.");
+			case EBuildPatchInstallError::UserCanceled: 
+				return NSLOCTEXT("BuildPatchInstallError", "BuildPatchInstallShortError_UserCanceled", "User cancelled.");
+			case EBuildPatchInstallError::PrerequisiteError: 
+				return NSLOCTEXT("BuildPatchInstallError", "BuildPatchInstallShortError_PrerequisiteError", "The necessary prerequisites have failed to install. Please contact support.");
+			case EBuildPatchInstallError::InitializationError: 
+				return NSLOCTEXT("BuildPatchInstallError", "BuildPatchInstallShortError_InitializationError", "The installer failed to initialize. Please contact support.");
+			case EBuildPatchInstallError::PathLengthExceeded: 
+				return NSLOCTEXT("BuildPatchInstallError", "BuildPatchInstallShortError_PathLengthExceeded", "Maximum path length exceeded. Please specify a shorter install location.");
+			case EBuildPatchInstallError::OutOfDiskSpace: 
+				return NSLOCTEXT("BuildPatchInstallError", "BuildPatchInstallShortError_OutOfDiskSpace", "Not enough disk space available. Please free up some disk space and try again.");
+			default: 
+				return NSLOCTEXT("BuildPatchInstallError", "BuildPatchInstallShortError_InvalidOrMax", "An unknown error ocurred. Please contact support.");
 		}
 	}
 
 	FText GetDiskSpaceMessage(const FString& Location, uint64 RequiredBytes, uint64 AvailableBytes, const FNumberFormattingOptions* FormatOptions)
 	{
 #if PLATFORM_DESKTOP
-        static const FText OutOfDiskSpace(NSLOCTEXT("BuildPatchInstallError", "InstallDirectoryDiskSpace", "There is not enough space at {Location}\n{RequiredBytes} is required.\n{AvailableBytes} is available.\nYou need an additional {SpaceAdditional} to perform the installation."));
+        const FText OutOfDiskSpace(NSLOCTEXT("BuildPatchInstallError", "InstallDirectoryDiskSpace", "There is not enough space at {Location}\n{RequiredBytes} is required.\n{AvailableBytes} is available.\nYou need an additional {SpaceAdditional} to perform the installation."));
 #else
-        static const FText OutOfDiskSpace(NSLOCTEXT("BuildPatchInstallError", "InstallDirectoryDiskSpaceDevice", "There is not enough space on your device.\n{RequiredBytes} is required.\n{AvailableBytes} is available.\nYou need an additional {SpaceAdditional} to perform the installation."));
+        const FText OutOfDiskSpace(NSLOCTEXT("BuildPatchInstallError", "InstallDirectoryDiskSpaceDevice", "There is not enough space on your device.\n{RequiredBytes} is required.\n{AvailableBytes} is available.\nYou need an additional {SpaceAdditional} to perform the installation."));
 #endif // PLATFORM_DESKTOP
-        static const FNumberFormattingOptions DefaultOptions = FNumberFormattingOptions()
+        const FNumberFormattingOptions DefaultOptions = FNumberFormattingOptions()
 			.SetMinimumFractionalDigits(2)
 			.SetMaximumFractionalDigits(2);
 		FormatOptions = FormatOptions == nullptr ? &DefaultOptions : FormatOptions;
@@ -199,38 +197,23 @@ namespace BuildPatchServices
 		return new FInstallerError();
 	}
 
-	const FString& EnumToString(const EBuildPatchInstallError& InErrorType)
+	FString EnumToString(const EBuildPatchInstallError& InErrorType)
 	{
-		// Const enum strings, special case no error.
-		static const FString NoError("SUCCESS");
-		static const FString DownloadError("EBuildPatchInstallError::DownloadError");
-		static const FString FileConstructionFail("EBuildPatchInstallError::FileConstructionFail");
-		static const FString MoveFileToInstall("EBuildPatchInstallError::MoveFileToInstall");
-		static const FString BuildVerifyFail("EBuildPatchInstallError::BuildVerifyFail");
-		static const FString ApplicationClosing("EBuildPatchInstallError::ApplicationClosing");
-		static const FString ApplicationError("EBuildPatchInstallError::ApplicationError");
-		static const FString UserCanceled("EBuildPatchInstallError::UserCanceled");
-		static const FString PrerequisiteError("EBuildPatchInstallError::PrerequisiteError");
-		static const FString InitializationError("EBuildPatchInstallError::InitializationError");
-		static const FString PathLengthExceeded("EBuildPatchInstallError::PathLengthExceeded");
-		static const FString OutOfDiskSpace("EBuildPatchInstallError::OutOfDiskSpace");
-		static const FString InvalidOrMax("EBuildPatchInstallError::InvalidOrMax");
-
 		switch (InErrorType)
 		{
-			case EBuildPatchInstallError::NoError: return NoError;
-			case EBuildPatchInstallError::DownloadError: return DownloadError;
-			case EBuildPatchInstallError::FileConstructionFail: return FileConstructionFail;
-			case EBuildPatchInstallError::MoveFileToInstall: return MoveFileToInstall;
-			case EBuildPatchInstallError::BuildVerifyFail: return BuildVerifyFail;
-			case EBuildPatchInstallError::ApplicationClosing: return ApplicationClosing;
-			case EBuildPatchInstallError::ApplicationError: return ApplicationError;
-			case EBuildPatchInstallError::UserCanceled: return UserCanceled;
-			case EBuildPatchInstallError::PrerequisiteError: return PrerequisiteError;
-			case EBuildPatchInstallError::InitializationError: return InitializationError;
-			case EBuildPatchInstallError::PathLengthExceeded: return PathLengthExceeded;
-			case EBuildPatchInstallError::OutOfDiskSpace: return OutOfDiskSpace;
-			default: return InvalidOrMax;
+			case EBuildPatchInstallError::NoError: return TEXT("SUCCESS");
+			case EBuildPatchInstallError::DownloadError: return TEXT("EBuildPatchInstallError::DownloadError");
+			case EBuildPatchInstallError::FileConstructionFail: return TEXT("EBuildPatchInstallError::FileConstructionFail");
+			case EBuildPatchInstallError::MoveFileToInstall: return TEXT("EBuildPatchInstallError::MoveFileToInstall");
+			case EBuildPatchInstallError::BuildVerifyFail: return TEXT("EBuildPatchInstallError::BuildVerifyFail");
+			case EBuildPatchInstallError::ApplicationClosing: return TEXT("EBuildPatchInstallError::ApplicationClosing");
+			case EBuildPatchInstallError::ApplicationError: return TEXT("EBuildPatchInstallError::ApplicationError");
+			case EBuildPatchInstallError::UserCanceled: return TEXT("EBuildPatchInstallError::UserCanceled");
+			case EBuildPatchInstallError::PrerequisiteError: return TEXT("EBuildPatchInstallError::PrerequisiteError");
+			case EBuildPatchInstallError::InitializationError: return TEXT("EBuildPatchInstallError::InitializationError");
+			case EBuildPatchInstallError::PathLengthExceeded: return TEXT("EBuildPatchInstallError::PathLengthExceeded");
+			case EBuildPatchInstallError::OutOfDiskSpace: return TEXT("EBuildPatchInstallError::OutOfDiskSpace");
+			default: return TEXT("EBuildPatchInstallError::InvalidOrMax");
 		}
 	}
 
diff --git a/Engine/Source/Runtime/Online/BuildPatchServices/Private/Installer/InstallerError.h b/Engine/Source/Runtime/Online/BuildPatchServices/Private/Installer/InstallerError.h
--- a/Engine/Source/Runtime/Online/BuildPatchServices/Private/Installer/InstallerError.h
+++ b/Engine/Source/Runtime/Online/BuildPatchServices/Private/Installer/InstallerError.h
@@ -229,5 +229,5 @@ namespace BuildPatchServices
 	 * @param ErrorType     The error type value.
 	 * @return A string value for the EBuildPatchInstallError enum.
 	 */
-	const FString& EnumToString(const EBuildPatchInstallError& ErrorType);
+	FString EnumToString(const EBuildPatchInstallError& ErrorType);
 }
