BPS: Fixed stale disk space messaging causing write errors to report exabytes of additional disk space required.

diff --git a/Engine/Source/Runtime/Online/BuildPatchServices/Private/BuildPatchFileConstructor.cpp b/Engine/Source/Runtime/Online/BuildPatchServices/Private/BuildPatchFileConstructor.cpp
--- a/Engine/Source/Runtime/Online/BuildPatchServices/Private/BuildPatchFileConstructor.cpp
+++ b/Engine/Source/Runtime/Online/BuildPatchServices/Private/BuildPatchFileConstructor.cpp
@@ -74,27 +74,6 @@ namespace FileConstructorHelpers
 		}
 	}
 
-	bool CheckRemainingDiskSpace(const FString& InstallDirectory, uint64 RemainingBytesRequired, uint64& OutAvailableDiskSpace)
-	{
-		bool bContinueConstruction = true;
-		uint64 TotalSize = 0;
-		OutAvailableDiskSpace = 0;
-		if (FPlatformMisc::GetDiskTotalAndFreeSpace(InstallDirectory, TotalSize, OutAvailableDiskSpace))
-		{
-			if (OutAvailableDiskSpace < RemainingBytesRequired)
-			{
-				bContinueConstruction = false;
-			}
-		}
-		else
-		{
-			// If we can't get the disk space free then the most likely reason is the drive is no longer around...
-			bContinueConstruction = false;
-		}
-
-		return bContinueConstruction;
-	}
-
 	uint64 CalculateRequiredDiskSpace(const FBuildPatchAppManifestPtr& CurrentManifest, const FBuildPatchAppManifestRef& BuildManifest, const EInstallMode& InstallMode, const TSet<FString>& InInstallTags)
 	{
 		// Make tags expected
@@ -182,7 +161,7 @@ enum class EConstructionError : uint8
 {
 	None = 0,
 	CannotCreateFile,
-	OutOfDiskSpace,
+	FailedWrite,
 	FailedInitialSizeCheck,
 	MissingChunk,
 	SerializeError,
@@ -1055,10 +1034,6 @@ public:
 	{
 		InstallationFreeSpaceRequired = InInstallationFreeSpaceRequired;
 	}
-	void DisableDynamicDiskSpaceHeadroom()
-	{
-		InstallationFreeSpaceRequired = 0;
-	}
 
 	// false means consistency failure
 	// Constructor thread
@@ -2790,29 +2765,29 @@ bool FBuildPatchFileConstructor::HandleInitialDiskSizeCheck(const FFileManifest&
 			LocalDiskSpaceRequired = CalculateDiskSpaceRequirementsWithDeleteDuringInstall();
 		}
 
+		bool bHaveSufficientFreeSpace = false;
+		bool bGotValidFreeSpace = false;
 		uint64 LocalDiskSpaceAvailable = 0;
 		{
 			uint64 TotalSize = 0;
 			uint64 AvailableSpace = 0;
-			if (FPlatformMisc::GetDiskTotalAndFreeSpace(Configuration.InstallDirectory, TotalSize, AvailableSpace))
+			bGotValidFreeSpace = FPlatformMisc::GetDiskTotalAndFreeSpace(Configuration.InstallDirectory, TotalSize, AvailableSpace);
+			if (bGotValidFreeSpace)
 			{
+				bHaveSufficientFreeSpace = AvailableSpace > LocalDiskSpaceRequired;
 				LocalDiskSpaceAvailable = AvailableSpace;
 				BackingStore->SetDynamicDiskSpaceHeadroom(LocalDiskSpaceRequired);
 			}
-			else
-			{
-				BackingStore->DisableDynamicDiskSpaceHeadroom();
-			}
 		}
 
-		UE_LOG(LogBuildPatchServices, Verbose, TEXT("Initial Disk Sizes: Required: %s Available: %s"), *FormatNumber(LocalDiskSpaceRequired), *FormatNumber(LocalDiskSpaceAvailable));
+		UE_LOG(LogBuildPatchServices, Display, TEXT("Initial Disk Sizes: Required: %s Available: %s"), *FormatNumber(LocalDiskSpaceRequired), bGotValidFreeSpace ? *FormatNumber(LocalDiskSpaceAvailable) : TEXT("<failed>"));
 
 		AvailableDiskSpace.store(LocalDiskSpaceAvailable, std::memory_order_release);
 		RequiredDiskSpace.store(LocalDiskSpaceRequired, std::memory_order_release);	
 
-		if (!FileConstructorHelpers::CheckRemainingDiskSpace(Configuration.InstallDirectory, LocalDiskSpaceRequired, LocalDiskSpaceAvailable))
+		if (!bGotValidFreeSpace || !bHaveSufficientFreeSpace)
 		{
-			UE_LOG(LogBuildPatchServices, Error, TEXT("Out of HDD space. Needs %llu bytes, Free %llu bytes"), LocalDiskSpaceRequired, LocalDiskSpaceAvailable);
+			UE_LOG(LogBuildPatchServices, Error, TEXT("Insufficient disk space on initial check."));
 			InstallerError->SetError(
 				EBuildPatchInstallError::OutOfDiskSpace,
 				DiskSpaceErrorCodes::InitialSpaceCheck,
@@ -3142,16 +3117,43 @@ void FBuildPatchFileConstructor::CompleteConstructedFile(FFileConstructionState&
 				InstallerError->SetError(EBuildPatchInstallError::FileConstructionFail, ConstructionErrorCodes::OutboundCorrupt);
 				break;
 			}
-		case EConstructionError::OutOfDiskSpace:
+		case EConstructionError::FailedWrite:
 			{
-				uint64 LocalAvailableDiskSpace = AvailableDiskSpace.load(std::memory_order_acquire);
-				uint64 LocalRequiredDiskSpace = RequiredDiskSpace.load(std::memory_order_acquire);
-				UE_LOG(LogBuildPatchServices, Error, TEXT("Out of HDD space. Needs %llu bytes, Free %llu bytes"), LocalRequiredDiskSpace, LocalAvailableDiskSpace);
+				// We get this when we fail during write. The problem is our cached values are out of date:
+				// AvailableDiskSpace is either 0, because we ran out (likely), or non zero, because of a write error (unlikely).
+				// RequiredDiskSpace is not accurate because presumably we've written something. So instead we
+				// tell them we failed, and to verify disk space. It's not useful to share the actual value we
+				// get back from the operating system since we are either in a write media failure mode, or so
+				// close to the edge of free space that values can look like we have enough, but don't.
+				uint64 TotalSize, AvailableSpace;
+				bool bGotFreeSpace = FPlatformMisc::GetDiskTotalAndFreeSpace(Configuration.InstallDirectory, TotalSize, AvailableSpace);
+
+				uint64 LocalDiskSpaceRequired;
+				if (Configuration.bDeleteChunkDBFilesAfterUse)
+				{
+					LocalDiskSpaceRequired = CalculateDiskSpaceRequirementsWithDeleteDuringInstall();
+				}
+				else
+				{
+					LocalDiskSpaceRequired = CalculateInProgressDiskSpaceRequired(*CurrentFile.FileManifest, CurrentFile.Progress);
+				}
+
+				UE_LOG(LogBuildPatchServices, Error, TEXT("Write failure during installation. FreeSpace: Valid=%s, Available=%llu"), bGotFreeSpace ? TEXT("true") : TEXT("false"), AvailableSpace);
+
+				const FNumberFormattingOptions FormatOptions = FNumberFormattingOptions()
+					.SetMinimumFractionalDigits(2)
+					.SetMaximumFractionalDigits(2);
+				FFormatNamedArguments Arguments;
+				Arguments.Emplace(TEXT("RequiredBytes"), FText::AsMemory(LocalDiskSpaceRequired, &FormatOptions, nullptr, EMemoryUnitStandard::SI));
+				FText OutOfDiskSpaceDuringInstall = FText::Format(NSLOCTEXT("BuildPatchInstallError", "InstallationWriteFailed", "Low-level failure to write files to disk during installation. Verify there is sufficient disk space. {RequiredBytes} more is required."), Arguments);
+
+				// Note that because it's almost certainly because of disk space, we report it as disk space.
+				// We can't be more specific because UE doesn't surface the exact reason for the write failure.
 				InstallerError->SetError(
 					EBuildPatchInstallError::OutOfDiskSpace,
 					DiskSpaceErrorCodes::DuringInstallation,
 					0,
-					BuildPatchServices::GetDiskSpaceMessage(Configuration.InstallDirectory, LocalRequiredDiskSpace, LocalAvailableDiskSpace));
+					MoveTemp(OutOfDiskSpaceDuringInstall));
 				break;
 			}
 		case EConstructionError::CannotCreateFile:
@@ -3241,7 +3243,7 @@ void FBuildPatchFileConstructor::CompleteConstructedFile(FFileConstructionState&
 		switch (CurrentFile.ConstructionError)
 		{
 			// Errors we expect that we can conceptually resume:
-			case EConstructionError::OutOfDiskSpace:
+			case EConstructionError::FailedWrite:
 			case EConstructionError::MissingChunk:
 			case EConstructionError::Aborted:
 			{
@@ -3477,7 +3479,7 @@ void FBuildPatchFileConstructor::ConstructFiles(const FResumeData& ResumeData)
 						// way to tease it apart, we run with the most likely issue: we ran out of disk space. Unfortunately
 						// it's unreliable to ask the file system what is available and check if it's low and use that for the 
 						// conversion.
-						File->ConstructionError = EConstructionError::OutOfDiskSpace;
+						File->ConstructionError = EConstructionError::FailedWrite;
 					}
 				}
 
