Fix for missing ChunkDownloadComplete analytic events

#rnx

diff --git a/Engine/Source/Runtime/Online/BackgroundHTTP/Private/BackgroundHttpManagerImpl.cpp b/Engine/Source/Runtime/Online/BackgroundHTTP/Private/BackgroundHttpManagerImpl.cpp
--- a/Engine/Source/Runtime/Online/BackgroundHTTP/Private/BackgroundHttpManagerImpl.cpp
+++ b/Engine/Source/Runtime/Online/BackgroundHTTP/Private/BackgroundHttpManagerImpl.cpp
@@ -365,6 +365,7 @@ bool FBackgroundHttpManagerImpl::AssociateWithAnyExistingRequest(const FBackgrou
 			Request->OnProgressUpdated().ExecuteIfBound(Request, ExistingFileSize, ExistingFileSize);
 
 			//Now complete with this completed response data
+			Request->SetMetrics(ExistingFileSize, 0.0f);
 			Request->CompleteWithExistingResponseData(NewResponseWithExistingFile);
 		}
 	}
diff --git a/Engine/Source/Runtime/Online/BackgroundHTTP/Private/BackgroundHttpRequestImpl.cpp b/Engine/Source/Runtime/Online/BackgroundHTTP/Private/BackgroundHttpRequestImpl.cpp
--- a/Engine/Source/Runtime/Online/BackgroundHTTP/Private/BackgroundHttpRequestImpl.cpp
+++ b/Engine/Source/Runtime/Online/BackgroundHTTP/Private/BackgroundHttpRequestImpl.cpp
@@ -76,9 +76,17 @@ void FBackgroundHttpRequestImpl::OnBackgroundDownloadComplete()
 	
 	if (OptionalMetricsInfo.IsSet())
 	{
+		if (!OnRequestMetrics().IsBound())
+		{
+			UE_LOG(LogBackgroundHttpRequest, Warning, TEXT("OptionalMetricsInfo OnRequestMetrics() is NOT bound - RequestID:%s"), *GetRequestID());
+		}
 		OnRequestMetrics().ExecuteIfBound(SharedThis(this), OptionalMetricsInfo->TotalBytesDownloaded, OptionalMetricsInfo->DownloadDuration);
 		OptionalMetricsInfo.Reset();
 	}
+	else
+	{
+		UE_LOG(LogBackgroundHttpRequest, Warning, TEXT("OptionalMetricsInfo missing, possibly iOS not yet didFinishCollectingMetrics: RequestID:%s"), *GetRequestID());
+	}
 }
 
 void FBackgroundHttpRequestImpl::NotifyNotificationObjectOfComplete(bool bWasSuccess)
@@ -95,8 +103,14 @@ void FBackgroundHttpRequestImpl::NotifyNotificationObjectOfComplete(bool bWasSuc
 	}
 }
 
+void FBackgroundHttpRequestImpl::SetMetrics(const int32 TotalBytesDownloaded, const float DownloadDuration)
+{
+	NotifyRequestMetricsAvailable(TotalBytesDownloaded, DownloadDuration);
+}
+
 void FBackgroundHttpRequestImpl::NotifyRequestMetricsAvailable(const int32 TotalBytesDownloaded, const float DownloadDuration)
 {
+	ensureAlwaysMsgf(!OptionalMetricsInfo.IsSet(), TEXT("Overwriting metrics data!!!"));
 	OptionalMetricsInfo = {
 		.TotalBytesDownloaded = TotalBytesDownloaded,
 		.DownloadDuration = DownloadDuration
diff --git a/Engine/Source/Runtime/Online/BackgroundHTTP/Private/IOS/ApplePlatformBackgroundHttpManager.cpp b/Engine/Source/Runtime/Online/BackgroundHTTP/Private/IOS/ApplePlatformBackgroundHttpManager.cpp
--- a/Engine/Source/Runtime/Online/BackgroundHTTP/Private/IOS/ApplePlatformBackgroundHttpManager.cpp
+++ b/Engine/Source/Runtime/Online/BackgroundHTTP/Private/IOS/ApplePlatformBackgroundHttpManager.cpp
@@ -153,4 +153,5 @@ void FApplePlatformBackgroundHttpManager::OnDownloadMetrics(const uint64 Downloa
 			}
 		}
 	}
+	UE_LOG(LogBackgroundHttpManager, Warning, TEXT("Fail to NotifyRequestMetricsAvailable - RequestID:%llu | TotalBytesDownloaded:%d | DownloadDuration:%f"), DownloadId, TotalBytesDownloaded, DownloadDuration);
 }
diff --git a/Engine/Source/Runtime/Online/BackgroundHTTP/Private/IOS/ApplePlatformBackgroundHttpRequest.cpp b/Engine/Source/Runtime/Online/BackgroundHTTP/Private/IOS/ApplePlatformBackgroundHttpRequest.cpp
--- a/Engine/Source/Runtime/Online/BackgroundHTTP/Private/IOS/ApplePlatformBackgroundHttpRequest.cpp
+++ b/Engine/Source/Runtime/Online/BackgroundHTTP/Private/IOS/ApplePlatformBackgroundHttpRequest.cpp
@@ -49,6 +49,39 @@ void FApplePlatformBackgroundHttpRequest::UpdateProgress()
 	FString TemporaryFilePath;
 	if (FBackgroundURLSessionHandler::IsDownloadFinished(DownloadId, ResultHTTPCode, TemporaryFilePath))
 	{
+		if (!OptionalMetricsInfo.IsSet())
+		{
+			// need to wait for metrics collection
+			if (MetricsWaitCounter == -1)
+			{
+				MetricsWaitCounter = 30;	// start a countdown
+				return;
+			}
+			else
+			{
+				MetricsWaitCounter--;
+				if (MetricsWaitCounter == 0)
+				{
+					// time up, still no metrics? report 0s
+					MetricsWaitCounter = -1;
+					UE_LOG(LogBackgroundHttpRequest, Error, TEXT("OptionalMetricsInfo still missing after waiting for 30 cycles: RequestID:%s"), *GetRequestID());
+					OptionalMetricsInfo = {
+						.TotalBytesDownloaded = 0,
+						.DownloadDuration = 0.0f
+					};
+				}
+				else
+				{
+					// count down, try again next cycle
+					return;
+				}
+			}
+		}
+		else
+		{
+			MetricsWaitCounter = -1;
+		}
+		
 		const bool bFileExists = !TemporaryFilePath.IsEmpty() && FPlatformFileManager::Get().GetPlatformFile().FileExists(*TemporaryFilePath);
 
 		// Fail request if we can't access the file
diff --git a/Engine/Source/Runtime/Online/BackgroundHTTP/Public/BackgroundHttpRequestImpl.h b/Engine/Source/Runtime/Online/BackgroundHTTP/Public/BackgroundHttpRequestImpl.h
--- a/Engine/Source/Runtime/Online/BackgroundHTTP/Public/BackgroundHttpRequestImpl.h
+++ b/Engine/Source/Runtime/Online/BackgroundHTTP/Public/BackgroundHttpRequestImpl.h
@@ -48,6 +48,7 @@ public:
 	
 	BACKGROUNDHTTP_API virtual void NotifyNotificationObjectOfComplete(const bool bWasSuccess);
 	BACKGROUNDHTTP_API virtual void NotifyRequestMetricsAvailable(const int32 TotalBytesDownloaded, const float DownloadDuration);
+	BACKGROUNDHTTP_API virtual void SetMetrics(const int32 TotalBytesDownloaded, const float DownloadDuration) override;
 	
 protected:
 	UE::FMutex DownloadCompleteMutex;
diff --git a/Engine/Source/Runtime/Online/BackgroundHTTP/Public/IOS/ApplePlatformBackgroundHttpRequest.h b/Engine/Source/Runtime/Online/BackgroundHTTP/Public/IOS/ApplePlatformBackgroundHttpRequest.h
--- a/Engine/Source/Runtime/Online/BackgroundHTTP/Public/IOS/ApplePlatformBackgroundHttpRequest.h
+++ b/Engine/Source/Runtime/Online/BackgroundHTTP/Public/IOS/ApplePlatformBackgroundHttpRequest.h
@@ -39,6 +39,8 @@ private:
 	uint64 DownloadId;
 
 	uint64 LastReportedDownloadedBytes = 0;
+	
+	int MetricsWaitCounter = -1;
 };
 
 typedef TSharedPtr<class FApplePlatformBackgroundHttpRequest, ESPMode::ThreadSafe> FAppleBackgroundHttpRequestPtr;
diff --git a/Engine/Source/Runtime/Online/BackgroundHTTP/Public/Interfaces/IBackgroundHttpRequest.h b/Engine/Source/Runtime/Online/BackgroundHTTP/Public/Interfaces/IBackgroundHttpRequest.h
--- a/Engine/Source/Runtime/Online/BackgroundHTTP/Public/Interfaces/IBackgroundHttpRequest.h
+++ b/Engine/Source/Runtime/Online/BackgroundHTTP/Public/Interfaces/IBackgroundHttpRequest.h
@@ -165,6 +165,11 @@ public:
 	 * Delegate called when the request's metrics were calculated
 	 */
 	virtual FBackgroundHttpRequestMetricsDelegate& OnRequestMetrics() = 0;
+	
+	/**
+	 * Set metrics for analytics event
+	 */
+	virtual void SetMetrics(const int32 TotalBytesDownloaded, const float DownloadDuration) {}
 
 	/**
 	* Called to begin processing the request.
