Support for nothreading in GrpcClient

diff --git a/Engine/Source/Runtime/Online/Experimental/GrpcClient/Private/GrpcClient.cpp b/Engine/Source/Runtime/Online/Experimental/GrpcClient/Private/GrpcClient.cpp
--- a/Engine/Source/Runtime/Online/Experimental/GrpcClient/Private/GrpcClient.cpp
+++ b/Engine/Source/Runtime/Online/Experimental/GrpcClient/Private/GrpcClient.cpp
@@ -3,6 +3,7 @@
 #include "GrpcClient.h"
 
 #include "GrpcClientModule.h"
+#include "Containers/BackgroundableTicker.h"
 
 const TCHAR* LexToString(EGrpcStatus Value)
 {
@@ -29,9 +30,14 @@ const TCHAR* LexToString(EGrpcStatus Value)
 	}
 }
 
+FGrpcAsyncManager::FGrpcAsyncManager()
+	: FTSTickerObjectBase(0.0f, FTSBackgroundableTicker::GetCoreTicker())
+{
+}
+
 void FGrpcAsyncManager::Start()
 {
-	Thread = TUniquePtr<FRunnableThread>(FRunnableThread::Create(this, TEXT("GrpcAsyncManager")));
+	Thread.Reset(FRunnableThread::Create(this, TEXT("GrpcAsyncManager")));
 }
 
 grpc::CompletionQueue& FGrpcAsyncManager::GetCompletionQueue()
@@ -49,20 +55,48 @@ uint32 FGrpcAsyncManager::Run()
 	{
 		if (CompletionQueue.Next(&Tag, &bOk))
 		{
-			FGrpcTag* GrpcTag = static_cast<FGrpcTag*>(Tag);
-			if (ActiveCalls.Contains(GrpcTag->Call))
-			{
-				TSharedPtr<FGrpcBaseCall> Call = ActiveCalls[GrpcTag->Call];
-				Call->HandleNext(Tag, bOk);
-			}
-			else
+			HandleNext(Tag, bOk);
+		}
+	}
+
+	return 0;
+}
+
+bool FGrpcAsyncManager::Tick(float DeltaTime)
+{
+	if (!Thread)
+	{
+		void *Tag;
+		bool bOk;
+
+		gpr_timespec Zero { 0, 0, GPR_CLOCK_MONOTONIC };
+		while (true)
+		{
+			grpc::CompletionQueue::NextStatus Status = CompletionQueue.AsyncNext(&Tag, &bOk, Zero);
+			if (Status != grpc::CompletionQueue::NextStatus::GOT_EVENT)
 			{
-				UE_LOG(LogGrpcClient, Warning, TEXT("GrpcAsyncManager - Received Tag (%x) for unmapped call"), Tag);
+				break;
 			}
+
+			HandleNext(Tag, bOk);
 		}
 	}
 
-	return 0;
+	return true;
+}
+
+void FGrpcAsyncManager::HandleNext(void *Tag, bool bOk)
+{
+	FGrpcTag* GrpcTag = static_cast<FGrpcTag*>(Tag);
+	if (ActiveCalls.Contains(GrpcTag->Call))
+	{
+		TSharedPtr<FGrpcBaseCall> Call = ActiveCalls[GrpcTag->Call];
+		Call->HandleNext(Tag, bOk);
+	}
+	else
+	{
+		UE_LOG(LogGrpcClient, Warning, TEXT("GrpcAsyncManager - Received Tag (%x) for unmapped call"), Tag);
+	}
 }
 
 void FGrpcAsyncManager::Stop()
@@ -83,3 +117,4 @@ void FGrpcAsyncManager::Exit()
 
 	Thread->WaitForCompletion();
 }
+
diff --git a/Engine/Source/Runtime/Online/Experimental/GrpcClient/Public/GrpcClient.h b/Engine/Source/Runtime/Online/Experimental/GrpcClient/Public/GrpcClient.h
--- a/Engine/Source/Runtime/Online/Experimental/GrpcClient/Public/GrpcClient.h
+++ b/Engine/Source/Runtime/Online/Experimental/GrpcClient/Public/GrpcClient.h
@@ -3,6 +3,7 @@
 #pragma once
 
 #include "GrpcClientModule.h"
+#include "Containers/Ticker.h"
 #include "HAL/Runnable.h"
 #include "HAL/RunnableThread.h"
 
@@ -48,8 +49,11 @@ struct FGrpcTag
 
 class FGrpcAsyncManager 
 	: public FRunnable
+	, public FTSTickerObjectBase
 {
 public:
+	FGrpcAsyncManager();
+
 	void Start();
 
 	GRPCCLIENT_API grpc::CompletionQueue& GetCompletionQueue();
@@ -60,6 +64,10 @@ public:
 	virtual void Exit() override;
 	//~ End FRunnable Interface
 
+	//~ Begin FTSTickerObjectBase interface
+	virtual bool Tick(float DeltaTime) override;
+	//~ End FTSTickerObjectBase interface
+
 protected:
 	bool bIsRunning = false;
 	grpc::CompletionQueue CompletionQueue;
@@ -67,6 +75,9 @@ protected:
 
 	TMap<void*, TSharedPtr<FGrpcBaseCall>> ActiveCalls;
 
+private:
+	void HandleNext(void* Tag, bool bOk);
+
 	friend class FGrpcClientCall;
 };
 
