Fix compilation issues

#rnx

diff --git a/Engine/Source/Runtime/Online/HTTP/Private/Apple/AppleHttpBackgroundSession.h b/Engine/Source/Runtime/Online/HTTP/Private/Apple/AppleHttpBackgroundSession.h
--- a/Engine/Source/Runtime/Online/HTTP/Private/Apple/AppleHttpBackgroundSession.h
+++ b/Engine/Source/Runtime/Online/HTTP/Private/Apple/AppleHttpBackgroundSession.h
@@ -8,12 +8,12 @@
 
 @interface FApplePlatformHttpBackgroundSessionDelegate : NSObject<NSURLSessionDataDelegate>
 
-- (void)URLSession:(NSURLSession *)session task:(NSURLSessionTask *)task didSendBodyData:(int64_t)bytesSent totalBytesSent:(int64_t)totalBytesSent totalBytesExpectedToSend:(int64_t)totalBytesExpectedToSend;
-- (void)URLSession:(NSURLSession *)session dataTask:(NSURLSessionDataTask *)dataTask didReceiveResponse:(NSURLResponse *)response completionHandler:(void (^)(NSURLSessionResponseDisposition))completionHandler;
-- (void)URLSession:(NSURLSession *)session dataTask:(NSURLSessionDataTask *)dataTask didReceiveData:(NSData *)data;
-- (void)URLSession:(NSURLSession *)session task:(NSURLSessionTask *)task didCompleteWithError:(NSError *)error;
-- (void)URLSession:(NSURLSession *)session dataTask:(NSURLSessionDataTask *)dataTask willCacheResponse:(NSCachedURLResponse *)proposedResponse completionHandler:(void (^)(NSCachedURLResponse * _Nullable))completionHandler;
-- (void)URLSessionDidFinishEventsForBackgroundURLSession:(NSURLSession *)session;
+- (void)URLSession:(NSURLSession* _Nonnull)session task:(NSURLSessionTask* _Nonnull)task didSendBodyData:(int64_t)bytesSent totalBytesSent:(int64_t)totalBytesSent totalBytesExpectedToSend:(int64_t)totalBytesExpectedToSend;
+- (void)URLSession:(NSURLSession* _Nonnull)session dataTask:(NSURLSessionDataTask* _Nonnull)dataTask didReceiveResponse:(NSURLResponse* _Nullable)response completionHandler:(void (^ _Nonnull)(NSURLSessionResponseDisposition))completionHandler;
+- (void)URLSession:(NSURLSession* _Nonnull)session dataTask:(NSURLSessionDataTask* _Nonnull)dataTask didReceiveData:(NSData* _Nonnull)data;
+- (void)URLSession:(NSURLSession* _Nonnull)session task:(NSURLSessionTask* _Nonnull)task didCompleteWithError:(NSError* _Nullable)error;
+- (void)URLSession:(NSURLSession* _Nonnull)session dataTask:(NSURLSessionDataTask* _Nonnull)dataTask willCacheResponse:(NSCachedURLResponse* _Nullable)proposedResponse completionHandler:(void (^ _Nonnull)(NSCachedURLResponse * _Nullable))completionHandler;
+- (void)URLSessionDidFinishEventsForBackgroundURLSession:(NSURLSession* _Nonnull)session;
 
 @end
 
