Fix TVOS compilation errors

#rnx

diff --git a/Engine/Source/Runtime/Online/HTTP/HTTP.Build.cs b/Engine/Source/Runtime/Online/HTTP/HTTP.Build.cs
--- a/Engine/Source/Runtime/Online/HTTP/HTTP.Build.cs
+++ b/Engine/Source/Runtime/Online/HTTP/HTTP.Build.cs
@@ -86,7 +86,7 @@ public class HTTP : ModuleRules
 			}
 			);
 
-		if (Target.bCompileAgainstApplicationCore && Target.Platform == UnrealTargetPlatform.IOS)
+		if (Target.bCompileAgainstApplicationCore && (Target.Platform == UnrealTargetPlatform.IOS || Target.Platform == UnrealTargetPlatform.TVOS))
 		{
 			PrivateDependencyModuleNames.AddRange(
 				new string[] {
