Fix Compile UnrealGame VisionOS failures

#rnx

diff --git a/Engine/Source/Runtime/Online/HTTP/HTTP.Build.cs b/Engine/Source/Runtime/Online/HTTP/HTTP.Build.cs
--- a/Engine/Source/Runtime/Online/HTTP/HTTP.Build.cs
+++ b/Engine/Source/Runtime/Online/HTTP/HTTP.Build.cs
@@ -86,7 +86,7 @@ public class HTTP : ModuleRules
 			}
 			);
 
-		if (Target.bCompileAgainstApplicationCore && (Target.Platform == UnrealTargetPlatform.IOS || Target.Platform == UnrealTargetPlatform.TVOS))
+		if (Target.bCompileAgainstApplicationCore && Target.Platform.IsInGroup(UnrealPlatformGroup.IOS)) // ios, tvos and visionos fall into this group
 		{
 			PrivateDependencyModuleNames.AddRange(
 				new string[] {
