BPS: Fix race condition around harvesting chunks. A chunk in the file to be deleted could need parts from another file that's under construction, which requires more dependency tracking than its worth - just lock around it.
BPS: Fix bug when a chunk can be provided from multiple tags/bundles and only one is installed. Based on ordering, the install source would try to read from the file that wasn't installed. Removed the accessor functions from Manifest that provide the incorrect data.

diff --git a/Engine/Source/Runtime/Online/BuildPatchServices/Private/BuildPatchFileConstructor.cpp b/Engine/Source/Runtime/Online/BuildPatchServices/Private/BuildPatchFileConstructor.cpp
--- a/Engine/Source/Runtime/Online/BuildPatchServices/Private/BuildPatchFileConstructor.cpp
+++ b/Engine/Source/Runtime/Online/BuildPatchServices/Private/BuildPatchFileConstructor.cpp
@@ -1467,6 +1467,12 @@ FBuildPatchFileConstructor::FBuildPatchFileConstructor(
 
 			// We need to set up a dependency chain so that files can know when they can start. We can't start constructing a file
 			// that needs parts from a file that isn't done as the patch authoring tool expects the file to be done.
+			//
+			// There are a couple major concerns here:
+			//	1. If a later file requires a chunk from an earlier file, the file will have been deleted, so we Harvest the
+			//		chunks. We can't start the file until that harvest completes. This dependency checking prevents this.
+			//	2. When harvesting chunks from a single file, the chunk we harvest could come from multiple files, including ones
+			//		in flight. Rather than track such dependencies here we provide file access locking inside the install source.
 			TMap<FStringView, int32> FileToIndexMap;
 			FileToIndexMap.Reserve(Configuration.ConstructList.Num());
 			for (int32 FileConstructIndex = 0; FileConstructIndex < Configuration.ConstructList.Num(); FileConstructIndex++)
@@ -1653,8 +1659,14 @@ void FBuildPatchFileConstructor::Abort()
 
 void FBuildPatchFileConstructor::SetChunkLocation(const FGuid& InGuid, EConstructorChunkLocation InNewLocation)
 {
-	// NOTE - reentrant. We assume a) that ChunkLocations is filled before threading
+	// This is almost always called on the constructor thread, but on error conditions gets
+	// called in completion handlers to set chunk locations to the cloud.
+	// We assume a) that ChunkLocations is filled before threading
 	// and that b) no guids are used across concurrent actions.
+	//
+	// This might be too complicated for its own good - if the only reason we touch this on other threads
+	// is for error handling, maybe we move errors to sideband info and only touch ChunkLocations on the
+	// constructor thread.
 	FRWScopeLock ChunkLock(ChunkLocationsLock, SLT_Write);
 	EConstructorChunkLocation* Location = ChunkLocations.Find(InGuid);
 	if (!Location)
diff --git a/Engine/Source/Runtime/Online/BuildPatchServices/Private/BuildPatchManifest.cpp b/Engine/Source/Runtime/Online/BuildPatchServices/Private/BuildPatchManifest.cpp
--- a/Engine/Source/Runtime/Online/BuildPatchServices/Private/BuildPatchManifest.cpp
+++ b/Engine/Source/Runtime/Online/BuildPatchServices/Private/BuildPatchManifest.cpp
@@ -1205,18 +1205,6 @@ void FBuildPatchAppManifest::RemoveCustomField(const FString& FieldName)
 	CustomFields.Fields.Remove(FieldName);
 }
 
-int32 FBuildPatchAppManifest::EnumerateProducibleChunks(const FString& InstallDirectory, const TSet<FGuid>& ChunksRequired, TSet<FGuid>& ChunksAvailable) const
-{
-	TMap<FString, int64> InstallationFileSizes;
-	return EnumerateProducibleChunks_Internal([&](const FString& Filename)
-	{
-		if (InstallationFileSizes.Contains(Filename) == false)
-		{
-			InstallationFileSizes.Add(Filename, IFileManager::Get().FileSize(*(InstallDirectory / Filename)));
-		}
-		return GetFileSize(Filename) == InstallationFileSizes[Filename];
-	}, ChunksRequired, ChunksAvailable);
-}
 
 int32 FBuildPatchAppManifest::EnumerateProducibleChunks(const TSet<FString>& TagSet, const TSet<FGuid>& ChunksRequired, TSet<FGuid>& ChunksAvailable) const
 {
@@ -1259,35 +1247,6 @@ int32 FBuildPatchAppManifest::EnumerateProducibleChunks_Internal(const TFunction
 	return Count;
 }
 
-TArray<FFileChunkPart> FBuildPatchAppManifest::GetFilePartsForChunk(const FGuid& ChunkId) const
-{
-	TArray<FFileChunkPart> FileParts;
-	FBlockStructure FoundParts;
-	for (const FFileManifest& FileManifest: FileManifestList.FileList)
-	{
-		uint64 FileOffset = 0;
-		for (const FChunkPart& ChunkPart : FileManifest.ChunkParts)
-		{
-			if (ChunkId == ChunkPart.Guid)
-			{
-				FFileChunkPart FileChunkPart;
-				FileChunkPart.Filename = FileManifest.Filename;
-				FileChunkPart.ChunkPart = ChunkPart;
-				FileChunkPart.FileOffset = FileOffset;
-				FileParts.Add(FileChunkPart);
-				FoundParts.Add(ChunkPart.Offset, ChunkPart.Size, ESearchDir::FromEnd);
-			}
-			FileOffset += ChunkPart.Size;
-		}
-	}
-
-	// If the structure is not a single complete block, then the chunk is not recoverable.
-	if (FoundParts.GetHead() == nullptr || FoundParts.GetHead() != FoundParts.GetTail() /* || @TODO Possible to check FoundParts.GetHead()->GetSize() != UncompressedDataSize here? */)
-	{
-		FileParts.Empty();
-	}
-	return FileParts;
-}
 
 bool FBuildPatchAppManifest::HasFileAttributes() const
 {
diff --git a/Engine/Source/Runtime/Online/BuildPatchServices/Private/BuildPatchManifest.h b/Engine/Source/Runtime/Online/BuildPatchServices/Private/BuildPatchManifest.h
--- a/Engine/Source/Runtime/Online/BuildPatchServices/Private/BuildPatchManifest.h
+++ b/Engine/Source/Runtime/Online/BuildPatchServices/Private/BuildPatchManifest.h
@@ -292,14 +292,6 @@ public:
 	 */
 	virtual bool GetFilePartHash(const FGuid& FilePartGuid, uint64& OutHash) const;
 
-	/**
-	 * Populates an array of chunks that should be producible from this local build, given the list of chunks needed. Also checks source files exist and match size.
-	 * @param InstallDirectory	IN		The directory of the build where chunks would be sourced from.
-	 * @param ChunksRequired	IN		A list of chunks that are needed.
-	 * @param ChunksAvailable	OUT		A list to receive the chunks that could be constructed locally.
-	 * @return the number of chunks added to the ChunksAvailable set.
-	 */
-	virtual int32 EnumerateProducibleChunks(const FString& InstallDirectory, const TSet<FGuid>& ChunksRequired, TSet<FGuid>& ChunksAvailable) const;
 
 	/**
 	 * Populates an array of chunks that should be producible from a local build, given the tagset that is expected to be installed.
@@ -335,12 +327,6 @@ public:
 	 */
 	virtual void GetRemovableFiles(const FBuildPatchAppManifest& OldManifest, TArray<FString>& RemovableFiles) const;
 
-	/**
-	 * Gets a list of file parts that can be used to recreate a chunk from this installation.
-	 * @param ChunkId       The guid for the desired chunk.
-	 * @return The array of file parts that can produce this chunk. Array is empty if the chunk cannot be produced.
-	 */
-	virtual TArray<BuildPatchServices::FFileChunkPart> GetFilePartsForChunk(const FGuid& ChunkId) const;
 
 	/** @return True if any files in this manifest have file attributes to be set */
 	virtual bool HasFileAttributes() const;
diff --git a/Engine/Source/Runtime/Online/BuildPatchServices/Private/Installer/InstallChunkSource.cpp b/Engine/Source/Runtime/Online/BuildPatchServices/Private/Installer/InstallChunkSource.cpp
--- a/Engine/Source/Runtime/Online/BuildPatchServices/Private/Installer/InstallChunkSource.cpp
+++ b/Engine/Source/Runtime/Online/BuildPatchServices/Private/Installer/InstallChunkSource.cpp
@@ -8,6 +8,7 @@
 #include "Common/FileSystem.h"
 #include "Common/StatsCollector.h"
 #include "Core/BlockStructure.h"
+#include "HAL/FileManager.h"
 #include "Misc/Paths.h"
 #include "Serialization/MemoryReader.h"
 #include "Tasks/Task.h"
@@ -40,13 +41,26 @@ namespace BuildPatchServices
 			// Make sure we close our handle before the deletion occurs.
 			// With multiple files in flight in the constructor we can be deleting a file at the same time as we are
 			// reading chunks for other files, which means we can hit OpenedFileHandles from multiple threads.
-			UE::TUniqueLock _(FileHandleLock);
-			TUniquePtr<FOpenedFile>* File = OpenedFileHandles.Find(FilePath);
-			if (File && (*File)->UsageCount != 0)
+			UE::TUniqueLock _(OpenedFileHandlesLock);
 			{
-				UE_LOG(LogInstallChunkSource, Error, TEXT("Tried to delete a file we're reading from! %s"), *FilePath);
+				TUniquePtr<FOpenedFile>* File = OpenedFileHandles.Find(FilePath);
+				if (File)
+				{
+					if ((*File)->FileHandleLock.TryLock() == false)
+					{
+						// We can't ever be trying to delete a file when a read is active on it,
+						// as this means that a file is trying to read from it after harvesting has
+						// completed, which should ensure no chunks are read.
+						UE_LOG(LogInstallChunkSource, Error, TEXT("Critical error: trying to delete %s while read lock is held."), *FilePath);
+
+						// We don't want to just crash, so take the lock in order to wait for the
+						// other thread to finish. Then we hope that nothing else is queued :grimacing:
+						(*File)->FileHandleLock.Lock();
+					}
+
+					OpenedFileHandles.Remove(FilePath);
+				}
 			}
-			OpenedFileHandles.Remove(FilePath);
 		}
 		virtual int32 GetChunkUnavailableAt(const FGuid& DataId) const override;
 		virtual void SetFileRetirementPositions(TMap<FString, int32>&& InFileRetirementPositions) override
@@ -80,7 +94,7 @@ namespace BuildPatchServices
 				// afaict the file manifest filename is normalized because in the manifest builder it
 				// generates it from file spans, which are created in directorybuildstreamer which makes them
 				// relative, and internal to that function they are normalized.
-				Callback(NormalizedInstallDirectory, ChunkDetails.FileManifest->Filename);
+				Callback(NormalizedInstallDirectory, ChunkDetails.InstalledFileManifest->Filename);
 			}
 		}
 
@@ -96,29 +110,31 @@ namespace BuildPatchServices
 		TSet<FGuid> AvailableInBuilds;
 		TArray<TPair<FString, FBuildPatchAppManifestRef>> InstallationSources;
 
-		// This is here to try and track down an unexplained crash where it looks like maybe we're hitting the
-		// same file during delete as we are during read (which shouldn't happen).
 		struct FOpenedFile
 		{
 			FOpenedFile(TUniquePtr<FArchive>&& InFileHandle) : FileHandle(MoveTemp(InFileHandle)) {}
 
 			TUniquePtr<FArchive> FileHandle;
-			std::atomic_int32_t UsageCount;
+			UE::FMutex FileHandleLock;
 		};
-		UE::FMutex FileHandleLock; // Protects OpenedFileHandles
+
+		UE::FMutex OpenedFileHandlesLock; 
 		TMap<FString, TUniquePtr<FOpenedFile>> OpenedFileHandles;
 
 		// The index (ChunkReferenceTracker->GetCurrentUsageIndex) at which our files will get deleted due to destructive install to make room for
 		// the new file. 
 		TMap<FString, int32> FileRetirementPositions;
 
+		// Contains the info to build the chunk from various pieces strewn across the installation.
 		struct FChunkSourceDetails
 		{
-			const FFileManifest* FileManifest;
-			const FChunkPart* ChunkPart;
-			uint64 FileOffset;
+			const FFileManifest* InstalledFileManifest;
+			uint32 OffsetInChunk;
+			uint32 SizeOfChunkPart;
+			uint64 OffsetInInstalledFile;
 		};
 
+		// Contains the list of pieces to assemble chunks from, for each chunk.
 		TMap<FGuid, TArray<FChunkSourceDetails>> ChunkSources;
 	};
 
@@ -129,41 +145,171 @@ namespace BuildPatchServices
 	{
 		TRACE_CPUPROFILER_EVENT_SCOPE(InstallChunkSource_ctor);
 
-		// Cache faster lookup information.		
-		for (const TPair<FString, FBuildPatchAppManifestRef>& Pair : InInstallationSources)
+		// InstallationSources is a map of installation directory to the manifest that was installed there.
+		for (const TPair<FString, FBuildPatchAppManifestRef>& InstallationSourcePair : InInstallationSources)
 		{
-			if (Pair.Value->EnumerateProducibleChunks(Pair.Key, ChunksThatWillBeNeeded, AvailableInBuilds) > 0)
+			const FString& InstallationDirectory = InstallationSourcePair.Key;
+			const FBuildPatchAppManifestRef& InstalledManifest = InstallationSourcePair.Value;
+
+			// Patches are only allowed to pull chunks from the install of the same tags, otherwise they must be
+			// supplied via another source. However we can't get the tag list at this point; all we have is a list of
+			// "registered installations" from the caller, so we get the full list of files and filter based on chunks we care
+			// about. This may result in us attempting to pull chunks from files outside our tag if it can be serviced
+			// by a file in a different install tag, however if the file doesn't exist, we'll continue to check other
+			// files until we presumably find the correct file in our own tag.
+			TSet<FString> InstalledFiles;
+			InstalledManifest->GetFileList(InstalledFiles);
+
+			struct FChunkAssemblyInfo
 			{
-				InstallationSources.Add(Pair);
-			}
-		}
-		UE_LOG(LogInstallChunkSource, Log, TEXT("Useful Sources:%d. Available Chunks:%d."), InstallationSources.Num(), AvailableInBuilds.Num());
+				TArray<FChunkSourceDetails> SourceDetails;
+				bool bComplete;
+			};
 
-		// Cache what file we get everything from so we aren't linearly looking over this every request
-		for (const TPair<FString, FBuildPatchAppManifestRef>& InstallationSource : InstallationSources)
-		{
-			FBuildPatchAppManifestRef Manifest = InstallationSource.Value;
+			// Find the chunks we can get from these files.
+			TMap<FGuid, FBlockStructure> ChunkStructures;
+			TMap<FGuid, FChunkAssemblyInfo> ChunkAssemblyInfos;
 
-			TSet<FString> FilesInManifest;
-			Manifest->GetFileList(FilesInManifest);
+			int32 MaxChunkStructures = 0;
 
-			for (const FString& FileName : FilesInManifest)
+			for (const FString& InstalledFileName : InstalledFiles)
 			{
-				const FFileManifest* FileManifest = Manifest->GetFileManifest(FileName);
+				const FFileManifest* InstalledFileManifest = InstalledManifest->GetFileManifest(InstalledFileName);
+				if (!InstalledFileManifest)
+				{
+					// Should never happen?
+					UE_LOG(LogInstallChunkSource, Error, TEXT("Missing file manifest from file provided from manifest? %s"), *InstalledFileName);
+					continue;
+				}
+
+				bool bExistenceCheckPassed = false;
 
+				uint32 ChunkPartSize = 0;
 				uint64 FileOffset = 0;
-				for (const FChunkPart& ChunkPart : FileManifest->ChunkParts)
+				for (int32 ChunkPartIndex = 0; ChunkPartIndex < InstalledFileManifest->ChunkParts.Num(); ChunkPartIndex++, FileOffset += ChunkPartSize)
 				{
-					if (AvailableInBuilds.Contains(ChunkPart.Guid))
+					const FChunkPart& ChunkPart = InstalledFileManifest->ChunkParts[ChunkPartIndex];
+					ChunkPartSize = ChunkPart.Size;
+
+					if (!ChunksThatWillBeNeeded.Contains(ChunkPart.Guid))
+					{
+						continue;
+					}
+
+					// We do the file existence check _after_ we know we need the chunk so that we avoid
+					// hitting the disk for every single file in the manifest when the tag might only be installing
+					// 5 of them.
+					if (!bExistenceCheckPassed)
+					{
+						uint64 SizeOnDisk = IFileManager::Get().FileSize(*(InstallationDirectory / InstalledFileName));
+						uint64 ExpectedSize = InstalledManifest->GetFileSize(InstalledFileName);
+						if (SizeOnDisk == ExpectedSize)
+						{
+							bExistenceCheckPassed = true;
+						}
+						else
+						{
+							// We can't get any chunks from this file, so bail.
+							break;
+						}
+					}
+					
+
+					// If we get the full chunk here, we don't need to assemble anything.
+					FChunkAssemblyInfo& ChunkAssemblyInfo = ChunkAssemblyInfos.FindOrAdd(ChunkPart.Guid);
+					if (ChunkAssemblyInfo.SourceDetails.Num() == 0)
+					{
+						// First time we've seen this chunk - check if we are getting it all at once (often)
+						const FChunkInfo* ChunkInfoPtr = InstalledManifest->GetChunkInfo(ChunkPart.Guid);
+						if (!ChunkInfoPtr)
+						{
+							UE_LOG(LogInstallChunkSource, Display, TEXT("Missing chunk info for guid: %s"), *LexToString(ChunkPart.Guid));
+							continue;
+						}
+
+						if (ChunkPart.Offset == 0 &&
+							ChunkInfoPtr->WindowSize == ChunkPart.Size)
+						{
+							// Got the whole thing, don't need to fill out the chunk parts.
+							FChunkSourceDetails NewInfo;
+							NewInfo.InstalledFileManifest = InstalledFileManifest;
+							NewInfo.OffsetInInstalledFile = FileOffset;
+							NewInfo.OffsetInChunk = ChunkPart.Offset;
+							NewInfo.SizeOfChunkPart = ChunkPart.Size;
+							ChunkAssemblyInfo.SourceDetails.Add(NewInfo);
+							ChunkAssemblyInfo.bComplete = true;
+							continue;
+						}
+
+						// If it's just a part of the chunk we fall through
+					}
+
+					if (ChunkAssemblyInfo.bComplete)
+					{
+						// We already have all we need to make this chunk, anything else is redundant.
+						continue;
+					}
+
+					FBlockStructure& PartialStructure = ChunkStructures.FindOrAdd(ChunkPart.Guid);
+					PartialStructure.Add(ChunkPart.Offset, ChunkPart.Size, ESearchDir::FromEnd);
+					
+					if (ChunkStructures.Num() > MaxChunkStructures)
 					{
-						TArray<FChunkSourceDetails>& ChunkSource = ChunkSources.FindOrAdd(ChunkPart.Guid);
-						ChunkSource.Add({FileManifest, &ChunkPart, FileOffset});
+						MaxChunkStructures = ChunkStructures.Num();
 					}
 
-					FileOffset += ChunkPart.Size;
+					FChunkSourceDetails PartInfo;
+					PartInfo.InstalledFileManifest = InstalledFileManifest;
+					PartInfo.OffsetInInstalledFile = FileOffset;
+					PartInfo.OffsetInChunk = ChunkPart.Offset;
+					PartInfo.SizeOfChunkPart = ChunkPart.Size;
+					ChunkAssemblyInfo.SourceDetails.Add(PartInfo);
+
+					const FChunkInfo* ChunkInfoPtr = InstalledManifest->GetChunkInfo(ChunkPart.Guid);
+
+					if (PartialStructure.GetHead() == PartialStructure.GetTail() &&
+						PartialStructure.GetHead()->GetSize() == ChunkInfoPtr->WindowSize)
+					{
+						// Done with this chunk, free the structures space.
+						ChunkAssemblyInfo.bComplete = true;
+						ChunkStructures.Remove(ChunkPart.Guid);
+					}					
 				}
 			}
+
+			// Move the full chunks we have available over to our internal structures.
+			uint32 IncompleteChunks = 0;
+			uint32 ChunksInAnotherSource = 0;
+			uint32 CompleteChunks = 0;
+			for (TPair<FGuid, FChunkAssemblyInfo>& ChunkAssemblyPair : ChunkAssemblyInfos)
+			{
+				if (!ChunkAssemblyPair.Value.bComplete)
+				{
+					IncompleteChunks++;
+					continue;
+				}
+
+				TArray<FChunkSourceDetails>& ChunkSource = ChunkSources.FindOrAdd(ChunkAssemblyPair.Key);
+				if (ChunkSource.Num())
+				{
+					ChunksInAnotherSource++;
+					continue;
+				}
+
+				ChunkSource = MoveTemp(ChunkAssemblyPair.Value.SourceDetails);
+				CompleteChunks++;
+
+				AvailableInBuilds.Add(ChunkAssemblyPair.Key);
+			}
+
+			UE_LOG(LogInstallChunkSource, Verbose, TEXT("Got %llu complete chunks, %llu incomplete chunks, %llu repeat chunks"), CompleteChunks, IncompleteChunks, ChunksInAnotherSource);
+			if (CompleteChunks)
+			{
+				InstallationSources.Add(InstallationSourcePair);
+			}
+
 		}
+		UE_LOG(LogInstallChunkSource, Log, TEXT("Useful Sources:%d. Available Chunks:%d."), InstallationSources.Num(), AvailableInBuilds.Num());
 	}
 
 	FInstallChunkSource::~FInstallChunkSource()
@@ -186,7 +332,23 @@ namespace BuildPatchServices
 		}
 		if (FileManifest != nullptr)
 		{
-			Algo::Transform(FileManifest->ChunkParts, OutChunks, &FChunkPart::Guid);
+			for (const TPair<FGuid, TArray<FChunkSourceDetails>>& ChunkSource : ChunkSources)
+			{
+				bool bNeedsThisFile = false;
+				for (const FChunkSourceDetails& Details : ChunkSource.Value)
+				{
+					if (Details.InstalledFileManifest == FileManifest)
+					{
+						bNeedsThisFile = true;
+						break;
+					}
+				}
+
+				if (bNeedsThisFile)
+				{
+					OutChunks.Add(ChunkSource.Key);
+				}
+			}
 		}
 	}
 
@@ -201,8 +363,7 @@ namespace BuildPatchServices
 			return [](bool) {return;};
 		}
 
-		// Must be run sequentially! No threaded protection provided. The expectation is that all process functions
-		// returned here must be run sequentially, and that the underlying file must not be deleted.
+		// This request can be called from multiple threads, including on the same source file.
 		return [this, FoundInstallManifest, FoundInstallDirectory = *FoundInstallDirectory, DataId, DestinationBuffer, UserPtr, CompleteFn](bool bIsAborted)
 		{
 			ISpeedRecorder::FRecord ActivityRecord;
@@ -248,7 +409,7 @@ namespace BuildPatchServices
 				const FChunkSourceDetails& FileChunkPart = ChunkSource->operator[](FileChunkPartsIdx);
 
 				// Validate the chunk can load into the destination.
-				uint32 ChunkEndLocation = FileChunkPart.ChunkPart->Offset + FileChunkPart.ChunkPart->Size;
+				uint32 ChunkEndLocation = FileChunkPart.OffsetInChunk + FileChunkPart.SizeOfChunkPart;
 				if (ChunkEndLocation > DestinationBuffer.GetSize())
 				{
 					// The chunk metadata tried to assemble larger than the chunk itself - error
@@ -257,24 +418,22 @@ namespace BuildPatchServices
 					break;
 				}
 
-				FString FullFilename = FoundInstallDirectory / FileChunkPart.FileManifest->Filename;
+				FString FullFilename = FoundInstallDirectory / FileChunkPart.InstalledFileManifest->Filename;
+
 
 				// We use the internal FArchive pointer so that we don't have to hold the file handle
 				// lock over the read - the uniqueptr might move around, but the managed pointer will not.
-				// We do know that our specific file won't get deleted until we are done.
+				// We do know that our specific file won't get deleted until we are done, but it's possible that
+				// multiple threads are here looking for the same file.
 				FOpenedFile* OpenedFile = nullptr;
 				{
-					FileHandleLock.Lock();
+					OpenedFileHandlesLock.Lock();
 					TUniquePtr<FOpenedFile>* OpenedFilePtr = OpenedFileHandles.Find(FullFilename);
 					if (OpenedFilePtr == nullptr)
 					{
 						TRACE_CPUPROFILER_EVENT_SCOPE(Install_OpenSource);
 
-						// Don't lock over the file open since that could take a while.
-						FileHandleLock.Unlock();
 						TUniquePtr<FArchive> NewReader = FileSystem->CreateFileReader(*FullFilename);
-						FileHandleLock.Lock();
-
 						if (NewReader.IsValid())
 						{
 							TUniquePtr<FOpenedFile> NewOpenedFile = MakeUnique<FOpenedFile>(MoveTemp(NewReader));
@@ -287,53 +446,38 @@ namespace BuildPatchServices
 					{
 						OpenedFile = OpenedFilePtr->Get();
 					}
-					FileHandleLock.Unlock();
+					OpenedFileHandlesLock.Unlock();
 				}
 
 				if (OpenedFile == nullptr)
 				{
+					UE_LOG(LogInstallChunkSource, Error, TEXT("Failed to open: %s on chunk %s"), *FileChunkPart.InstalledFileManifest->Filename, *LexToString(DataId));
 					Result = IInstallChunkSourceStat::ELoadResult::OpenFileFail;
 					break;
 				}
 
-				// Should only ever be one thread hitting files.
-				if (OpenedFile->UsageCount.fetch_add(1, std::memory_order_acq_rel) != 0)
 				{
-					UE_LOG(LogInstallChunkSource, Warning, TEXT("Invalid usage count in install source! %s"), *FileChunkPart.FileManifest->Filename);
-				}
+					UE::TUniqueLock HandleLock(OpenedFile->FileHandleLock);
 
-				{
 					TRACE_CPUPROFILER_EVENT_SCOPE(Install_Serialize);
-					OpenedFile->FileHandle->Seek(FileChunkPart.FileOffset);
 
-					// We're hunting down some weird crashes where we can't tell if it's because our dest pointer is invalid or not,
-					// so do a silly check here to try and catch it if so.
-					{
-						uint8* BasePtr = (uint8*)DestinationBuffer.GetData() + FileChunkPart.ChunkPart->Offset;
-						BasePtr[0] = 0;
-						BasePtr[FileChunkPart.ChunkPart->Size-1] = 0;
-					}
-					OpenedFile->FileHandle->Serialize((uint8*)DestinationBuffer.GetData() + FileChunkPart.ChunkPart->Offset, FileChunkPart.ChunkPart->Size);
+					OpenedFile->FileHandle->Seek(FileChunkPart.OffsetInInstalledFile);
+					OpenedFile->FileHandle->Serialize((uint8*)DestinationBuffer.GetData() + FileChunkPart.OffsetInChunk, FileChunkPart.SizeOfChunkPart);
 
-					ActivityRecord.Size += FileChunkPart.ChunkPart->Size;
+					ActivityRecord.Size += FileChunkPart.SizeOfChunkPart;
 
 					FBlockStructure NewChunk;
-					NewChunk.Add(FileChunkPart.ChunkPart->Offset, FileChunkPart.ChunkPart->Size);
+					NewChunk.Add(FileChunkPart.OffsetInChunk, FileChunkPart.SizeOfChunkPart);
 					if (NewChunk.Intersect(ChunkBlocks).GetHead() != nullptr)
 					{
 						// This used to be allowed but in advance of multi threaded reading we want to make sure this
 						// doesn't happen anymore (already shouldn't be...)
+						OpenedFile->FileHandleLock.Unlock();
 						UE_LOG(LogInstallChunkSource, Error, TEXT("Chunk %s had overlapping chunk parts"), *WriteToString<40>(DataId));
 						Result = IInstallChunkSourceStat::ELoadResult::InvalidChunkParts;
 						break;
 					}
-					ChunkBlocks.Add(FileChunkPart.ChunkPart->Offset, FileChunkPart.ChunkPart->Size);
-				}
-
-				// Should only ever be one thread hitting files.
-				if (OpenedFile->UsageCount.fetch_add(-1, std::memory_order_acq_rel) != 1)
-				{
-					UE_LOG(LogInstallChunkSource, Warning, TEXT("Invalid usage count in install source post read! %s"), *FileChunkPart.FileManifest->Filename);
+					ChunkBlocks.Add(FileChunkPart.OffsetInChunk, FileChunkPart.SizeOfChunkPart);
 				}
 
 				// The expectation is that we get the full chunk only once we've assembled all of the parts
@@ -425,7 +569,7 @@ namespace BuildPatchServices
 		{
 			for (const FChunkSourceDetails& Part : *ChunkSource)
 			{
-				const int32* FirstIndexAfterFile = FileRetirementPositions.Find(Part.FileManifest->Filename);
+				const int32* FirstIndexAfterFile = FileRetirementPositions.Find(Part.InstalledFileManifest->Filename);
 				if (FirstIndexAfterFile &&
 					*FirstIndexAfterFile < ChunkUnavailableAt)
 				{
diff --git a/Engine/Source/Runtime/Online/BuildPatchServices/Private/Tests/Mock/Manifest.mock.h b/Engine/Source/Runtime/Online/BuildPatchServices/Private/Tests/Mock/Manifest.mock.h
--- a/Engine/Source/Runtime/Online/BuildPatchServices/Private/Tests/Mock/Manifest.mock.h
+++ b/Engine/Source/Runtime/Online/BuildPatchServices/Private/Tests/Mock/Manifest.mock.h
@@ -350,13 +350,6 @@ namespace BuildPatchServices
 			return false;
 		}
 
-		virtual int32 EnumerateProducibleChunks(const FString& InstallDirectory, const TSet<FGuid>& ChunksRequired, TSet<FGuid>& ChunksAvailable) const override
-		{
-			int32 RetVal = ProducibleChunks.Difference(ChunksAvailable).Num();
-			ChunksAvailable.Append(ProducibleChunks);
-			return RetVal;
-		}
-
 		virtual void GetOutdatedFiles(const IBuildManifestRef& OldManifest, TSet<FString>& OutOutdatedFiles) const override { GetOutdatedFiles(FBuildPatchAppManifestPtr(StaticCastSharedRef<FBuildPatchAppManifest>(OldManifest)), TEXT(""), OutOutdatedFiles); }
 		virtual void GetOutdatedFiles(const FBuildPatchAppManifestPtr& OldManifest, const FString& InstallDirectory, TSet<FString>& OutOutdatedFiles) const override { GetOutdatedFiles(OldManifest.Get(), InstallDirectory, OutOutdatedFiles); }
 		virtual void GetOutdatedFiles(const FBuildPatchAppManifest*    OldManifest, const FString& InstallDirectory, TSet<FString>& OutOutdatedFiles) const override { GetOutdatedFiles(OldManifest, InstallDirectory, TSet<FString>(), OutOutdatedFiles); }
@@ -370,12 +363,6 @@ namespace BuildPatchServices
 			return true;
 		}
 
-		virtual TArray<FFileChunkPart> GetFilePartsForChunk(const FGuid& ChunkId) const override
-		{
-			checkNoEntry(); // AFAICT this function is deprecated and never used.
-			return TArray<FFileChunkPart>();
-		}
-
 		virtual bool HasFileAttributes() const override
 		{
 			return true;
