BPS: Fix log format

#rnx

diff --git a/Engine/Source/Runtime/Online/BuildPatchServices/Private/Installer/InstallChunkSource.cpp b/Engine/Source/Runtime/Online/BuildPatchServices/Private/Installer/InstallChunkSource.cpp
--- a/Engine/Source/Runtime/Online/BuildPatchServices/Private/Installer/InstallChunkSource.cpp
+++ b/Engine/Source/Runtime/Online/BuildPatchServices/Private/Installer/InstallChunkSource.cpp
@@ -302,7 +302,7 @@ namespace BuildPatchServices
 				AvailableInBuilds.Add(ChunkAssemblyPair.Key);
 			}
 
-			UE_LOG(LogInstallChunkSource, Verbose, TEXT("Got %llu complete chunks, %llu incomplete chunks, %llu repeat chunks"), CompleteChunks, IncompleteChunks, ChunksInAnotherSource);
+			UE_LOG(LogInstallChunkSource, Verbose, TEXT("Got %u complete chunks, %u incomplete chunks, %u repeat chunks"), CompleteChunks, IncompleteChunks, ChunksInAnotherSource);
 			if (CompleteChunks)
 			{
 				InstallationSources.Add(InstallationSourcePair);
