Fix race condition in background download code on iOS
#rnx

diff --git a/Engine/Source/Runtime/Online/BackgroundHTTP/Private/BackgroundHttpManagerImpl.cpp b/Engine/Source/Runtime/Online/BackgroundHTTP/Private/BackgroundHttpManagerImpl.cpp
--- a/Engine/Source/Runtime/Online/BackgroundHTTP/Private/BackgroundHttpManagerImpl.cpp
+++ b/Engine/Source/Runtime/Online/BackgroundHTTP/Private/BackgroundHttpManagerImpl.cpp
@@ -220,8 +220,8 @@ void FBackgroundHttpManagerImpl::GatherTempFilesWithoutURLMappings(TArray<FStrin
 	TArray<FString>& FilesToCheckRef = *FileListToCheckPtr;
 	for (const FString& File : FilesToCheckRef)
 	{
-		const FString* FoundURL = GetFileHashHelper()->FindMappedURLForTempFilename(File);
-		if (nullptr == FoundURL)
+		FString FoundURL = {};
+		if (!GetFileHashHelper()->FindMappedURLForTempFilename(File, FoundURL))
 		{
 			OutTempFilesMissingURLMappings.Add(File);
 		}
@@ -389,10 +389,10 @@ bool FBackgroundHttpManagerImpl::CheckForExistingCompletedDownload(const FBackgr
 	const TArray<FString>& URLList = Request->GetURLList();
 	for (const FString& URL : URLList)
 	{
-		const FString* FoundTempFilename = GetFileHashHelper()->FindTempFilenameMappingForURL(URL);
-		if (nullptr != FoundTempFilename)
+		FString FoundTempFilename = {};
+		if (GetFileHashHelper()->FindTempFilenameMappingForURL(URL, FoundTempFilename))
 		{
-			const FString& FileDestination = FBackgroundHttpFileHashHelper::GetFullPathOfTempFilename(*FoundTempFilename);
+			const FString& FileDestination = FBackgroundHttpFileHashHelper::GetFullPathOfTempFilename(FoundTempFilename);
 			if (PlatformFile.FileExists(*FileDestination))
 			{
 				bDidFindExistingDownload = true;
diff --git a/Engine/Source/Runtime/Online/BackgroundHTTPFileHash/Private/BackgroundHttpFileHashHelper.cpp b/Engine/Source/Runtime/Online/BackgroundHTTPFileHash/Private/BackgroundHttpFileHashHelper.cpp
--- a/Engine/Source/Runtime/Online/BackgroundHTTPFileHash/Private/BackgroundHttpFileHashHelper.cpp
+++ b/Engine/Source/Runtime/Online/BackgroundHTTPFileHash/Private/BackgroundHttpFileHashHelper.cpp
@@ -3,7 +3,6 @@
 #include "BackgroundHttpFileHashPrivate.h"
 
 #include "HAL/FileManager.h"
-#include "HAL/PlatformFileManager.h"
 #include "HAL/PlatformFile.h"
 
 #include "Hash/CityHash.h"
@@ -14,59 +13,63 @@
 
 DEFINE_LOG_CATEGORY(LogBackgroundHttpFileHash);
 
-
 FString FURLTempFileMapping::ToString() const
 {
-	//if changed make sure to update the format expected in InitFromString!
+	// if changed make sure to update the format expected in InitFromString!
 	return FString::Printf(TEXT("-URL=%s -TempFilename=%s"), *URL, *TempFilename);
 }
 
 bool FURLTempFileMapping::InitFromString(const FString& StringIn)
 {
 	FURLTempFileMapping ReturnedMapping;
-	
+
 	FString URLIn;
 	FString TempFilenameIn;
-	
+
 	const bool bDidSuccessfullyParse = FParse::Value(*StringIn, TEXT("-URL="), URLIn)
 								&& FParse::Value(*StringIn, TEXT("-TempFilename="), TempFilenameIn);
-	
-	//Only set our values if all values parsed successfully
+
+	// Only set our values if all values parsed successfully
 	if (bDidSuccessfullyParse)
 	{
 		URL = URLIn;
 		TempFilename = TempFilenameIn;
 	}
-	
+
 	return bDidSuccessfullyParse;
 }
 
 void FBackgroundHttpFileHashHelper::LoadData()
 {
-	//No reason to load this data more then once
+	FWriteScopeLock _(URLFileMappingsLock);
+
+	// No reason to load this data more then once
 	if (!bHasLoadedURLData)
 	{
 		FString FileContents;
-		const bool bLoadedFile = FFileHelper::LoadFileToString(FileContents, *GetURLMappingFilePath());
+		const bool bLoadedFile = FFileHelper::LoadFileToString(FileContents, *GetURLMappingFilePathNoLock());
 		if (bLoadedFile)
 		{
-			InitFromString(FileContents);
+			InitFromStringNoLock(FileContents);
 		}
 	}
-	
+
 	bHasLoadedURLData = true;
 }
 
 void FBackgroundHttpFileHashHelper::SaveData()
 {
-	if (bIsDirty)
+	FWriteScopeLock _(URLFileMappingsLock); // write due to bIsDirty
+
+	if (!bIsDirty.exchange(false))
 	{
-		//Only save to disk if we have already loaded the URLData or we will overwrite it with empty data!
-		if (ensureAlwaysMsgf(bHasLoadedURLData, TEXT("Attempt to Save URL Mapping Data before calling LoadData! Skipping as otherwise we would overwrite our persistent data!")))
-		{
-			FFileHelper::SaveStringToFile(ToString(), *GetURLMappingFilePath());
-			bIsDirty = false;
-		}
+		return;
+	}
+
+	// Only save to disk if we have already loaded the URLData or we will overwrite it with empty data!
+	if (ensureAlwaysMsgf(bHasLoadedURLData, TEXT("Attempt to Save URL Mapping Data before calling LoadData! Skipping as otherwise we would overwrite our persistent data!")))
+	{
+		FFileHelper::SaveStringToFile(ToStringNoLock(), *GetURLMappingFilePathNoLock());
 	}
 }
 
@@ -82,61 +85,74 @@ const FString& FBackgroundHttpFileHashHelper::GetTempFileExtension()
 	return TempExtension;
 }
 
-const FString& FBackgroundHttpFileHashHelper::GetURLMappingFilePath()
+const FString& FBackgroundHttpFileHashHelper::GetURLMappingFilePathNoLock()
 {
-	static FString URLMappingFilePath = FPaths::Combine(GetTemporaryRootPath(),TEXT("URLMap"), TEXT("TempFileURLMappings.urlmap"));
+	static FString URLMappingFilePath = FPaths::Combine(GetTemporaryRootPath(), TEXT("URLMap"), TEXT("TempFileURLMappings.urlmap"));
 	return URLMappingFilePath;
 }
 
 FString FBackgroundHttpFileHashHelper::GetFullPathOfTempFilename(const FString& TempFilename)
 {
+	// don't take lock here, as other methods taking lock call it
 	const FString FullFileName = FPaths::SetExtension(FPaths::GetCleanFilename(TempFilename), GetTempFileExtension());
 	return FPaths::Combine(GetTemporaryRootPath(), FullFileName);
 }
 
 FString FBackgroundHttpFileHashHelper::ToString() const
+{
+	FReadScopeLock _(URLFileMappingsLock);
+	return ToStringNoLock();
+}
+
+FString FBackgroundHttpFileHashHelper::ToStringNoLock() const
 {
 	FString ReturnedString;
 	for(const TPair<FString, FURLTempFileMapping>& URLMapPair : URLFileMappings)
 	{
 		const FURLTempFileMapping& URLMapping = URLMapPair.Value;
 		ReturnedString.Append(URLMapping.ToString());
-		
-		//Add a newline between each entry.
-		//This needs to match what we are looking for between entries in InitFromString if changed!
+
+		// Add a newline between each entry.
+		// This needs to match what we are looking for between entries in InitFromString if changed!
 		ReturnedString.Append(TEXT("\r\n"));
 	}
-	
+
 	return MoveTemp(ReturnedString);
 }
 
 bool FBackgroundHttpFileHashHelper::InitFromString(const FString& StringIn)
 {
-	//Grab a single line into it's own string
+	FWriteScopeLock _(URLFileMappingsLock);
+	return InitFromStringNoLock(StringIn);
+}
+
+bool FBackgroundHttpFileHashHelper::InitFromStringNoLock(const FString& StringIn)
+{
+	// Grab a single line into it's own string
 	FString CurrentLine;
 	FString RestOfString = StringIn;
-	
+
 	bool bInitSuccess = true;
-	
-	//Parse each line of the string until we run out of data
+
+	// Parse each line of the string until we run out of data
 	bool bStillHasMoreToParse = !RestOfString.IsEmpty();
 	while (bStillHasMoreToParse)
 	{
-		//Entries are assumed to be split by a newline. Should be the same as what is setup in the ToString function!
+		// Entries are assumed to be split by a newline. Should be the same as what is setup in the ToString function!
 		RestOfString.Split("\r\n", &CurrentLine, &RestOfString);
-		
-		//We still have something to parse, so lets go ahead and create an entry in the new mapp
+
+		// We still have something to parse, so lets go ahead and create an entry in the new mapp
 		if (!CurrentLine.IsEmpty())
 		{
 			FURLTempFileMapping NewMapping;
 			const bool bParsedSuccessfully = NewMapping.InitFromString(CurrentLine);
-			
+
 			if (ensureAlwaysMsgf(bParsedSuccessfully, TEXT("Invalid data found in supplied string! Invalid Line: %s"), *CurrentLine))
 			{
-				//This isn't saved separately as its just always the TempFilename
+				// This isn't saved separately as its just always the TempFilename
 				FString HashForMap = NewMapping.TempFilename;
-				
-				//Actually add an entry for the URLFileMapping
+
+				// Actually add an entry for the URLFileMapping
 				URLFileMappings.Add(HashForMap, MoveTemp(NewMapping));
 			}
 			else
@@ -144,35 +160,53 @@ bool FBackgroundHttpFileHashHelper::InitFromString(const FString& StringIn)
 				bInitSuccess = false;
 			}
 		}
-		
+
 		bStillHasMoreToParse = !RestOfString.IsEmpty();
 	}
-	
+
 	bIsDirty = true;
 	return bInitSuccess;
 }
 
+bool FBackgroundHttpFileHashHelper::FindTempFilenameMappingForURL(const FString& URL, FString& OutTempFilename) const
+{
+	FReadScopeLock _(URLFileMappingsLock);
+
+	const FURLTempFileMapping* FoundMapping = FindMappingForURLNoLock(URL);
+	if (nullptr != FoundMapping)
+	{
+		OutTempFilename = FoundMapping->TempFilename;
+		return true;
+	}
+
+	return false;
+}
+
 const FString* FBackgroundHttpFileHashHelper::FindTempFilenameMappingForURL(const FString& URL) const
 {
-	const FURLTempFileMapping* FoundMapping = FindMappingForURL(URL);
+	FReadScopeLock _(URLFileMappingsLock);
+
+	const FURLTempFileMapping* FoundMapping = FindMappingForURLNoLock(URL);
 	if (nullptr != FoundMapping)
 	{
-		return &(FoundMapping->TempFilename);
+		return &FoundMapping->TempFilename;
 	}
-	
+
 	return nullptr;
 }
 
 const FString& FBackgroundHttpFileHashHelper::FindOrAddTempFilenameMappingForURL(const FString& URL)
 {
-	const FURLTempFileMapping* FoundMapping = FindMappingForURL(URL);
+	FWriteScopeLock _(URLFileMappingsLock);
+
+	const FURLTempFileMapping* FoundMapping = FindMappingForURLNoLock(URL);
 	if (nullptr != FoundMapping)
 	{
 		return (FoundMapping->TempFilename);
 	}
-	
-	//Didn't find a valid old one, so lets add it
-	FString Hash = FindValidFilenameHashForURL(URL);
+
+	// Didn't find a valid old one, so lets add it
+	FString Hash = FindValidFilenameHashForURLNoLock(URL);
 	FURLTempFileMapping& NewMapping = URLFileMappings.FindOrAdd(Hash);
 	NewMapping.URL = URL;
 	NewMapping.TempFilename = Hash;
@@ -181,21 +215,42 @@ const FString& FBackgroundHttpFileHashHelper::FindOrAddTempFilenameMappingForURL
 	return NewMapping.TempFilename;
 }
 
+bool FBackgroundHttpFileHashHelper::FindMappedURLForTempFilename(const FString& TempFilename, FString& OutMappedURL) const
+{
+	// Make sure we are only looking for the base temp filename instead of the full path
+	FString BaseFilename = FPaths::GetBaseFilename(TempFilename);
+
+	FReadScopeLock _(URLFileMappingsLock);
+
+	const FURLTempFileMapping* FoundMapping = URLFileMappings.Find(BaseFilename);
+	if (nullptr != FoundMapping)
+	{
+		OutMappedURL = FoundMapping->URL;
+		return true;
+	}
+	else
+	{
+		return false;
+	}
+}
+
 const FString* FBackgroundHttpFileHashHelper::FindMappedURLForTempFilename(const FString& TempFilename) const
 {
-	//Make sure we are only looking for the base temp filename instead of the full path
+	// Make sure we are only looking for the base temp filename instead of the full path
 	FString BaseFilename = FPaths::GetBaseFilename(TempFilename);
-	
+
+	FReadScopeLock _(URLFileMappingsLock);
+
 	const FURLTempFileMapping* FoundMapping = URLFileMappings.Find(BaseFilename);
 	return (nullptr != FoundMapping) ? &(FoundMapping->URL) : nullptr;
 }
 
-const FURLTempFileMapping* FBackgroundHttpFileHashHelper::FindMappingForURL(const FString& URL) const
+const FURLTempFileMapping* FBackgroundHttpFileHashHelper::FindMappingForURLNoLock(const FString& URL) const
 {
-	FString Hash = FindValidFilenameHashForURL(URL);
+	FString Hash = FindValidFilenameHashForURLNoLock(URL);
 	const FURLTempFileMapping* FoundMapping = URLFileMappings.Find(Hash);
-	
-	//We found an already existing entry for the given Hash, so return that temp file mapping
+
+	// We found an already existing entry for the given Hash, so return that temp file mapping
 	if (nullptr != FoundMapping)
 	{
 		//Sanity check that these things match, shouldn't be required as we should always find a valid hash from FindValidFilenameHashForURL
@@ -204,31 +259,31 @@ const FURLTempFileMapping* FBackgroundHttpFileHashHelper::FindMappingForURL(cons
 			return FoundMapping;
 		}
 	}
-	
+
 	return nullptr;
 }
 
-FString FBackgroundHttpFileHashHelper::FindValidFilenameHashForURL(const FString& URL) const
+FString FBackgroundHttpFileHashHelper::FindValidFilenameHashForURLNoLock(const FString& URL) const
 {
 	FString HashToReturn;
-	
-	//Using a simple counter here that we should never hit in practice, but want to protect against
-	//and infinite loop in the case of a logic error here
+
+	// Using a simple counter here that we should never hit in practice, but want to protect against
+	// and infinite loop in the case of a logic error here
 	const uint32 MaxCollisionsToTry = 200;
 	uint32 CollisionCount = 0;
-	
-	//Put a simple counter here that we should NEVER hit, but wanted to protect against an infinite loop
-	//in case there is a serious logic error here
+
+	// Put a simple counter here that we should NEVER hit, but wanted to protect against an infinite loop
+	// in case there is a serious logic error here
 	while (CollisionCount < MaxCollisionsToTry)
 	{
-		const FString Hash = GenerateHashedFilenameForURL(URL, CollisionCount);
+		const FString Hash = GenerateHashedFilenameForURLNoLock(URL, CollisionCount);
 		const FURLTempFileMapping* FoundMapping = URLFileMappings.Find(Hash);
 		const bool bFoundValidMapping = (nullptr != FoundMapping);
-		
-		//Data already in our mapping
+
+		// Data already in our mapping
 		if (bFoundValidMapping)
 		{
-			//Existing entry was ours so use it
+			// Existing entry was ours so use it
 			if (FoundMapping->URL.Equals(URL, ESearchCase::IgnoreCase))
 			{
 				UE_LOG(LogBackgroundHttpFileHash, VeryVerbose, TEXT("Found valid existing URL Temp File Mapping for %s at %s"), *URL, *Hash);
@@ -236,26 +291,26 @@ FString FBackgroundHttpFileHashHelper::FindValidFilenameHashForURL(const FString
 				HashToReturn = Hash;
 				break;
 			}
-			//The existing entry was not ours, we have a collision and need to move on
+			// The existing entry was not ours, we have a collision and need to move on
 			else
 			{
 				UE_LOG(LogBackgroundHttpFileHash, Log, TEXT("Collision for URL Mapping for %s with %s with hash %s"), *URL, *(FoundMapping->URL), *Hash);
 			}
 		}
-		//Data not in our mapping, so its definitely safe to use for this URL
+		// Data not in our mapping, so its definitely safe to use for this URL
 		else
 		{
 			UE_LOG(LogBackgroundHttpFileHash, Verbose, TEXT("Found valid unused URL Temp File Mapping for %s at %s"), *URL, *Hash);
 			HashToReturn = Hash;
-			
-			//No need to keep searching as we found a valid result
+
+			// No need to keep searching as we found a valid result
 			break;
 		}
-		
+
 		++CollisionCount;
 	}
-	
-	//If we didn't manage to find a valid hash above, fallback to our sanitized URL
+
+	// If we didn't manage to find a valid hash above, fallback to our sanitized URL
 	if (HashToReturn.IsEmpty())
 	{
 		const FString TempFilename = FPaths::MakeValidFileName(URL);
@@ -263,22 +318,24 @@ FString FBackgroundHttpFileHashHelper::FindValidFilenameHashForURL(const FString
 		
 		HashToReturn = TempFilename;
 	}
-	
+
 	return MoveTemp(HashToReturn);
 }
 
-FString FBackgroundHttpFileHashHelper::GenerateHashedFilenameForURL(const FString& URL, uint32 CollisionCount) const
+FString FBackgroundHttpFileHashHelper::GenerateHashedFilenameForURLNoLock(const FString& URL, uint32 CollisionCount)
 {
 	const TArray<TCHAR, FString::AllocatorType>& UnderlyingCharData = URL.GetCharArray();
 	const uint32 HashedURL = CityHash32(reinterpret_cast<const char*>(UnderlyingCharData.GetData()), (UnderlyingCharData.Num() * sizeof(TCHAR)));
-	
+
 	return FPaths::MakeValidFileName(FString::Printf(TEXT("%d%d"), HashedURL, CollisionCount));
 }
 
 void FBackgroundHttpFileHashHelper::DeleteURLMappingsWithoutTempFiles()
 {
+	FWriteScopeLock _(URLFileMappingsLock);
+
 	TArray<FString> MappingsToRemove;
-	
+
 	for (const TPair<FString, FURLTempFileMapping>& Entry : URLFileMappings)
 	{
 		const FString FullFilePathForEntry = GetFullPathOfTempFilename(Entry.Value.TempFilename);
@@ -287,7 +344,7 @@ void FBackgroundHttpFileHashHelper::DeleteURLMappingsWithoutTempFiles()
 			MappingsToRemove.Add(Entry.Key);
 		}
 	}
-	
+
 	for (const FString& Key : MappingsToRemove)
 	{
 		URLFileMappings.Remove(Key);
@@ -298,7 +355,9 @@ void FBackgroundHttpFileHashHelper::DeleteURLMappingsWithoutTempFiles()
 
 void FBackgroundHttpFileHashHelper::RemoveURLMapping(const FString& URL)
 {
-	const FString Hash = FindValidFilenameHashForURL(URL);
+	FWriteScopeLock _(URLFileMappingsLock);
+
+	const FString Hash = FindValidFilenameHashForURLNoLock(URL);
 	URLFileMappings.Remove(Hash);
 	
 	bIsDirty = true;
diff --git a/Engine/Source/Runtime/Online/BackgroundHTTPFileHash/Public/BackgroundHttpFileHashHelper.h b/Engine/Source/Runtime/Online/BackgroundHTTPFileHash/Public/BackgroundHttpFileHashHelper.h
--- a/Engine/Source/Runtime/Online/BackgroundHTTPFileHash/Public/BackgroundHttpFileHashHelper.h
+++ b/Engine/Source/Runtime/Online/BackgroundHTTPFileHash/Public/BackgroundHttpFileHashHelper.h
@@ -48,67 +48,84 @@ public:
 		: bHasLoadedURLData(false)
 		, bIsDirty(false)
 		, URLFileMappings()
+		, URLFileMappingsLock()
 	{
 	}
-	
-	//Handles loading URL Mapping data from disk. No operation if data has previously been loaded.
+
+	// Handles loading URL Mapping data from disk. No operation if data has previously been loaded.
 	BACKGROUNDHTTPFILEHASH_API void LoadData();
-	
-	//Handles saving URL Mapping data to disk. No operation if data isn't dirty
+
+	// Handles saving URL Mapping data to disk. No operation if data isn't dirty
 	BACKGROUNDHTTPFILEHASH_API void SaveData();
-	
-	//Creates an FString representation of this FBackgroundHttpFileHashHelper's URL mappings. Used for serializing to/from disk
+
+	// Creates an FString representation of this FBackgroundHttpFileHashHelper's URL mappings. Used for serializing to/from disk
 	BACKGROUNDHTTPFILEHASH_API FString ToString() const;
 
-	//Sets the values of this FBackgroundHttpFileHashHelper from the string representation
-	//Returns if the parse was successful or not
+	// Sets the values of this FBackgroundHttpFileHashHelper from the string representation
+	// Returns if the parse was successful or not
 	BACKGROUNDHTTPFILEHASH_API bool InitFromString(const FString& StringIn);
 
-	//Looks for a temp filename mapping for the given URL. If one isn't found returns nullptr and does NOT generate one.
+	// Looks for a temp filename mapping for the given URL. If one isn't found returns nullptr and does NOT generate one.
+	BACKGROUNDHTTPFILEHASH_API bool FindTempFilenameMappingForURL(const FString& URL, FString& OutTempFilename) const;
+
+	UE_DEPRECATED(5.7, "Please use FindTempFilenameMappingForURL with reference out argument instead")
 	BACKGROUNDHTTPFILEHASH_API const FString* FindTempFilenameMappingForURL(const FString& URL) const;
-	
-	//Looks for a temp filename mapping for the given URL. If one isn't found a new one is generated and returned.
+
+	// Looks for a temp filename mapping for the given URL. If one isn't found a new one is generated and returned.
 	BACKGROUNDHTTPFILEHASH_API const FString& FindOrAddTempFilenameMappingForURL(const FString& URL);
 
-	//Removes URL mapping
+	// Removes URL mapping
 	BACKGROUNDHTTPFILEHASH_API void RemoveURLMapping(const FString& URL);
-	
-	//Looks for a URL mapped to the supplied TempFilename. Nullptr returned if one isn't found
+
+	// Looks for a URL mapped to the supplied TempFilename. Nullptr returned if one isn't found
+	BACKGROUNDHTTPFILEHASH_API bool FindMappedURLForTempFilename(const FString& TempFilename, FString& OutMappedURL) const;
+
+	UE_DEPRECATED(5.7, "Please use FindMappedURLForTempFilename with reference out argument instead")
 	BACKGROUNDHTTPFILEHASH_API const FString* FindMappedURLForTempFilename(const FString& TempFilename) const;
-	
-	//Deletes any URLMapping that doesn't have a corresponding Temp file actually on disk.
+
+	// Deletes any URLMapping that doesn't have a corresponding Temp file actually on disk.
 	BACKGROUNDHTTPFILEHASH_API void DeleteURLMappingsWithoutTempFiles();
-	
-	//Get our base directory used to store Temp files
+
+	// Get our base directory used to store Temp files
 	static BACKGROUNDHTTPFILEHASH_API const FString& GetTemporaryRootPath();
-	
-	//Helper function that returns the file extension used by our BackgroundHTTP temp files
+
+	// Helper function that returns the file extension used by our BackgroundHTTP temp files
 	static BACKGROUNDHTTPFILEHASH_API const FString& GetTempFileExtension();
-	
-	//Helper function that returns a full path representation of hte given TempFilename
+
+	// Helper function that returns a full path representation of hte given TempFilename
 	static BACKGROUNDHTTPFILEHASH_API FString GetFullPathOfTempFilename(const FString& TempFilename);
-	
+
 private:
-	
-	//Helper function that just returns the filepath we use to save/load our URL mapping data
-	static BACKGROUNDHTTPFILEHASH_API const FString& GetURLMappingFilePath();
-	
-	//Looks for a temp filename mapping for the given URL. If one isn't found returns nullptr an does NOT generate one.
-	BACKGROUNDHTTPFILEHASH_API const FURLTempFileMapping* FindMappingForURL(const FString& URL) const;
-	
-	//Finds the hash that should be used to refer to this URL. Either finds the hash used in the map to reference the
-	//URLFileMappings entry for this URL or returns the hash we should use (the next valid hash entry) without inserting it into the map.
-	BACKGROUNDHTTPFILEHASH_API FString FindValidFilenameHashForURL(const FString& URL) const;
-	//Creates a hashed Filename with a modifier for the given collision count.
-	BACKGROUNDHTTPFILEHASH_API FString GenerateHashedFilenameForURL(const FString& URL, uint32 CollisionCount) const;
-
-	//Used to make sure that we have called LoadData before any operations that depend on it
-	bool bHasLoadedURLData;
-	//Tracks if we have changed data since our last call to SaveData()
-	bool bIsDirty;
-	
-	//Stores our FURLTempFileMapping mapped to a Hashed filename
+	// Same as ToString but doesn't take a lock 
+	BACKGROUNDHTTPFILEHASH_API FString ToStringNoLock() const;
+
+	// Same as InitFromString but doesn't take a lock 
+	BACKGROUNDHTTPFILEHASH_API bool InitFromStringNoLock(const FString& StringIn);
+
+	// Helper function that just returns the filepath we use to save/load our URL mapping data
+	static BACKGROUNDHTTPFILEHASH_API const FString& GetURLMappingFilePathNoLock();
+
+	// Looks for a temp filename mapping for the given URL. If one isn't found returns nullptr an does NOT generate one.
+	BACKGROUNDHTTPFILEHASH_API const FURLTempFileMapping* FindMappingForURLNoLock(const FString& URL) const;
+
+	// Finds the hash that should be used to refer to this URL. Either finds the hash used in the map to reference the
+	// URLFileMappings entry for this URL or returns the hash we should use (the next valid hash entry) without inserting it into the map.
+	BACKGROUNDHTTPFILEHASH_API FString FindValidFilenameHashForURLNoLock(const FString& URL) const;
+
+	// Creates a hashed Filename with a modifier for the given collision count.
+	static BACKGROUNDHTTPFILEHASH_API FString GenerateHashedFilenameForURLNoLock(const FString& URL, uint32 CollisionCount);
+
+	// Used to make sure that we have called LoadData before any operations that depend on it
+	std::atomic<bool> bHasLoadedURLData;
+
+	// Tracks if we have changed data since our last call to SaveData()
+	std::atomic<bool> bIsDirty;
+
+	// Stores our FURLTempFileMapping mapped to a Hashed filename
 	TMap<FString, FURLTempFileMapping> URLFileMappings;
+
+	// Access lock for URLFileMappings
+	mutable FRWLock URLFileMappingsLock;
 };
 
 typedef TSharedRef<FBackgroundHttpFileHashHelper, ESPMode::ThreadSafe> BackgroundHttpFileHashHelperRef;
