Fix websocket connection closing events not firing when libwebsocket bPollService is false.

When bPollService is true, lws_service is called with non-zero timeout, and requires waking with lws_cancel_service. When closing a socket locally, lws is stuck waiting and cant wake. This CL makes sure Wake is called on close in order to wake lws_service.

diff --git a/Engine/Source/Runtime/Online/WebSockets/Private/Lws/LwsWebSocket.cpp b/Engine/Source/Runtime/Online/WebSockets/Private/Lws/LwsWebSocket.cpp
--- a/Engine/Source/Runtime/Online/WebSockets/Private/Lws/LwsWebSocket.cpp
+++ b/Engine/Source/Runtime/Online/WebSockets/Private/Lws/LwsWebSocket.cpp
@@ -164,6 +164,8 @@ void FLwsWebSocket::Close(int32 Code, const FString& Reason)
 		CloseRequest.Code = Code;
 		CloseRequest.Reason = ANSIReason;
 	}
+
+	FLwsWebSocketsManager::Get().WakeService();
 }
 
 void FLwsWebSocket::Send(const void* Data, SIZE_T Size, bool bIsBinary)
