Fix some clang static analysis issues

#rnx

diff --git a/Engine/Source/Runtime/Core/Private/ProfilingDebugging/Unix/UnixMemoryTrace.cpp b/Engine/Source/Runtime/Core/Private/ProfilingDebugging/Unix/UnixMemoryTrace.cpp
--- a/Engine/Source/Runtime/Core/Private/ProfilingDebugging/Unix/UnixMemoryTrace.cpp
+++ b/Engine/Source/Runtime/Core/Private/ProfilingDebugging/Unix/UnixMemoryTrace.cpp
@@ -26,6 +26,10 @@ FMalloc* MemoryTrace_Create(FMalloc* InMalloc)
 			ArgV[ArgC++] = Arg;
 			Arg = nullptr; // getdelim will allocate buffer for next Arg
 		}
+		if (Arg != nullptr)
+		{
+			free(Arg);
+		}
 		fclose(CmdLineFile);
 
 		FMalloc* OutMalloc = MemoryTrace_CreateInternal(InMalloc, ArgC, ArgV);
diff --git a/Engine/Source/Runtime/Core/Private/Unix/UnixPlatformProcess.cpp b/Engine/Source/Runtime/Core/Private/Unix/UnixPlatformProcess.cpp
--- a/Engine/Source/Runtime/Core/Private/Unix/UnixPlatformProcess.cpp
+++ b/Engine/Source/Runtime/Core/Private/Unix/UnixPlatformProcess.cpp
@@ -2261,6 +2261,10 @@ uint32 FUnixPlatformProcess::FProcEnumInfo::GetParentPID() const
 	sprintf(Buf, "/proc/%d/stat", GetPID());
 	
 	FILE* FilePtr = fopen(Buf, "r");
+	if (FilePtr == nullptr)
+	{
+		return 1;
+	}
 	if (fscanf(FilePtr, "%d %s %c %d", &DummyNumber, Buf, &DummyChar, &ParentPID) != 4)
 	{
 		ParentPID = 1;
diff --git a/Engine/Source/Runtime/Core/Public/HAL/IConsoleManager.h b/Engine/Source/Runtime/Core/Public/HAL/IConsoleManager.h
--- a/Engine/Source/Runtime/Core/Public/HAL/IConsoleManager.h
+++ b/Engine/Source/Runtime/Core/Public/HAL/IConsoleManager.h
@@ -16,7 +16,7 @@
 
 #define TRACK_CONSOLE_FIND_COUNT !(UE_BUILD_SHIPPING || UE_BUILD_TEST)
 
-#if DO_CHECK && (!UE_BUILD_SHIPPING) // Disable even if checks in shipping are enabled.
+#if DO_CHECK && (!UE_BUILD_SHIPPING) && !defined(__clang_analyzer__) // Disable even if checks in shipping are enabled.
 	#define cvarCheckCode( Code )		checkCode( Code )
 #else
 	#define cvarCheckCode(...)
diff --git a/Engine/Source/Runtime/Online/Experimental/EventLoop/Private/EventLoop/BSDSocket/EventLoopNotifierSocket.cpp b/Engine/Source/Runtime/Online/Experimental/EventLoop/Private/EventLoop/BSDSocket/EventLoopNotifierSocket.cpp
--- a/Engine/Source/Runtime/Online/Experimental/EventLoop/Private/EventLoop/BSDSocket/EventLoopNotifierSocket.cpp
+++ b/Engine/Source/Runtime/Online/Experimental/EventLoop/Private/EventLoop/BSDSocket/EventLoopNotifierSocket.cpp
@@ -22,6 +22,10 @@ bool FEventLoopNotifierSocket::Init()
 
 	// TODO: Investigate if it's possible to use pipe or IPPROTO_RAW, or socketpair instead
 	FileDescriptorListener = socket(AF_INET, SOCK_STREAM, IPPROTO_TCP);
+	if (FileDescriptorListener == INVALID_SOCKET)
+	{
+		return false;
+	}
 	sockaddr_in AddrIn;
 	memset(&AddrIn, 0, sizeof(AddrIn));
 	AddrIn.sin_family = AF_INET;
@@ -57,6 +61,12 @@ bool FEventLoopNotifierSocket::Init()
 	getsockname(FileDescriptorListener, &Addr, &AddrLen);
 
 	FileDescriptorRead = socket(AF_INET, SOCK_STREAM, IPPROTO_TCP);
+	if (FileDescriptorRead == INVALID_SOCKET)
+	{
+		close(FileDescriptorListener);
+		FileDescriptorListener = INVALID_SOCKET;
+		return false;
+	}
 	if (connect(FileDescriptorRead, &Addr, sizeof(Addr)) == SOCKET_RESULT_FAILED)
 	{
 		close(FileDescriptorListener);
