BPS: Fixes for install source automation test failures.

#rnx

diff --git a/Engine/Source/Runtime/Online/BuildPatchServices/Private/Installer/InstallChunkSource.cpp b/Engine/Source/Runtime/Online/BuildPatchServices/Private/Installer/InstallChunkSource.cpp
--- a/Engine/Source/Runtime/Online/BuildPatchServices/Private/Installer/InstallChunkSource.cpp
+++ b/Engine/Source/Runtime/Online/BuildPatchServices/Private/Installer/InstallChunkSource.cpp
@@ -201,9 +201,11 @@ namespace BuildPatchServices
 					// 5 of them.
 					if (!bExistenceCheckPassed)
 					{
-						uint64 SizeOnDisk = IFileManager::Get().FileSize(*(InstallationDirectory / InstalledFileName));
+						int64 SizeOnDisk = 0;
+						bool bGotFileSizeSuccessfully = InFileSystem->GetFileSize(*(InstallationDirectory / InstalledFileName), SizeOnDisk);
 						uint64 ExpectedSize = InstalledManifest->GetFileSize(InstalledFileName);
-						if (SizeOnDisk == ExpectedSize)
+						if (bGotFileSizeSuccessfully &&
+							(uint64)SizeOnDisk == ExpectedSize)
 						{
 							bExistenceCheckPassed = true;
 						}
diff --git a/Engine/Source/Runtime/Online/BuildPatchServices/Private/Tests/Unit/InstallChunkSource.spec.cpp b/Engine/Source/Runtime/Online/BuildPatchServices/Private/Tests/Unit/InstallChunkSource.spec.cpp
--- a/Engine/Source/Runtime/Online/BuildPatchServices/Private/Tests/Unit/InstallChunkSource.spec.cpp
+++ b/Engine/Source/Runtime/Online/BuildPatchServices/Private/Tests/Unit/InstallChunkSource.spec.cpp
@@ -373,6 +373,8 @@ void FInstallChunkSourceSpec::InventUsableChunkData()
 			const FGuid& TheChunk = MockManifest->DataList[ChunkIndex + ChunkCountPerSource*ManifestIndex];
 			SomeAvailableChunks.Add(TheChunk);
 			ThisMockManifest->ProducibleChunks.Add(TheChunk);
+			ThisMockManifest->ChunksRequiredForFiles.Add(TheChunk);
+			MockManifest->ChunksRequiredForFiles.Add(TheChunk);
 
 			// make sure we can evenly divide the chunk in to 4 files without worrying about leftovers,
 			// and that we can fill using unsigned ints.
