[Http]Add support to mock different responses through different verbs for same url.

diff --git a/Engine/Source/Programs/WebTests/Private/TestHttp.cpp b/Engine/Source/Programs/WebTests/Private/TestHttp.cpp
--- a/Engine/Source/Programs/WebTests/Private/TestHttp.cpp
+++ b/Engine/Source/Programs/WebTests/Private/TestHttp.cpp
@@ -573,6 +573,32 @@ TEST_CASE_METHOD(FWaitUntilCompleteHttpFixture, "Can mock response failure throu
 	HttpRequest->ProcessRequest();
 }
 
+TEST_CASE_METHOD(FWaitUntilCompleteHttpFixture, "Can mock different responses for different verbs", HTTP_TAG)
+{
+	CVarHttpUrlPatternsToMockResponse->Set(TEXT("epicgames.com++GET->500 epicgames.com++POST->503"));
+
+	TSharedRef<IHttpRequest> HttpRequestGet = CreateRequest();
+	HttpRequestGet->SetVerb(FHttpConstants::VERB_GET);
+	HttpRequestGet->SetURL(TEXT("https://www.epicgames.com/"));
+	HttpRequestGet->OnProcessRequestComplete().BindLambda([](FHttpRequestPtr HttpRequest, FHttpResponsePtr HttpResponse, bool bSucceeded) {
+		CHECK(bSucceeded);
+		REQUIRE(HttpResponse != nullptr);
+		CHECK(HttpResponse->GetResponseCode() == 500);
+	});
+	HttpRequestGet->ProcessRequest();
+
+	TSharedRef<IHttpRequest> HttpRequestPost = CreateRequest();
+	HttpRequestPost->SetVerb(FHttpConstants::VERB_POST);
+	HttpRequestPost->SetURL(TEXT("https://www.epicgames.com/"));
+	HttpRequestPost->OnProcessRequestComplete().BindLambda([](FHttpRequestPtr HttpRequest, FHttpResponsePtr HttpResponse, bool bSucceeded) {
+		CHECK(bSucceeded);
+		REQUIRE(HttpResponse != nullptr);
+		CHECK(HttpResponse->GetResponseCode() == 503);
+		CVarHttpUrlPatternsToMockResponse->Set(TEXT(""));
+	});
+	HttpRequestPost->ProcessRequest();
+}
+
 TEST_CASE_METHOD(FWaitUntilCompleteHttpFixture, "Can mock successful response through CVar", HTTP_TAG)
 {
 	CVarHttpUrlPatternsToMockResponse->Set(TEXT("epicgames.com->200 unrealengine_wrong_domain.com->200++{\"var_str\":\"abcd\",\"var_int\":20}++Content-Type:application/json"));
diff --git a/Engine/Source/Runtime/Online/HTTP/Private/GenericPlatform/HttpRequestCommon.cpp b/Engine/Source/Runtime/Online/HTTP/Private/GenericPlatform/HttpRequestCommon.cpp
--- a/Engine/Source/Runtime/Online/HTTP/Private/GenericPlatform/HttpRequestCommon.cpp
+++ b/Engine/Source/Runtime/Online/HTTP/Private/GenericPlatform/HttpRequestCommon.cpp
@@ -103,7 +103,7 @@ bool FHttpRequestCommon::PreCheck() const
 
 bool FHttpRequestCommon::TriggerMockResponse()
 {
-	TOptional<FHttpManager::FMockResponse> MockResponse = FHttpModule::Get().GetHttpManager().GetMockResponse(GetURL());
+	TOptional<FHttpManager::FMockResponse> MockResponse = FHttpModule::Get().GetHttpManager().GetMockResponse(GetURL(), GetVerb());
 	if (MockResponse.IsSet())
 	{
 		if (MockResponse.GetValue().StatusCode == EHttpResponseCodes::Unknown)
diff --git a/Engine/Source/Runtime/Online/HTTP/Private/HttpManager.cpp b/Engine/Source/Runtime/Online/HTTP/Private/HttpManager.cpp
--- a/Engine/Source/Runtime/Online/HTTP/Private/HttpManager.cpp
+++ b/Engine/Source/Runtime/Online/HTTP/Private/HttpManager.cpp
@@ -757,15 +757,16 @@ void FHttpManager::UpdateUrlPatternsToMockResponse(IConsoleVariable* CVar)
 	for (const FString& UrlPatternToMockResponseString : UrlPatternsToMockResponseStrings)
 	{
 		// Example runtime args(NOTE: It can't contain space in body of each url pattern): 
-		//  -ini:Engine:[ConsoleVariables]:http.UrlPatternsToMockResponse="www.mydomain.com/mypath->403++{\"friendly_error\":\"online.presence.restricted\",\"error_code\":3456}++Content-Type:application/json"
+		//  -ini:Engine:[ConsoleVariables]:http.UrlPatternsToMockResponse="www.mydomain.com/mypath++POST->403++{\"friendly_error\":\"online.presence.restricted\",\"error_code\":3456}++Content-Type:application/json"
 		TArray<FString> UrlPattern;
 		UrlPatternToMockResponseString.ParseIntoArray(UrlPattern, TEXT("->"));
+		const FStringView PayloadSeparator = TEXT("++");
+
 		if (UrlPattern.Num() == 2)
 		{
 			FMockResponse MockResponse;
 
 			FString MockResponseStr = UrlPattern[1];
-			const FStringView PayloadSeparator = TEXT("++");
 			const int32 PayloadSeparatorIndex = UE::String::FindFirst(MockResponseStr, PayloadSeparator);
 			if (PayloadSeparatorIndex != INDEX_NONE)
 			{
@@ -797,21 +798,32 @@ void FHttpManager::UpdateUrlPatternsToMockResponse(IConsoleVariable* CVar)
 				MockResponse.StatusCode = FCString::Atoi(*MockResponseStr);
 			}
 
-			UrlPatternsToMockResponse.Emplace(UrlPattern[0], MoveTemp(MockResponse));
+			TArray<FString> UrlInfo;
+			UrlPattern[0].ParseIntoArray(UrlInfo, PayloadSeparator.GetData());
+
+			FUrlToMatch UrlToMatch;
+			UrlToMatch.Str = UrlInfo[0];
+			if (UrlInfo.Num() > 1)
+			{
+				UrlToMatch.Verb = UrlInfo[1];
+			}
+			UrlPatternsToMockResponse.Emplace(UrlToMatch, MoveTemp(MockResponse));
 		}
 	}
 }
 
-TOptional<FHttpManager::FMockResponse> FHttpManager::GetMockResponse(FStringView Url) const
+TOptional<FHttpManager::FMockResponse> FHttpManager::GetMockResponse(FStringView Url, FStringView Verb) const
 {
 	TOptional<FMockResponse> Result;
 	const FScopeLock CacheLock(&UrlPatternsToMockResponseCriticalSection);
 
-	for (const TPair<FString, FMockResponse>& UrlPattern : UrlPatternsToMockResponse)
+	for (const TPair<FUrlToMatch, FMockResponse>& UrlPatternPair : UrlPatternsToMockResponse)
 	{
-		if (Url.Contains(UrlPattern.Key))
+		const FUrlToMatch& UrlToMatch = UrlPatternPair.Key;
+		if (Url.Contains(UrlToMatch.Str) && (UrlToMatch.Verb.IsEmpty() || UrlToMatch.Verb.Compare(FString(Verb), ESearchCase::IgnoreCase) == 0))
 		{
-			Result = UrlPattern.Value;
+			Result = UrlPatternPair.Value;
+			break;
 		}
 	}
 
diff --git a/Engine/Source/Runtime/Online/HTTP/Public/HttpManager.h b/Engine/Source/Runtime/Online/HTTP/Public/HttpManager.h
--- a/Engine/Source/Runtime/Online/HTTP/Public/HttpManager.h
+++ b/Engine/Source/Runtime/Online/HTTP/Public/HttpManager.h
@@ -408,6 +408,12 @@ PACKAGE_SCOPE:
 		TMap<FString, FString> ResponseHeaders;
 	};
 
+	struct FUrlToMatch
+	{
+		FString Str;
+		FString Verb;
+	};
+
 	/** Used to lock access to add/remove/find requests */
 	mutable FCriticalSection RequestLock;
 
@@ -444,7 +450,7 @@ PACKAGE_SCOPE:
 	HTTP_API bool ShouldLogResponse(FStringView Url) const;
 	HTTP_API bool ShouldDisableFailedLog(FStringView Url) const;
 
-	HTTP_API TOptional<FMockResponse> GetMockResponse(FStringView Url) const;
+	HTTP_API TOptional<FMockResponse> GetMockResponse(FStringView Url, FStringView Verb) const;
 
 	HTTP_API void MarkCurrentThreadCompletingRequest(bool bCompleting);
 	HTTP_API bool IsCurrentThreadCompletingRequest() const;
@@ -457,6 +463,6 @@ PACKAGE_SCOPE:
 	HTTP_API void AddRequest(const FHttpRequestRef& Request);
 
 protected:
-	TMap<FString, FMockResponse> UrlPatternsToMockResponse;
+	TArray<TPair<FUrlToMatch, FMockResponse>> UrlPatternsToMockResponse;
 	mutable FCriticalSection UrlPatternsToMockResponseCriticalSection;
 };
