BPS: Update DiffManifest to optionally export patch descriptor files.

#rnx

diff --git a/Engine/Source/Programs/BuildPatchTool/Private/ToolModes/DiffManifestMode.cpp b/Engine/Source/Programs/BuildPatchTool/Private/ToolModes/DiffManifestMode.cpp
--- a/Engine/Source/Programs/BuildPatchTool/Private/ToolModes/DiffManifestMode.cpp
+++ b/Engine/Source/Programs/BuildPatchTool/Private/ToolModes/DiffManifestMode.cpp
@@ -46,6 +46,9 @@ public:
 			UE_LOG(LogBuildPatchTool, Display, TEXT("  -InstallTagsB=\"\"       Specifies in quotes a comma seperated list of install tags used on ManifestB. Same rules apply as InstallTagsA."));
 			UE_LOG(LogBuildPatchTool, Display, TEXT("  -CompareTagSet=\"\"      Specifies in quotes a comma seperated list of install tags used to calculate differential statistics betweeen the manifests. Multiple lists are allowed. Same rules apply as InstallTagsA."));
 			UE_LOG(LogBuildPatchTool, Display, TEXT("  -OutputFile=\"\"         Specifies in quotes the file path where the diff will be exported as a JSON object."));
+			UE_LOG(LogBuildPatchTool, Display, TEXT("  -PatchDescriptorPath=\"\" If specified, will write out metadata files specifying what parts of each patched file changed."));
+			UE_LOG(LogBuildPatchTool, Display, TEXT("  -OnlyPatchDescriptors    If specified with PatchDescriptorPath, DiffManifests will only emit the files and not emit diff info."));
+			UE_LOG(LogBuildPatchTool, Display, TEXT("  -RequireOptimizedDelta   If specified, fail if an optimized delta file doesn't exist."));
 			UE_LOG(LogBuildPatchTool, Display, TEXT(""));
 			return EReturnCode::OK;
 		}
@@ -76,6 +79,9 @@ public:
 		Configuration.TagSetB = TagSetB;
 		Configuration.CompareTagSets = CompareTagSets;
 		Configuration.OutputFilePath = OutputFile;
+		Configuration.OutputPatchDescriptorPath = PatchDescriptorPath;
+		Configuration.bOnlyPatchDescriptors = bOnlyPatchDescriptors;
+		Configuration.bRequireOptimizedDelta = bRequireOptimizedDelta;
 
 		// Run the merge manifest routine.
 		bool bSuccess = BpsInterface.DiffManifests(Configuration);
@@ -114,6 +120,12 @@ private:
 
 		ParseSwitches(TEXT("CompareTagSet="), CompareTagsArray, Switches);
 
+		PARSE_SWITCH(PatchDescriptorPath);
+		NormalizeUriFile(PatchDescriptorPath);
+
+		bRequireOptimizedDelta = ParseOption(TEXT("RequireOptimizedDelta"), Switches);
+		bOnlyPatchDescriptors = ParseOption(TEXT("OnlyPatchDescriptors"), Switches);
+
 		return true;
 #undef PARSE_SWITCH
 	}
@@ -144,6 +156,9 @@ private:
 	FString InstallTagsA;
 	FString InstallTagsB;
 	FString OutputFile;
+	FString PatchDescriptorPath;
+	bool bRequireOptimizedDelta;
+	bool bOnlyPatchDescriptors;
 };
 
 BuildPatchTool::IToolModeRef BuildPatchTool::FDiffManifestToolModeFactory::Create(IBuildPatchServicesModule& BpsInterface)
diff --git a/Engine/Source/Runtime/Online/BuildPatchServices/Private/Diffing/DiffManifests.cpp b/Engine/Source/Runtime/Online/BuildPatchServices/Private/Diffing/DiffManifests.cpp
--- a/Engine/Source/Runtime/Online/BuildPatchServices/Private/Diffing/DiffManifests.cpp
+++ b/Engine/Source/Runtime/Online/BuildPatchServices/Private/Diffing/DiffManifests.cpp
@@ -297,6 +297,7 @@ namespace BuildPatchServices
 		bool AsyncRun();
 		void HandleDownloadComplete(int32 RequestId, const FDownloadRef& Download);
 
+		void ExportPatchDescriptorFiles(const TSet<FString>& InstallTags, FBuildPatchAppManifestPtr NewManifest, FBuildPatchAppManifestPtr OldManifest);
 	private:
 		const FDiffManifestsConfiguration Configuration;
 		FTSTicker& CoreTicker;
@@ -393,6 +394,175 @@ namespace BuildPatchServices
 		return Thread.Get();
 	}
 
+	void FDiffManifests::ExportPatchDescriptorFiles(const TSet<FString>& InstallTags, FBuildPatchAppManifestPtr NewManifest, FBuildPatchAppManifestPtr OldManifest)
+	{
+		UE_LOG(LogDiffManifests, Display, TEXT("Generating PatchDescription Files for %d tags and writing to %s"), InstallTags.Num(), *Configuration.OutputPatchDescriptorPath);
+
+		// We want a separate report for each tag, since tags can't assume files from other tags are installed and it's easy to 
+		// accumulate after if we care.
+		for (const FString& InstallTag : InstallTags)
+		{
+			UE_LOG(LogDiffManifests, Display, TEXT("%s:"), *InstallTag);
+			TSet<FString> CurrentInstallTag;
+			CurrentInstallTag.Add(InstallTag);
+
+			TSet<FString> OldFiles;
+			OldManifest->GetTaggedFileList(CurrentInstallTag, OldFiles);
+
+			TArray<FString> NewFiles;
+			NewManifest->GetTaggedFileList(CurrentInstallTag, NewFiles);
+						
+			// Filter to files that need data sent. We don't use GetOutdatedFiles since we don't want files that get deleted.
+			// We have to use a Set here for GetChunksRequiredForFiles later.
+			TSet<FString> FilesToPatch;
+			for (const FString& NewFile : NewFiles)
+			{
+				const FFileManifest* FileManifest = NewManifest->GetFileManifest(NewFile);
+				if (!OldFiles.Contains(NewFile))
+				{
+					FilesToPatch.Add(NewFile);
+					continue;
+				}
+
+				// File exists in both - check if anything changed
+				FSHAHash NewHash;
+				FSHAHash OldHash;
+				bool bNewHashExists = NewManifest->GetFileHash(NewFile, NewHash);
+				bool bOldHashExists = OldManifest->GetFileHash(NewFile, OldHash);
+
+				if (!bNewHashExists || !bOldHashExists)
+				{
+					UE_LOG(LogDiffManifests, Error, TEXT("Hash missing for file! %s Old: %d New: %d"), *NewFile, bOldHashExists, bNewHashExists);
+					continue;
+				}
+
+				if (NewHash != OldHash)
+				{
+					FilesToPatch.Add(NewFile);
+				}
+			}
+
+			// OK got the list of files that need patching for this tag list - we need to generate the chunks we need.
+			TSet<FGuid> ChunksForNewFiles;
+			NewManifest->GetChunksRequiredForFiles(FilesToPatch, ChunksForNewFiles);
+
+			// Now we need to get the chunks that the old installation can provide *for this tag list*
+			TSet<FGuid> ChunksOnDisk;
+			TSet<FGuid> ChunksForOldFiles;
+			OldManifest->GetChunksRequiredForFiles(OldFiles, ChunksForOldFiles);
+			OldManifest->EnumerateProducibleChunks(CurrentInstallTag, ChunksForOldFiles, ChunksOnDisk);
+
+			TSet<FGuid> ChunksToDownload = ChunksForNewFiles.Difference(ChunksOnDisk);
+			uint64 TotalNewChunkBytes = 0;
+			uint64 TotalNewChunkDownloadBytes = 0;
+
+			for (FGuid& Guid : ChunksToDownload)
+			{
+				const BuildPatchServices::FChunkInfo* ChunkInfo = NewManifest->GetChunkInfo(Guid);
+				TotalNewChunkBytes += ChunkInfo->WindowSize;
+				TotalNewChunkDownloadBytes += ChunkInfo->FileSize;
+			}
+
+			// These numbers are different than the actual patch sizes since we might not be using the entire
+			// chunk, and we might be using parts more than once.
+			UE_LOG(LogDiffManifests, Display, TEXT("    Chunk uncompressed / download sizes: %s, %s"), 
+				*FText::AsNumber(TotalNewChunkBytes).ToString(),
+				*FText::AsNumber(TotalNewChunkDownloadBytes).ToString()
+			);
+
+			// Now we know which chunks have to be downloaded vs are available on disk, we can iterate the
+			// chunk parts and report whether its New or Matched in the format. We try to accumulate runs
+			// rather than send a string of adjacent matches.
+			uint64 TotalBytesWritten = 0;
+			uint64 TotalNewBytes = 0;
+			TSet<FGuid> ReferencedChunks;
+			for (const FString& PatchFile : FilesToPatch)
+			{
+				const FFileManifest* FileManifest = NewManifest->GetFileManifest(PatchFile);
+
+				TArray<TPair<uint64, bool>> MatchedOrNot;
+
+				uint64 FileOffset = 0;
+				for (const FChunkPart& ChunkPart : FileManifest->ChunkParts)
+				{
+					if (ChunksOnDisk.Contains(ChunkPart.Guid))
+					{
+						// Match - add if we don't have any entries, or we're not currently matching (i.e. can't extend)
+						if (!MatchedOrNot.Num() || !MatchedOrNot.Top().Value)
+						{
+							MatchedOrNot.Emplace(FileOffset, true);
+						}
+					}
+					else
+					{
+						ReferencedChunks.Add(ChunkPart.Guid);
+						TotalNewBytes += ChunkPart.Size;
+
+						// Patch - add if we don't have any entries, or we are currently matching (i.e. can't extend)
+						if (!MatchedOrNot.Num() || MatchedOrNot.Top().Value)
+						{
+							MatchedOrNot.Emplace(FileOffset, false);
+						}
+					}
+					FileOffset += ChunkPart.Size;
+				}
+				TotalBytesWritten += FileOffset;
+
+				// Emit the listing for this file.
+				TUtf8StringBuilder<1024> Listing;
+				for (int32 i = 0; i < MatchedOrNot.Num(); i++)
+				{
+					uint64 Offset = MatchedOrNot[i].Key;
+					bool bIsMatch = MatchedOrNot[i].Value;
+
+					uint64 End = FileOffset;
+					if (i < MatchedOrNot.Num() - 1)
+					{
+						End = MatchedOrNot[i+1].Key;
+					}
+
+					if (bIsMatch)
+					{
+						// The final 0 is supposed to be the offset in the source file where the data comes from,
+						// but we don't have that data (and we don't need it for the purposes of this utility, which is
+						// just to evaluate which parts of a file changed).
+						Listing.Appendf(UTF8TEXT("M,%llu,%llu,0\r\n"), Offset, End-Offset);
+					}
+					else
+					{
+						Listing.Appendf(UTF8TEXT("N,%llu,%llu\r\n"), Offset, End-Offset);
+					}
+				}
+
+				FString FileName = (Configuration.OutputPatchDescriptorPath / PatchFile) + ".patch";
+
+				TUniquePtr<FArchive> Ar(IFileManager::Get().CreateFileWriter(*FileName, 0));
+				if (!Ar)
+				{
+					UE_LOG(LogTemp, Warning, TEXT("Unable to open output file: %s"), *FileName);
+				}
+				else
+				{
+					FUtf8StringView OutputView = MakeStringView(Listing);
+					UTF8CHAR UTF8BOM[] = { (UTF8CHAR)0xEF, (UTF8CHAR)0xBB, (UTF8CHAR)0xBF };
+					Ar->Serialize(&UTF8BOM, sizeof(UTF8BOM));
+					Ar->Serialize((void*)OutputView.GetData(), OutputView.Len() * sizeof(UTF8CHAR));
+					Ar->Close();
+				}
+						
+			} // end each file
+
+			// This can be more than the uncompressed download size if chunks are reused, or smaller if we didn't need
+			// an entire chunk.
+			UE_LOG(LogDiffManifests, Display, TEXT("    Tag patches %d files, with %s new uncompressed bytes out of %s total bytes (%.1f %%)."), 
+				FilesToPatch.Num(), 
+				*FText::AsNumber(TotalNewBytes).ToString(),
+				*FText::AsNumber(TotalBytesWritten).ToString(),
+				TotalBytesWritten ? (100.0f * TotalNewBytes / TotalBytesWritten) : 0
+			);
+		}
+	}
+
 	bool FDiffManifests::AsyncRun()
 	{
 		FBuildPatchAppManifestPtr ManifestA = FutureManifestA.Get();
@@ -413,11 +583,17 @@ namespace BuildPatchServices
 			// Check for delta file, replacing ManifestB if we find one
 			FOptimisedDeltaConfiguration OptimisedDeltaConfiguration(ManifestB.ToSharedRef());
 			OptimisedDeltaConfiguration.SourceManifest = ManifestA;
-			OptimisedDeltaConfiguration.DeltaPolicy = EDeltaPolicy::TryFetchContinueWithout;
+			OptimisedDeltaConfiguration.DeltaPolicy = Configuration.bRequireOptimizedDelta ? EDeltaPolicy::Expect : EDeltaPolicy::TryFetchContinueWithout;
 			OptimisedDeltaConfiguration.CloudDirectories = { FPaths::GetPath(Configuration.ManifestBUri) };
 			FOptimisedDeltaDependencies OptimisedDeltaDependencies;
 			OptimisedDeltaDependencies.DownloadService = DownloadService.Get();
 			TUniquePtr<IOptimisedDelta> OptimisedDelta(FOptimisedDeltaFactory::Create(OptimisedDeltaConfiguration, MoveTemp(OptimisedDeltaDependencies)));
+			if (OptimisedDelta->GetResult().HasError())
+			{
+				UE_LOG(LogDiffManifests, Error, TEXT("Optimized delta not found and RequireOptimizedDelta was passed - failing"));
+				bShouldRun = false;
+				return false;
+			}
 			ManifestB = OptimisedDelta->GetResult().GetValue();
 			const int32 MetaDownloadBytes = OptimisedDelta->GetMetaDownloadSize();
 
@@ -433,6 +609,19 @@ namespace BuildPatchServices
 				TagsB = TagsB.Intersect(Configuration.TagSetB);
 			}
 
+			if (Configuration.OutputPatchDescriptorPath.Len())
+			{
+				TSet<FString> SharedTags = TagsA.Intersect(TagsB);
+				ExportPatchDescriptorFiles(SharedTags, ManifestA, ManifestB);
+
+				if (Configuration.bOnlyPatchDescriptors)
+				{
+					// Early out. 
+					bShouldRun = false;
+					return true;
+				}
+			}
+
 			int64 NewChunksCount = 0;
 			int64 TotalChunkSize = 0;
 			TSet<FString> TaggedFileSetA;
diff --git a/Engine/Source/Runtime/Online/BuildPatchServices/Public/BuildPatchSettings.h b/Engine/Source/Runtime/Online/BuildPatchServices/Public/BuildPatchSettings.h
--- a/Engine/Source/Runtime/Online/BuildPatchServices/Public/BuildPatchSettings.h
+++ b/Engine/Source/Runtime/Online/BuildPatchServices/Public/BuildPatchSettings.h
@@ -469,6 +469,17 @@ namespace BuildPatchServices
 		TArray<TSet<FString>> CompareTagSets;
 		// A full file path where a JSON object will be saved for the diff details.Empty string if not desired.
 		FString OutputFilePath;
+
+		// A destination base directory to use for writing out patch description files for the patch between the two
+		// manifest.
+		FString OutputPatchDescriptorPath;
+
+		// If true, the tool will throw an error if the optimized delta file is missing.
+		bool bRequireOptimizedDelta=false;
+
+		// If true, only emit patch descriptors and not the entirety of the diff manifests output.
+		bool bOnlyPatchDescriptors=false;
+		
 	};
 
 	/**
