BPS: Fix for patch descriptor dumping
BPS: Add a way to suppress installation time simulation.

#rnx

diff --git a/Engine/Source/Programs/BuildPatchTool/Private/ToolModes/DiffManifestMode.cpp b/Engine/Source/Programs/BuildPatchTool/Private/ToolModes/DiffManifestMode.cpp
--- a/Engine/Source/Programs/BuildPatchTool/Private/ToolModes/DiffManifestMode.cpp
+++ b/Engine/Source/Programs/BuildPatchTool/Private/ToolModes/DiffManifestMode.cpp
@@ -82,6 +82,7 @@ public:
 		Configuration.OutputPatchDescriptorPath = PatchDescriptorPath;
 		Configuration.bOnlyPatchDescriptors = bOnlyPatchDescriptors;
 		Configuration.bRequireOptimizedDelta = bRequireOptimizedDelta;
+		Configuration.bEmitInstallTime = !bSkipInstallTime;
 
 		// Run the merge manifest routine.
 		bool bSuccess = BpsInterface.DiffManifests(Configuration);
@@ -125,6 +126,7 @@ private:
 
 		bRequireOptimizedDelta = ParseOption(TEXT("RequireOptimizedDelta"), Switches);
 		bOnlyPatchDescriptors = ParseOption(TEXT("OnlyPatchDescriptors"), Switches);
+		bSkipInstallTime = ParseOption(TEXT("SkipInstallTime"), Switches);
 
 		return true;
 #undef PARSE_SWITCH
@@ -159,6 +161,7 @@ private:
 	FString PatchDescriptorPath;
 	bool bRequireOptimizedDelta;
 	bool bOnlyPatchDescriptors;
+	bool bSkipInstallTime;
 };
 
 BuildPatchTool::IToolModeRef BuildPatchTool::FDiffManifestToolModeFactory::Create(IBuildPatchServicesModule& BpsInterface)
diff --git a/Engine/Source/Runtime/Online/BuildPatchServices/Private/Diffing/DiffManifests.cpp b/Engine/Source/Runtime/Online/BuildPatchServices/Private/Diffing/DiffManifests.cpp
--- a/Engine/Source/Runtime/Online/BuildPatchServices/Private/Diffing/DiffManifests.cpp
+++ b/Engine/Source/Runtime/Online/BuildPatchServices/Private/Diffing/DiffManifests.cpp
@@ -297,7 +297,7 @@ namespace BuildPatchServices
 		bool AsyncRun();
 		void HandleDownloadComplete(int32 RequestId, const FDownloadRef& Download);
 
-		void ExportPatchDescriptorFiles(const TSet<FString>& InstallTags, FBuildPatchAppManifestPtr NewManifest, FBuildPatchAppManifestPtr OldManifest);
+		void ExportPatchDescriptorFiles(const TSet<FString>& InstallTags, FBuildPatchAppManifestPtr OldManifest, FBuildPatchAppManifestPtr NewManifest);
 	private:
 		const FDiffManifestsConfiguration Configuration;
 		FTSTicker& CoreTicker;
@@ -394,7 +394,7 @@ namespace BuildPatchServices
 		return Thread.Get();
 	}
 
-	void FDiffManifests::ExportPatchDescriptorFiles(const TSet<FString>& InstallTags, FBuildPatchAppManifestPtr NewManifest, FBuildPatchAppManifestPtr OldManifest)
+	void FDiffManifests::ExportPatchDescriptorFiles(const TSet<FString>& InstallTags, FBuildPatchAppManifestPtr OldManifest, FBuildPatchAppManifestPtr NewManifest)
 	{
 		UE_LOG(LogDiffManifests, Display, TEXT("Generating PatchDescription Files for %d tags and writing to %s"), InstallTags.Num(), *Configuration.OutputPatchDescriptorPath);
 
@@ -771,45 +771,49 @@ namespace BuildPatchServices
 			}
 
 			// Hit a destructive and nondestructive simulation for a few different specs.
-			TArray<DiffHelpers::FSimConfig> SimConfigs;
-			// Add some lower spec values, taken from around 25 percentile of stats at the time of writing [July 2019].
-			SimConfigs.AddDefaulted_GetRef().InstallMode = EInstallMode::DestructiveInstall;
-			SimConfigs.AddDefaulted_GetRef().InstallMode = EInstallMode::NonDestructiveInstall;
-			SimConfigs[0].DownloadSpeed  = SimConfigs[1].DownloadSpeed  =   1200000.0; // 1.2 MB/s
-			SimConfigs[0].DiskReadSpeed  = SimConfigs[1].DiskReadSpeed  =  30000000.0; // 30 MB/s
-			SimConfigs[0].DiskWriteSpeed = SimConfigs[1].DiskWriteSpeed =  25000000.0; // 25 MB/s
-			// We didn't have stats for BackupSerialisationSpeed, but it runs much slower than disk speed and so we tune it to be sure destructive is relatively penalizing.
-			SimConfigs[0].BackupSerialisationSpeed = SimConfigs[1].BackupSerialisationSpeed = 10000000.0; // 10 MB/s
-			// Add some lower spec values, taken from around 50 percentile of stats at the time of writing [July 2019].
-			SimConfigs.AddDefaulted_GetRef().InstallMode = EInstallMode::DestructiveInstall;
-			SimConfigs.AddDefaulted_GetRef().InstallMode = EInstallMode::NonDestructiveInstall;
-			SimConfigs[2].DownloadSpeed  = SimConfigs[3].DownloadSpeed  =   3500000.0; // 3.5 MB/s
-			SimConfigs[2].DiskReadSpeed  = SimConfigs[3].DiskReadSpeed  = 145000000.0; // 145 MB/s
-			SimConfigs[2].DiskWriteSpeed = SimConfigs[3].DiskWriteSpeed =  75000000.0; // 75 MB/s
-			// We didn't have stats for BackupSerialisationSpeed, but it runs much slower than disk speed and so we tune it to be sure destructive is relatively penalizing.
-			SimConfigs[2].BackupSerialisationSpeed = SimConfigs[3].BackupSerialisationSpeed = 20000000.0; // 20 MB/s
-			// Add some higher spec values, taken from around 75 percentile of stats at the time of writing [July 2019].
-			SimConfigs.AddDefaulted_GetRef().InstallMode = EInstallMode::DestructiveInstall;
-			SimConfigs.AddDefaulted_GetRef().InstallMode = EInstallMode::NonDestructiveInstall;
-			SimConfigs[4].DownloadSpeed  = SimConfigs[5].DownloadSpeed  =  13000000.0; // 13 MB/s
-			SimConfigs[4].DiskReadSpeed  = SimConfigs[5].DiskReadSpeed  = 295000000.0; // 295 MB/s
-			SimConfigs[4].DiskWriteSpeed = SimConfigs[5].DiskWriteSpeed = 125000000.0; // 125 MB/s
-			// We didn't have stats for BackupSerialisationSpeed, but it runs much slower than disk speed and so we tune it to be sure destructive is relatively penalizing.
-			SimConfigs[4].BackupSerialisationSpeed = SimConfigs[5].BackupSerialisationSpeed = 40000000.0; // 40 MB/s
-
-			// Run the calculations and log.
-			TArray<double> InstallTimeCoefficients = DiffHelpers::CalculateInstallTimeCoefficient(ManifestA.ToSharedRef(), TagsA, ManifestB.ToSharedRef(), TagsB, SimConfigs);
-			checkf(6 == InstallTimeCoefficients.Num() && 6 == SimConfigs.Num(), TEXT("Unexpected result size from CalculateInstallTimeCoefficient."));
-			UE_LOG(LogDiffManifests, Display, TEXT(""));
-			UE_LOG(LogDiffManifests, Display, TEXT("Install time coefficients are not accurate timing representations, but are comparable from patch to patch."));
-			UE_LOG(LogDiffManifests, Display, TEXT("They can be used to spot out of the ordinary time requirements for installing an update."));
-			UE_LOG(LogDiffManifests, Display, TEXT("Install Time Coefficients:"));
-			UE_LOG(LogDiffManifests, Display, TEXT("    Low-Spec  DestructiveInstall:    %s"), *FPlatformTime::PrettyTime(InstallTimeCoefficients[0]));
-			UE_LOG(LogDiffManifests, Display, TEXT("    Low-Spec  NonDestructiveInstall: %s"), *FPlatformTime::PrettyTime(InstallTimeCoefficients[1]));
-			UE_LOG(LogDiffManifests, Display, TEXT("    Mid-Spec  DestructiveInstall:    %s"), *FPlatformTime::PrettyTime(InstallTimeCoefficients[2]));
-			UE_LOG(LogDiffManifests, Display, TEXT("    Mid-Spec  NonDestructiveInstall: %s"), *FPlatformTime::PrettyTime(InstallTimeCoefficients[3]));
-			UE_LOG(LogDiffManifests, Display, TEXT("    High-Spec DestructiveInstall:    %s"), *FPlatformTime::PrettyTime(InstallTimeCoefficients[4]));
-			UE_LOG(LogDiffManifests, Display, TEXT("    High-Spec NonDestructiveInstall: %s"), *FPlatformTime::PrettyTime(InstallTimeCoefficients[5]));
+			TArray<double> InstallTimeCoefficients;
+			if (Configuration.bEmitInstallTime)
+			{
+				TArray<DiffHelpers::FSimConfig> SimConfigs;
+				// Add some lower spec values, taken from around 25 percentile of stats at the time of writing [July 2019].
+				SimConfigs.AddDefaulted_GetRef().InstallMode = EInstallMode::DestructiveInstall;
+				SimConfigs.AddDefaulted_GetRef().InstallMode = EInstallMode::NonDestructiveInstall;
+				SimConfigs[0].DownloadSpeed  = SimConfigs[1].DownloadSpeed  =   1200000.0; // 1.2 MB/s
+				SimConfigs[0].DiskReadSpeed  = SimConfigs[1].DiskReadSpeed  =  30000000.0; // 30 MB/s
+				SimConfigs[0].DiskWriteSpeed = SimConfigs[1].DiskWriteSpeed =  25000000.0; // 25 MB/s
+				// We didn't have stats for BackupSerialisationSpeed, but it runs much slower than disk speed and so we tune it to be sure destructive is relatively penalizing.
+				SimConfigs[0].BackupSerialisationSpeed = SimConfigs[1].BackupSerialisationSpeed = 10000000.0; // 10 MB/s
+				// Add some lower spec values, taken from around 50 percentile of stats at the time of writing [July 2019].
+				SimConfigs.AddDefaulted_GetRef().InstallMode = EInstallMode::DestructiveInstall;
+				SimConfigs.AddDefaulted_GetRef().InstallMode = EInstallMode::NonDestructiveInstall;
+				SimConfigs[2].DownloadSpeed  = SimConfigs[3].DownloadSpeed  =   3500000.0; // 3.5 MB/s
+				SimConfigs[2].DiskReadSpeed  = SimConfigs[3].DiskReadSpeed  = 145000000.0; // 145 MB/s
+				SimConfigs[2].DiskWriteSpeed = SimConfigs[3].DiskWriteSpeed =  75000000.0; // 75 MB/s
+				// We didn't have stats for BackupSerialisationSpeed, but it runs much slower than disk speed and so we tune it to be sure destructive is relatively penalizing.
+				SimConfigs[2].BackupSerialisationSpeed = SimConfigs[3].BackupSerialisationSpeed = 20000000.0; // 20 MB/s
+				// Add some higher spec values, taken from around 75 percentile of stats at the time of writing [July 2019].
+				SimConfigs.AddDefaulted_GetRef().InstallMode = EInstallMode::DestructiveInstall;
+				SimConfigs.AddDefaulted_GetRef().InstallMode = EInstallMode::NonDestructiveInstall;
+				SimConfigs[4].DownloadSpeed  = SimConfigs[5].DownloadSpeed  =  13000000.0; // 13 MB/s
+				SimConfigs[4].DiskReadSpeed  = SimConfigs[5].DiskReadSpeed  = 295000000.0; // 295 MB/s
+				SimConfigs[4].DiskWriteSpeed = SimConfigs[5].DiskWriteSpeed = 125000000.0; // 125 MB/s
+				// We didn't have stats for BackupSerialisationSpeed, but it runs much slower than disk speed and so we tune it to be sure destructive is relatively penalizing.
+				SimConfigs[4].BackupSerialisationSpeed = SimConfigs[5].BackupSerialisationSpeed = 40000000.0; // 40 MB/s
+
+				// Run the calculations and log.
+				InstallTimeCoefficients = DiffHelpers::CalculateInstallTimeCoefficient(ManifestA.ToSharedRef(), TagsA, ManifestB.ToSharedRef(), TagsB, SimConfigs);
+				checkf(6 == InstallTimeCoefficients.Num() && 6 == SimConfigs.Num(), TEXT("Unexpected result size from CalculateInstallTimeCoefficient."));
+				UE_LOG(LogDiffManifests, Display, TEXT(""));
+				UE_LOG(LogDiffManifests, Display, TEXT("Install time coefficients are not accurate timing representations, but are comparable from patch to patch."));
+				UE_LOG(LogDiffManifests, Display, TEXT("They can be used to spot out of the ordinary time requirements for installing an update."));
+				UE_LOG(LogDiffManifests, Display, TEXT("Install Time Coefficients:"));
+				UE_LOG(LogDiffManifests, Display, TEXT("    Low-Spec  DestructiveInstall:    %s"), *FPlatformTime::PrettyTime(InstallTimeCoefficients[0]));
+				UE_LOG(LogDiffManifests, Display, TEXT("    Low-Spec  NonDestructiveInstall: %s"), *FPlatformTime::PrettyTime(InstallTimeCoefficients[1]));
+				UE_LOG(LogDiffManifests, Display, TEXT("    Mid-Spec  DestructiveInstall:    %s"), *FPlatformTime::PrettyTime(InstallTimeCoefficients[2]));
+				UE_LOG(LogDiffManifests, Display, TEXT("    Mid-Spec  NonDestructiveInstall: %s"), *FPlatformTime::PrettyTime(InstallTimeCoefficients[3]));
+				UE_LOG(LogDiffManifests, Display, TEXT("    High-Spec DestructiveInstall:    %s"), *FPlatformTime::PrettyTime(InstallTimeCoefficients[4]));
+				UE_LOG(LogDiffManifests, Display, TEXT("    High-Spec NonDestructiveInstall: %s"), *FPlatformTime::PrettyTime(InstallTimeCoefficients[5]));
+			}
 
 			// Save the output.
 			if (bSuccess && Configuration.OutputFilePath.IsEmpty() == false)
diff --git a/Engine/Source/Runtime/Online/BuildPatchServices/Public/BuildPatchSettings.h b/Engine/Source/Runtime/Online/BuildPatchServices/Public/BuildPatchSettings.h
--- a/Engine/Source/Runtime/Online/BuildPatchServices/Public/BuildPatchSettings.h
+++ b/Engine/Source/Runtime/Online/BuildPatchServices/Public/BuildPatchSettings.h
@@ -480,6 +480,8 @@ namespace BuildPatchServices
 		// If true, only emit patch descriptors and not the entirety of the diff manifests output.
 		bool bOnlyPatchDescriptors=false;
 		
+		// If true, emit simulated installation time estimations
+		bool bEmitInstallTime=true;
 	};
 
 	/**
