Fix hang in FLwsWebSocketsManager::ShutdownWebSockets

Prevent lws_service from being called again after we've decided to shutdown.

diff --git a/Engine/Source/Runtime/Online/WebSockets/Private/Lws/LwsWebSocketsManager.cpp b/Engine/Source/Runtime/Online/WebSockets/Private/Lws/LwsWebSocketsManager.cpp
--- a/Engine/Source/Runtime/Online/WebSockets/Private/Lws/LwsWebSocketsManager.cpp
+++ b/Engine/Source/Runtime/Online/WebSockets/Private/Lws/LwsWebSocketsManager.cpp
@@ -295,12 +295,6 @@ void FLwsWebSocketsManager::ShutdownWebSockets()
 		OnChildEndFramePostForkHandle.Reset();
 	}
 
-	if (LwsContext)
-	{
-		// Cancel service which may be blocking the HTTP thread, before we attempt the kill Thread below 
-		lws_cancel_service(LwsContext);
-	}
-
 	if (Thread)
 	{
 		Thread->Kill(true);
@@ -366,6 +360,11 @@ uint32 FLwsWebSocketsManager::Run()
 void FLwsWebSocketsManager::Stop()
 {
 	ExitRequest.Set(true);
+
+	if (LwsContext)
+	{
+		WakeService();
+	}
 }
 
 void FLwsWebSocketsManager::Exit()
