UE : HTTP - Array ReallocGrow3 crash fix when ReceiveResponseBodyCallback in FCurlHttpRequest.
#tests webtests
[RN] fix crash in FCurlHttpRequest ReceiveResponseBodyCallback

diff --git a/Engine/Source/Programs/WebTests/Private/TestHttp.cpp b/Engine/Source/Programs/WebTests/Private/TestHttp.cpp
--- a/Engine/Source/Programs/WebTests/Private/TestHttp.cpp
+++ b/Engine/Source/Programs/WebTests/Private/TestHttp.cpp
@@ -1234,6 +1234,42 @@ TEST_CASE_METHOD(FWaitUntilCompleteHttpFixture, "Can download big file exceeds 3
 	UE_LOG(LogHttp, Log, TEXT("%d GB of big download request spent %.2fs to complete!"), DataSizeInGB, DurationInSeconds);
 }
 
+TEST_CASE_METHOD(FWaitUntilCompleteHttpFixture, "Can download big file exceeds 32 bits without stream delegate", HTTP_TAG)
+{
+	DisableWarningsInThisTest();
+
+	if (!bRunHeavyTests) // TODO: This can be removed once we can make the http request fail early. 
+	{
+		return;
+	}
+
+	const double StartTime = FPlatformTime::Seconds();
+
+	// Smaller size to stay within int32 limits for in-memory array: 2GB, which is close to int32::Max (MAX_int32 = 2147483647)
+	constexpr uint64 DataSizeInGB = 2;
+	constexpr uint64 Chunks = DataSizeInGB * 1024;
+	constexpr uint64 ChunkSize = 1024 * 1024;
+
+	TSharedRef<IHttpRequest> HttpRequest = CreateRequest();
+	HttpRequest->SetURL(UrlStreamDownload(Chunks, ChunkSize));
+	HttpRequest->SetVerb(TEXT("GET"));
+
+	HttpRequest->OnProcessRequestComplete().BindLambda([Chunks, ChunkSize](FHttpRequestPtr HttpRequest, FHttpResponsePtr HttpResponse, bool bSucceeded)
+	{
+		CHECK(!bSucceeded);
+		CHECK(HttpRequest->GetStatus() == EHttpRequestStatus::Failed);
+		CHECK(HttpRequest->GetFailureReason() == EHttpFailureReason::ResponseTooLarge);
+		// REQUIRE(HttpResponse == nullptr); // TODO: HttpResponse needs to be NULL when failed due to ResponseTooLarge
+	});
+	HttpRequest->ProcessRequest();
+
+	WaitUntilAllHttpRequestsComplete();
+
+	const double DurationInSeconds = FPlatformTime::Seconds() - StartTime;
+
+	UE_LOG(LogHttp, Log, TEXT("%llu GB download without stream delegate spent %.2fs to complete!"), DataSizeInGB, DurationInSeconds);
+}
+
 TEST_CASE_METHOD(FWaitUntilCompleteHttpFixture, "Streaming http upload from memory", HTTP_TAG)
 {
 	TSharedRef<IHttpRequest> HttpRequest = CreateRequest();
diff --git a/Engine/Source/Runtime/Online/HTTP/Private/Curl/CurlHttp.cpp b/Engine/Source/Runtime/Online/HTTP/Private/Curl/CurlHttp.cpp
--- a/Engine/Source/Runtime/Online/HTTP/Private/Curl/CurlHttp.cpp
+++ b/Engine/Source/Runtime/Online/HTTP/Private/Curl/CurlHttp.cpp
@@ -570,13 +570,13 @@ size_t FCurlHttpRequest::ReceiveResponseBodyCallback(void* Ptr, size_t SizeInBlo
 	// callback function, it will signal an error condition to the library. This will cause the transfer 
 	// to get aborted and the libcurl function used will return CURLE_WRITE_ERROR.
 	size_t NumberOfBytesProcessed = 0;
-	  
-	uint64 SizeToDownload = SizeInBlocks * BlockSizeInBytes;
-
+	const size_t SizeToDownload = SizeInBlocks * BlockSizeInBytes;
+	const int64 TotalRead = TotalBytesRead.load();
+	const size_t CurrentBytesRead = TotalRead > 0 ? TotalRead : 0;
 	TSharedPtr<FCurlHttpResponse> Response = StaticCastSharedPtr<FCurlHttpResponse>(ResponseCommon);
 	UE_LOG(LogHttp, Verbose, TEXT("%p: ReceiveResponseBodyCallback: %llu bytes out of %llu received. (SizeInBlocks=%llu, BlockSizeInBytes=%llu, TotalBytesRead=%llu, Response->GetContentLength()=%llu, SizeToDownload=%llu (<-this will get returned from the callback))"),
-		this, TotalBytesRead.load() + SizeToDownload, Response->GetContentLength(),
-		SizeInBlocks, BlockSizeInBytes, TotalBytesRead.load(), Response->GetContentLength(), SizeToDownload);
+		this, CurrentBytesRead + SizeToDownload, Response->GetContentLength(),
+		SizeInBlocks, BlockSizeInBytes, CurrentBytesRead, Response->GetContentLength(), SizeToDownload);
 
 	// note that we can be passed 0 bytes if file transmitted has 0 length
 	if (SizeToDownload == 0)
@@ -599,9 +599,44 @@ size_t FCurlHttpRequest::ReceiveResponseBodyCallback(void* Ptr, size_t SizeInBlo
 	}
 	else
 	{
-		Response->Payload.AddUninitialized(SizeToDownload);
-		check(TotalBytesRead.load() + SizeToDownload <= Response->Payload.Num());
-		FMemory::Memcpy(static_cast<uint8*>(Response->Payload.GetData()) + TotalBytesRead.load(), Ptr, SizeToDownload);
+		const size_t NewTotalSize = CurrentBytesRead + SizeToDownload;
+		const uint64 ContentLength = Response->GetContentLength();
+		if (NewTotalSize > TNumericLimits<int32>::Max())
+		{
+			UE_LOG(LogHttp, Warning, TEXT("%p: ReceiveResponseBodyCallback: Payload size exceeds int32 limits: %llu bytes, consider using SetResponseBodyReceiveStream instead"), this, NewTotalSize);
+			SetFailureReason(EHttpFailureReason::ResponseTooLarge);
+			return 0;
+		}
+
+		// First time chunk added, reserve the capacity, if we know the content length, avoiding relocations
+		if (Response->Payload.IsEmpty() && ContentLength > 0 && ContentLength <= TNumericLimits<int32>::Max())
+		{
+			Response->Payload.Reserve(static_cast<int32>(ContentLength));
+		}
+
+		// Grow the array if needed
+		if (NewTotalSize > static_cast<size_t>(Response->Payload.Num()))
+		{
+			const size_t PayloadMax = Response->Payload.Max();
+			if (ContentLength == 0 && NewTotalSize > PayloadMax)
+			{
+				constexpr size_t SizeMultiplier = 2;
+				const size_t NewReserve = FMath::Max<size_t>(NewTotalSize, PayloadMax * SizeMultiplier);
+				if (NewReserve > TNumericLimits<int32>::Max())
+				{
+					UE_LOG(LogHttp, Warning, TEXT("%p: ReceiveResponseBodyCallback: Payload size exceeds int32 limits: TotalSize = %llu bytes, NewTotalSize = %llu, consider using SetResponseBodyReceiveStream instead"), this, NewTotalSize, NewReserve);
+					SetFailureReason(EHttpFailureReason::ResponseTooLarge);
+					return 0;
+				}
+				Response->Payload.Reserve(NewReserve);
+			}
+
+			Response->Payload.SetNumUninitialized(static_cast<int32>(NewTotalSize));
+		}
+
+		// Finally, copy the data
+		uint8* Dest = Response->Payload.GetData() + CurrentBytesRead;
+		FMemory::Memcpy(Dest, Ptr, SizeToDownload);
 
 		NumberOfBytesProcessed = SizeToDownload;
 	}
diff --git a/Engine/Source/Runtime/Online/HTTP/Public/Interfaces/IHttpBase.h b/Engine/Source/Runtime/Online/HTTP/Public/Interfaces/IHttpBase.h
--- a/Engine/Source/Runtime/Online/HTTP/Public/Interfaces/IHttpBase.h
+++ b/Engine/Source/Runtime/Online/HTTP/Public/Interfaces/IHttpBase.h
@@ -66,7 +66,8 @@ enum class EHttpFailureReason : uint8
 	// request timeout etc. It's not practical to unify to same behavior between them through their APIs. 
 	// When user's configured time is up, it's TimedOut no matter which step this Http request is in.
 	TimedOut,
-	Other
+	Other,
+	ResponseTooLarge
 };
 
 /** @return the stringified version of the enum passed in */
@@ -79,6 +80,7 @@ inline const TCHAR* LexToString(EHttpFailureReason HttpFailureReason)
 	case EHttpFailureReason::Cancelled: return TEXT("Cancelled");
 	case EHttpFailureReason::TimedOut: return TEXT("TimedOut");
 	case EHttpFailureReason::Other: return TEXT("Other");
+	case EHttpFailureReason::ResponseTooLarge: return TEXT("ResponseTooLarge");
 	default: checkNoEntry(); return TEXT("Invalid");
 	}
 }
