HTTP Client: Log Authorization headers as "[REDACTED]" rather than eliding them

Aids in determining whether authorization headers were present, without revealing their contents.

Changed the check for "Authorization" to non-case-sensitive.

#tests Real requests with Authorization headers
[RN] HTTP Client: Authorization headers are now logged as "[REDACTED]" rather than not logged at all.

diff --git a/Engine/Source/Runtime/Online/HTTP/Private/Curl/CurlHttp.cpp b/Engine/Source/Runtime/Online/HTTP/Private/Curl/CurlHttp.cpp
--- a/Engine/Source/Runtime/Online/HTTP/Private/Curl/CurlHttp.cpp
+++ b/Engine/Source/Runtime/Online/HTTP/Private/Curl/CurlHttp.cpp
@@ -997,28 +997,55 @@ bool FCurlHttpRequest::SetupRequestHttpThread()
 			const int32 NumAllHeaders = AllHeaders.Num();
 			for (int32 Idx = 0; Idx < NumAllHeaders; ++Idx)
 			{
-				const bool bCanLogHeaderValue = !AllHeaders[Idx].Contains(TEXT("Authorization"));
-				if (bCanLogHeaderValue)
-				{
-					UE_LOG(LogHttp, Verbose, TEXT("%p: Adding header '%s'"), this, *AllHeaders[Idx]);
-				}
+				const FString& FullHeader = AllHeaders[Idx];
+				// Note that these FStringViews are _not_ NUL-terminated, and must be printed with `"%.*s", View.Len(), View.GetData()`.
+				FStringView HeaderName;
+				FStringView Separator;
+				FStringView HeaderValue; // Including surrounding whitespace, if any.
 
-				curl_slist* NewHeaderList = curl_slist_append(HeaderList, TCHAR_TO_UTF8(*AllHeaders[Idx]));
-				if (!NewHeaderList)
 				{
-					if (bCanLogHeaderValue)
+					int32 ColonIndex = -1;
+					if (FullHeader.FindChar(':', ColonIndex))
 					{
-						UE_LOG(LogHttp, Warning, TEXT("Failed to append header '%s'"), *AllHeaders[Idx]);
+						FStringView FullHeaderView(FullHeader);
+						HeaderName = FullHeaderView.Left(ColonIndex);
+						Separator = FullHeaderView.Mid(ColonIndex, 1);
+						HeaderValue = FullHeaderView.RightChop(ColonIndex + 1);
 					}
 					else
 					{
-						UE_LOG(LogHttp, Warning, TEXT("Failed to append header 'Authorization'"));
+						HeaderName = FullHeader;
+					}
+				}
+
+				{
+					static const FStringView Authorization(FHttpConstants::HEADER_AUTHORIZATION);
+					if (HeaderName.Contains(Authorization, ESearchCase::IgnoreCase))
+					{
+						static constexpr TCHAR Redacted[] = TEXT(" [REDACTED]");
+						HeaderValue = Redacted;
 					}
 				}
-				else
+
+				UE_LOG(LogHttp, Verbose,
+					TEXT("%p: Adding header '%.*s%.*s%.*s'"),
+					this,
+					HeaderName.Len(), HeaderName.GetData(),
+					Separator.Len(), Separator.GetData(),
+					HeaderValue.Len(), HeaderValue.GetData());
+
+				curl_slist* NewHeaderList = curl_slist_append(HeaderList, TCHAR_TO_UTF8(ToCStr(FullHeader)));
+				if (!NewHeaderList)
 				{
-					HeaderList = NewHeaderList;
+					UE_LOG(LogHttp, Warning,
+						TEXT("Failed to append header '%.*s%.*s%.*s'"),
+						HeaderName.Len(), HeaderName.GetData(),
+						Separator.Len(), Separator.GetData(),
+						HeaderValue.Len(), HeaderValue.GetData());
+					continue;
 				}
+
+				HeaderList = NewHeaderList;
 			}
 		}
 
