Fix for chunkdb source not deleting the last chunkdb during delete-during-install

#rnx

diff --git a/Engine/Source/Runtime/Online/BuildPatchServices/Private/Installer/ChunkDbChunkSource.cpp b/Engine/Source/Runtime/Online/BuildPatchServices/Private/Installer/ChunkDbChunkSource.cpp
--- a/Engine/Source/Runtime/Online/BuildPatchServices/Private/Installer/ChunkDbChunkSource.cpp
+++ b/Engine/Source/Runtime/Online/BuildPatchServices/Private/Installer/ChunkDbChunkSource.cpp
@@ -117,7 +117,7 @@ namespace BuildPatchServices
 			for (FChunkDbDataAccess& ChunkDbDataAccess : ChunkDbDataAccesses)
 			{
 				if (!ChunkDbDataAccess.bIsRetired &&
-					ChunkDbDataAccess.RetireAt > RemainingChunkCount)
+					ChunkDbDataAccess.RetireAt >= RemainingChunkCount)
 				{
 					ChunkDbDataAccess.Retire(FileSystem, Configuration.bDeleteChunkDBAfterUse);
 				}
