UE - HTTP: Changing TimeStartedWaitingInQueue to double to avoid precision losses.
#tests webtests

diff --git a/Engine/Source/Runtime/Online/HTTP/Private/GenericPlatform/HttpRequestCommon.cpp b/Engine/Source/Runtime/Online/HTTP/Private/GenericPlatform/HttpRequestCommon.cpp
--- a/Engine/Source/Runtime/Online/HTTP/Private/GenericPlatform/HttpRequestCommon.cpp
+++ b/Engine/Source/Runtime/Online/HTTP/Private/GenericPlatform/HttpRequestCommon.cpp
@@ -647,9 +647,9 @@ void FHttpRequestCommon::StartWaitingInQueue()
 	TimeStartedWaitingInQueue = FPlatformTime::Seconds();
 }
 
-float FHttpRequestCommon::GetTimeStartedWaitingInQueue() const
+double FHttpRequestCommon::GetTimeStartedWaitingInQueue() const
 {
-	check(TimeStartedWaitingInQueue != 0);
+	check(TimeStartedWaitingInQueue != 0.0);
 	return TimeStartedWaitingInQueue;
 }
 
diff --git a/Engine/Source/Runtime/Online/HTTP/Private/GenericPlatform/HttpRequestCommon.h b/Engine/Source/Runtime/Online/HTTP/Private/GenericPlatform/HttpRequestCommon.h
--- a/Engine/Source/Runtime/Online/HTTP/Private/GenericPlatform/HttpRequestCommon.h
+++ b/Engine/Source/Runtime/Online/HTTP/Private/GenericPlatform/HttpRequestCommon.h
@@ -56,7 +56,7 @@ public:
 	virtual void TickThreadedRequest(float DeltaSeconds) = 0;
 
 	HTTP_API void StartWaitingInQueue();
-	HTTP_API float GetTimeStartedWaitingInQueue() const;
+	HTTP_API double GetTimeStartedWaitingInQueue() const;
 
 	HTTP_API virtual void SetURL(const FString& InURL) override;
 	HTTP_API virtual const FString& GetURL() const override;
@@ -198,7 +198,7 @@ protected:
 	float ElapsedTime = 0.0f;
 
 	/** Record the time started to wait in the queue */
-	float TimeStartedWaitingInQueue = 0.0f;
+	double TimeStartedWaitingInQueue = 0.0;
 
 	EHttpRequestPriority Priority = EHttpRequestPriority::Normal;
 };
diff --git a/Engine/Source/Runtime/Online/HTTP/Private/HttpThread.cpp b/Engine/Source/Runtime/Online/HTTP/Private/HttpThread.cpp
--- a/Engine/Source/Runtime/Online/HTTP/Private/HttpThread.cpp
+++ b/Engine/Source/Runtime/Online/HTTP/Private/HttpThread.cpp
@@ -325,9 +325,9 @@ void FHttpThreadBase::StartRequestsWaitingInQueue(TArray<FHttpRequestCommon*>& R
 			FHttpRequestCommon* ReadyThreadedRequest = RateLimitedThreadedRequests[0];
 			RateLimitedThreadedRequests.RemoveAt(0);
 
-			float DurationInQueue = FPlatformTime::Seconds() - ReadyThreadedRequest->GetTimeStartedWaitingInQueue();
+			const double DurationInQueue = FPlatformTime::Seconds() - ReadyThreadedRequest->GetTimeStartedWaitingInQueue();
 			UE_CLOG(!FApp::IsUnattended() && DurationInQueue > CVarHttpDurationInQueueToWarnInSeconds.GetValueOnAnyThread(), LogHttp, Warning, TEXT("Request (%p) %s %s waited in queue for %.2fs before starting"), ReadyThreadedRequest, *ReadyThreadedRequest->GetVerb(), *ReadyThreadedRequest->GetURL(), DurationInQueue);
-			float StartImmediately = 0.01f;
+			constexpr float StartImmediately = 0.01f;
 			if (DurationInQueue > StartImmediately)
 			{
 				HttpManager.RecordMaxTimeToWaitInQueue(DurationInQueue);
