When check if http request needs to log http response, use Contains instead of Compare when check if the response is in json encoding format.

diff --git a/Engine/Source/Runtime/Online/HTTP/Private/GenericPlatform/HttpRequestCommon.cpp b/Engine/Source/Runtime/Online/HTTP/Private/GenericPlatform/HttpRequestCommon.cpp
--- a/Engine/Source/Runtime/Online/HTTP/Private/GenericPlatform/HttpRequestCommon.cpp
+++ b/Engine/Source/Runtime/Online/HTTP/Private/GenericPlatform/HttpRequestCommon.cpp
@@ -770,7 +770,8 @@ void FHttpRequestCommon::LogResponse(const TSharedPtr<IHttpResponse>& InResponse
 
 	if (UE::HttpRequestCommon::Private::CVarHttpLogJsonResponseOnly.GetValueOnAnyThread())
 	{
-		bool bIsContentTypeJson = !InResponse->GetHeader(TEXT("Content-Type")).Compare(TEXT("application/json"), ESearchCase::IgnoreCase);
+		// Use Contains instead of Compare, for response header like "Content-Type: application/json;charset=utf-8"
+		bool bIsContentTypeJson = InResponse->GetHeader(TEXT("Content-Type")).Contains(TEXT("application/json"), ESearchCase::IgnoreCase);
 		if (!bIsContentTypeJson)
 			return;
 	}
