Implement ETag and Cache-Control Support for Cosmo Downloader Cache Validation

#rnx

diff --git a/Engine/Source/Runtime/Online/HTTP/Private/Curl/CurlHttp.cpp b/Engine/Source/Runtime/Online/HTTP/Private/Curl/CurlHttp.cpp
--- a/Engine/Source/Runtime/Online/HTTP/Private/Curl/CurlHttp.cpp
+++ b/Engine/Source/Runtime/Online/HTTP/Private/Curl/CurlHttp.cpp
@@ -538,7 +538,7 @@ size_t FCurlHttpRequest::ReceiveResponseHeaderCallback(void* Ptr, size_t SizeInB
 				long HttpCode = 0;
 				if (CURLE_OK == curl_easy_getinfo(EasyHandle, CURLINFO_RESPONSE_CODE, &HttpCode))
 				{
-					bRedirected = (HttpCode >= 300 && HttpCode < 400);
+					bRedirected = (HttpCode >= 300 && HttpCode < 400) && (HttpCode != 304);
 
 					if (!bRedirected)
 					{
