BPS: Add clarity to error case where downloaded chunk is corrupt

#rnx

diff --git a/Engine/Source/Runtime/Online/BuildPatchServices/Private/Data/ChunkData.cpp b/Engine/Source/Runtime/Online/BuildPatchServices/Private/Data/ChunkData.cpp
--- a/Engine/Source/Runtime/Online/BuildPatchServices/Private/Data/ChunkData.cpp
+++ b/Engine/Source/Runtime/Online/BuildPatchServices/Private/Data/ChunkData.cpp
@@ -9,6 +9,7 @@
 #include "Common/FileSystem.h"
 #include "BuildPatchHash.h"
 #include "BuildPatchManifest.h"
+#include "BuildPatchServicesPrivate.h"
 
 // The chunk header magic codeword, for quick checking that the opened file is a chunk file.
 #define CHUNK_HEADER_MAGIC      0xB1FE3AA2
@@ -529,26 +530,36 @@ namespace BuildPatchServices
 				((OutHeader.StoredAs & EChunkStorageFlags::Encrypted) != EChunkStorageFlags::None) || // Must not be encrypted
 				(InDestinationBuffer.GetSize() > OutHeader.DataSizeUncompressed)) // Unexpected buffer size.
 			{
+				UE_LOG(LogBuildPatchServices, Error, TEXT("Failed to validate chunk. Header reported Guid=%s HashType=%d StoredAs=%d UncompressedSize=%d DestinationSize=%llu"),
+					*WriteToString<64>(OutHeader.Guid),
+					OutHeader.HashType,
+					OutHeader.StoredAs,
+					OutHeader.DataSizeUncompressed,
+					InDestinationBuffer.GetSize());
 				return false;
 			}
 
 			// If the chunk isn't compressed and is the entire block, we can read directly in to the destination buffer.
 			// We seem to be getting chunkinfos that report window sizes that are not the full chunk size, so we need to be
 			// able to catch these subchunks and copy them out as needed.
+			uint64 ReadSize;
 			if ((OutHeader.StoredAs & EChunkStorageFlags::Compressed) == EChunkStorageFlags::None &&
 				InDestinationBuffer.GetSize() == OutHeader.DataSizeUncompressed)
 			{
+				ReadSize = InDestinationBuffer.GetSize();
 				InArchive.Serialize(InDestinationBuffer.GetData(), InDestinationBuffer.GetSize());
 			}
 			else
 			{
 				// We need to allocate and read in to another buffer for decompression.
+				ReadSize = OutHeader.DataSizeCompressed;
 				OutCompressedBuffer = FUniqueBuffer::Alloc(OutHeader.DataSizeCompressed);
 				InArchive.Serialize(OutCompressedBuffer.GetData(), OutCompressedBuffer.GetSize());
 			}
 
 			if (InArchive.IsError())
 			{
+				UE_LOG(LogBuildPatchServices, Error, TEXT("Archive error while deserializing chunk %s, ReadSize = %llu"), *WriteToString<64>(OutHeader.Guid), ReadSize);
 				return false;
 			}
 			return true;
diff --git a/Engine/Source/Runtime/Online/BuildPatchServices/Private/Installer/CloudChunkSource.cpp b/Engine/Source/Runtime/Online/BuildPatchServices/Private/Installer/CloudChunkSource.cpp
--- a/Engine/Source/Runtime/Online/BuildPatchServices/Private/Installer/CloudChunkSource.cpp
+++ b/Engine/Source/Runtime/Online/BuildPatchServices/Private/Installer/CloudChunkSource.cpp
@@ -1202,7 +1202,8 @@ namespace BuildPatchServices
 
 					if (!Result)
 					{
-						// The header or chunk data was bad.				
+						// The header or chunk data was bad.
+						UE_LOG(LogCloudChunkSource, Error, TEXT("Failed to ValidateAndRead chunk %s, download size %d"), *WriteToString<64>(Read->DataId), DownloadedData.Num());
 						Read->CompleteFn.Execute(Read->DataId, false, true, Read->UserPtr);
 						delete Read;
 					}
@@ -1389,6 +1390,7 @@ namespace BuildPatchServices
 			Read->ExpectedSize = ManifestSet->GetDownloadSize(Read->DataId);
 			Read->SecondsAtRequested = FStatsCollector::GetSeconds();
 			Read->UrlUsed = ChunkUriRequest.CloudDirectory / ChunkUriRequest.RelativePath;
+			UE_LOG(LogCloudChunkSource, VeryVerbose, TEXT("Downloading: %s"), *Read->UrlUsed);
 
 			InFlightURLRequests.Add(Read->ReadId, Read);
 			bool bMessageSent = MessagePump->SendRequest(ChunkUriRequest, [this, ReadId = Read->ReadId](FChunkUriResponse Response)
