Add more stats counter to trace http request adding.

diff --git a/Engine/Source/Runtime/Online/HTTP/Private/HttpManager.cpp b/Engine/Source/Runtime/Online/HTTP/Private/HttpManager.cpp
--- a/Engine/Source/Runtime/Online/HTTP/Private/HttpManager.cpp
+++ b/Engine/Source/Runtime/Online/HTTP/Private/HttpManager.cpp
@@ -566,6 +566,7 @@ void FHttpManager::FlushTick(float DeltaSeconds)
 void FHttpManager::AddRequest(const FHttpRequestRef& Request)
 {
 	{
+		QUICK_SCOPE_CYCLE_COUNTER(STAT_FHttpManager_AddRequest);
 		FScopeLock ScopeLock(&RequestLock);
 		UE_CLOG(bFlushing, LogHttp, Warning, TEXT("Adding request %s to http manager while flushing"), *Request->GetURL());
 		Requests.Add(Request);
diff --git a/Engine/Source/Runtime/Online/HTTP/Private/HttpThread.cpp b/Engine/Source/Runtime/Online/HTTP/Private/HttpThread.cpp
--- a/Engine/Source/Runtime/Online/HTTP/Private/HttpThread.cpp
+++ b/Engine/Source/Runtime/Online/HTTP/Private/HttpThread.cpp
@@ -170,6 +170,7 @@ void FHttpThreadBase::StopThread()
 
 void FHttpThreadBase::AddRequest(FHttpRequestCommon* Request)
 {
+	QUICK_SCOPE_CYCLE_COUNTER(STAT_FHttpThreadBase_AddRequest);
 	NewThreadedRequests.Enqueue(Request);
 }
 
