Remove the redundant log of response headers in CurlHttp. It's already logged by FHttpRequestCommon::LogResponse.

diff --git a/Engine/Source/Runtime/Online/HTTP/Private/Curl/CurlHttp.cpp b/Engine/Source/Runtime/Online/HTTP/Private/Curl/CurlHttp.cpp
--- a/Engine/Source/Runtime/Online/HTTP/Private/Curl/CurlHttp.cpp
+++ b/Engine/Source/Runtime/Online/HTTP/Private/Curl/CurlHttp.cpp
@@ -481,8 +481,6 @@ size_t FCurlHttpRequest::ReceiveResponseHeaderCallback(void* Ptr, size_t SizeInB
 		Header = Header.Replace(TEXT("\n"), TEXT(""));
 		Header = Header.Replace(TEXT("\r"), TEXT(""));
 
-		UE_LOG(LogHttp, Verbose, TEXT("%p: Received response header '%s'."), this, *Header);
-
 		FString HeaderKey, HeaderValue;
 		if (Header.Split(TEXT(":"), &HeaderKey, &HeaderValue))
 		{
