UE: safe divisor for Bandwidth in Http : BroadcastHttpRequestCompleted

#tests webtests

diff --git a/Engine/Source/Runtime/Online/HTTP/Private/HttpManager.cpp b/Engine/Source/Runtime/Online/HTTP/Private/HttpManager.cpp
--- a/Engine/Source/Runtime/Online/HTTP/Private/HttpManager.cpp
+++ b/Engine/Source/Runtime/Online/HTTP/Private/HttpManager.cpp
@@ -661,11 +661,19 @@ void FHttpManager::BroadcastHttpRequestCompleted(const FHttpRequestRef& Request)
 		HttpStatsHistory.TotalDownloadedBytes += SizeBytes;
 		HttpStatsHistory.DownloadedBytes[HttpStatsHistory.HistoryIndex] = SizeBytes;
 
-		HttpStats.BandwidthMbps = ((HttpStatsHistory.TotalDownloadedBytes * 8) / (HttpStatsHistory.TotalDuration + 1) / 1000);
-		HttpStats.HttpDurationMsAvg = HttpStatsHistory.TotalDuration / FHttpStatsHistory::HttpHistoryCount;
-
-		// Increment index
-		HttpStatsHistory.HistoryIndex = (HttpStatsHistory.HistoryIndex + 1) % FHttpStatsHistory::HttpHistoryCount;
+		const int64 TotalDuration = FMath::Max<int64>(0, HttpStatsHistory.TotalDuration);
+		const int64 DurationDivisor = TotalDuration + 1;
+		HttpStats.BandwidthMbps = (HttpStatsHistory.TotalDownloadedBytes * 8) / DurationDivisor / 1000;
+		if (FHttpStatsHistory::HttpHistoryCount > 0)
+		{
+			HttpStats.HttpDurationMsAvg = TotalDuration / FHttpStatsHistory::HttpHistoryCount;
+			// Increment index
+			HttpStatsHistory.HistoryIndex = (HttpStatsHistory.HistoryIndex + 1) % FHttpStatsHistory::HttpHistoryCount;
+		}
+		else
+		{
+			HttpStats.HttpDurationMsAvg = 0;
+		}
 	}
 		
 	RequestCompletedDelegate.ExecuteIfBound(Request);
