Decrease logspam by only logging a cdn swap if it actually swapped cdns.

#rnx

diff --git a/Engine/Source/Runtime/Online/BuildPatchServices/Private/Installer/CloudChunkSource.cpp b/Engine/Source/Runtime/Online/BuildPatchServices/Private/Installer/CloudChunkSource.cpp
--- a/Engine/Source/Runtime/Online/BuildPatchServices/Private/Installer/CloudChunkSource.cpp
+++ b/Engine/Source/Runtime/Online/BuildPatchServices/Private/Installer/CloudChunkSource.cpp
@@ -1272,7 +1272,10 @@ namespace BuildPatchServices
 						}
 
 						CurrentBestCloudDir = MinAtIndex;
-						UE_LOG(LogCloudChunkSource, Warning, TEXT("CDN %s failed download, updating CDN selection to: %s"), *Configuration.CloudRoots[Read->CloudDirUsed], *Configuration.CloudRoots[MinAtIndex]);
+						if (Read->CloudDirUsed != MinAtIndex)
+						{
+							UE_LOG(LogCloudChunkSource, Warning, TEXT("CDN %s failed download, updating CDN selection to: %s"), *Configuration.CloudRoots[Read->CloudDirUsed], *Configuration.CloudRoots[MinAtIndex]);
+						}
 
 						MessagePump->SendMessage({FGenericMessage::EType::CDNDownloadFailed, Read->DataId, Configuration.CloudRoots[Read->CloudDirUsed], Configuration.CloudRoots[MinAtIndex]});
 					}
