Move request headers into request common;
Log request headers and content as well in request common.

#Synced-CL 48203440

diff --git a/Engine/Source/Programs/WebTests/Private/TestHttp.cpp b/Engine/Source/Programs/WebTests/Private/TestHttp.cpp
--- a/Engine/Source/Programs/WebTests/Private/TestHttp.cpp
+++ b/Engine/Source/Programs/WebTests/Private/TestHttp.cpp
@@ -47,6 +47,8 @@ extern TAutoConsoleVariable<int32> CVarHttpMaxConcurrentRequests;
 extern TAutoConsoleVariable<int32> CVarHttpEventLoopEnableChance;
 extern TAutoConsoleVariable<FString> CVarHttpUrlPatternsToMockResponse;
 extern TAutoConsoleVariable<FString> CVarHttpUrlPatternsToDisableFailedLog;
+extern TAutoConsoleVariable<FString> CVarHttpUrlPatternsToInspect;
+extern TAutoConsoleVariable<bool> CVarHttpLogJsonResponseOnly;
 
 class FMockHttpModule : public FHttpModule
 {
@@ -452,6 +454,48 @@ TEST_CASE_METHOD(FWaitUntilCompleteHttpFixture, "Http methods with payloads", HT
 	HttpRequest->ProcessRequest();
 }
 
+TEST_CASE_METHOD(FWaitUntilCompleteHttpFixture, "Http request can LogContent correctly with different strings", HTTP_TAG)
+{
+	DisableWarningsInThisTest();
+
+	CVarHttpUrlPatternsToInspect->Set(TEXT("methods"));
+	CVarHttpLogJsonResponseOnly->Set(false);
+
+	TSharedRef<IHttpRequest> HttpRequest = CreateRequest();
+	HttpRequest->SetURL(UrlToTestMethods());
+	HttpRequest->SetVerb(FHttpConstants::VERB_POST);
+
+	const TCHAR* Content = GENERATE(
+		TEXT(""),
+		TEXT("\n"),
+		TEXT("\n\n"),
+		TEXT("\r"),
+		TEXT("\r\r"),
+		TEXT("\n\r"),
+		TEXT("\r\n"),
+		TEXT("l"),
+		TEXT("li"),
+		TEXT("line"),
+		TEXT("line1\r\nl"),
+		TEXT("line1\n\rline2"),
+		TEXT("line1\r\nline2"),
+		TEXT("line1\r\nline2\r"),
+		TEXT("line1\r\nline2\n"),
+		TEXT("line1\r\nline2\r\n"),
+		TEXT("line1\r\nline2\n\r")
+	);
+	HttpRequest->SetContentAsString(Content);
+	HttpRequest->SetHeader(FHttpConstants::HEADER_CONTENT_TYPE, FHttpConstants::MEDIATYPE_JSON);
+
+	HttpRequest->OnProcessRequestComplete().BindLambda([](FHttpRequestPtr HttpRequest, FHttpResponsePtr HttpResponse, bool bSucceeded) {
+		// TODO: Find a way to validate output logs
+
+		CVarHttpUrlPatternsToInspect->Set(TEXT(""));
+		CVarHttpLogJsonResponseOnly->Set(true);
+	});
+	HttpRequest->ProcessRequest();
+}
+
 TEST_CASE_METHOD(FWaitUntilCompleteHttpFixture, "Http methods which support to upload payload can work well without any payload", HTTP_TAG)
 {
 	TSharedRef<IHttpRequest> HttpRequest = CreateRequest();
diff --git a/Engine/Source/Runtime/Online/HTTP/Private/Apple/AppleHttp.cpp b/Engine/Source/Runtime/Online/HTTP/Private/Apple/AppleHttp.cpp
--- a/Engine/Source/Runtime/Online/HTTP/Private/Apple/AppleHttp.cpp
+++ b/Engine/Source/Runtime/Online/HTTP/Private/Apple/AppleHttp.cpp
@@ -463,19 +463,12 @@ FAppleHttpRequest::~FAppleHttpRequest()
     [Session release];
 }
 
-FString FAppleHttpRequest::GetHeader(const FString& HeaderName) const
-{
-	SCOPED_AUTORELEASE_POOL;
-	FString Header([Request valueForHTTPHeaderField:HeaderName.GetNSString()]);
-	return Header;
-}
-
-
 void FAppleHttpRequest::SetHeader(const FString& HeaderName, const FString& HeaderValue)
 {
 	SCOPED_AUTORELEASE_POOL;
-	UE_LOG(LogHttp, Verbose, TEXT("FAppleHttpRequest::SetHeader() - %s / %s"), *HeaderName, *HeaderValue );
 	[Request setValue: HeaderValue.GetNSString() forHTTPHeaderField: HeaderName.GetNSString()];
+
+	FHttpRequestCommon::SetHeader(HeaderName, HeaderValue);
 }
 
 void FAppleHttpRequest::AppendToHeader(const FString& HeaderName, const FString& AdditionalHeaderValue)
@@ -725,7 +718,7 @@ bool FAppleHttpRequest::SetupRequest()
 	if(GetContentLength() > 0)
 	{
 		UE_LOG(LogHttp, VeryVerbose, TEXT("Setting content length: %d"), GetContentLength());
-		[Request setValue:[NSString stringWithFormat:@"%llu", GetContentLength()] forHTTPHeaderField:@"Content-Length"];
+		SetHeader(TEXT("Content-Length"), FString::Printf(TEXT("%llu"), GetContentLength()));
 	}
 
 	PostProcess();
@@ -943,7 +936,7 @@ bool FAppleHttpRequest::StartThreadedRequest()
 
 	[[Task retain] resume];
 	
-	return true;
+	return FFHttpRequestCommon::StartThreadedRequest();
 }
 
 bool FAppleHttpRequest::IsThreadedRequestComplete()
diff --git a/Engine/Source/Runtime/Online/HTTP/Private/Apple/AppleHttp.h b/Engine/Source/Runtime/Online/HTTP/Private/Apple/AppleHttp.h
--- a/Engine/Source/Runtime/Online/HTTP/Private/Apple/AppleHttp.h
+++ b/Engine/Source/Runtime/Online/HTTP/Private/Apple/AppleHttp.h
@@ -25,7 +25,6 @@ public:
 	friend class FAppleHttpResponse;
 
 	//~ Begin IHttpBase Interface
-	virtual FString GetHeader(const FString& HeaderName) const override;
 	virtual TArray<FString> GetAllHeaders() const override;	
 	virtual FString GetContentType() const override;
 	virtual uint64 GetContentLength() const override;
diff --git a/Engine/Source/Runtime/Online/HTTP/Private/Curl/CurlHttp.cpp b/Engine/Source/Runtime/Online/HTTP/Private/Curl/CurlHttp.cpp
--- a/Engine/Source/Runtime/Online/HTTP/Private/Curl/CurlHttp.cpp
+++ b/Engine/Source/Runtime/Online/HTTP/Private/Curl/CurlHttp.cpp
@@ -224,19 +224,6 @@ FCurlHttpRequest::~FCurlHttpRequest()
 	}
 }
 
-FString FCurlHttpRequest::GetHeader(const FString& HeaderName) const
-{
-	FString Result;
-
-	const FString* Header = Headers.Find(HeaderName);
-	if (Header != NULL)
-	{
-		Result = *Header;
-	}
-	
-	return Result;
-}
-
 FString FCurlHttpRequest::CombineHeaderKeyValue(const FString& HeaderKey, const FString& HeaderValue)
 {
 	FString Combined;
@@ -252,8 +239,8 @@ FString FCurlHttpRequest::CombineHeaderKeyValue(const FString& HeaderKey, const
 TArray<FString> FCurlHttpRequest::GetAllHeaders() const
 {
 	TArray<FString> Result;
-	Result.Reserve(Headers.Num());
-	for (const TPair<FString, FString>& It : Headers)
+	Result.Reserve(RequestHeaders.Num());
+	for (const TPair<FString, FString>& It : RequestHeaders)
 	{
 		Result.Emplace(CombineHeaderKeyValue(It.Key, It.Value));
 	}
@@ -359,17 +346,6 @@ bool FCurlHttpRequest::SetContentFromStream(TSharedRef<FArchive, ESPMode::Thread
 	return true;
 }
 
-void FCurlHttpRequest::SetHeader(const FString& HeaderName, const FString& HeaderValue)
-{
-	if (CompletionStatus == EHttpRequestStatus::Processing)
-	{
-		UE_LOG(LogHttp, Warning, TEXT("FCurlHttpRequest::SetHeader() - attempted to set header on a request that is inflight"));
-		return;
-	}
-
-	Headers.Add(HeaderName, HeaderValue);
-}
-
 void FCurlHttpRequest::AppendToHeader(const FString& HeaderName, const FString& AdditionalHeaderValue)
 {
 	if (CompletionStatus == EHttpRequestStatus::Processing)
@@ -380,7 +356,7 @@ void FCurlHttpRequest::AppendToHeader(const FString& HeaderName, const FString&
 
 	if (!HeaderName.IsEmpty() && !AdditionalHeaderValue.IsEmpty())
 	{
-		FString* PreviousValue = Headers.Find(HeaderName);
+		FString* PreviousValue = RequestHeaders.Find(HeaderName);
 		FString NewValue;
 		if (PreviousValue != nullptr && !PreviousValue->IsEmpty())
 		{
@@ -857,7 +833,7 @@ bool FCurlHttpRequest::SetupRequest()
 	CurlAddToMultiResult = CURLM_OK;
 	LastReportedBytesSent = 0;
 
-	UE_LOG(LogHttp, Verbose, TEXT("%p: Custom headers are %s"), this, Headers.Num() ? TEXT("present") : TEXT("NOT present"));
+	UE_LOG(LogHttp, Verbose, TEXT("%p: Custom headers are %s"), this, RequestHeaders.Num() ? TEXT("present") : TEXT("NOT present"));
 	UE_LOG(LogHttp, Verbose, TEXT("%p: Payload size=%llu"), this, RequestPayload->GetContentLength());
 
 	// content-length should be present http://www.w3.org/Protocols/rfc2616/rfc2616-sec4.html#sec4.4
@@ -879,17 +855,17 @@ bool FCurlHttpRequest::SetupRequest()
 	return true;
 }
 
-bool FCurlHttpRequest::SetupRequestHttpThread()
+bool FCurlHttpRequest::StartThreadedRequest()
 {
-	QUICK_SCOPE_CYCLE_COUNTER(STAT_FCurlHttpRequest_SetupRequestHttpThread);
+	QUICK_SCOPE_CYCLE_COUNTER(STAT_FCurlHttpRequest_StartThreadedRequest);
 	{
-		QUICK_SCOPE_CYCLE_COUNTER(STAT_FCurlHttpRequest_SetupRequest_SLIST_FREE_HEADERS);
+		QUICK_SCOPE_CYCLE_COUNTER(STAT_FCurlHttpRequest_StartThreadedRequest_SLIST_FREE_HEADERS);
 		curl_slist_free_all(HeaderList);
 		HeaderList = nullptr;
 	}
 
 	{
-		QUICK_SCOPE_CYCLE_COUNTER(STAT_FCurlHttpRequest_SetupRequest_EASY_SETOPT);
+		QUICK_SCOPE_CYCLE_COUNTER(STAT_FCurlHttpRequest_StartThreadedRequest_EASY_SETOPT);
 
 		curl_easy_setopt(EasyHandle, CURLOPT_URL, TCHAR_TO_ANSI(*GetURL()));
 
@@ -989,57 +965,15 @@ bool FCurlHttpRequest::SetupRequestHttpThread()
 		}
 
 		{
-			QUICK_SCOPE_CYCLE_COUNTER(STAT_FCurlHttpRequest_SetupRequest_SLIST_APPEND_HEADERS);
+			QUICK_SCOPE_CYCLE_COUNTER(STAT_FCurlHttpRequest_StartThreadedRequest_SLIST_APPEND_HEADERS);
 
-			TArray<FString> AllHeaders = GetAllHeaders();
-			const int32 NumAllHeaders = AllHeaders.Num();
-			for (int32 Idx = 0; Idx < NumAllHeaders; ++Idx)
+			for (const TPair<FString, FString>& HeaderPair : RequestHeaders)
 			{
-				const FString& FullHeader = AllHeaders[Idx];
-				// Note that these FStringViews are _not_ NUL-terminated, and must be printed with `"%.*s", View.Len(), View.GetData()`.
-				FStringView HeaderName;
-				FStringView Separator;
-				FStringView HeaderValue; // Including surrounding whitespace, if any.
-
-				{
-					int32 ColonIndex = -1;
-					if (FullHeader.FindChar(':', ColonIndex))
-					{
-						FStringView FullHeaderView(FullHeader);
-						HeaderName = FullHeaderView.Left(ColonIndex);
-						Separator = FullHeaderView.Mid(ColonIndex, 1);
-						HeaderValue = FullHeaderView.RightChop(ColonIndex + 1);
-					}
-					else
-					{
-						HeaderName = FullHeader;
-					}
-				}
-
-				{
-					static const FStringView Authorization(FHttpConstants::HEADER_AUTHORIZATION);
-					if (HeaderName.Contains(Authorization, ESearchCase::IgnoreCase))
-					{
-						static constexpr TCHAR Redacted[] = TEXT(" [REDACTED]");
-						HeaderValue = Redacted;
-					}
-				}
-
-				UE_LOG(LogHttp, Verbose,
-					TEXT("%p: Adding header '%.*s%.*s%.*s'"),
-					this,
-					HeaderName.Len(), HeaderName.GetData(),
-					Separator.Len(), Separator.GetData(),
-					HeaderValue.Len(), HeaderValue.GetData());
-
+				FString FullHeader = FString::Format(TEXT("{0}: {1}"), { HeaderPair.Key, HeaderPair.Value });
 				curl_slist* NewHeaderList = curl_slist_append(HeaderList, TCHAR_TO_UTF8(ToCStr(FullHeader)));
 				if (!NewHeaderList)
 				{
-					UE_LOG(LogHttp, Warning,
-						TEXT("Failed to append header '%.*s%.*s%.*s'"),
-						HeaderName.Len(), HeaderName.GetData(),
-						Separator.Len(), Separator.GetData(),
-						HeaderValue.Len(), HeaderValue.GetData());
+					UE_LOG(LogHttp, Warning, TEXT("Failed to append header '%s'"), *HeaderPair.Key);
 					continue;
 				}
 
@@ -1067,7 +1001,7 @@ bool FCurlHttpRequest::SetupRequestHttpThread()
 #if !WITH_CURL_XCURL
 	{
 		//Tracking the locking in the CURLOPT_SHARE branch of the curl_easy_setopt implementation
-		QUICK_SCOPE_CYCLE_COUNTER(STAT_FCurlHttpRequest_SetupRequest_EASY_CURLOPT_SHARE);
+		QUICK_SCOPE_CYCLE_COUNTER(STAT_FCurlHttpRequest_StartThreadedRequest_EASY_CURLOPT_SHARE);
 
 		curl_easy_setopt(EasyHandle, CURLOPT_SHARE, FCurlHttpManager::GShareHandle);
 	}
@@ -1078,9 +1012,9 @@ bool FCurlHttpRequest::SetupRequestHttpThread()
 	curl_easy_setopt(EasyHandle, CURLOPT_QUICK_EXIT, 1L);
 #endif
 
-	UE_LOG(LogHttp, Log, TEXT("%p: Starting %s request to URL='%s'"), this, *Verb, *GetURL());
-
-	return true;
+	ElapsedTime = 0.0f;
+	
+	return FHttpRequestCommon::StartThreadedRequest();
 }
 
 void FCurlHttpRequest::SetupOptionUnixSocketPath()
@@ -1167,16 +1101,6 @@ void FCurlHttpRequest::ClearInCaseOfRetry()
 	}
 }
 
-bool FCurlHttpRequest::StartThreadedRequest()
-{
-	// reset timeout
-	ElapsedTime = 0.0f;
-	
-	UE_LOG(LogHttp, Verbose, TEXT("%p: request (easy handle:%p) has started threaded processing"), this, EasyHandle);
-
-	return true;
-}
-
 bool FCurlHttpRequest::IsThreadedRequestComplete()
 {
 	if (bCurlRequestCompleted && ElapsedTime >= FHttpModule::Get().GetHttpDelayTime())
diff --git a/Engine/Source/Runtime/Online/HTTP/Private/Curl/CurlHttp.h b/Engine/Source/Runtime/Online/HTTP/Private/Curl/CurlHttp.h
--- a/Engine/Source/Runtime/Online/HTTP/Private/Curl/CurlHttp.h
+++ b/Engine/Source/Runtime/Online/HTTP/Private/Curl/CurlHttp.h
@@ -143,7 +143,6 @@ public:
 	friend class FCurlHttpResponse;
 
 	//~ Begin IHttpBase Interface
-	virtual FString GetHeader(const FString& HeaderName) const override;
 	virtual TArray<FString> GetAllHeaders() const override;
 	virtual FString GetContentType() const override;
 	virtual uint64 GetContentLength() const override;
@@ -159,7 +158,6 @@ public:
 	virtual void SetContentAsString(const FString& ContentString) override;
 	virtual bool SetContentAsStreamedFile(const FString& Filename) override;
 	virtual bool SetContentFromStream(TSharedRef<FArchive, ESPMode::ThreadSafe> Stream) override;
-	virtual void SetHeader(const FString& HeaderName, const FString& HeaderValue) override;
 	virtual void AppendToHeader(const FString& HeaderName, const FString& AdditionalHeaderValue) override;
 	virtual bool ProcessRequest() override;
 	virtual void Tick(float DeltaSeconds) override;
@@ -172,13 +170,6 @@ public:
 	virtual void TickThreadedRequest(float DeltaSeconds) override;
 	//~ End IHttpRequestThreaded Interface
 
-	/**
-	 * Perform the http-thread setup of the request
-	 *
-	 * @return true if the request was successfully setup
-	 */
-	bool SetupRequestHttpThread();
-
 	/**
 	 * Perform the http-thread cleanup of the request
 	 */
@@ -382,8 +373,6 @@ private:
 	CURLcode		CurlCompletionResult;
 	/** Is the request payload seekable? */
 	bool bIsRequestPayloadSeekable = false;
-	/** Mapping of header section to values. */
-	TMap<FString, FString> Headers;
 	/** Have we had any HTTP activity with the host? Sending headers, SSL handshake, etc */
 	bool bAnyHttpActivity;
 	/** Newly received headers we need to inform listeners about */
diff --git a/Engine/Source/Runtime/Online/HTTP/Private/Curl/CurlHttpThread.cpp b/Engine/Source/Runtime/Online/HTTP/Private/Curl/CurlHttpThread.cpp
--- a/Engine/Source/Runtime/Online/HTTP/Private/Curl/CurlHttpThread.cpp
+++ b/Engine/Source/Runtime/Online/HTTP/Private/Curl/CurlHttpThread.cpp
@@ -74,7 +74,7 @@ bool FCurlHttpThread::StartThreadedRequest(FHttpRequestCommon* Request)
 	CURL* EasyHandle = CurlRequest->GetEasyHandle();
 	ensure(!HandlesToRequests.Contains(EasyHandle));
 
-	if (!CurlRequest->SetupRequestHttpThread())
+	if (!FLegacyHttpThread::StartThreadedRequest(Request))
 	{
 		UE_LOG(LogHttp, Warning, TEXT("Could not set libcurl options for easy handle, processing HTTP request failed. Increase verbosity for additional information."));
 		return false;
@@ -91,7 +91,7 @@ bool FCurlHttpThread::StartThreadedRequest(FHttpRequestCommon* Request)
 
 	HandlesToRequests.Add(EasyHandle, Request);
 
-	return FLegacyHttpThread::StartThreadedRequest(Request);
+	return true;
 }
 
 void FCurlHttpThread::CompleteThreadedRequest(FHttpRequestCommon* Request)
diff --git a/Engine/Source/Runtime/Online/HTTP/Private/Curl/CurlMultiPollEventLoopHttpThread.cpp b/Engine/Source/Runtime/Online/HTTP/Private/Curl/CurlMultiPollEventLoopHttpThread.cpp
--- a/Engine/Source/Runtime/Online/HTTP/Private/Curl/CurlMultiPollEventLoopHttpThread.cpp
+++ b/Engine/Source/Runtime/Online/HTTP/Private/Curl/CurlMultiPollEventLoopHttpThread.cpp
@@ -74,7 +74,7 @@ bool FCurlMultiPollEventLoopHttpThread::StartThreadedRequest(FHttpRequestCommon*
 	CURL* EasyHandle = CurlRequest->GetEasyHandle();
 	ensure(!HandlesToRequests.Contains(EasyHandle));
 
-	if (!CurlRequest->SetupRequestHttpThread())
+	if (!FEventLoopHttpThread::StartThreadedRequest(Request))
 	{
 		UE_LOG(LogHttp, Warning, TEXT("Could not set libcurl options for easy handle, processing HTTP request failed. Increase verbosity for additional information."));
 		return false;
@@ -91,7 +91,7 @@ bool FCurlMultiPollEventLoopHttpThread::StartThreadedRequest(FHttpRequestCommon*
 
 	HandlesToRequests.Add(EasyHandle, Request);
 
-	return FEventLoopHttpThread::StartThreadedRequest(Request);
+	return true;
 }
 
 void FCurlMultiPollEventLoopHttpThread::CompleteThreadedRequest(FHttpRequestCommon* Request)
diff --git a/Engine/Source/Runtime/Online/HTTP/Private/Curl/CurlMultiWaitEventLoopHttpThread.cpp b/Engine/Source/Runtime/Online/HTTP/Private/Curl/CurlMultiWaitEventLoopHttpThread.cpp
--- a/Engine/Source/Runtime/Online/HTTP/Private/Curl/CurlMultiWaitEventLoopHttpThread.cpp
+++ b/Engine/Source/Runtime/Online/HTTP/Private/Curl/CurlMultiWaitEventLoopHttpThread.cpp
@@ -119,7 +119,7 @@ bool FCurlMultiWaitEventLoopHttpThread::StartThreadedRequest(FHttpRequestCommon*
 	CURL* EasyHandle = CurlRequest->GetEasyHandle();
 	ensure(!HandlesToRequests.Contains(EasyHandle));
 
-	if (!CurlRequest->SetupRequestHttpThread())
+	if (!FEventLoopHttpThread::StartThreadedRequest(Request))
 	{
 		UE_LOG(LogHttp, Warning, TEXT("Could not set libcurl options for easy handle, processing HTTP request failed. Increase verbosity for additional information."));
 		return false;
@@ -136,7 +136,7 @@ bool FCurlMultiWaitEventLoopHttpThread::StartThreadedRequest(FHttpRequestCommon*
 
 	HandlesToRequests.Add(EasyHandle, Request);
 
-	return FEventLoopHttpThread::StartThreadedRequest(Request);
+	return true;
 }
 
 void FCurlMultiWaitEventLoopHttpThread::CompleteThreadedRequest(FHttpRequestCommon* Request)
diff --git a/Engine/Source/Runtime/Online/HTTP/Private/Curl/CurlSocketEventLoopHttpThread.cpp b/Engine/Source/Runtime/Online/HTTP/Private/Curl/CurlSocketEventLoopHttpThread.cpp
--- a/Engine/Source/Runtime/Online/HTTP/Private/Curl/CurlSocketEventLoopHttpThread.cpp
+++ b/Engine/Source/Runtime/Online/HTTP/Private/Curl/CurlSocketEventLoopHttpThread.cpp
@@ -65,7 +65,7 @@ bool FCurlSocketEventLoopHttpThread::StartThreadedRequest(FHttpRequestCommon* Re
 	CURL* EasyHandle = CurlRequest->GetEasyHandle();
 	ensure(!HandlesToRequests.Contains(EasyHandle));
 
-	if (!CurlRequest->SetupRequestHttpThread())
+	if (!FEventLoopHttpThread::StartThreadedRequest(Request))
 	{
 		UE_LOG(LogHttp, Warning, TEXT("Could not set libcurl options for easy handle, processing HTTP request failed. Increase verbosity for additional information."));
 		return false;
@@ -82,7 +82,7 @@ bool FCurlSocketEventLoopHttpThread::StartThreadedRequest(FHttpRequestCommon* Re
 
 	HandlesToRequests.Add(EasyHandle, Request);
 
-	return FEventLoopHttpThread::StartThreadedRequest(Request);
+	return true;
 }
 
 void FCurlSocketEventLoopHttpThread::CompleteThreadedRequest(FHttpRequestCommon* Request)
diff --git a/Engine/Source/Runtime/Online/HTTP/Private/GenericPlatform/HttpRequestCommon.cpp b/Engine/Source/Runtime/Online/HTTP/Private/GenericPlatform/HttpRequestCommon.cpp
--- a/Engine/Source/Runtime/Online/HTTP/Private/GenericPlatform/HttpRequestCommon.cpp
+++ b/Engine/Source/Runtime/Online/HTTP/Private/GenericPlatform/HttpRequestCommon.cpp
@@ -5,14 +5,12 @@
 #include "HAL/Event.h"
 #include "HAL/IConsoleManager.h"
 #include "Http.h"
+#include "HttpConstants.h"
 #include "HttpManager.h"
 #include "Logging/StructuredLog.h"
 #include "Misc/CommandLine.h"
 #include "Stats/Stats.h"
 
-namespace UE::HttpRequestCommon::Private
-{
-
 TAutoConsoleVariable<bool> CVarHttpLogJsonResponseOnly(
 	TEXT("http.LogJsonResponseOnly"),
 	true,
@@ -20,8 +18,6 @@ TAutoConsoleVariable<bool> CVarHttpLogJsonResponseOnly(
 	ECVF_SaveForNextBoot
 );
 
-}
-
 FHttpRequestCommon::FHttpRequestCommon()
 	: RequestStartTimeAbsoluteSeconds(FPlatformTime::Seconds())
 	, ActivityTimeoutAt(0.0)
@@ -38,6 +34,30 @@ FString FHttpRequestCommon::GetURLParameter(const FString& ParameterName) const
 	return ReturnValue;
 }
 
+void FHttpRequestCommon::SetHeader(const FString& HeaderName, const FString& HeaderValue)
+{
+	if (CompletionStatus == EHttpRequestStatus::Processing)
+	{
+		UE_LOG(LogHttp, Warning, TEXT("FHttpRequestCommon::SetHeader() - attempted to set header on a request that is inflight"));
+		return;
+	}
+
+	RequestHeaders.Add(HeaderName, HeaderValue);
+}
+
+FString FHttpRequestCommon::GetHeader(const FString& HeaderName) const
+{
+	FString Result;
+
+	const FString* Header = RequestHeaders.Find(HeaderName);
+	if (Header != NULL)
+	{
+		Result = *Header;
+	}
+	
+	return Result;
+}
+
 EHttpRequestStatus::Type FHttpRequestCommon::GetStatus() const
 {
 	return CompletionStatus;
@@ -642,6 +662,12 @@ float FHttpRequestCommon::GetElapsedTime() const
 	return ElapsedTime;
 }
 
+bool FHttpRequestCommon::StartThreadedRequest()
+{
+	LogRequest();
+	return true;
+}
+
 void FHttpRequestCommon::StartWaitingInQueue()
 {
 	TimeStartedWaitingInQueue = FPlatformTime::Seconds();
@@ -749,34 +775,63 @@ void FHttpRequestCommon::CloseRequestPayloadDefaultImpl()
 	}
 }
 
-void FHttpRequestCommon::LogResponse(const TSharedPtr<IHttpResponse>& InResponse)
+void FHttpRequestCommon::LogRequest() const
+{
+	bool bShouldLogDetails = FHttpModule::Get().GetHttpManager().ShouldInspect(GetURL());
+	UE_HTTP_LOG_AS_WARNING_IF(bShouldLogDetails, TEXT("%p %s %s started from http thread"), this, *GetVerb(), *GetURL());
+
+	LogHeaders(RequestHeaders, bShouldLogDetails, false/*bIsResponse*/);
+	LogContent(bShouldLogDetails, *this, false/*bIsResponse*/);
+}
+
+void FHttpRequestCommon::LogResponse(const TSharedPtr<FHttpResponseCommon>& InResponse) const
 {
-	bool bShouldLogResponse = FHttpModule::Get().GetHttpManager().ShouldLogResponse(GetURL());
-	UE_HTTP_LOG_AS_WARNING_IF(bShouldLogResponse, TEXT("%p %s %s completed with code %d after %.2fs. Content length: %ld"), this, *GetVerb(), *GetURL(), InResponse->GetResponseCode(), ElapsedTime, InResponse->GetContentLength());
+	bool bShouldLogDetails = FHttpModule::Get().GetHttpManager().ShouldInspect(GetURL());
+	UE_HTTP_LOG_AS_WARNING_IF(bShouldLogDetails, TEXT("%p %s %s completed with code %d after %.2fs"), this, *GetVerb(), *GetURL(), InResponse->GetResponseCode(), ElapsedTime);
 
-	TArray<FString> AllHeaders = InResponse->GetAllHeaders();
-	for (const FString& HeaderStr : AllHeaders)
+	LogHeaders(InResponse->Headers, bShouldLogDetails, true/*bIsResponse*/);
+	LogContent(bShouldLogDetails, *InResponse, true/*bIsResponse*/);
+}
+
+void FHttpRequestCommon::LogHeaders(const TMap<FString, FString>& Headers, bool bShouldLogAsWarning, bool bIsResponse) const
+{
+	for (const TPair<FString, FString>& HeaderPair : Headers)
 	{
-		if (!HeaderStr.StartsWith(TEXT("Authorization")) && !HeaderStr.StartsWith(TEXT("Set-Cookie")))
-		{
-			UE_HTTP_LOG_AS_WARNING_IF(bShouldLogResponse, TEXT("%p Response Header %s"), this, *HeaderStr);
-		}
+		bool bIsSensitiveData = 
+			HeaderPair.Key.Contains(FHttpConstants::HEADER_AUTHORIZATION, ESearchCase::IgnoreCase) || 
+			HeaderPair.Key.Contains(TEXT("Set-Cookie"), ESearchCase::IgnoreCase);
+
+		UE_HTTP_LOG_AS_WARNING_IF(bShouldLogAsWarning, TEXT("%p %s Header %s:%s"), 
+			this,
+			bIsResponse ? TEXT("Response") : TEXT("Request"),
+			*HeaderPair.Key,
+			bIsSensitiveData ? TEXT("[REDACTED]") : *HeaderPair.Value);
 	}
+}
 
-	if (!bShouldLogResponse || InResponse->GetContentLength() == 0)
+void FHttpRequestCommon::LogContent(bool bShouldLogDetails, const IHttpBase& HttpBase, bool bIsResponse) const
+{
+	if (!bShouldLogDetails)
 	{
 		return;
 	}
 
-	if (UE::HttpRequestCommon::Private::CVarHttpLogJsonResponseOnly.GetValueOnAnyThread())
+	if (CVarHttpLogJsonResponseOnly.GetValueOnAnyThread())
 	{
 		// Use Contains instead of Compare, for response header like "Content-Type: application/json;charset=utf-8"
-		bool bIsContentTypeJson = InResponse->GetHeader(TEXT("Content-Type")).Contains(TEXT("application/json"), ESearchCase::IgnoreCase);
+		bool bIsContentTypeJson = HttpBase.GetHeader(FHttpConstants::HEADER_CONTENT_TYPE).Contains(FHttpConstants::MEDIATYPE_JSON, ESearchCase::IgnoreCase);
 		if (!bIsContentTypeJson)
 			return;
 	}
 
-	const TArray<uint8>& Content = InResponse->GetContent();
+	const TArray<uint8>& Content = HttpBase.GetContent();
+	if (Content.IsEmpty())
+	{
+		return;
+	}
+
+	UE_LOGFMT(LogHttp, Warning, "{0} Content({1}):", bIsResponse ? TEXT("Response") : TEXT("Request"), FText::AsMemory(HttpBase.GetContentLength(), EMemoryUnitStandard::IEC).ToString());
+
 	FUtf8StringView ResponseStringView(reinterpret_cast<const UTF8CHAR*>(Content.GetData()), Content.Num());
 	int32 StartPos = 0;
 	int32 EndPos = 0;
@@ -784,16 +839,21 @@ void FHttpRequestCommon::LogResponse(const TSharedPtr<IHttpResponse>& InResponse
 	while (StartPos < ResponseStringView.Len())
 	{
 		EndPos = ResponseStringView.Find("\n", StartPos);
-		if (EndPos != INDEX_NONE)
+		if (EndPos == INDEX_NONE)
 		{
-			FUtf8StringView Line(&ResponseStringView[StartPos], EndPos - StartPos);
-			UE_LOGFMT(LogHttp, Warning, "{Line}", Line);
+			EndPos = ResponseStringView.Len();
 		}
-		else
+
+		int32 LengthToLog = EndPos - StartPos;
+		while (LengthToLog > 0 && ResponseStringView[StartPos + LengthToLog - 1] == '\r')
 		{
-			FUtf8StringView Remain(&ResponseStringView[StartPos], ResponseStringView.Len() - StartPos);
-			UE_LOGFMT(LogHttp, Warning, "{Remain}", Remain);
-			break;
+			LengthToLog = LengthToLog - 1;
+		}
+
+		if (LengthToLog > 0)
+		{
+			FUtf8StringView Line(&ResponseStringView[StartPos], LengthToLog);
+			UE_LOGFMT(LogHttp, Warning, "{Line}", Line);
 		}
 
 		StartPos = EndPos + 1;
diff --git a/Engine/Source/Runtime/Online/HTTP/Private/GenericPlatform/HttpRequestCommon.h b/Engine/Source/Runtime/Online/HTTP/Private/GenericPlatform/HttpRequestCommon.h
--- a/Engine/Source/Runtime/Online/HTTP/Private/GenericPlatform/HttpRequestCommon.h
+++ b/Engine/Source/Runtime/Online/HTTP/Private/GenericPlatform/HttpRequestCommon.h
@@ -20,6 +20,8 @@ public:
 	HTTP_API virtual FString GetURLParameter(const FString& ParameterName) const override;
 
 	// IHttpRequest
+	HTTP_API virtual void SetHeader(const FString& HeaderName, const FString& HeaderValue) override;
+	HTTP_API virtual FString GetHeader(const FString& HeaderName) const override;
 	HTTP_API virtual EHttpRequestStatus::Type GetStatus() const override;
 	HTTP_API virtual const FString& GetEffectiveURL() const override;
 	HTTP_API virtual EHttpFailureReason GetFailureReason() const override;
@@ -52,7 +54,7 @@ public:
 	HTTP_API virtual float GetElapsedTime() const override;
 
 	virtual bool IsThreadedRequestComplete() = 0;
-	virtual bool StartThreadedRequest() = 0;
+	virtual bool StartThreadedRequest();
 	virtual void TickThreadedRequest(float DeltaSeconds) = 0;
 
 	HTTP_API void StartWaitingInQueue();
@@ -114,7 +116,10 @@ protected:
 	HTTP_API void CloseRequestPayloadDefaultImpl();
 
 	HTTP_API void LogFailure() const;
-	HTTP_API void LogResponse(const TSharedPtr<IHttpResponse>& InResponse);
+	HTTP_API void LogRequest() const;
+	HTTP_API void LogResponse(const TSharedPtr<FHttpResponseCommon>& InResponse) const;
+	HTTP_API void LogHeaders(const TMap<FString, FString>& Headers, bool bShouldLogAsWarning, bool bIsResponse) const;
+	HTTP_API void LogContent(bool bShouldLogDetails, const IHttpBase& HttpBase, bool bIsResponse) const;
 	HTTP_API bool TriggerMockResponse();
 
 	HTTP_API void InitResponse();
@@ -188,6 +193,9 @@ protected:
 	// null later on internally, the request itself won't cache received data anymore
 	std::atomic<bool> bInitializedWithValidStream = false;
 
+	/** Mapping of header section to values. */
+	TMap<FString, FString> RequestHeaders;
+
 	/** Payload to use with the request. Typically for POST, PUT, or PATCH */
 	TUniquePtr<FRequestPayload> RequestPayload;
 
diff --git a/Engine/Source/Runtime/Online/HTTP/Private/HttpManager.cpp b/Engine/Source/Runtime/Online/HTTP/Private/HttpManager.cpp
--- a/Engine/Source/Runtime/Online/HTTP/Private/HttpManager.cpp
+++ b/Engine/Source/Runtime/Online/HTTP/Private/HttpManager.cpp
@@ -26,8 +26,8 @@ TAutoConsoleVariable<int32> CVarHttpEventLoopEnableChance(
 	ECVF_SaveForNextBoot
 );
 
-TAutoConsoleVariable<FString> CVarHttpUrlPatternsToLogResponse(
-	TEXT("http.UrlPatternsToLogResponse"),
+TAutoConsoleVariable<FString> CVarHttpUrlPatternsToInspect(
+	TEXT("http.UrlPatternsToInspect"),
 	TEXT(""),
 	TEXT("List of url patterns to log headers and json content: \"epicgames.com unrealengine.com ...\""),
 	ECVF_SaveForNextBoot
@@ -121,8 +121,8 @@ void FHttpManager::Initialize()
 
 	UpdateConfigs();
 
-	UpdateUrlPatternsToLogResponse(CVarHttpUrlPatternsToLogResponse.AsVariable());
-	CVarHttpUrlPatternsToLogResponse.AsVariable()->OnChangedDelegate().AddRaw(this, &FHttpManager::UpdateUrlPatternsToLogResponse);
+	UpdateUrlPatternsToInspect(CVarHttpUrlPatternsToInspect.AsVariable());
+	CVarHttpUrlPatternsToInspect.AsVariable()->OnChangedDelegate().AddRaw(this, &FHttpManager::UpdateUrlPatternsToInspect);
 
 	UpdateUrlPatternsToDisableFailedLog(CVarHttpUrlPatternsToDisableFailedLog.AsVariable());
 	CVarHttpUrlPatternsToDisableFailedLog.AsVariable()->OnChangedDelegate().AddRaw(this, &FHttpManager::UpdateUrlPatternsToDisableFailedLog);
@@ -133,7 +133,7 @@ void FHttpManager::Initialize()
 
 void FHttpManager::Shutdown()
 {
-	CVarHttpUrlPatternsToLogResponse.AsVariable()->OnChangedDelegate().Clear();
+	CVarHttpUrlPatternsToInspect.AsVariable()->OnChangedDelegate().Clear();
 	CVarHttpUrlPatternsToDisableFailedLog.AsVariable()->OnChangedDelegate().Clear();
 	CVarHttpUrlPatternsToMockResponse.AsVariable()->OnChangedDelegate().Clear();
 
@@ -711,19 +711,19 @@ void FHttpManager::RecordPlatformStats(const FHttpStatsPlatform& PlatformStats)
 	HttpStats.PlatformStats = PlatformStats;
 }
 
-void FHttpManager::UpdateUrlPatternsToLogResponse(IConsoleVariable* CVar)
+void FHttpManager::UpdateUrlPatternsToInspect(IConsoleVariable* CVar)
 {
-	const FScopeLock CacheLock(&UrlPatternsToLogResponseCriticalSection);
-	const FString UrlPatternsToLogResponseStr = CVar->AsVariable()->GetString();
-	UrlPatternsToLogResponseStr.ParseIntoArray(UrlPatternsToLogResponse, TEXT(" "));
+	const FScopeLock CacheLock(&UrlPatternsToInspectCriticalSection);
+	const FString UrlPatternsToInspectStr = CVar->AsVariable()->GetString();
+	UrlPatternsToInspectStr.ParseIntoArray(UrlPatternsToInspect, TEXT(" "));
 }
 
-bool FHttpManager::ShouldLogResponse(FStringView Url) const
+bool FHttpManager::ShouldInspect(FStringView Url) const
 {
-	const FScopeLock CacheLock(&UrlPatternsToLogResponseCriticalSection);
-	for (const FString& UrlPatternToLogResponse : UrlPatternsToLogResponse)
+	const FScopeLock CacheLock(&UrlPatternsToInspectCriticalSection);
+	for (const FString& UrlPatternToInspect : UrlPatternsToInspect)
 	{
-		if (Url.Contains(UrlPatternToLogResponse))
+		if (Url.Contains(UrlPatternToInspect))
 		{
 			return true;
 		}
diff --git a/Engine/Source/Runtime/Online/HTTP/Public/HttpManager.h b/Engine/Source/Runtime/Online/HTTP/Public/HttpManager.h
--- a/Engine/Source/Runtime/Online/HTTP/Public/HttpManager.h
+++ b/Engine/Source/Runtime/Online/HTTP/Public/HttpManager.h
@@ -317,7 +317,7 @@ protected:
 
 	HTTP_API bool HasAnyBoundDelegate() const;
 
-	HTTP_API void UpdateUrlPatternsToLogResponse(IConsoleVariable* CVar);
+	HTTP_API void UpdateUrlPatternsToInspect(IConsoleVariable* CVar);
 	HTTP_API void UpdateUrlPatternsToDisableFailedLog(IConsoleVariable* CVar);
 	HTTP_API void UpdateUrlPatternsToMockResponse(IConsoleVariable* CVar);
 
@@ -389,8 +389,8 @@ protected:
 
 	bool bUseEventLoop = true;
 
-	TArray<FString> UrlPatternsToLogResponse;
-	mutable FCriticalSection UrlPatternsToLogResponseCriticalSection;
+	TArray<FString> UrlPatternsToInspect;
+	mutable FCriticalSection UrlPatternsToInspectCriticalSection;
 
 	TArray<FString> UrlPatternsToDisableFailedLog;
 	mutable FCriticalSection UrlPatternsToDisableFailedLogCriticalSection;
@@ -447,7 +447,7 @@ PACKAGE_SCOPE:
 	/** Record platform specific stats */
 	HTTP_API void RecordPlatformStats(const FHttpStatsPlatform& PlatformStats);
 
-	HTTP_API bool ShouldLogResponse(FStringView Url) const;
+	HTTP_API bool ShouldInspect(FStringView Url) const;
 	HTTP_API bool ShouldDisableFailedLog(FStringView Url) const;
 
 	HTTP_API TOptional<FMockResponse> GetMockResponse(FStringView Url, FStringView Verb) const;
