Rename CVarHttpLogJsonResponseOnly to CVarHttpLogJsonContentOnly.

#Synced-CL 48203440

diff --git a/Engine/Source/Programs/WebTests/Private/TestHttp.cpp b/Engine/Source/Programs/WebTests/Private/TestHttp.cpp
--- a/Engine/Source/Programs/WebTests/Private/TestHttp.cpp
+++ b/Engine/Source/Programs/WebTests/Private/TestHttp.cpp
@@ -48,7 +48,7 @@ extern TAutoConsoleVariable<int32> CVarHttpEventLoopEnableChance;
 extern TAutoConsoleVariable<FString> CVarHttpUrlPatternsToMockResponse;
 extern TAutoConsoleVariable<FString> CVarHttpUrlPatternsToDisableFailedLog;
 extern TAutoConsoleVariable<FString> CVarHttpUrlPatternsToInspect;
-extern TAutoConsoleVariable<bool> CVarHttpLogJsonResponseOnly;
+extern TAutoConsoleVariable<bool> CVarHttpLogJsonContentOnly;
 
 class FMockHttpModule : public FHttpModule
 {
@@ -459,7 +459,7 @@ TEST_CASE_METHOD(FWaitUntilCompleteHttpFixture, "Http request can LogContent cor
 	DisableWarningsInThisTest();
 
 	CVarHttpUrlPatternsToInspect->Set(TEXT("methods"));
-	CVarHttpLogJsonResponseOnly->Set(false);
+	CVarHttpLogJsonContentOnly->Set(false);
 
 	TSharedRef<IHttpRequest> HttpRequest = CreateRequest();
 	HttpRequest->SetURL(UrlToTestMethods());
@@ -491,7 +491,7 @@ TEST_CASE_METHOD(FWaitUntilCompleteHttpFixture, "Http request can LogContent cor
 		// TODO: Find a way to validate output logs
 
 		CVarHttpUrlPatternsToInspect->Set(TEXT(""));
-		CVarHttpLogJsonResponseOnly->Set(true);
+		CVarHttpLogJsonContentOnly->Set(true);
 	});
 	HttpRequest->ProcessRequest();
 }
diff --git a/Engine/Source/Runtime/Online/HTTP/Private/GenericPlatform/HttpRequestCommon.cpp b/Engine/Source/Runtime/Online/HTTP/Private/GenericPlatform/HttpRequestCommon.cpp
--- a/Engine/Source/Runtime/Online/HTTP/Private/GenericPlatform/HttpRequestCommon.cpp
+++ b/Engine/Source/Runtime/Online/HTTP/Private/GenericPlatform/HttpRequestCommon.cpp
@@ -11,10 +11,10 @@
 #include "Misc/CommandLine.h"
 #include "Stats/Stats.h"
 
-TAutoConsoleVariable<bool> CVarHttpLogJsonResponseOnly(
-	TEXT("http.LogJsonResponseOnly"),
+TAutoConsoleVariable<bool> CVarHttpLogJsonContentOnly(
+	TEXT("http.LogJsonContentOnly"),
 	true,
-	TEXT("When log response payload, log json content only"),
+	TEXT("When log request/response payload, log json content only"),
 	ECVF_SaveForNextBoot
 );
 
@@ -816,9 +816,9 @@ void FHttpRequestCommon::LogContent(bool bShouldLogDetails, const IHttpBase& Htt
 		return;
 	}
 
-	if (CVarHttpLogJsonResponseOnly.GetValueOnAnyThread())
+	if (CVarHttpLogJsonContentOnly.GetValueOnAnyThread())
 	{
-		// Use Contains instead of Compare, for response header like "Content-Type: application/json;charset=utf-8"
+		// Use Contains instead of Compare, for request/response header like "Content-Type: application/json;charset=utf-8"
 		bool bIsContentTypeJson = HttpBase.GetHeader(FHttpConstants::HEADER_CONTENT_TYPE).Contains(FHttpConstants::MEDIATYPE_JSON, ESearchCase::IgnoreCase);
 		if (!bIsContentTypeJson)
 			return;
