Deprecate legacy http thread by output warning if the chance is not 100 percent for event loop;

#Synced-CL 48344811

diff --git a/Engine/Source/Runtime/Online/HTTP/Private/HttpManager.cpp b/Engine/Source/Runtime/Online/HTTP/Private/HttpManager.cpp
--- a/Engine/Source/Runtime/Online/HTTP/Private/HttpManager.cpp
+++ b/Engine/Source/Runtime/Online/HTTP/Private/HttpManager.cpp
@@ -115,6 +115,11 @@ void FHttpManager::Initialize()
 		// Can't set cvar CVarHttpEventLoopEnableChance through runtime args or .ini files because http module initialized too early
 		FParse::Bool(FCommandLine::Get(), TEXT("useeventloop="), bUseEventLoop);
 
+		if (CVarHttpEventLoopEnableChance.GetValueOnGameThread() < 100 || !bUseEventLoop)
+		{
+			UE_LOG(LogHttp, Error, TEXT("Legacy http thread has been deprecated, please always use event loop thread instead."));
+		}
+
 		Thread = CreateHttpThread();
 		Thread->StartThread();
 	}
