Use libCurl by default  instead of xCurl on WinGDK. Reason: curl_multi_wakeup is slow with xCurl on WinGDK only which could impact the game thread performance;
Added perf stats for process request flow;
Add more optimization in FEventLoopNotifierEventFd: When Notify() get called several times in same frame, we only need to notify once;
Don't call Shutdown in FEventLoopNotifierSocket when failed to write/read.

#Synced-CL 48344811

diff --git a/Engine/Source/Runtime/Online/Experimental/EventLoop/Private/EventLoop/BSDSocket/EventLoopNotifierEventFd.cpp b/Engine/Source/Runtime/Online/Experimental/EventLoop/Private/EventLoop/BSDSocket/EventLoopNotifierEventFd.cpp
--- a/Engine/Source/Runtime/Online/Experimental/EventLoop/Private/EventLoop/BSDSocket/EventLoopNotifierEventFd.cpp
+++ b/Engine/Source/Runtime/Online/Experimental/EventLoop/Private/EventLoop/BSDSocket/EventLoopNotifierEventFd.cpp
@@ -21,6 +21,8 @@ static FAutoConsoleVariableRef CVarHttpEventLoopNotifierUseThreadPool(
 	ECVF_SaveForNextBoot
 );
 
+std::atomic<bool> FEventLoopNotifierEventFd::bNotified = false;
+
 bool FEventLoopNotifierEventFd::Init()
 {
 	FileDescriptorRead = eventfd(0, EFD_NONBLOCK);
@@ -57,10 +59,15 @@ void FEventLoopNotifierEventFd::Notify()
 
 #if PLATFORM_WRITES_ARE_SLOW
 	if (bHttpEventLoopNotifierUseThreadPool && GIOThreadPool)
-	{
-		AsyncPool(*GIOThreadPool, [Fd = FileDescriptorRead] {
-			EventLoopNotifierEventFd::NotifyImpl(Fd);
-		});
+	{ 
+		// When Notify() get called several times in same frame, we only need to notify once
+		if (!bNotified.exchange(true))
+		{
+			AsyncPool(*GIOThreadPool, [Fd = FileDescriptorRead] {
+				bNotified = false;
+				EventLoopNotifierEventFd::NotifyImpl(Fd);
+			});
+		}
 	}
 	else
 #endif
diff --git a/Engine/Source/Runtime/Online/Experimental/EventLoop/Private/EventLoop/BSDSocket/EventLoopNotifierSocket.cpp b/Engine/Source/Runtime/Online/Experimental/EventLoop/Private/EventLoop/BSDSocket/EventLoopNotifierSocket.cpp
--- a/Engine/Source/Runtime/Online/Experimental/EventLoop/Private/EventLoop/BSDSocket/EventLoopNotifierSocket.cpp
+++ b/Engine/Source/Runtime/Online/Experimental/EventLoop/Private/EventLoop/BSDSocket/EventLoopNotifierSocket.cpp
@@ -2,6 +2,7 @@
 
 #include "EventLoop/BSDSocket/EventLoopNotifierSocket.h"
 #include "BSDSocketTypesPrivate.h"
+#include "EventLoop/EventLoopLog.h"
 #include "Stats/Stats.h"
 
 #if PLATFORM_HAS_BSD_SOCKETS && !PLATFORM_HAS_BSD_SOCKET_FEATURE_WINSOCKETS
@@ -135,7 +136,8 @@ void FEventLoopNotifierSocket::Notify()
 		// The notifier just cares that the socket is set to a readable state.
 		if (write(FileDescriptorWrite, "1", 1) == SOCKET_RESULT_FAILED)
 		{
-			Shutdown();
+			int Err = errno;
+			UE_CLOG(Err != EAGAIN, LogEventLoop, Warning, TEXT("Failed to notify by socket fd write %d with error %d"), FileDescriptorWrite, Err);
 		}
 	}
 }
@@ -145,10 +147,7 @@ void FEventLoopNotifierSocket::Clear()
 	if (FileDescriptorRead != INVALID_SOCKET)
 	{
 		char Buffer[1];
-		if (read(FileDescriptorRead, &Buffer[0], 1) == SOCKET_RESULT_FAILED)
-		{
-			Shutdown();
-		}
+		ensure(read(FileDescriptorRead, &Buffer[0], 1) != SOCKET_RESULT_FAILED);
 	}
 }
 
diff --git a/Engine/Source/Runtime/Online/Experimental/EventLoop/Public/EventLoop/BSDSocket/EventLoopNotifierEventFd.h b/Engine/Source/Runtime/Online/Experimental/EventLoop/Public/EventLoop/BSDSocket/EventLoopNotifierEventFd.h
--- a/Engine/Source/Runtime/Online/Experimental/EventLoop/Public/EventLoop/BSDSocket/EventLoopNotifierEventFd.h
+++ b/Engine/Source/Runtime/Online/Experimental/EventLoop/Public/EventLoop/BSDSocket/EventLoopNotifierEventFd.h
@@ -4,6 +4,7 @@
 
 #include "BSDSocketTypes.h"
 #include "CoreTypes.h"
+#include <atomic>
 
 #if HAS_EVENTLOOP_NOTIFIER_EVENTFD
 
@@ -22,6 +23,7 @@ public:
 
 private:
 	int32 FileDescriptorRead = INVALID_SOCKET;
+	static std::atomic<bool> bNotified;
 };
 
 /* UE::EventLoop */ }
diff --git a/Engine/Source/Runtime/Online/HTTP/Private/Apple/AppleEventLoopHttpThread.h b/Engine/Source/Runtime/Online/HTTP/Private/Apple/AppleEventLoopHttpThread.h
--- a/Engine/Source/Runtime/Online/HTTP/Private/Apple/AppleEventLoopHttpThread.h
+++ b/Engine/Source/Runtime/Online/HTTP/Private/Apple/AppleEventLoopHttpThread.h
@@ -22,6 +22,7 @@ public:
 	
 	void Notify()
 	{
+		QUICK_SCOPE_CYCLE_COUNTER(STAT_FAppleHTTPIOAccess_Notify);
 		SharedEvent->Notify();
 	}
 	
diff --git a/Engine/Source/Runtime/Online/HTTP/Private/Apple/AppleHttp.cpp b/Engine/Source/Runtime/Online/HTTP/Private/Apple/AppleHttp.cpp
--- a/Engine/Source/Runtime/Online/HTTP/Private/Apple/AppleHttp.cpp
+++ b/Engine/Source/Runtime/Online/HTTP/Private/Apple/AppleHttp.cpp
@@ -647,6 +647,8 @@ void FAppleHttpRequest::SetVerb(const FString& Verb)
 
 bool FAppleHttpRequest::ProcessRequest()
 {
+	QUICK_SCOPE_CYCLE_COUNTER(STAT_FAppleHttpRequest_ProcessRequest);
+
 	SCOPED_AUTORELEASE_POOL;
 
 	if (!PreProcess())
diff --git a/Engine/Source/Runtime/Online/HTTP/Private/Curl/CurlMultiPollEventLoopHttpThread.cpp b/Engine/Source/Runtime/Online/HTTP/Private/Curl/CurlMultiPollEventLoopHttpThread.cpp
--- a/Engine/Source/Runtime/Online/HTTP/Private/Curl/CurlMultiPollEventLoopHttpThread.cpp
+++ b/Engine/Source/Runtime/Online/HTTP/Private/Curl/CurlMultiPollEventLoopHttpThread.cpp
@@ -35,6 +35,7 @@ void FCurlMultiPollIOManager::Shutdown()
 
 void FCurlMultiPollIOManager::Notify()
 {
+	QUICK_SCOPE_CYCLE_COUNTER(STAT_FCurlMultiPollIOManager_Notify);
 	CheckMultiCodeOk(curl_multi_wakeup(Params.MultiHandle));
 }
 
