Fix the non-unity compile error in AppleHttp.

#Synced-CL 48344811

diff --git a/Engine/Source/Runtime/Online/HTTP/Private/Apple/AppleHttp.cpp b/Engine/Source/Runtime/Online/HTTP/Private/Apple/AppleHttp.cpp
--- a/Engine/Source/Runtime/Online/HTTP/Private/Apple/AppleHttp.cpp
+++ b/Engine/Source/Runtime/Online/HTTP/Private/Apple/AppleHttp.cpp
@@ -15,6 +15,7 @@
 #include "Misc/EngineVersion.h"
 #include "Misc/App.h"
 #include "Misc/Base64.h"
+#include "Stats/Stats.h"
 
 @implementation FAppleHttpResponseDelegate
 @synthesize Response;
