Make sure OpenSSL is enabled on WinGDK.

#Synced-CL 48344811

diff --git a/Engine/Source/Runtime/Online/SSL/SSL.Build.cs b/Engine/Source/Runtime/Online/SSL/SSL.Build.cs
--- a/Engine/Source/Runtime/Online/SSL/SSL.Build.cs
+++ b/Engine/Source/Runtime/Online/SSL/SSL.Build.cs
@@ -10,7 +10,7 @@ public class SSL : ModuleRules
 		{
 			return
 				Target.Platform == UnrealTargetPlatform.Mac ||
-				Target.Platform == UnrealTargetPlatform.Win64 ||
+				Target.IsInPlatformGroup(UnrealPlatformGroup.Windows) ||
 				Target.IsInPlatformGroup(UnrealPlatformGroup.Unix) ||
 				Target.Platform == UnrealTargetPlatform.IOS ||
 				Target.Platform == UnrealTargetPlatform.Android;
@@ -44,7 +44,7 @@ public class SSL : ModuleRules
 
 			AddEngineThirdPartyPrivateStaticDependencies(Target, "OpenSSL");
 
-			if (Target.Platform == UnrealTargetPlatform.Win64)
+			if (Target.IsInPlatformGroup(UnrealPlatformGroup.Windows))
 			{
 				PublicSystemLibraries.Add("crypt32.lib");
 			}
