Move EventLoop out of Experimental.

#Synced-CL 48708107

diff --git a/Engine/Source/Runtime/Online/EventLoop/EventLoop.Build.cs b/Engine/Source/Runtime/Online/EventLoop/EventLoop.Build.cs
new file mode 100644
--- /dev/null
+++ b/Engine/Source/Runtime/Online/EventLoop/EventLoop.Build.cs
@@ -0,0 +1,46 @@
+// Copyright Epic Games, Inc. All Rights Reserved.
+
+using UnrealBuildTool;
+
+public class EventLoop : ModuleRules
+{
+	public EventLoop(ReadOnlyTargetRules Target) : base(Target)
+	{
+		PublicDependencyModuleNames.AddRange(
+			new string[] {
+			}
+		);
+
+		PrivateDependencyModuleNames.AddRange(
+			new string[] {
+				"Core",
+				"NetCommon"
+			});
+
+		PublicDefinitions.Add(string.Format("HAS_EVENTLOOP_PLATFORM_SOCKET_IMPLEMENTATION={0}", bHasPlatformSocketImplementation ? 1 : 0));
+		PublicDefinitions.Add(string.Format("HAS_EVENTLOOP_PLATFORM_BSD_SOCKET_HEADER={0}", bHasPlatformBSDSocketHeader ? 1 : 0));
+		PublicDefinitions.Add(string.Format("HAS_PLATFORM_EVENTLOOP_NOTIFIER_EVENTFD={0}", bHasPlatformEventLoopNotifierEventFd ? 1 : 0));
+		PublicDefinitions.Add(string.Format("HAS_EVENTLOOP_NOTIFIER_EVENTFD={0}", bHasEventLoopNotifierEventFd ? 1 : 0));
+	}
+
+	public virtual bool bHasPlatformSocketImplementation
+	{
+		get
+		{
+			return false;
+			// Other platforms may override this property.
+		}
+	}
+
+	public virtual bool bHasPlatformBSDSocketHeader
+	{
+		get
+		{
+			return false;
+			// Other platforms may override this property.
+		}
+	}
+
+	public virtual bool bHasPlatformEventLoopNotifierEventFd { get { return false; } }
+	public virtual bool bHasEventLoopNotifierEventFd { get { return Target.IsInPlatformGroup(UnrealPlatformGroup.Linux); } }
+}
diff --git a/Engine/Source/Runtime/Online/EventLoop/Private/EventLoop/BSDSocket/BSDSocketTypesPrivate.h b/Engine/Source/Runtime/Online/EventLoop/Private/EventLoop/BSDSocket/BSDSocketTypesPrivate.h
new file mode 100644
--- /dev/null
+++ b/Engine/Source/Runtime/Online/EventLoop/Private/EventLoop/BSDSocket/BSDSocketTypesPrivate.h
@@ -0,0 +1,79 @@
+// Copyright Epic Games, Inc. All Rights Reserved.
+
+#pragma once
+
+#include "Containers/UnrealString.h"
+#include "Containers/StringConv.h"
+#include "Containers/UnrealString.h"
+#include "EventLoop/BSDSocket/BSDSocketTypes.h"
+
+#if PLATFORM_HAS_BSD_SOCKETS
+
+#if PLATFORM_HAS_BSD_SOCKET_FEATURE_WINSOCKETS
+#include "Windows/AllowWindowsPlatformTypes.h"
+
+#include <winsock2.h>
+#include <ws2tcpip.h>
+
+typedef int32 SOCKLEN;
+#define strerror_r(errno,buf,len) strerror_s(buf,len,errno)
+
+#include "Windows/HideWindowsPlatformTypes.h"
+#else // PLATFORM_HAS_BSD_SOCKET_FEATURE_WINSOCKETS
+#if HAS_EVENTLOOP_PLATFORM_BSD_SOCKET_HEADER
+#include COMPILED_PLATFORM_HEADER(BSDSocketTypesPrivate.h)
+#else // HAS_EVENTLOOP_PLATFORM_BSD_SOCKET_HEADER
+#include <unistd.h>
+#include <sys/socket.h>
+#if PLATFORM_HAS_BSD_SOCKET_FEATURE_IOCTL
+#include <fcntl.h>
+#include <sys/types.h>
+#include <sys/ioctl.h>
+#endif // PLATFORM_HAS_BSD_SOCKET_FEATURE_IOCTL
+#if PLATFORM_HAS_BSD_SOCKET_FEATURE_POLL
+#include <poll.h>
+#endif // PLATFORM_HAS_BSD_SOCKET_FEATURE_POLL
+#include <netinet/in.h>
+#include <arpa/inet.h>
+#if PLATFORM_HAS_BSD_SOCKET_FEATURE_GETHOSTNAME
+#include <netdb.h>
+#endif // PLATFORM_HAS_BSD_SOCKET_FEATURE_GETHOSTNAME
+#if PLATFORM_HAS_BSD_SOCKET_FEATURE_NODELAY
+#include <netinet/tcp.h>
+#endif // PLATFORM_HAS_BSD_SOCKET_FEATURE_NODELAY
+
+#if HAS_EVENTLOOP_NOTIFIER_EVENTFD && !HAS_PLATFORM_EVENTLOOP_NOTIFIER_EVENTFD
+#include <sys/eventfd.h>
+#endif // HAS_EVENTLOOP_NOTIFIER_EVENTFD
+
+#define ioctlsocket ioctl
+#endif // HAS_EVENTLOOP_PLATFORM_BSD_SOCKET_HEADER
+
+#define SOCKET_ERROR -1
+#define INVALID_SOCKET -1
+
+typedef socklen_t SOCKLEN;
+typedef int32 SOCKET;
+typedef sockaddr_in SOCKADDR_IN;
+typedef struct timeval TIMEVAL;
+
+inline int32 closesocket(SOCKET Socket)
+{
+	shutdown(Socket, SHUT_RDWR); // gracefully shutdown if connected
+	return close(Socket);
+}
+
+#endif // PLATFORM_HAS_BSD_SOCKET_FEATURE_WINSOCKETS
+
+inline FString GetLastErrorString()
+{
+	constexpr uint32 BufferSize = 256;
+	ANSICHAR Buffer[BufferSize];
+	Buffer[0] = 0;
+
+	strerror_r(errno, Buffer, BufferSize);
+
+	return FString::Printf(TEXT("%s: (%d)"), StringCast<TCHAR>(Buffer).Get(), errno);
+}
+
+#endif // PLATFORM_HAS_BSD_SOCKETS
diff --git a/Engine/Source/Runtime/Online/EventLoop/Private/EventLoop/BSDSocket/EventLoopIOManagerBSDSocketSelect.cpp b/Engine/Source/Runtime/Online/EventLoop/Private/EventLoop/BSDSocket/EventLoopIOManagerBSDSocketSelect.cpp
new file mode 100644
--- /dev/null
+++ b/Engine/Source/Runtime/Online/EventLoop/Private/EventLoop/BSDSocket/EventLoopIOManagerBSDSocketSelect.cpp
@@ -0,0 +1,312 @@
+// Copyright Epic Games, Inc. All Rights Reserved.
+
+#include "EventLoop/BSDSocket/EventLoopIOManagerBSDSocketSelect.h"
+#include "EventLoop/BSDSocket/BSDSocketTypesPrivate.h"
+#include "EventLoop/EventLoopLog.h"
+#include "EventLoop/IEventLoop.h"
+#include "HAL/PlatformProcess.h"
+#include "HAL/PlatformTime.h"
+#include "Math/UnrealMathUtility.h"
+#include "Stats/Stats.h"
+
+#if PLATFORM_HAS_BSD_SOCKETS && PLATFORM_HAS_BSD_SOCKET_FEATURE_SELECT
+
+namespace UE::EventLoop {
+
+class FIOManagerBSDSocketSelectImpl final : public IIOManager
+{
+public:
+	FIOManagerBSDSocketSelectImpl(IEventLoop& EventLoop);
+	virtual ~FIOManagerBSDSocketSelectImpl() = default;
+	virtual bool Init() override;
+	virtual void Shutdown() override;
+	virtual void Notify() override;
+	virtual void Poll(FTimespan WaitTime) override;
+
+	void PollInternal(FTimespan WaitTime);
+
+	FIOAccessBSDSocket& GetIOAccess()
+	{
+		return IOAccess;
+	}
+
+private:
+	using FStorageType = FIOAccessBSDSocket::FStorageType;
+	using FInternalHandleArryType = FStorageType::FInternalHandleArryType;
+	using FInternalHandle = FStorageType::FInternalHandle;
+
+	struct FSelectData
+	{
+		fd_set SocketReadSet;
+		fd_set SocketWriteSet;
+		fd_set SocketExceptionSet;
+		SOCKET MaxFd = INVALID_SOCKET;
+	};
+
+	void RebuildSelectData(FSelectData& SelectData) const;
+	int32 PollSelectData(const FSelectData& SelectData, FSelectData& OutSignaledData, FTimespan WaitTime) const;
+
+	FSelectData SelectData;
+
+	IEventLoop& EventLoop;
+	FStorageType IORequestStorage;
+	FIOAccessBSDSocket IOAccess;
+	TAtomic<bool> bAsyncSignal;
+
+	// Indices
+	TMap<FInternalHandle, SOCKET> SocketInternalHandleIndex;
+	TMap<SOCKET, FInternalHandle> SocketIndex;
+};
+
+FIOManagerBSDSocketSelectImpl::FIOManagerBSDSocketSelectImpl(IEventLoop& InEventLoop)
+	: EventLoop(InEventLoop)
+	, IOAccess(IORequestStorage)
+	, bAsyncSignal(false)
+{
+}
+
+bool FIOManagerBSDSocketSelectImpl::Init()
+{
+	IORequestStorage.Init();
+	return true;
+}
+
+void FIOManagerBSDSocketSelectImpl::Shutdown()
+{
+}
+
+void FIOManagerBSDSocketSelectImpl::Notify()
+{
+	bAsyncSignal = true;
+}
+
+void FIOManagerBSDSocketSelectImpl::Poll(FTimespan WaitTime)
+{
+	QUICK_SCOPE_CYCLE_COUNTER(STAT_EventLoop_FIOManagerBSDSocketSelectImpl_Poll);
+
+	// The base BSD socket polling implementation relies on using socket select which cannot be
+	// woken by another thread. To mitigate this issue, many smaller calls to PollInternal are done
+	// while no signal is received.
+
+	const FTimespan MaxIterationWaitTime = FTimespan::FromMilliseconds(10);
+	FTimespan AccumulatedWaitTime;
+
+	// If a signal was received before polling, still poll for any ready events using a timeout of 0.
+	if (bAsyncSignal)
+	{
+		WaitTime = FTimespan::Zero();
+	}
+
+	// Run at least one iteration. Keep running while WaitTime has not been reached.
+	for (int32 Index = 0; Index == 0 || AccumulatedWaitTime < WaitTime; ++Index)
+	{
+		FTimespan CurrentWaitTime = FMath::Min(WaitTime - AccumulatedWaitTime, MaxIterationWaitTime);
+		PollInternal(CurrentWaitTime);
+
+		if (bAsyncSignal)
+		{
+			break;
+		}
+
+		AccumulatedWaitTime += MaxIterationWaitTime;
+	}
+
+	// Reset signal. Resetting is safe here as setting the async signal indicates that the poll should be exited.
+	bAsyncSignal = false;
+}
+
+void FIOManagerBSDSocketSelectImpl::PollInternal(FTimespan WaitTime)
+{
+	QUICK_SCOPE_CYCLE_COUNTER(STAT_EventLoop_FIOManagerBSDSocketSelectImpl_PollInternal);
+
+	// Process queued actions.
+	FInternalHandleArryType AddedHandles;
+	FInternalHandleArryType RemovedHandles;
+	uint32 NumChanges = IORequestStorage.Update(&AddedHandles, &RemovedHandles);
+
+	// Remove active requests before adding more to have an accurate count of available resources.
+	for (FInternalHandle RemovedHandle : RemovedHandles)
+	{
+		SOCKET RemovedSocket;
+		if (SocketInternalHandleIndex.RemoveAndCopyValue(RemovedHandle, RemovedSocket))
+		{
+			UE_LOG(LogEventLoop, VeryVerbose, TEXT("[FIOManagerBSDSocketSelectImpl::PollInternal] Request removed: Socket: %p, Handle: %s"), reinterpret_cast<void*>(RemovedSocket), *RemovedHandle.ToString());
+			SocketIndex.Remove(RemovedSocket);
+		}
+	}
+
+	for (FInternalHandle AddedHandle : AddedHandles)
+	{
+		// Request may have been added and removed in the same update.
+		if (FIORequestBSDSocket* NewRequest = IORequestStorage.Find(AddedHandle))
+		{
+			// Check whether there are enough resources available to handle the request.
+			if (SocketIndex.Num() >= FD_SETSIZE)
+			{
+				UE_LOG(LogEventLoop, Verbose, TEXT("[FIOManagerBSDSocketSelectImpl::PollInternal] Request failed: Socket: %p, Reason: %s"), reinterpret_cast<void*>(NewRequest->Socket), *LexToString(ESocketIoRequestStatus::NoResources));
+				NewRequest->Callback(NewRequest->Socket, ESocketIoRequestStatus::NoResources, EIOFlags::None);
+				IORequestStorage.Remove(AddedHandle);
+				continue;
+			}
+
+			// Only one request may be active for a socket at one time.
+			if (SocketIndex.Contains(NewRequest->Socket))
+			{
+				UE_LOG(LogEventLoop, Verbose, TEXT("[FIOManagerBSDSocketSelectImpl::PollInternal] Request failed: Socket: %p, Reason: %s"), reinterpret_cast<void*>(NewRequest->Socket), *LexToString(ESocketIoRequestStatus::Invalid));
+				NewRequest->Callback(NewRequest->Socket, ESocketIoRequestStatus::Invalid, EIOFlags::None);
+				IORequestStorage.Remove(AddedHandle);
+				continue;
+			}
+
+			UE_LOG(LogEventLoop, VeryVerbose, TEXT("[FIOManagerBSDSocketSelectImpl::PollInternal] Request added: Socket: %p, Handle: %s, Flags: 0x%08X"), reinterpret_cast<void*>(NewRequest->Socket), *AddedHandle.ToString(), NewRequest->Flags);
+
+			SocketInternalHandleIndex.Add(AddedHandle, NewRequest->Socket);
+			SocketIndex.Add(NewRequest->Socket, AddedHandle);
+		}
+	}
+
+	// Check whether the cached FD_SETs need to be rebuilt.
+	if (NumChanges > 0)
+	{
+		RebuildSelectData(SelectData);
+	}
+
+	// When there is no work, block the loop here.
+	if (IORequestStorage.IsEmpty())
+	{
+		QUICK_SCOPE_CYCLE_COUNTER(STAT_EventLoop_FIOManagerBSDSocketSelectImpl_PollInternal_IdleSleep);
+		FPlatformProcess::SleepNoStats(WaitTime.GetTotalSeconds());
+		return;
+	}
+
+	// Poll for socket status.
+	FSelectData SignaledSelectData;
+	int32 SelectStatus = PollSelectData(SelectData, SignaledSelectData, WaitTime);
+
+	// Check whether any events triggered.
+	if (SelectStatus > 0)
+	{
+		QUICK_SCOPE_CYCLE_COUNTER(STAT_EventLoop_FIOManagerBSDSocketSelectImpl_PollInternal_HandleIOState);
+
+		// Check request state and trigger callbacks if needed.
+		for (TPair<FInternalHandle, FIORequestBSDSocket&>& IORequestEntry : IORequestStorage)
+		{
+			FIORequestBSDSocket& IORequest = IORequestEntry.Value;
+			EIOFlags SignaledFlags = EIOFlags::None;
+
+			if (EnumHasAnyFlags(IORequest.Flags, EIOFlags::Read) && FD_ISSET(IORequest.Socket, &SignaledSelectData.SocketReadSet))
+			{
+				SignaledFlags |= EIOFlags::Read;
+			}
+
+			if (EnumHasAnyFlags(IORequest.Flags, EIOFlags::Write) &&
+				(FD_ISSET(IORequest.Socket, &SignaledSelectData.SocketWriteSet) || FD_ISSET(IORequest.Socket, &SignaledSelectData.SocketExceptionSet)))
+			{
+				SignaledFlags |= EIOFlags::Write;
+			}
+
+			if (SignaledFlags != EIOFlags::None)
+			{
+				// Signal event status.
+				UE_LOG(LogEventLoop, VeryVerbose, TEXT("[FIOManagerBSDSocketSelectImpl::PollInternal] Request signaled: Socket: %p, Handle: %s, Flags: 0x%08X"), reinterpret_cast<void*>(IORequest.Socket), *IORequestEntry.Key.ToString(), SignaledFlags);
+				IORequest.Callback(IORequest.Socket, ESocketIoRequestStatus::Ok, SignaledFlags);
+			}
+		}
+	}
+	else if (SelectStatus < 0)
+	{
+		UE_LOG(LogEventLoop, Fatal, TEXT("[FIOManagerBSDSocketSelectImpl::PollInternal] Select error: %s"), *GetLastErrorString());
+	}
+}
+
+void FIOManagerBSDSocketSelectImpl::RebuildSelectData(FSelectData& OutSelectData) const
+{
+	QUICK_SCOPE_CYCLE_COUNTER(STAT_EventLoop_FIOManagerBSDSocketSelectImpl_RebuildSelectData);
+	FD_ZERO(&OutSelectData.SocketReadSet);
+	FD_ZERO(&OutSelectData.SocketWriteSet);
+	FD_ZERO(&OutSelectData.SocketExceptionSet);
+	OutSelectData.MaxFd = INVALID_SOCKET;
+
+	for (const TPair<SOCKET, FInternalHandle>& IndexEntry : SocketIndex)
+	{
+		if (const FIORequestBSDSocket* IORequest = IORequestStorage.Find(IndexEntry.Value))
+		{
+			if (OutSelectData.MaxFd == INVALID_SOCKET)
+			{
+				OutSelectData.MaxFd = IORequest->Socket;
+			}
+			else
+			{
+				OutSelectData.MaxFd = FMath::Max(SelectData.MaxFd, IORequest->Socket);
+			}
+
+			if (EnumHasAnyFlags(IORequest->Flags, EIOFlags::Read))
+			{
+				FD_SET(IORequest->Socket, &OutSelectData.SocketReadSet);
+			}
+
+			if (EnumHasAnyFlags(IORequest->Flags, EIOFlags::Write))
+			{
+				FD_SET(IORequest->Socket, &OutSelectData.SocketWriteSet);
+				FD_SET(IORequest->Socket, &OutSelectData.SocketExceptionSet);
+			}
+		}
+	}
+}
+
+int32 FIOManagerBSDSocketSelectImpl::PollSelectData(const FSelectData& InSelectData, FIOManagerBSDSocketSelectImpl::FSelectData& OutSignaledData, FTimespan WaitTime) const
+{
+	QUICK_SCOPE_CYCLE_COUNTER(STAT_EventLoop_FIOManagerBSDSocketSelectImpl_PollSelectData);
+
+	// Copy cached FD sets.
+	OutSignaledData = InSelectData;
+
+	// convert WaitTime to a timeval
+	timeval Time;
+	Time.tv_sec = (int32)WaitTime.GetTotalSeconds();
+	Time.tv_usec = WaitTime.GetFractionMicro();
+
+	timeval* TimePointer = WaitTime.GetTicks() >= 0 ? &Time : nullptr;
+
+	// Poll for socket status.
+	return select(IntCastChecked<int>(OutSignaledData.MaxFd + 1), &OutSignaledData.SocketReadSet, &OutSignaledData.SocketWriteSet, &OutSignaledData.SocketExceptionSet, TimePointer);
+}
+
+//-----------------------
+// Pimpl implementation.
+//-----------------------
+
+FIOManagerBSDSocketSelect::FIOManagerBSDSocketSelect(IEventLoop& EventLoop, FParams&&)
+	: Impl(MakeShared<FIOManagerBSDSocketSelectImpl>(EventLoop))
+{
+}
+
+bool FIOManagerBSDSocketSelect::Init()
+{
+	return Impl->Init();
+}
+
+void FIOManagerBSDSocketSelect::Shutdown()
+{
+	Impl->Shutdown();
+}
+
+void FIOManagerBSDSocketSelect::Notify()
+{
+	Impl->Notify();
+}
+
+void FIOManagerBSDSocketSelect::Poll(FTimespan WaitTime)
+{
+	Impl->Poll(WaitTime);
+}
+
+FIOManagerBSDSocketSelect::FIOAccess& FIOManagerBSDSocketSelect::GetIOAccess()
+{
+	return Impl->GetIOAccess();
+}
+
+/* UE::EventLoop */ }
+
+#endif // PLATFORM_HAS_BSD_SOCKETS && PLATFORM_HAS_BSD_SOCKET_FEATURE_SELECT
+
diff --git a/Engine/Source/Runtime/Online/EventLoop/Private/EventLoop/BSDSocket/EventLoopNotifierEventFd.cpp b/Engine/Source/Runtime/Online/EventLoop/Private/EventLoop/BSDSocket/EventLoopNotifierEventFd.cpp
new file mode 100644
--- /dev/null
+++ b/Engine/Source/Runtime/Online/EventLoop/Private/EventLoop/BSDSocket/EventLoopNotifierEventFd.cpp
@@ -0,0 +1,88 @@
+// Copyright Epic Games, Inc. All Rights Reserved.
+
+#include "EventLoop/BSDSocket/EventLoopNotifierEventFd.h"
+
+#include "Async/Async.h"
+#include "EventLoop/EventLoopLog.h"
+#include "HAL//IConsoleManager.h"
+#include "Misc/QueuedThreadPool.h"
+#include "Stats/Stats.h"
+
+#if HAS_EVENTLOOP_NOTIFIER_EVENTFD
+#include "BSDSocketTypesPrivate.h"
+
+namespace UE::EventLoop {
+
+static bool bHttpEventLoopNotifierUseThreadPool = true;
+static FAutoConsoleVariableRef CVarHttpEventLoopNotifierUseThreadPool(
+	TEXT("http.EventLoopNotifierUseThreadPool"),
+	bHttpEventLoopNotifierUseThreadPool,
+	TEXT("When enabled, in event loop it will use thread pool to notify other thread to wake up"),
+	ECVF_SaveForNextBoot
+);
+
+std::atomic<bool> FEventLoopNotifierEventFd::bNotified = false;
+
+bool FEventLoopNotifierEventFd::Init()
+{
+	FileDescriptorRead = eventfd(0, EFD_NONBLOCK);
+	return FileDescriptorRead != INVALID_SOCKET;
+}
+
+void FEventLoopNotifierEventFd::Shutdown()
+{
+	if (FileDescriptorRead != INVALID_SOCKET)
+	{
+		close(FileDescriptorRead);
+		FileDescriptorRead = INVALID_SOCKET;
+	}
+}
+
+namespace EventLoopNotifierEventFd
+{
+
+static void NotifyImpl(int32 Fd)
+{
+	eventfd_t Value = 1;
+	if (eventfd_write(Fd, Value) == SOCKET_RESULT_FAILED)
+	{
+		int Err = errno;
+		UE_CLOG(Err != EAGAIN, LogEventLoop, Warning, TEXT("Failed to notify by event_fd_write %d with error %d"), Fd, Err);
+	}
+}
+
+}
+
+void FEventLoopNotifierEventFd::Notify()
+{
+	QUICK_SCOPE_CYCLE_COUNTER(STAT_FEventLoopNotifierEventFd_Notify);
+
+#if PLATFORM_WRITES_ARE_SLOW
+	if (bHttpEventLoopNotifierUseThreadPool && GIOThreadPool)
+	{ 
+		// When Notify() get called several times in same frame, we only need to notify once
+		if (!bNotified.exchange(true))
+		{
+			AsyncPool(*GIOThreadPool, [Fd = FileDescriptorRead] {
+				bNotified = false;
+				EventLoopNotifierEventFd::NotifyImpl(Fd);
+			});
+		}
+	}
+	else
+#endif
+	{
+		EventLoopNotifierEventFd::NotifyImpl(FileDescriptorRead);
+	}
+}
+
+void FEventLoopNotifierEventFd::Clear()
+{
+	eventfd_t Value = 0;
+	ensure(eventfd_read(FileDescriptorRead, &Value) != SOCKET_RESULT_FAILED);
+}
+
+/* UE::EventLoop */ }
+
+#endif // HAS_EVENTLOOP_NOTIFIER_EVENTFD
+
diff --git a/Engine/Source/Runtime/Online/EventLoop/Private/EventLoop/BSDSocket/EventLoopNotifierSocket.cpp b/Engine/Source/Runtime/Online/EventLoop/Private/EventLoop/BSDSocket/EventLoopNotifierSocket.cpp
new file mode 100644
--- /dev/null
+++ b/Engine/Source/Runtime/Online/EventLoop/Private/EventLoop/BSDSocket/EventLoopNotifierSocket.cpp
@@ -0,0 +1,157 @@
+// Copyright Epic Games, Inc. All Rights Reserved.
+
+#include "EventLoop/BSDSocket/EventLoopNotifierSocket.h"
+#include "BSDSocketTypesPrivate.h"
+#include "EventLoop/EventLoopLog.h"
+#include "Stats/Stats.h"
+
+#if PLATFORM_HAS_BSD_SOCKETS && !PLATFORM_HAS_BSD_SOCKET_FEATURE_WINSOCKETS
+
+namespace UE::EventLoop {
+
+namespace EventLoopNotifierSocket
+{
+
+static void SetSocketNonBlock(int32 FileDescriptor)
+{
+	int32 Flags = fcntl(FileDescriptor, F_GETFL, 0);
+	fcntl(FileDescriptor, F_SETFL, O_NONBLOCK | Flags );
+}
+
+}
+
+bool FEventLoopNotifierSocket::Init()
+{
+	// TODO: Validate and handle the case for socket recreation in case of suspend and resume
+
+	// TODO: Investigate if it's possible to use pipe or IPPROTO_RAW, or socketpair instead
+	FileDescriptorListener = socket(AF_INET, SOCK_STREAM, IPPROTO_TCP);
+	if (FileDescriptorListener == INVALID_SOCKET)
+	{
+		return false;
+	}
+	sockaddr_in AddrIn;
+	memset(&AddrIn, 0, sizeof(AddrIn));
+	AddrIn.sin_family = AF_INET;
+	AddrIn.sin_addr.s_addr = htonl(INADDR_LOOPBACK);
+	AddrIn.sin_port = 0;
+
+	if (bind(FileDescriptorListener, reinterpret_cast<sockaddr*>(&AddrIn), sizeof(AddrIn)) == SOCKET_RESULT_FAILED)
+	{
+		close(FileDescriptorListener);
+		FileDescriptorListener = INVALID_SOCKET;
+
+		return false;
+	}
+
+	if (listen(FileDescriptorListener, /*backlog*/1) == SOCKET_RESULT_FAILED)
+	{
+		close(FileDescriptorListener);
+		FileDescriptorListener = INVALID_SOCKET;
+
+		return false;
+	}
+
+	EventLoopNotifierSocket::SetSocketNonBlock(FileDescriptorListener);
+
+
+
+
+	// TODO: Set the write buffer to 1
+	sockaddr Addr;
+	memset(&Addr, 0, sizeof(Addr));
+
+	socklen_t AddrLen = sizeof(Addr);
+	getsockname(FileDescriptorListener, &Addr, &AddrLen);
+
+	FileDescriptorRead = socket(AF_INET, SOCK_STREAM, IPPROTO_TCP);
+	if (FileDescriptorRead == INVALID_SOCKET)
+	{
+		close(FileDescriptorListener);
+		FileDescriptorListener = INVALID_SOCKET;
+		return false;
+	}
+	if (connect(FileDescriptorRead, &Addr, sizeof(Addr)) == SOCKET_RESULT_FAILED)
+	{
+		close(FileDescriptorListener);
+		FileDescriptorListener = INVALID_SOCKET;
+
+		close(FileDescriptorRead);
+		FileDescriptorRead = INVALID_SOCKET;
+		return false;
+	}
+
+	EventLoopNotifierSocket::SetSocketNonBlock(FileDescriptorRead);
+
+
+
+
+	// TODO: Set the read buffer to 1
+	FileDescriptorWrite = accept(FileDescriptorListener, NULL, NULL);
+	if (FileDescriptorWrite == INVALID_SOCKET)
+	{
+		close(FileDescriptorListener);
+		FileDescriptorListener = INVALID_SOCKET;
+
+		close(FileDescriptorRead);
+		FileDescriptorRead = INVALID_SOCKET;
+		return false;
+	}
+
+	EventLoopNotifierSocket::SetSocketNonBlock(FileDescriptorWrite);
+
+
+
+
+	return true;
+}
+
+void FEventLoopNotifierSocket::Shutdown()
+{
+	if (FileDescriptorWrite != INVALID_SOCKET)
+	{
+		close(FileDescriptorWrite);
+		FileDescriptorWrite = INVALID_SOCKET;
+	}
+
+	if (FileDescriptorRead != INVALID_SOCKET)
+	{
+		close(FileDescriptorRead);
+		FileDescriptorRead = INVALID_SOCKET;
+	}
+
+	if (FileDescriptorListener != INVALID_SOCKET)
+	{
+		close(FileDescriptorListener);
+		FileDescriptorListener = INVALID_SOCKET;
+	}
+}
+
+void FEventLoopNotifierSocket::Notify()
+{
+	QUICK_SCOPE_CYCLE_COUNTER(STAT_FEventLoopNotifierSocket_Notify);
+	if (FileDescriptorWrite != INVALID_SOCKET)
+	{
+		// TODO: A failure to write may not be critical. Its possible that the receive buffer is just full. 
+		// The notifier just cares that the socket is set to a readable state.
+		if (write(FileDescriptorWrite, "1", 1) == SOCKET_RESULT_FAILED)
+		{
+			int Err = errno;
+			UE_CLOG(Err != EAGAIN, LogEventLoop, Warning, TEXT("Failed to notify by socket fd write %d with error %d"), FileDescriptorWrite, Err);
+		}
+	}
+}
+
+void FEventLoopNotifierSocket::Clear()
+{
+	if (FileDescriptorRead != INVALID_SOCKET)
+	{
+		char Buffer[1];
+		ensure(read(FileDescriptorRead, &Buffer[0], 1) != SOCKET_RESULT_FAILED);
+	}
+}
+
+/* UE::EventLoop */ }
+
+#endif // PLATFORM_HAS_BSD_SOCKETS && !PLATFORM_HAS_BSD_SOCKET_FEATURE_WINSOCKETS
+
diff --git a/Engine/Source/Runtime/Online/EventLoop/Private/EventLoop/EventLoopLog.cpp b/Engine/Source/Runtime/Online/EventLoop/Private/EventLoop/EventLoopLog.cpp
new file mode 100644
--- /dev/null
+++ b/Engine/Source/Runtime/Online/EventLoop/Private/EventLoop/EventLoopLog.cpp
@@ -0,0 +1,9 @@
+// Copyright Epic Games, Inc. All Rights Reserved.
+
+#include "EventLoop/EventLoopLog.h"
+
+namespace UE::EventLoop {
+
+DEFINE_LOG_CATEGORY(LogEventLoop);
+
+/* UE::EventLoop */ }
diff --git a/Engine/Source/Runtime/Online/EventLoop/Private/EventLoop/EventLoopLog.h b/Engine/Source/Runtime/Online/EventLoop/Private/EventLoop/EventLoopLog.h
new file mode 100644
--- /dev/null
+++ b/Engine/Source/Runtime/Online/EventLoop/Private/EventLoop/EventLoopLog.h
@@ -0,0 +1,11 @@
+// Copyright Epic Games, Inc. All Rights Reserved.
+
+#pragma once
+
+#include "Logging/LogMacros.h"
+
+namespace UE::EventLoop {
+
+DECLARE_LOG_CATEGORY_EXTERN(LogEventLoop, Log, All);
+
+/* UE::EventLoop */ }
diff --git a/Engine/Source/Runtime/Online/EventLoop/Private/EventLoop/EventLoopModule.cpp b/Engine/Source/Runtime/Online/EventLoop/Private/EventLoop/EventLoopModule.cpp
new file mode 100644
--- /dev/null
+++ b/Engine/Source/Runtime/Online/EventLoop/Private/EventLoop/EventLoopModule.cpp
@@ -0,0 +1,17 @@
+// Copyright Epic Games, Inc. All Rights Reserved.
+
+#include "EventLoop/EventLoopModule.h"
+
+IMPLEMENT_MODULE(UE::EventLoop::FEventLoopModule, EventLoop);
+
+namespace UE::EventLoop {
+
+void FEventLoopModule::StartupModule()
+{
+}
+
+void FEventLoopModule::ShutdownModule()
+{
+}
+
+/* UE::EventLoop */ }
diff --git a/Engine/Source/Runtime/Online/EventLoop/Public/EventLoop/BSDSocket/BSDSocketTypes.h b/Engine/Source/Runtime/Online/EventLoop/Public/EventLoop/BSDSocket/BSDSocketTypes.h
new file mode 100644
--- /dev/null
+++ b/Engine/Source/Runtime/Online/EventLoop/Public/EventLoop/BSDSocket/BSDSocketTypes.h
@@ -0,0 +1,19 @@
+// Copyright Epic Games, Inc. All Rights Reserved.
+
+#pragma once
+
+#include "CoreTypes.h"
+
+#if PLATFORM_HAS_BSD_SOCKETS
+
+#if PLATFORM_HAS_BSD_SOCKET_FEATURE_WINSOCKETS
+	typedef UPTRINT SOCKET;
+	#define INVALID_SOCKET  (SOCKET)(~0)
+#else // PLATFORM_HAS_BSD_SOCKET_FEATURE_WINSOCKETS
+	typedef int32 SOCKET;
+	#define INVALID_SOCKET -1
+#endif // PLATFORM_HAS_BSD_SOCKET_FEATURE_WINSOCKETS
+
+#define SOCKET_RESULT_FAILED -1
+
+#endif // PLATFORM_HAS_BSD_SOCKETS
\ No newline at end of file
diff --git a/Engine/Source/Runtime/Online/EventLoop/Public/EventLoop/BSDSocket/EventLoopIOAccessBSDSocket.h b/Engine/Source/Runtime/Online/EventLoop/Public/EventLoop/BSDSocket/EventLoopIOAccessBSDSocket.h
new file mode 100644
--- /dev/null
+++ b/Engine/Source/Runtime/Online/EventLoop/Public/EventLoop/BSDSocket/EventLoopIOAccessBSDSocket.h
@@ -0,0 +1,101 @@
+// Copyright Epic Games, Inc. All Rights Reserved.
+
+#pragma once
+
+#include "CoreTypes.h"
+
+#if PLATFORM_HAS_BSD_SOCKETS
+
+#include "Containers/UnrealString.h"
+#include "EventLoop/BSDSocket/BSDSocketTypes.h"
+#include "EventLoop/IEventLoopIOManager.h"
+#include "EventLoop/EventLoopManagedStorage.h"
+#include "Templates/Function.h"
+
+namespace UE::EventLoop {
+
+enum class ESocketIoRequestStatus : uint8
+{
+	/**
+	 * No error occurred, signaled flags are valid.
+	 * After signaling the handle associated with the request will remain valid until the user
+	 * deletes the associated request.
+	 */
+	Ok,
+
+	/**
+	 * There are no resources available to serve the IO request.
+	 * This will occur if more requests have been made than the system is capable of handling.
+	 * After signaling the handle associated with the request is removed.
+	 */
+	NoResources,
+
+	/**
+	 * The request is invalid.
+	 * A request may be invalid if another request for the same socket is already in progress.
+	 * After signaling the handle associated with the request is removed.
+	 */
+	Invalid,
+};
+
+using FSocketIOCallback = TUniqueFunction<void(SOCKET Socket, ESocketIoRequestStatus Status, EIOFlags SignaledFlags)>;
+using FSocketIORequestDestroyedCallback = FManagedStorageOnRemoveComplete;
+
+struct FIORequestBSDSocket
+{
+	SOCKET Socket = INVALID_SOCKET;
+	EIOFlags Flags = EIOFlags::None;
+	FSocketIOCallback Callback;
+};
+
+class FIOAccessBSDSocket final : public FNoncopyable
+{
+public:
+	struct FStorageTraits : public FManagedStorageDefaultTraits
+	{
+		using FExternalHandle = FIORequestHandle;
+	};
+
+	using FStorageType = TManagedStorage<FIORequestBSDSocket, FStorageTraits>;
+
+	FIOAccessBSDSocket(FStorageType& InIORequestStorage)
+		: IORequestStorage(InIORequestStorage)
+	{
+	}
+
+	FIORequestHandle CreateSocketIORequest(FIORequestBSDSocket&& Request)
+	{
+		if (Request.Socket == INVALID_SOCKET || Request.Flags == EIOFlags::None || !Request.Callback)
+		{
+			return FIORequestHandle();
+		}
+
+		return IORequestStorage.Add({MoveTemp(Request)});
+	}
+
+	void DestroyIORequest(FIORequestHandle& Handle, FSocketIORequestDestroyedCallback&& OnDestroyComplete = FSocketIORequestDestroyedCallback())
+	{
+		IORequestStorage.Remove(Handle, MoveTemp(OnDestroyComplete));
+	}
+
+private:
+	FStorageType& IORequestStorage;
+};
+
+/* UE::EventLoop */ }
+
+inline FString LexToString(UE::EventLoop::ESocketIoRequestStatus Status)
+{
+	switch (Status)
+	{
+	case UE::EventLoop::ESocketIoRequestStatus::Ok: return TEXT("Ok");
+	case UE::EventLoop::ESocketIoRequestStatus::NoResources: return TEXT("NoResources");
+	case UE::EventLoop::ESocketIoRequestStatus::Invalid: return TEXT("Invalid");
+
+	default:
+		checkNoEntry();
+		return TEXT("Unknown");
+	}
+}
+
+#endif // PLATFORM_HAS_BSD_SOCKETS
diff --git a/Engine/Source/Runtime/Online/EventLoop/Public/EventLoop/BSDSocket/EventLoopIOManagerBSDSocket.h b/Engine/Source/Runtime/Online/EventLoop/Public/EventLoop/BSDSocket/EventLoopIOManagerBSDSocket.h
new file mode 100644
--- /dev/null
+++ b/Engine/Source/Runtime/Online/EventLoop/Public/EventLoop/BSDSocket/EventLoopIOManagerBSDSocket.h
@@ -0,0 +1,21 @@
+// Copyright Epic Games, Inc. All Rights Reserved.
+
+#pragma once
+
+#include "CoreTypes.h"
+#include "EventLoop/BSDSocket/EventLoopIOManagerBSDSocketSelect.h"
+
+namespace UE::EventLoop {
+
+#if PLATFORM_HAS_BSD_SOCKETS
+#if HAS_EVENTLOOP_PLATFORM_SOCKET_IMPLEMENTATION
+	#include COMPILED_PLATFORM_HEADER(EventLoopIOManagerBSDSocket.h)
+#else
+#if PLATFORM_HAS_BSD_SOCKET_FEATURE_SELECT
+	// Default implementation.
+	using FIOManagerBSDSocket = FIOManagerBSDSocketSelect;
+#endif // PLATFORM_HAS_BSD_SOCKET_FEATURE_SELECT
+#endif // EVENTLOOP_PLATFORM_IMPLEMENTATION
+#endif // PLATFORM_HAS_BSD_SOCKETS
+
+/* UE::EventLoop */ }
diff --git a/Engine/Source/Runtime/Online/EventLoop/Public/EventLoop/BSDSocket/EventLoopIOManagerBSDSocketSelect.h b/Engine/Source/Runtime/Online/EventLoop/Public/EventLoop/BSDSocket/EventLoopIOManagerBSDSocketSelect.h
new file mode 100644
--- /dev/null
+++ b/Engine/Source/Runtime/Online/EventLoop/Public/EventLoop/BSDSocket/EventLoopIOManagerBSDSocketSelect.h
@@ -0,0 +1,45 @@
+// Copyright Epic Games, Inc. All Rights Reserved.
+
+#pragma once
+
+#include "CoreTypes.h"
+
+#if PLATFORM_HAS_BSD_SOCKETS && PLATFORM_HAS_BSD_SOCKET_FEATURE_SELECT
+
+#include "EventLoop/BSDSocket/BSDSocketTypes.h"
+#include "EventLoop/BSDSocket/EventLoopIOAccessBSDSocket.h"
+#include "EventLoop/IEventLoopIOManager.h"
+#include "EventLoop/EventLoopManagedStorage.h"
+#include "Templates/SharedPointer.h"
+
+namespace UE::EventLoop {
+
+class FIOManagerBSDSocketSelect final : public IIOManager
+{
+public:
+	using FIOAccess = FIOAccessBSDSocket;
+
+	struct FParams
+	{
+	};
+
+	struct FConfig
+	{
+	};
+
+	EVENTLOOP_API FIOManagerBSDSocketSelect(IEventLoop& EventLoop, FParams&& Params);
+	virtual ~FIOManagerBSDSocketSelect() = default;
+	EVENTLOOP_API virtual bool Init() override;
+	EVENTLOOP_API virtual void Shutdown() override;
+	EVENTLOOP_API virtual void Notify() override;
+	EVENTLOOP_API virtual void Poll(FTimespan WaitTime) override;
+
+	EVENTLOOP_API FIOAccess& GetIOAccess();
+
+private:
+	TSharedRef<class FIOManagerBSDSocketSelectImpl> Impl;
+};
+
+/* UE::EventLoop */ }
+
+#endif // PLATFORM_HAS_BSD_SOCKETS && PLATFORM_HAS_BSD_SOCKET_FEATURE_SELECT
diff --git a/Engine/Source/Runtime/Online/EventLoop/Public/EventLoop/BSDSocket/EventLoopNotifier.h b/Engine/Source/Runtime/Online/EventLoop/Public/EventLoop/BSDSocket/EventLoopNotifier.h
new file mode 100644
--- /dev/null
+++ b/Engine/Source/Runtime/Online/EventLoop/Public/EventLoop/BSDSocket/EventLoopNotifier.h
@@ -0,0 +1,19 @@
+// Copyright Epic Games, Inc. All Rights Reserved.
+
+#pragma once
+
+#include "CoreTypes.h"
+
+#if HAS_EVENTLOOP_NOTIFIER_EVENTFD
+#include "EventLoop/BSDSocket/EventLoopNotifierEventFd.h"
+namespace UE::EventLoop
+{
+using FEventLoopNotifier = FEventLoopNotifierEventFd;
+}
+#elif PLATFORM_HAS_BSD_SOCKETS && !PLATFORM_HAS_BSD_SOCKET_FEATURE_WINSOCKETS
+#include "EventLoop/BSDSocket/EventLoopNotifierSocket.h"
+namespace UE::EventLoop
+{
+using FEventLoopNotifier = FEventLoopNotifierSocket;
+}
+#endif
diff --git a/Engine/Source/Runtime/Online/EventLoop/Public/EventLoop/BSDSocket/EventLoopNotifierEventFd.h b/Engine/Source/Runtime/Online/EventLoop/Public/EventLoop/BSDSocket/EventLoopNotifierEventFd.h
new file mode 100644
--- /dev/null
+++ b/Engine/Source/Runtime/Online/EventLoop/Public/EventLoop/BSDSocket/EventLoopNotifierEventFd.h
@@ -0,0 +1,31 @@
+// Copyright Epic Games, Inc. All Rights Reserved.
+
+#pragma once
+
+#include "BSDSocketTypes.h"
+#include "CoreTypes.h"
+#include <atomic>
+
+#if HAS_EVENTLOOP_NOTIFIER_EVENTFD
+
+namespace UE::EventLoop {
+
+class FEventLoopNotifierEventFd
+{
+public:
+	bool Init();
+	void Shutdown();
+
+	void Notify();
+	void Clear();
+
+	int32 GetFileDescriptorRead() const { return FileDescriptorRead; }
+
+private:
+	int32 FileDescriptorRead = INVALID_SOCKET;
+	static std::atomic<bool> bNotified;
+};
+
+/* UE::EventLoop */ }
+
+#endif // HAS_EVENTLOOP_NOTIFIER_EVENTFD
diff --git a/Engine/Source/Runtime/Online/EventLoop/Public/EventLoop/BSDSocket/EventLoopNotifierSocket.h b/Engine/Source/Runtime/Online/EventLoop/Public/EventLoop/BSDSocket/EventLoopNotifierSocket.h
new file mode 100644
--- /dev/null
+++ b/Engine/Source/Runtime/Online/EventLoop/Public/EventLoop/BSDSocket/EventLoopNotifierSocket.h
@@ -0,0 +1,32 @@
+// Copyright Epic Games, Inc. All Rights Reserved.
+
+#pragma once
+
+#include "BSDSocketTypes.h"
+#include "CoreTypes.h"
+
+#if PLATFORM_HAS_BSD_SOCKETS && !PLATFORM_HAS_BSD_SOCKET_FEATURE_WINSOCKETS
+
+namespace UE::EventLoop {
+
+class FEventLoopNotifierSocket
+{
+public:
+	bool Init();
+	void Shutdown();
+
+	void Notify();
+	void Clear();
+
+	int32 GetFileDescriptorRead() const { return FileDescriptorRead; }
+
+private:
+	int32 FileDescriptorListener = INVALID_SOCKET;
+	int32 FileDescriptorWrite = INVALID_SOCKET;
+	int32 FileDescriptorRead = INVALID_SOCKET;
+};
+
+/* UE::EventLoop */ }
+
+#endif // PLATFORM_HAS_BSD_SOCKETS && !PLATFORM_HAS_BSD_SOCKET_FEATURE_WINSOCKETS
+
diff --git a/Engine/Source/Runtime/Online/EventLoop/Public/EventLoop/EventLoop.h b/Engine/Source/Runtime/Online/EventLoop/Public/EventLoop/EventLoop.h
new file mode 100644
--- /dev/null
+++ b/Engine/Source/Runtime/Online/EventLoop/Public/EventLoop/EventLoop.h
@@ -0,0 +1,334 @@
+// Copyright Epic Games, Inc. All Rights Reserved.
+
+#pragma once
+
+#include "CoreTypes.h"
+#include "Algo/ForEach.h"
+#include "Async/Mutex.h"
+#include "Containers/MpscQueue.h"
+#include "EventLoop/IEventLoop.h"
+#include "EventLoop/EventLoopTimer.h"
+#include "EventLoop/IEventLoopIOManager.h"
+#include "GenericPlatform/GenericPlatformProcess.h"
+#include "HAL/PlatformTime.h"
+#include "Math/UnrealMathUtility.h"
+#include "Misc/ScopeExit.h"
+#include "Stats/Stats.h"
+#include <limits>
+
+namespace UE::EventLoop {
+
+struct FEventLoopTraitsBase
+{
+	static inline uint32 GetCurrentThreadId()
+	{
+		return FPlatformTLS::GetCurrentThreadId();
+	}
+
+	static inline FTimespan GetCurrentTime()
+	{
+		return FTimespan(FPlatformTime::Seconds() * ETimespan::TicksPerSecond);
+	}
+
+	static inline bool IsEventLoopThread(uint32 EventLoopThreadId)
+	{
+		return EventLoopThreadId == GetCurrentThreadId();
+	}
+
+	static inline bool IsInitialized(uint32 EventLoopThreadId)
+	{
+		return (EventLoopThreadId != 0);
+	}
+
+	static inline void CheckInitialized(uint32 EventLoopThreadId)
+	{
+		check(IsInitialized(EventLoopThreadId));
+	}
+
+	static inline void CheckNotInitialized(uint32 EventLoopThreadId)
+	{
+		check(!IsInitialized(EventLoopThreadId));
+	}
+
+	static inline void CheckIsEventLoopThread(uint32 EventLoopThreadId)
+	{
+		check(IsEventLoopThread(EventLoopThreadId));
+	}
+
+	static inline bool IsShutdownRequested(bool bShutdownRequested)
+	{
+		return bShutdownRequested;
+	}
+};
+
+struct FEventLoopDefaultTraits : public FEventLoopTraitsBase
+{
+	using FTimerManagerTraits = FTimerManagerDefaultTraits;
+};
+
+template <typename IOManagerType, typename Traits = FEventLoopDefaultTraits>
+class TEventLoop final : public IEventLoop
+{
+private:
+	IOManagerType IOManager;
+	TTimerManager<typename Traits::FTimerManagerTraits> TimerManager;
+	TMpscQueue<FAsyncTask> AsyncTasks;
+	TArray<FOnShutdownComplete> ShutdownRequests;
+	FTimespan LoopTime;
+	TAtomic<uint32> EventLoopThreadId;
+	TAtomic<bool> bShutdownRequested;
+	TAtomic<bool> bShutdownCompleted;
+	TAtomic<bool> bNotifyOnInit;
+	UE::FMutex InitMutex;
+	UE::FMutex ShutdownMutex;
+
+public:
+	struct FParams
+	{
+		// Initial IO manager parameters.
+		typename IOManagerType::FParams IOManagerParams;
+	};
+
+	TEventLoop(FParams&& Params = FParams())
+		: IOManager(*this, MoveTemp(Params.IOManagerParams))
+		, TimerManager()
+		, AsyncTasks()
+		, LoopTime()
+		, EventLoopThreadId(0)
+		, bShutdownRequested(false)
+		, bShutdownCompleted(false)
+		, bNotifyOnInit(false)
+	{
+		// Make sure request manager type derives from IIOManager.
+		static_assert(std::is_base_of_v<IIOManager, IOManagerType> == true);
+	}
+
+	virtual ~TEventLoop() = default;
+
+	typename IOManagerType::FIOAccess& GetIOAccess()
+	{
+		return IOManager.GetIOAccess();
+	}
+
+	virtual bool Init() override
+	{
+		bool bInitSuccessful = false;
+
+		{
+			InitMutex.Lock();
+			ON_SCOPE_EXIT
+			{
+				InitMutex.Unlock();
+			};
+
+			Traits::CheckNotInitialized(EventLoopThreadId);
+			LoopTime = Traits::GetCurrentTime();
+			TimerManager.Init();
+			bInitSuccessful = IOManager.Init();
+			EventLoopThreadId = Traits::GetCurrentThreadId();
+		}
+
+		if (bNotifyOnInit)
+		{
+			IOManager.Notify();
+		}
+
+		return bInitSuccessful;
+	}
+
+	virtual void RequestShutdown(FOnShutdownComplete&& OnShutdownComplete = FOnShutdownComplete()) override
+	{
+		bool bRunLocally = false;
+
+		{
+			ShutdownMutex.Lock();
+			ON_SCOPE_EXIT
+			{
+				ShutdownMutex.Unlock();
+			};
+
+			// Check if this call to RequestShutdown is the first call.
+			bool bExpectedValue = false;
+			if (bShutdownRequested.CompareExchange(bExpectedValue, true))
+			{
+				// This is the first call to RequestShutdown. Notify the request manager to wake
+				// it so that the loop will see the request.
+				ShutdownRequests.Add(MoveTemp(OnShutdownComplete));
+				TryNotify();
+			}
+			else if (!bShutdownCompleted)
+			{
+				// Shutdown has been requested, but not yet acknowledged.
+				ShutdownRequests.Add(MoveTemp(OnShutdownComplete));
+			}
+			else
+			{
+				// Loop shutdown has already completed. Run the callback locally outside of the lock.
+				bRunLocally = true;
+			}
+		}
+
+		// Loop shutdown has already completed.
+		if (bRunLocally && OnShutdownComplete)
+		{
+			OnShutdownComplete();
+		}
+	}
+
+	virtual FTimerHandle SetTimer(FTimerCallback&& Callback, FTimespan InRate, bool InbRepeat = false, TOptional<FTimespan> InFirstDelay = TOptional<FTimespan>()) override
+	{
+		FTimerHandle OutHandle = TimerManager.SetTimer(MoveTemp(Callback), InRate, InbRepeat, InFirstDelay);
+		TryNotify();
+		return OutHandle;
+	}
+
+	virtual void ClearTimer(FTimerHandle& InHandle, FOnTimerCleared&& OnTimerCleared = FOnTimerCleared()) override
+	{
+		TimerManager.ClearTimer(InHandle, MoveTemp(OnTimerCleared));
+		TryNotify();
+	}
+
+	virtual void PostAsyncTask(FAsyncTask&& Task) override
+	{
+		AsyncTasks.Enqueue(MoveTemp(Task));
+		TryNotify();
+	}
+
+	virtual void Run() override
+	{
+		QUICK_SCOPE_CYCLE_COUNTER(STAT_EventLoop_TEventLoop_Run);
+		Traits::CheckInitialized(EventLoopThreadId);
+
+		while (true)
+		{
+			const FTimespan InfiniteWait(std::numeric_limits<int64>::max());
+			FTimespan WaitTime = InfiniteWait;
+
+			if (!RunOnce(WaitTime))
+			{
+				break;
+			}
+		};
+	}
+
+	virtual bool RunOnce(FTimespan WaitTime) override
+	{
+		QUICK_SCOPE_CYCLE_COUNTER(STAT_EventLoop_TEventLoop_RunOnce);
+		Traits::CheckInitialized(EventLoopThreadId);
+
+		if (bShutdownCompleted)
+		{
+			// Terminated.
+			return false;
+		}
+
+		// Check if there are any running timers. The maximum wait time will then be shortened
+		// to the expiration of the next timer.
+		FTimespan TimerExpirationTime;
+		if (TimerManager.GetNextTimeout(TimerExpirationTime))
+		{
+			WaitTime = FMath::Min(WaitTime, TimerExpirationTime);
+		}
+
+		IOManager.Poll(WaitTime);
+
+		// Run timers.
+		FTimespan CurrentTime = Traits::GetCurrentTime();
+		FTimespan ElapsedTime = FMath::Max(CurrentTime - LoopTime, 0.0);
+		LoopTime = CurrentTime;
+		TimerManager.Tick(ElapsedTime);
+
+		// If any timers ran due to the above call to Tick, check whether any repeat timers are
+		// waiting to be rescheduled. Advancing the loop time and ticking the timers again will
+		// allow those timers to reschedule without requiring polling the IO manager again.
+		if (TimerManager.HasPendingRepeatTimer())
+		{
+			CurrentTime = Traits::GetCurrentTime();
+			ElapsedTime = FMath::Max(CurrentTime - LoopTime, 0.0);
+			LoopTime = CurrentTime;
+			TimerManager.Tick(ElapsedTime);
+		}
+
+		RunAsyncTasks();
+
+		if (Traits::IsShutdownRequested(bShutdownRequested))
+		{
+			IOManager.Shutdown();
+
+			TArray<FOnShutdownComplete> LocalShutdownRequests;
+			{
+				ShutdownMutex.Lock();
+				ON_SCOPE_EXIT
+				{
+					ShutdownMutex.Unlock();
+				};
+
+				bShutdownCompleted = true;
+				LocalShutdownRequests = MoveTemp(ShutdownRequests);
+			}
+
+			// Run post-shutdown tasks.
+			Algo::ForEach(LocalShutdownRequests, [](FOnShutdownComplete& OnShutdownComplete){
+				if (OnShutdownComplete)
+				{
+					OnShutdownComplete();
+				}
+			});
+
+			// Terminated.
+			return false;
+		}
+
+		// Continue running.
+		return true;
+	}
+
+	virtual FTimespan GetLoopTime() const override
+	{
+		Traits::CheckInitialized(EventLoopThreadId);
+		return LoopTime;
+	}
+
+private:
+	void RunAsyncTasks()
+	{
+		FAsyncTask AsyncTask;
+		while (AsyncTasks.Dequeue(AsyncTask))
+		{
+			AsyncTask();
+		}
+	}
+
+	void TryNotify()
+	{
+		QUICK_SCOPE_CYCLE_COUNTER(STAT_TEventLoop_TryNotify);
+
+		if (Traits::IsInitialized(EventLoopThreadId))
+		{
+			// Init already completed. Notify the IO manager.
+			IOManager.Notify();
+		}
+		else
+		{
+			InitMutex.Lock();
+			ON_SCOPE_EXIT
+			{
+				InitMutex.Unlock();
+			};
+
+			if (!Traits::IsInitialized(EventLoopThreadId))
+			{
+				// Init has not yet started. Let it know to notify the IO manager of
+				// pending actions once it has completed.
+				bNotifyOnInit = true;
+			}
+			else
+			{
+				// Init competed before the lock was acquired. Notify the IO manager.
+				IOManager.Notify();
+			}
+		}
+	}
+};
+
+/* UE::EventLoop */ }
diff --git a/Engine/Source/Runtime/Online/EventLoop/Public/EventLoop/EventLoopHandle.h b/Engine/Source/Runtime/Online/EventLoop/Public/EventLoop/EventLoopHandle.h
new file mode 100644
--- /dev/null
+++ b/Engine/Source/Runtime/Online/EventLoop/Public/EventLoop/EventLoopHandle.h
@@ -0,0 +1,96 @@
+// Copyright Epic Games, Inc. All Rights Reserved.
+#pragma once
+
+#include "CoreTypes.h"
+#include "Containers/UnrealString.h"
+#include "Templates/Atomic.h"
+
+namespace UE::EventLoop {
+
+template <typename Traits>
+struct TResourceHandle final
+{
+	enum EGenerateNewHandleType
+	{
+		GenerateNewHandle
+	};
+
+	/** Creates an initially unset handle */
+	TResourceHandle()
+		: ID(0)
+	{
+	}
+
+	/** Creates a handle pointing to a new instance */
+	explicit TResourceHandle(EGenerateNewHandleType)
+		: ID(GenerateNewID())
+	{
+	}
+
+	/** True if this handle was ever initialized by the timer manager */
+	bool IsValid() const
+	{
+		return ID > 0;
+	}
+
+	/** Explicitly clear handle */
+	void Invalidate()
+	{
+		ID = 0;
+	}
+
+	bool operator==(const TResourceHandle& Other) const
+	{
+		return ID == Other.ID;
+	}
+
+	bool operator!=(const TResourceHandle& Other) const
+	{
+		return ID != Other.ID;
+	}
+
+	friend inline uint32 GetTypeHash(const TResourceHandle& InHandle)
+	{
+		return GetTypeHash(InHandle.ID);
+	}
+
+	FString ToString() const
+	{
+		return IsValid() ? FString::Printf(TEXT("%s:%llu"), Traits::Name, ID) : FString::Printf(TEXT("%s:<invalid>"), Traits::Name);
+	}
+
+	static const TCHAR* GetTypeName()
+	{
+		return Traits::Name;
+	}
+
+private:
+	/**
+	 * Generates a new ID for use the delegate handle.
+	 *
+	 * @return A unique ID for the delegate.
+	 */
+	static uint64 GenerateNewID()
+	{
+		// Just increment a counter to generate an ID.
+		uint64 Result = NextID++;
+
+		// Check for the next-to-impossible event that we wrap round to 0, because we reserve 0 for null delegates.
+		if (Result == 0)
+		{
+			// Increment it again - it might not be zero, so don't just assign it to 1.
+			Result = ++NextID;
+		}
+
+		return Result;
+	}
+
+	static TAtomic<uint64> NextID;
+
+	uint64 ID;
+};
+
+template <typename Traits>
+TAtomic<uint64> TResourceHandle<Traits>::NextID = 1;
+
+/* UE::EventLoop */ }
diff --git a/Engine/Source/Runtime/Online/EventLoop/Public/EventLoop/EventLoopIOManagerNull.h b/Engine/Source/Runtime/Online/EventLoop/Public/EventLoop/EventLoopIOManagerNull.h
new file mode 100644
--- /dev/null
+++ b/Engine/Source/Runtime/Online/EventLoop/Public/EventLoop/EventLoopIOManagerNull.h
@@ -0,0 +1,63 @@
+// Copyright Epic Games, Inc. All Rights Reserved.
+#pragma once
+
+#include "CoreTypes.h"
+#include "Async/ManualResetEvent.h"
+#include "EventLoop/IEventLoop.h"
+#include "EventLoop/IEventLoopIOManager.h"
+#include "Misc/Timespan.h"
+
+namespace UE::EventLoop {
+
+class FIOAccessNull final : public FNoncopyable
+{
+public:
+	// Nothing needed for now.
+};
+
+class FIOManagerNull final : public IIOManager
+{
+public:
+	using FIOAccess = FIOAccessNull;
+
+	struct FParams
+	{
+	};
+
+	FIOManagerNull(IEventLoop& EventLoop, FParams&& Params)
+	{
+	}
+
+	virtual ~FIOManagerNull() = default;
+
+	virtual bool Init() override
+	{
+		return true;
+	}
+
+	virtual void Shutdown() override
+	{
+	}
+
+	virtual void Notify() override
+	{
+		Event.Notify();
+	}
+
+	virtual void Poll(FTimespan WaitTime) override
+	{
+		Event.WaitFor(FMonotonicTimeSpan::FromSeconds(WaitTime.GetTotalSeconds()));
+		Event.Reset();
+	}
+
+	FIOAccess& GetIOAccess()
+	{
+		return IOAccess;
+	}
+
+private:
+	FIOAccess IOAccess;
+	UE::FManualResetEvent Event;
+};
+
+/* UE::EventLoop */ }
diff --git a/Engine/Source/Runtime/Online/EventLoop/Public/EventLoop/EventLoopManagedStorage.h b/Engine/Source/Runtime/Online/EventLoop/Public/EventLoop/EventLoopManagedStorage.h
new file mode 100644
--- /dev/null
+++ b/Engine/Source/Runtime/Online/EventLoop/Public/EventLoop/EventLoopManagedStorage.h
@@ -0,0 +1,513 @@
+// Copyright Epic Games, Inc. All Rights Reserved.
+#pragma once
+
+#include "CoreTypes.h"
+#include "Containers/Map.h"
+#include "Containers/SparseArray.h"
+#include "Containers/MpscQueue.h"
+#include "EventLoop/EventLoopHandle.h"
+#include "HAL/PlatformTLS.h"
+#include "Misc/CoreMiscDefines.h"
+#include "Misc/Optional.h"
+#include "Stats/Stats.h"
+#include "Templates/Function.h"
+
+namespace UE::EventLoop {
+
+using FManagedStorageOnRemoveComplete = TUniqueFunction<void()>;
+
+struct FManagedStorageDefaultExternalHandleTraits
+{
+	static constexpr TCHAR Name[] = TEXT("ManagedStorageDefault");
+};
+
+struct FManagedStorageTraitsBase
+{
+	static inline uint32 GetCurrentThreadId()
+	{
+		return FPlatformTLS::GetCurrentThreadId();
+	}
+
+	static inline bool IsManagerThread(uint32 ManagerThreadId)
+	{
+		return ManagerThreadId == GetCurrentThreadId();
+	}
+
+	static inline void CheckNotInitialized(uint32 ManagerThreadId)
+	{
+		check(ManagerThreadId == 0);
+	}
+
+	static inline void CheckIsManagerThread(uint32 ManagerThreadId)
+	{
+		check(IsManagerThread(ManagerThreadId));
+	}
+};
+
+/*
+ * Default traits for managed storage. Traits are used to implement functionality which can be
+ * mocked to allow testing the class.
+ *
+ * In most cases the default traits can be used without modification.
+ */
+struct FManagedStorageDefaultTraits : public FManagedStorageTraitsBase
+{
+
+	/**
+	 * The allocator to use for the internal handle array.
+	 * 
+	 * An internal handle array is passed to Update to get the handles added since the last call.
+	 * Using inline storage can prevent an unnecessary allocation when retrieving this info.
+	 */
+	using InternalHandleArryAllocatorType = TInlineAllocator<32>;
+
+	/**
+	 * Whether to check thread safety when accessing internal storage.
+	 * Defaulted to off for performance, but can be added to implementation traits for debugging.
+	 */
+	static constexpr bool bStorageAccessThreadChecksEnabled = false;
+
+	/**
+	 * The type of handle which will ne externally usable.
+	 */
+	using FExternalHandle = TResourceHandle<FManagedStorageDefaultExternalHandleTraits>;
+};
+
+/**
+ * Managed storage class which allows resource creation and removal from any thread.
+ * The owning thread may access the storage without thread safety concerns.
+ */
+template <typename ElementType, typename Traits = FManagedStorageDefaultTraits>
+class TManagedStorage final : public FNoncopyable
+{
+public:
+	/*
+	 * Handle which extends the external resource handle to allow O(1) access into the internal storage array.
+	 */
+	template <typename ExternalHandleType>
+	struct TManagedStorageInternalHandle
+	{
+		enum EGenerateNewHandleType
+		{
+			GenerateNewHandle
+		};
+
+		/** Creates an initially unset handle */
+		TManagedStorageInternalHandle()
+			: ExternalHandle()
+			, InternalIndex(INDEX_NONE)
+		{
+		}
+
+		/** Creates a handle pointing to a new instance */
+		explicit TManagedStorageInternalHandle(EGenerateNewHandleType, ExternalHandleType InExternalHandle, int32 InInternalIndex)
+			: ExternalHandle(InExternalHandle)
+			, InternalIndex(InInternalIndex)
+		{
+		}
+
+		/** True if this handle was ever initialized by the timer manager */
+		bool IsValid() const
+		{
+			return ExternalHandle.IsValid();
+		}
+
+		/** Explicitly clear handle */
+		void Invalidate()
+		{
+			ExternalHandle.Invalidate();
+		}
+
+		bool operator==(const TManagedStorageInternalHandle& Other) const
+		{
+			return ExternalHandle == Other.ExternalHandle;
+		}
+
+		bool operator!=(const TManagedStorageInternalHandle& Other) const
+		{
+			return ExternalHandle != Other.ExternalHandle;
+		}
+
+		friend inline uint32 GetTypeHash(const TManagedStorageInternalHandle& InHandle)
+		{
+			return GetTypeHash(InHandle.ExternalHandle);
+		}
+
+		FString ToString() const
+		{
+			return IsValid() ? FString::Printf(TEXT("Internal:[ext:%s, idx:%d]"), *ExternalHandle.ToString(), InternalIndex) : FString("<invalid>");
+		}
+
+		/** Gets the external view of this handle. */
+		ExternalHandleType GetExternalHandle() const
+		{
+			return ExternalHandle;
+		}
+
+		/** Gets the internal storage index for the handles data. */
+		int32 GetInternalIndex() const
+		{
+			return InternalIndex;
+		}
+
+	private:
+		ExternalHandleType ExternalHandle;
+		int32 InternalIndex;
+	};
+
+	using FExternalHandle = typename Traits::FExternalHandle;
+	using FInternalHandle = TManagedStorageInternalHandle<FExternalHandle>;
+	using FInternalHandleArryType = TArray<FInternalHandle, typename Traits::InternalHandleArryAllocatorType>;
+
+private:
+	struct FAddRequest
+	{
+		FExternalHandle Handle;
+		ElementType Data;
+	};
+
+	struct FRemoveRequest
+	{
+		FExternalHandle Handle;
+		FManagedStorageOnRemoveComplete OnComplete;
+	};
+
+	struct FStorageEntry
+	{
+		FInternalHandle InternalHandle;
+		ElementType Data;
+	};
+
+	using FStorageType = TSparseArray<FStorageEntry>;
+
+	/** The internal storage of entries. */
+	FStorageType Storage;
+	/* Mapping of external handles to the array index of the element in Storage. */
+	TMap<FExternalHandle, int32> StorageHandleIndex;
+	/** Requests to add entries to storage. Processed in Update(). */
+	TMpscQueue<FAddRequest> AddRequests;
+	/** Requests to remove entries from storage. Processed in Update(). */
+	TMpscQueue<FRemoveRequest> RemoveRequests;
+	/**
+	 * The thread id set when calling Init. This thread becomes the 'Manager' thread and is safe
+	 * to call non-thread safe methods from.
+	 */
+	TAtomic<uint32> ManagerThreadId;
+
+public:
+	TManagedStorage()
+		: Storage()
+		, StorageHandleIndex()
+		, AddRequests()
+		, RemoveRequests()
+		, ManagerThreadId(0)
+	{
+	}
+
+	~TManagedStorage() = default;
+
+	/**
+	 * External access for checking whether the current thread is the manager thread.
+	 */
+	inline bool IsManagerThread() const
+	{
+		return Traits::IsManagerThread(ManagerThreadId);
+	}
+
+	/**
+	 * Initialize storage.
+	 * 
+	 * NOT thread safe.
+	 */
+	void Init()
+	{
+		Traits::CheckNotInitialized(ManagerThreadId);
+		ManagerThreadId = Traits::GetCurrentThreadId();
+	}
+
+	 /**
+	  * Process all pending storage modifications.
+	  * 
+	  * NOT thread safe.
+	  * 
+	  * @param OutAddedHandles	Handles which were added since the last call to Update.
+	  * 						Element may not be present in the internal storage if element was
+	  *							also removed since the last call to update.
+	  * @param OutRemovedHandles Handles which were removed since the last update.
+	  * @return The number of changes which occurred.
+	  */
+	uint32 Update(FInternalHandleArryType* OutAddedHandles = nullptr, FInternalHandleArryType* OutRemovedHandles = nullptr)
+	{
+		QUICK_SCOPE_CYCLE_COUNTER(STAT_EventLoop_TManagedStorage_Update);
+
+		Traits::CheckIsManagerThread(ManagerThreadId);
+
+		uint32 NumChanges = 0;
+
+		// Add new queued items.
+		if (OutAddedHandles)
+		{
+			FAddRequest AddRequest;
+			while (AddRequests.Dequeue(AddRequest))
+			{
+				OutAddedHandles->Add(AddImpl(AddRequest.Handle, MoveTemp(AddRequest.Data)));
+				++NumChanges;
+			}
+		}
+		else
+		{
+			FAddRequest AddRequest;
+			while (AddRequests.Dequeue(AddRequest))
+			{
+				AddImpl(AddRequest.Handle, MoveTemp(AddRequest.Data));
+				++NumChanges;
+			}
+		}
+
+		// Remove queued items.
+		if (OutRemovedHandles)
+		{
+			FRemoveRequest RemoveRequest;
+			while (RemoveRequests.Dequeue(RemoveRequest))
+			{
+				OutRemovedHandles->Add(RemoveImpl(RemoveRequest.Handle, RemoveRequest.OnComplete));
+				++NumChanges;
+			}
+		}
+		else
+		{
+			FRemoveRequest RemoveRequest;
+			while (RemoveRequests.Dequeue(RemoveRequest))
+			{
+				RemoveImpl(RemoveRequest.Handle, RemoveRequest.OnComplete);
+				++NumChanges;
+			}
+		}
+
+		return NumChanges;
+	}
+
+	 /**
+	  * Requests to add a resource.
+	  * 
+	  * Thread safe.
+	  *
+	  * @param Data The data to be added and managed.
+	  * @return A new handle for the added resource.
+	  */
+	FExternalHandle Add(ElementType&& Data)
+	{
+		// Assign new handle for data.
+		FExternalHandle Handle(FExternalHandle::GenerateNewHandle);
+		AddRequests.Enqueue(FAddRequest{Handle, MoveTemp(Data)});
+		return Handle;
+	}
+
+	/**
+	 * Requests to remove a resource.
+	 * 
+	 * Thread safe.
+	 *
+	 * @param Handle The handle for the resource to be removed.
+	 * @param OnComplete A callback to be signaled once the resource has been destroyed.
+	 */
+	void Remove(const FExternalHandle Handle, FManagedStorageOnRemoveComplete&& OnRemoveComplete = FManagedStorageOnRemoveComplete())
+	{
+		RemoveRequests.Enqueue(FRemoveRequest{Handle, MoveTemp(OnRemoveComplete)});
+	}
+
+	/**
+	 * Removes a resource directly using its internal handle.
+	 * 
+	 * NOT thread safe.
+	 *
+	 * @param Handle The handle for the resource to be removed.
+	 */
+	void Remove(const FInternalHandle Handle)
+	{
+		RemoveImpl(Handle.GetExternalHandle(), nullptr);
+	}
+
+public:
+	/*
+	 * STORAGE ACCESS
+	 * 
+	 * NOT thread safe.
+	 */
+
+	inline int32 Num() const
+	{
+		CheckThreadForStorageAccess();
+		return Storage.Num();
+	}
+
+	inline bool IsEmpty() const
+	{
+		CheckThreadForStorageAccess();
+		return Storage.IsEmpty();
+	}
+
+	inline ElementType* Find(FExternalHandle Handle)
+	{
+		CheckThreadForStorageAccess();
+
+		// When using an external handle the internal index will be valid if present.
+		if (int32* InternalIndex = StorageHandleIndex.Find(Handle))
+		{
+			FStorageEntry& Entry = Storage[*InternalIndex];
+			return &Entry.Data;
+		}
+
+		return nullptr;
+	}
+
+	inline const ElementType* Find(FExternalHandle Handle) const
+	{
+		return const_cast<TManagedStorage*>(this)->Find(Handle);
+	}
+
+	inline ElementType* Find(FInternalHandle Handle)
+	{
+		CheckThreadForStorageAccess();
+
+		// When using an internal handle the index must be verified in case the handle is stale.
+		const int32 InternalIndex = Handle.GetInternalIndex();
+		if (Storage.IsValidIndex(InternalIndex))
+		{
+			// Check that the internal handle matches the expected value in case
+			// the slot in the array was reused.
+			FStorageEntry& Entry = Storage[InternalIndex];
+			return Entry.InternalHandle == Handle ? &Entry.Data : nullptr;
+		}
+
+		return nullptr;
+	}
+
+	inline const ElementType* Find(FInternalHandle Handle) const
+	{
+		return const_cast<TManagedStorage*>(this)->Find(Handle);
+	}
+
+	/**
+	 * DO NOT USE DIRECTLY
+	 * STL-like iterators to enable range-based for loop support.
+	 * 
+	 * NOT thread safe.
+	 */
+	template <bool bConst>
+	class TBaseRangeForIterator
+	{
+	private:
+		using TInternalIterator = typename FStorageType::TRangedForIterator;
+		using StorageIteratorType = std::conditional_t<bConst, typename FStorageType::TRangedForConstIterator, typename FStorageType::TRangedForIterator>;
+		using InternalElementType = std::conditional_t<bConst, const ElementType, ElementType>;
+		using ItElementType = TPair<FInternalHandle, InternalElementType&>;
+
+	public:
+		explicit TBaseRangeForIterator(StorageIteratorType InStorageIterator)
+			: StorageIterator(InStorageIterator)
+		{
+			if (StorageIterator)
+			{
+				ElementAccess.Emplace(ItElementType{StorageIterator->InternalHandle, StorageIterator->Data});
+			}
+		}
+
+		inline TBaseRangeForIterator& operator++()
+		{
+			++StorageIterator;
+			if (StorageIterator)
+			{
+				ElementAccess.Emplace(ItElementType{StorageIterator->InternalHandle, StorageIterator->Data});
+			}
+			return *this;
+		}
+
+		inline bool operator==(const TBaseRangeForIterator& Rhs) const { return StorageIterator == Rhs.StorageIterator; }
+		inline bool operator!=(const TBaseRangeForIterator& Rhs) const { return StorageIterator != Rhs.StorageIterator; }
+
+		/** conversion to "bool" returning true if the iterator is valid. */
+		inline explicit operator bool() const
+		{
+			return !!StorageIterator;
+		}
+
+		/** inverse of the "bool" operator */
+		inline bool operator !() const
+		{
+			return !(bool)*this;
+		}
+
+		inline ItElementType& operator*() const { return *ElementAccess; }
+		inline ItElementType* operator->() const { return &*ElementAccess; }
+
+	private:
+		StorageIteratorType StorageIterator;
+		mutable TOptional<ItElementType> ElementAccess;
+	};
+
+	using TRangedForIterator = TBaseRangeForIterator<false>;
+	using TRangedForConstIterator = TBaseRangeForIterator<true>;
+
+	inline TRangedForIterator begin()
+	{
+		CheckThreadForStorageAccess();
+		return TRangedForIterator(Storage.begin());
+	}
+
+	inline TRangedForConstIterator begin() const
+	{
+		CheckThreadForStorageAccess();
+		return TRangedForConstIterator(Storage.begin());
+	}
+
+	inline TRangedForIterator end()
+	{
+		CheckThreadForStorageAccess();
+		return TRangedForIterator(Storage.end());
+	}
+
+	inline TRangedForConstIterator end() const
+	{
+		CheckThreadForStorageAccess();
+		return TRangedForConstIterator(Storage.end());
+	}
+
+private:
+	inline void CheckThreadForStorageAccess() const
+	{
+		if constexpr (Traits::bStorageAccessThreadChecksEnabled)
+		{
+			Traits::CheckIsManagerThread(ManagerThreadId);
+		}
+	}
+
+	FInternalHandle AddImpl(const FExternalHandle Handle, ElementType&& Data)
+	{
+		FSparseArrayAllocationInfo Allocation = Storage.AddUninitialized();
+		FInternalHandle InternalHandle(FInternalHandle::GenerateNewHandle, Handle, Allocation.Index);
+		new(Allocation) FStorageEntry{ InternalHandle, MoveTemp(Data)};
+		StorageHandleIndex.Add(Handle, Allocation.Index);
+		return InternalHandle;
+	}
+
+	FInternalHandle RemoveImpl(const FExternalHandle Handle, const FManagedStorageOnRemoveComplete& OnComplete)
+	{
+		int32 FoundIndex = INDEX_NONE;
+		FInternalHandle InternalHandle;
+		if (StorageHandleIndex.RemoveAndCopyValue(Handle, FoundIndex) && Storage.IsValidIndex(FoundIndex))
+		{
+			InternalHandle = Storage[FoundIndex].InternalHandle;
+			Storage.RemoveAt(FoundIndex);
+		}
+
+		if (OnComplete)
+		{
+			OnComplete();
+		}
+
+		return InternalHandle;
+	}
+};
+
+/* UE::EventLoop */ }
diff --git a/Engine/Source/Runtime/Online/EventLoop/Public/EventLoop/EventLoopModule.h b/Engine/Source/Runtime/Online/EventLoop/Public/EventLoop/EventLoopModule.h
new file mode 100644
--- /dev/null
+++ b/Engine/Source/Runtime/Online/EventLoop/Public/EventLoop/EventLoopModule.h
@@ -0,0 +1,39 @@
+// Copyright Epic Games, Inc. All Rights Reserved.
+
+#pragma once
+
+#include "CoreTypes.h"
+#include "Modules/ModuleInterface.h"
+#include "Modules/ModuleManager.h"
+
+namespace UE::EventLoop {
+
+class FEventLoopModule : public IModuleInterface
+{
+public:
+	// IModuleInterface
+
+	/**
+	 * Called right after the module DLL has been loaded and the module object has been created
+	 * Overloaded to allow the default subsystem a chance to load
+	 */
+	virtual void StartupModule() override;
+
+	/**
+	 * Called before the module is unloaded, right before the module object is destroyed.
+	 * Overloaded to shut down all loaded online subsystems
+	 */
+	virtual void ShutdownModule() override;
+
+	/**
+	 * Override this to set whether your module is allowed to be unloaded on the fly
+	 *
+	 * @return	Whether the module supports shutdown separate from the rest of the engine.
+	 */
+	virtual bool SupportsDynamicReloading() override
+	{
+		return false;
+	}
+};
+
+/* UE::EventLoop */ }
diff --git a/Engine/Source/Runtime/Online/EventLoop/Public/EventLoop/EventLoopTimer.h b/Engine/Source/Runtime/Online/EventLoop/Public/EventLoop/EventLoopTimer.h
new file mode 100644
--- /dev/null
+++ b/Engine/Source/Runtime/Online/EventLoop/Public/EventLoop/EventLoopTimer.h
@@ -0,0 +1,428 @@
+// Copyright Epic Games, Inc. All Rights Reserved.
+#pragma once
+
+#include "CoreTypes.h"
+#include "EventLoop/EventLoopManagedStorage.h"
+#include "Templates/Function.h"
+#include "Misc/ScopeExit.h"
+#include "Misc/Timespan.h"
+#include "Stats/Stats.h"
+
+namespace UE::EventLoop {
+
+struct FTimerHandleTraits
+{
+	static constexpr TCHAR Name[] = TEXT("TimerHandle");
+};
+
+using FTimerCallback = TUniqueFunction<void()>;
+
+using FOnTimerCleared = FManagedStorageOnRemoveComplete;
+
+using FTimerHandle = TResourceHandle<FTimerHandleTraits>;
+
+enum class ETimerStatus : uint8
+{
+	/**
+	 * Timer has been queued, but has not started running yet.
+	 * Timer will become active the next time the timer manager is ticked.
+	 */
+	Pending,
+	/**
+	 * Timer is active and will be triggered once its expiration elapses in Tick.
+	 */
+	Active,
+	/**
+	 * Timer will be removed during the next tick.
+	 */
+	PendingRemoval,
+	/**
+	 * Timer will be rescheduled during the next tick.
+	 */
+	PendingReschedule,
+	/**
+	 * The timer has elapsed and its callback is executing.
+	 */
+	Executing,
+};
+
+struct FTimerData
+{
+	FTimerData(FTimerCallback&& InCallback, TOptional<FTimespan> InFirstDelay, FTimespan InRate, bool bInRepeat)
+		: Callback(MoveTemp(InCallback))
+		, FirstDelay(InFirstDelay)
+		, Rate(InRate)
+		, bRepeat(bInRepeat)
+	{
+	}
+
+	/** Holds the callback to call. */
+	FTimerCallback Callback;
+
+	/** Time to delay the first execution of the timer. Relevant to looping timers. */
+	TOptional<FTimespan> FirstDelay;
+
+	/** Time between set and fire, or repeat frequency if looping. */
+	FTimespan Rate;
+
+	/** When the timer is active, the expiration is the absolute time relative to the Timer Manager at which the timer should be fired. */
+	FTimespan Expiration;
+
+	ETimerStatus Status = ETimerStatus::Pending;
+
+	/** If true, this timer will repeat indefinitely. Otherwise, it will be destroyed when it expires. */
+	uint8 bRepeat : 1;
+
+	FTimerData() = default;
+
+	// Movable only
+	FTimerData(FTimerData&&) = default;
+	FTimerData(const FTimerData&) = delete;
+	FTimerData& operator=(FTimerData&&) = default;
+	FTimerData& operator=(const FTimerData&) = delete;
+};
+
+struct FTimerManagerTraitsBase
+{
+	static inline void CheckIsManagerThread(bool bIsManagerThread)
+	{
+		check(bIsManagerThread);
+	}
+};
+
+/*
+ * Default traits for timer manager. Traits are used to implement functionality which can be
+ * mocked to allow testing the class.
+ *
+ * In most cases the default traits can be used without modification.
+ */
+struct FTimerManagerDefaultTraits : public FTimerManagerTraitsBase
+{
+	/**
+	 * The allocator to use for the timer heap.
+	 */
+	using FTimerHeapAllocatorType = TInlineAllocator<32>;
+
+	/**
+	 * The allocator to use for timers waiting to be rescheduled.
+	 */
+	using FTimerRepeatAllocatorType = TInlineAllocator<32>;
+
+	/**
+	 * Default traits used for timer storage.
+	 */
+	struct FStorageTraits : public FManagedStorageDefaultTraits
+	{
+		using FExternalHandle = FTimerHandle;
+	};
+};
+
+/**
+ * Timer manager for posting timers to an event loop.
+ * 
+ * The event loop timer manager provides the guarantee that its timers will never fire early while
+ * only progressing its view of time based on the delta passed to its Tick method. To provide this
+ * guarantee, timers may fire later than requested which is by design.
+ * 
+ * Adding and removing timers are thread-safe, however the timer callback and clear callbacks will
+ * both be called from within the Tick method, which may be a different thread. It is up to the
+ * user to handle thread safety within their callbacks.
+ */
+template <typename Traits = FTimerManagerDefaultTraits>
+class TTimerManager final : public FNoncopyable
+{
+public:
+	using FStorageType = TManagedStorage<FTimerData, typename Traits::FStorageTraits>;
+	using FTimerHandle = typename FStorageType::FExternalHandle;
+	using FInternalHandle = typename FStorageType::FInternalHandle;
+	using FInternalHandleArryType = typename FStorageType::FInternalHandleArryType;
+
+public:
+	/**
+	 * Initialize the timer manager and its storage.
+	 *
+	 * NOT thread safe.
+	 */
+	void Init()
+	{
+		Storage.Init();
+	}
+
+	/**
+	 * Accumulate the timer managers elapsed time and fire timers which have elapsed.
+	 *
+	 * NOT thread safe.
+	 * 
+	 * @param DeltaTime The amount of time which has elapsed since the last call to Tick.
+	 */
+	void Tick(FTimespan DeltaTime);
+
+	/**
+	 * Retrieve the remaining time since the last call to Tick before the next timer will fire.
+	 *
+	 * NOT thread safe.
+	 * 
+	 * @param OutTimeout Time amount of time relative to TimerManagers' internal time until the next valid timer will run.
+	 * @return true if a valid timer is currently set.
+	 */
+	inline bool GetNextTimeout(FTimespan& OutTimeout) const
+	{
+		if (!ActiveTimerHeap.IsEmpty())
+		{
+			FInternalHandle TopHandle = ActiveTimerHeap.HeapTop();
+			const FTimerData* TimerData = Storage.Find(TopHandle);
+			OutTimeout = TimerData->Expiration - InternalTime;
+			return true;
+		}
+
+		return false;
+	}
+
+	/**
+	 * Returns the total number of timers. Includes active timers and timers waiting to be rescheduled.
+	 *
+	 * NOT thread safe.
+	 * 
+	 * @return the total number of timers.
+	 */
+	inline uint32 GetNumTimers() const
+	{
+		return Storage.Num();
+	}
+
+	/**
+	 * Set a new timer. The timer callback will be triggered from within the call to Tick,
+	 * which may occur on a different thread from the one which set the timer.
+	 * 
+	 * Thread safe.
+	 *
+	 * @param Callback Callback to call when timer fires.
+	 * @param InRate The amount of time between set and firing.
+	 * @param InbRepeat true to keep firing at Rate intervals, false to fire only once.
+	 * @param InFirstDelay The time for the first iteration of a looping timer.
+	 * @return handle to the registered timer.
+	 */
+	inline FTimerHandle SetTimer(FTimerCallback&& Callback, FTimespan InRate, bool InbRepeat = false, TOptional<FTimespan> InFirstDelay = TOptional<FTimespan>())
+	{
+		// Callback must be set.
+		const bool bValidCallback = Callback.operator bool();
+		// Rate cannot be less than zero.
+		const bool bValidRate = InRate >= FTimespan::Zero();
+		// First delay is only valid when repeating.
+		const bool bFirstDelaySetValid = !InFirstDelay.IsSet() || InbRepeat;
+		// First delay cannot be less than zero.
+		const bool bFirstDelayValidTime = !InFirstDelay.IsSet() || (*InFirstDelay >= FTimespan::Zero());
+
+		if (!(bValidCallback && bValidRate && bFirstDelaySetValid && bFirstDelayValidTime))
+		{
+			return FTimerHandle();
+		}
+
+		return Storage.Add(FTimerData(MoveTemp(Callback), InFirstDelay, InRate, InbRepeat));
+	}
+
+	/**
+	* Clears a previously set timer.
+	* 
+	* Thread safe.
+	*
+	* @param InHandle The handle of the timer to clear.
+	* @param OnTimerCleared Callback to be fired when the timer has been removed. The callback will
+	*                       be fired within the Tick method which may be a different thread from
+	*                       which the timer clear was requested.
+	*/
+	inline void ClearTimer(FTimerHandle& InHandle, FOnTimerCleared&& OnTimerCleared = FOnTimerCleared())
+	{
+		// If ClearTimer is run from the manager thread and the timer manager is currently ticking,
+		// mark the timer as pending removal so that it will not be run.
+		if (Storage.IsManagerThread() && bIsTicking)
+		{
+			if (FTimerData* TimerData = Storage.Find(InHandle))
+			{
+				TimerData->Status = ETimerStatus::PendingRemoval;
+			}
+		}
+
+		// Queue removal request.
+		Storage.Remove(InHandle, MoveTemp(OnTimerCleared));
+
+		// Invalidate the handle.
+		InHandle.Invalidate();
+	}
+
+	/**
+	 * Returns whether any timers are waiting to be rescheduled.
+	 *
+	 * NOT thread safe.
+	 * 
+	 * @return true if a valid timer is waiting to be rescheduled.
+	 */
+	inline bool HasPendingRepeatTimer() const
+	{
+		return !PendingRepeatTimers.IsEmpty();
+	}
+
+private:
+	struct FTimerHeapOrder
+	{
+		explicit FTimerHeapOrder(const FStorageType& InStorage)
+			: Storage(InStorage)
+			, NumTimers(InStorage.Num())
+		{
+		}
+
+		bool operator()(FInternalHandle LhsHandle, FInternalHandle RhsHandle) const
+		{
+			const FTimerData* LhsData = Storage.Find(LhsHandle);
+			const FTimerData* RhsData = Storage.Find(RhsHandle);
+
+			// Move to top if empty.
+			if (LhsData == nullptr)
+			{
+				return true;
+			}
+
+			// Move to top if empty.
+			if (RhsData == nullptr)
+			{
+				return false;
+			}
+
+			return LhsData->Expiration < RhsData->Expiration;
+		}
+
+		const FStorageType& Storage;
+		int32 NumTimers;
+	};
+
+	/** Storage for timer entries. */
+	FStorageType Storage;
+	/** Heap of actively running timers. */
+	TArray<FInternalHandle, typename Traits::FTimerHeapAllocatorType> ActiveTimerHeap;
+	/** Timers which will be rescheduled on the next call to Tick. */
+	TArray<FInternalHandle, typename Traits::FTimerRepeatAllocatorType> PendingRepeatTimers;
+	/** An internally consistent clock. Advances during ticking. */
+	FTimespan InternalTime;
+	/** Track whether the timer manager is in the process of ticking. */
+	bool bIsTicking = false;
+};
+
+template <typename Traits>
+void TTimerManager<Traits>::Tick(FTimespan DeltaTime)
+{
+	QUICK_SCOPE_CYCLE_COUNTER(STAT_EventLoop_TTimerManager_Tick);
+
+	Traits::CheckIsManagerThread(Storage.IsManagerThread());
+
+	bIsTicking = true;
+	ON_SCOPE_EXIT{ bIsTicking = false; };
+
+	// Progress internal time.
+	InternalTime += DeltaTime;
+
+	// Process queued actions.
+	FInternalHandleArryType AddedHandles;
+	FInternalHandleArryType RemovedHandles;
+	Storage.Update(&AddedHandles, &RemovedHandles);
+
+	{
+		QUICK_SCOPE_CYCLE_COUNTER(STAT_EventLoop_TTimerManager_Tick_AddTimers);
+
+		// Activate new timers.
+		// Timers are always added on the next call to Tick regardless of thread to provide the
+		// guarantee that a timer may fire late, but never early.
+		for (FInternalHandle InternalHandle : AddedHandles)
+		{
+			// Timer data may not be found if the timer was added and removed before Tick was called.
+			if (FTimerData* TimerData = Storage.Find(InternalHandle))
+			{
+				const FTimespan EffectiveInitialRate = TimerData->FirstDelay ? *TimerData->FirstDelay : TimerData->Rate;
+				TimerData->Expiration = InternalTime + EffectiveInitialRate;
+				TimerData->Status = ETimerStatus::Active;
+				ActiveTimerHeap.HeapPush(InternalHandle, FTimerHeapOrder(Storage));
+			}
+		}
+	}
+
+	{
+		QUICK_SCOPE_CYCLE_COUNTER(STAT_EventLoop_TTimerManager_Tick_RemoveTimers);
+
+		// Clear removed timers.
+		for (FInternalHandle InternalHandle : RemovedHandles)
+		{
+			// Index may not be found if the timer was added and removed before Tick was called.
+			int32 InternalHandleIndex = ActiveTimerHeap.Find(InternalHandle);
+			if (InternalHandleIndex != INDEX_NONE)
+			{
+				ActiveTimerHeap.HeapRemoveAt(InternalHandleIndex, FTimerHeapOrder(Storage), EAllowShrinking::No);
+			}
+		}
+	}
+
+	{
+		QUICK_SCOPE_CYCLE_COUNTER(STAT_EventLoop_TTimerManager_Tick_RescheduleTimers);
+
+		// Schedule timers waiting to be repeated.
+		for (FInternalHandle InternalHandle : PendingRepeatTimers)
+		{
+			// Index may not be found if the timer was removed before Tick was called.
+			if (FTimerData* TimerData = Storage.Find(InternalHandle))
+			{
+				TimerData->Expiration = InternalTime + TimerData->Rate;
+				TimerData->Status = ETimerStatus::Active;
+				ActiveTimerHeap.HeapPush(InternalHandle, FTimerHeapOrder(Storage));
+			}
+		}
+		PendingRepeatTimers.Empty();
+	}
+
+	{
+		QUICK_SCOPE_CYCLE_COUNTER(STAT_EventLoop_TTimerManager_Tick_RunTimers);
+
+		while (ActiveTimerHeap.Num() > 0)
+		{
+			FInternalHandle TopHandle = ActiveTimerHeap.HeapTop();
+			FTimerData* TimerData = Storage.Find(TopHandle);
+
+			// TimerHeap is now resolved with Storage. All entries are valid.
+			check(TimerData);
+
+			// Pop from heap and continue if timer was removed during iteration.
+			// The timer data will be removed on the next call to Update.
+			if (TimerData->Status == ETimerStatus::PendingRemoval)
+			{
+				FInternalHandle TimerHandle;
+				ActiveTimerHeap.HeapPop(TimerHandle, FTimerHeapOrder(Storage), EAllowShrinking::No);
+				continue;
+			}
+
+			if (InternalTime >= TimerData->Expiration)
+			{
+				FInternalHandle TimerHandle;
+				ActiveTimerHeap.HeapPop(TimerHandle, FTimerHeapOrder(Storage), EAllowShrinking::No);
+
+				// Fire user timer callback.
+				TimerData->Status = ETimerStatus::Executing;
+				TimerData->Callback();
+
+				// Reschedule repeating timers.
+				if (TimerData->bRepeat)
+				{
+					TimerData->Status = ETimerStatus::PendingReschedule;
+					PendingRepeatTimers.Add(TimerHandle);
+				}
+				else
+				{
+					// Timer has been fired and is non-repeating. Remove the timer.
+					Storage.Remove(TopHandle);
+				}
+			}
+			else
+			{
+				break;
+			}
+		}
+	}
+}
+
+/* UE::EventLoop */ }
diff --git a/Engine/Source/Runtime/Online/EventLoop/Public/EventLoop/IEventLoop.h b/Engine/Source/Runtime/Online/EventLoop/Public/EventLoop/IEventLoop.h
new file mode 100644
--- /dev/null
+++ b/Engine/Source/Runtime/Online/EventLoop/Public/EventLoop/IEventLoop.h
@@ -0,0 +1,108 @@
+// Copyright Epic Games, Inc. All Rights Reserved.
+
+#pragma once
+
+#include "CoreTypes.h"
+#include "EventLoop/EventLoopTimer.h"
+#include "Misc/Timespan.h"
+#include "Templates/Function.h"
+#include "Templates/SharedPointer.h"
+
+namespace UE::EventLoop {
+
+class IRequestManager;
+
+enum class EEventLoopStatus : uint8
+{
+	Unintialized,
+	Running,
+	Terminated,
+};
+
+using FAsyncTask = TUniqueFunction<void()>;
+using FOnShutdownComplete = TUniqueFunction<void()>;
+
+class IEventLoop : public TSharedFromThis<IEventLoop, ESPMode::ThreadSafe>
+{
+public:
+	virtual ~IEventLoop() = default;
+
+	/**
+	 * Initialize the event loop.
+	 * 
+	 * NOT thread safe.
+	 */
+	virtual bool Init() = 0;
+
+	/**
+	 * Signals a shutdown request to the event loop. May be signaled from any thread.
+	 * Once Shutdown has been called, continue calling Poll until
+	 * Invalidates the timer handle as it should no longer be used.
+	 * 
+	 * Thread safe.
+	 */
+	virtual void RequestShutdown(FOnShutdownComplete&& OnShutdownComplete = FOnShutdownComplete()) = 0;
+
+	/**
+	 * Set a new timer. The timer callback will be triggered from within the call to RunOnce,
+	 * which may occur on a different thread from the one which set the timer.
+	 * 
+	 * Thread safe.
+	 *
+	 * @param Callback Callback to call when timer fires.
+	 * @param InRate The amount of time between set and firing.
+	 * @param InbRepeat true to keep firing at Rate intervals, false to fire only once.
+	 * @param InFirstDelay The time for the first iteration of a looping timer.
+	 * @return handle to the registered timer.
+	 */
+	virtual FTimerHandle SetTimer(FTimerCallback&& Callback, FTimespan InRate, bool InbRepeat = false, TOptional<FTimespan> InFirstDelay = TOptional<FTimespan>()) = 0;
+
+	/**
+	* Clears a previously set timer.
+	* 
+	* Thread safe.
+	*
+	* @param InHandle The handle of the timer to clear.
+	* @param OnTimerCleared Callback to be fired when the timer has been removed. The callback will
+	*                       be fired from within the event loops RunOnce method.
+	*/
+	virtual void ClearTimer(FTimerHandle& InHandle, FOnTimerCleared&& OnTimerCleared = FOnTimerCleared()) = 0;
+
+	/**
+	* Post a task to be run by the event loop. Queues the task to be run and signals the event loop to wake.
+	* 
+	* Thread safe.
+	*
+	* @param Task The task to be run.
+	*/
+	virtual void PostAsyncTask(FAsyncTask&& Task) = 0;
+
+	/*
+	 * Run the event loop until shutdown is called.
+	 * The event loop will wake up when a timer needs to run or an async task has been posted.
+	 * 
+	 * NOT thread safe.
+	 */
+	virtual void Run() = 0;
+
+	/*
+	 * Run one iteration of the event loop.
+	 * 
+	 * NOT thread safe.
+	 * 
+	 * @param WaitTime Maximum amount of time to wait for events.
+	 * @return true if RunOnce should be called again.
+	 */
+	virtual bool RunOnce(FTimespan WaitTime) = 0;
+
+	/*
+	 * Return the event loops current loop time.
+	 * 
+	 * NOT thread safe.
+	 * 
+	 * @return the time the event loop last woke up.
+	 */
+	virtual FTimespan GetLoopTime() const = 0;
+};
+
+/* UE::EventLoop */ }
diff --git a/Engine/Source/Runtime/Online/EventLoop/Public/EventLoop/IEventLoopIOManager.h b/Engine/Source/Runtime/Online/EventLoop/Public/EventLoop/IEventLoopIOManager.h
new file mode 100644
--- /dev/null
+++ b/Engine/Source/Runtime/Online/EventLoop/Public/EventLoop/IEventLoopIOManager.h
@@ -0,0 +1,76 @@
+// Copyright Epic Games, Inc. All Rights Reserved.
+
+#pragma once
+
+#include "EventLoop/EventLoopHandle.h"
+#include "Misc/EnumClassFlags.h"
+#include "Misc/Timespan.h"
+
+namespace UE::EventLoop {
+
+class IEventLoop;
+
+struct FIORequestHandleTraits
+{
+	static constexpr auto Name = TEXT("EventLoopIORequestHandle");
+};
+using FIORequestHandle = TResourceHandle<FIORequestHandleTraits>;
+
+enum class EIOFlags : uint32
+{
+	None = 0,
+	Read = 1 << 0,
+	Write = 1 << 1,
+};
+ENUM_CLASS_FLAGS(EIOFlags);
+
+class IIOManager : public FNoncopyable
+{
+public:
+
+	/**
+	 * Parameters for starting an IO manager. Derived classes are expected to implement the
+	 * structure with any needed values. Will be passed from the event loop constructor into the
+	 * IO manager.
+	 */
+	struct FParams
+	{
+	};
+
+	/**
+	 * IIOManager(IEventLoop& EventLoop, FParams&& Params);
+	 */
+
+	virtual ~IIOManager() = default;
+
+
+	/**
+	 * Initialize the IO manager. Called from within IEventLoop::Init.
+	 *
+	 * @return true if successfully initialized.
+	 */
+	virtual bool Init() = 0;
+
+	/**
+	 * Cleanup any resources and prepare for shutdown. Called from IEventLoop::RunOnce when the
+	 * event loop has received a shutdown request.
+	 */
+	virtual void Shutdown() = 0;
+
+	/**
+	 * Notify the IO manager to interrupt a waiting call to Poll. Used to wake the request
+	 * manager when activity is required.
+	 * 
+	 * Thread safe.
+	 */
+	virtual void Notify() = 0;
+
+	/**
+	 * Poll the request manager for activity.
+	 * 
+	 * @param WaitTime The maximum amount of time to wait for request activity.
+	 */
+	virtual void Poll(FTimespan WaitTime) = 0;
+};
+
+/* UE::EventLoop */ }
diff --git a/Engine/Source/Runtime/Online/EventLoopTests/Tests/EventLoop/BSDSocketTypesPrivate.h b/Engine/Source/Runtime/Online/EventLoopTests/Tests/EventLoop/BSDSocketTypesPrivate.h
new file mode 100644
--- /dev/null
+++ b/Engine/Source/Runtime/Online/EventLoopTests/Tests/EventLoop/BSDSocketTypesPrivate.h
@@ -0,0 +1,60 @@
+// Copyright Epic Games, Inc. All Rights Reserved.
+
+#pragma once
+
+#include "CoreTypes.h"
+#include "EventLoop/BSDSocket/BSDSocketTypes.h"
+
+#if PLATFORM_HAS_BSD_SOCKETS
+
+#if PLATFORM_HAS_BSD_SOCKET_FEATURE_WINSOCKETS
+#include "Windows/AllowWindowsPlatformTypes.h"
+
+#include <winsock2.h>
+#include <ws2tcpip.h>
+
+typedef int32 SOCKLEN;
+
+#include "Windows/HideWindowsPlatformTypes.h"
+#else // PLATFORM_HAS_BSD_SOCKET_FEATURE_WINSOCKETS
+#if HAS_EVENTLOOP_PLATFORM_BSD_SOCKET_HEADER
+#include COMPILED_PLATFORM_HEADER(BSDSocketTypesPrivate.h)
+#else // HAS_EVENTLOOP_PLATFORM_BSD_SOCKET_HEADER
+#include <unistd.h>
+#include <sys/socket.h>
+#if PLATFORM_HAS_BSD_SOCKET_FEATURE_IOCTL
+#include <fcntl.h>
+#include <sys/types.h>
+#include <sys/ioctl.h>
+#endif // PLATFORM_HAS_BSD_SOCKET_FEATURE_IOCTL
+#if PLATFORM_HAS_BSD_SOCKET_FEATURE_POLL
+#include <poll.h>
+#endif // PLATFORM_HAS_BSD_SOCKET_FEATURE_POLL
+#include <netinet/in.h>
+#include <arpa/inet.h>
+#if PLATFORM_HAS_BSD_SOCKET_FEATURE_GETHOSTNAME
+#include <netdb.h>
+#endif // PLATFORM_HAS_BSD_SOCKET_FEATURE_GETHOSTNAME
+#if PLATFORM_HAS_BSD_SOCKET_FEATURE_NODELAY
+#include <netinet/tcp.h>
+#endif // PLATFORM_HAS_BSD_SOCKET_FEATURE_NODELAY
+
+#define ioctlsocket ioctl
+#endif // HAS_EVENTLOOP_PLATFORM_BSD_SOCKET_HEADER
+
+#define SOCKET_ERROR -1
+#define INVALID_SOCKET -1
+
+typedef socklen_t SOCKLEN;
+typedef int32 SOCKET;
+typedef sockaddr_in SOCKADDR_IN;
+typedef struct timeval TIMEVAL;
+
+inline int32 closesocket(SOCKET Socket)
+{
+	shutdown(Socket, SHUT_RDWR); // gracefully shutdown if connected
+	return close(Socket);
+}
+
+#endif // PLATFORM_HAS_BSD_SOCKET_FEATURE_WINSOCKETS
+#endif // PLATFORM_HAS_BSD_SOCKETS
diff --git a/Engine/Source/Runtime/Online/EventLoopTests/Tests/EventLoop/EventLoopHandleTests.cpp b/Engine/Source/Runtime/Online/EventLoopTests/Tests/EventLoop/EventLoopHandleTests.cpp
new file mode 100644
--- /dev/null
+++ b/Engine/Source/Runtime/Online/EventLoopTests/Tests/EventLoop/EventLoopHandleTests.cpp
@@ -0,0 +1,54 @@
+// Copyright Epic Games, Inc. All Rights Reserved.
+
+#include "EventLoop/EventLoopHandle.h"
+
+#include "TestHarness.h"
+
+namespace UE::EventLoop
+{
+	struct FTestHandleTraits
+	{
+		static constexpr auto Name = TEXT("test");
+	};
+
+	TEST_CASE("EventLoop::EventLoopHandle::Invalid Handle", "[Online][EventLoop][Smoke]")
+	{
+		TResourceHandle<FTestHandleTraits> Handle;
+		CHECK(!Handle.IsValid());
+	}
+
+	TEST_CASE("EventLoop::EventLoopHandle::Valid Handle", "[Online][EventLoop][Smoke]")
+	{
+		TResourceHandle<FTestHandleTraits> Handle(TResourceHandle<FTestHandleTraits>::GenerateNewHandle);
+		CHECK(Handle.IsValid());
+	}
+
+	struct FCount1HandleTraits
+	{
+		static constexpr auto Name = TEXT("count1");
+	};
+
+	struct FCount2HandleTraits
+	{
+		static constexpr auto Name = TEXT("count2");
+	};
+
+	TEST_CASE("EventLoop::EventLoopHandle::Counting", "[Online][EventLoop][Smoke]")
+	{
+		TResourceHandle<FCount1HandleTraits> Handle1(TResourceHandle<FCount1HandleTraits>::GenerateNewHandle);
+		TResourceHandle<FCount1HandleTraits> Handle2(TResourceHandle<FCount1HandleTraits>::GenerateNewHandle);
+		TResourceHandle<FCount2HandleTraits> Handle3(TResourceHandle<FCount2HandleTraits>::GenerateNewHandle);
+		FString ExpectedString;
+
+		// Check that handles count as expected.
+		ExpectedString = FString::Printf(TEXT("%s:1"), FCount1HandleTraits::Name);
+		CHECK(Handle1.ToString() == ExpectedString);
+
+		ExpectedString = FString::Printf(TEXT("%s:2"), FCount1HandleTraits::Name);
+		CHECK(Handle2.ToString() == ExpectedString);
+
+		// A handle using a different traits time uses its own counter.
+		ExpectedString = FString::Printf(TEXT("%s:1"), FCount2HandleTraits::Name);
+		CHECK(Handle3.ToString() == ExpectedString);
+	}
+} // UE::EventLoop
diff --git a/Engine/Source/Runtime/Online/EventLoopTests/Tests/EventLoop/EventLoopIOManagerBSDSocketSelectTests.cpp b/Engine/Source/Runtime/Online/EventLoopTests/Tests/EventLoop/EventLoopIOManagerBSDSocketSelectTests.cpp
new file mode 100644
--- /dev/null
+++ b/Engine/Source/Runtime/Online/EventLoopTests/Tests/EventLoop/EventLoopIOManagerBSDSocketSelectTests.cpp
@@ -0,0 +1,851 @@
+// Copyright Epic Games, Inc. All Rights Reserved.
+
+#include "CoreTypes.h"
+
+#if PLATFORM_HAS_BSD_SOCKETS
+
+#include "BSDSocketTypesPrivate.h"
+#include "EventLoop/IEventLoop.h"
+#include "EventLoop/BSDSocket/EventLoopIOManagerBSDSocketSelect.h"
+#include "TestHarness.h"
+
+#include "SocketSubsystem.h"
+
+namespace UE::EventLoop
+{
+	struct FBSDSocketFuncs
+	{
+		static SOCKET Open(int32 Family, int32 Type, int32 Protocol)
+		{
+			return socket(Family, Type, Protocol);
+		}
+
+		static void Close(SOCKET Socket)
+		{
+			closesocket(Socket);
+		}
+
+		static bool SetSendBufferSize(SOCKET Socket, int32 Size)
+		{
+			SOCKLEN SizeSize = sizeof(int32);
+			bool bOk = setsockopt(Socket, SOL_SOCKET, SO_SNDBUF, (char*)&Size, sizeof(int32)) == 0;
+
+			// Read the value back in case the size was modified
+			int32 NewSize;
+			getsockopt(Socket, SOL_SOCKET, SO_SNDBUF, (char*)&NewSize, &SizeSize);
+
+			return bOk && Size == NewSize;
+		}
+
+		static bool SetRecvBufferSize(SOCKET Socket, int32 Size)
+		{
+			SOCKLEN SizeSize = sizeof(int32);
+			bool bOk = setsockopt(Socket, SOL_SOCKET, SO_RCVBUF, (char*)&Size, sizeof(int32)) == 0;
+
+			// Read the value back in case the size was modified
+			int32 NewSize;
+			getsockopt(Socket, SOL_SOCKET, SO_RCVBUF, (char*)&NewSize, &SizeSize);
+
+			return bOk && Size == NewSize;
+		}
+
+		static bool SetNonBlocking(SOCKET Socket, bool bIsNonBlocking)
+		{
+	#if PLATFORM_HAS_BSD_SOCKET_FEATURE_WINSOCKETS
+			u_long Value = bIsNonBlocking ? true : false;
+			return ioctlsocket(Socket, FIONBIO, &Value) == 0;
+	#else
+			int Flags = fcntl(Socket, F_GETFL, 0);
+			//Set the flag or clear it, without destroying the other flags.
+			Flags = bIsNonBlocking ? Flags | O_NONBLOCK : Flags ^ (Flags & O_NONBLOCK);
+			int err = fcntl(Socket, F_SETFL, Flags);
+			return (err == 0 ? true : false);
+	#endif
+		}
+
+		static bool SetReusePort(SOCKET Socket, bool bIsReusePort)
+		{
+			int32 Resuse = bIsReusePort ? true : false;
+			bool bOk = setsockopt(Socket, SOL_SOCKET, SO_BROADCAST, (char*)&Resuse, sizeof(int32)) == 0;
+			bOk &= setsockopt(Socket, SOL_SOCKET, SO_REUSEADDR, (char*)&Resuse, sizeof(int32)) == 0;
+			return bOk;
+		}
+
+		static bool BindLoopback(SOCKET Socket)
+		{
+			sockaddr_in Addr;
+			FMemory::Memzero(&Addr, sizeof(Addr));
+			Addr.sin_addr.s_addr = htonl(INADDR_LOOPBACK);
+			Addr.sin_family = AF_INET;
+
+			return Bind(Socket, Addr);
+		}
+
+		static bool Bind(SOCKET Socket, const sockaddr_in& Addr)
+		{
+			return bind(Socket, (const sockaddr*)&Addr, sizeof(sockaddr_in)) == 0;
+		}
+
+		static bool Listen(SOCKET Socket, int32 MaxBacklog)
+		{
+			return listen(Socket, MaxBacklog) == 0;
+		}
+
+		static bool GetAddress(SOCKET Socket, sockaddr_in& OutAddr)
+		{
+			SOCKLEN Size = sizeof(sockaddr_in);
+
+			// Figure out what ip/port we are bound to
+			return getsockname(Socket, (sockaddr*)&OutAddr, &Size) == 0;
+		}
+
+		static bool Connect(SOCKET Socket, const sockaddr_in& Addr)
+		{
+			SOCKLEN Size = sizeof(sockaddr_in);
+			int32 Return = connect(Socket, (const sockaddr*)&Addr, Size);
+
+			ISocketSubsystem* SocketSubsystem = ISocketSubsystem::Get(PLATFORM_SOCKETSUBSYSTEM);
+			check(SocketSubsystem);
+			ESocketErrors Error = SocketSubsystem->TranslateErrorCode(Return);
+
+			// EWOULDBLOCK is not an error, and EINPROGRESS is fine on initial connection as it may still be creating for nonblocking sockets
+			return ((Error == SE_NO_ERROR) || (Error == SE_EWOULDBLOCK) || (Error == SE_EINPROGRESS));
+		}
+
+		static SOCKET Accept(SOCKET Socket)
+		{
+			return accept(Socket, NULL, NULL);
+		}
+
+		static bool Send(SOCKET Socket, const uint8* Data, int32 Count, int32& BytesSent)
+		{
+			BytesSent = send(Socket,(const char*)Data,Count,0);
+			return BytesSent >= 0;
+		}
+
+		static bool Recv(SOCKET Socket, uint8* Data, int32 BufferSize, int32& BytesRead)
+		{
+			int SocketType = 0;
+			SOCKLEN SocketTypeLength = sizeof(int32);
+			if (getsockopt(Socket, SOL_SOCKET, SO_TYPE, (char*)&SocketType, &SocketTypeLength) == SOCKET_ERROR)
+			{
+				return false;
+			}
+
+			bool bSuccess = false;
+			const bool bStreamSocket = (SocketType == SOCK_STREAM);
+			BytesRead = recv(Socket, (char*)Data, BufferSize, 0);
+
+			if (BytesRead >= 0)
+			{
+				// For Streaming sockets, 0 indicates a graceful failure
+				bSuccess = !bStreamSocket || (BytesRead > 0);
+			}
+			else
+			{
+				// For Streaming sockets, don't treat SE_EWOULDBLOCK as an error
+				ISocketSubsystem* SocketSubsystem = ISocketSubsystem::Get(PLATFORM_SOCKETSUBSYSTEM);
+				check(SocketSubsystem);
+				bSuccess = bStreamSocket && (SocketSubsystem->TranslateErrorCode(BytesRead) == SE_EWOULDBLOCK);
+				BytesRead = 0;
+			}
+
+			return bSuccess;
+		}
+
+		static bool SendTo(SOCKET Socket, const uint8* Data, int32 Count, int32& BytesSent, const sockaddr_in& Addr)
+		{
+			SOCKLEN Size = sizeof(sockaddr_in);
+			// Write the data and see how much was written
+			BytesSent = sendto(Socket, (const char*)Data, Count, 0, (const sockaddr*)&(Addr), Size);
+			return BytesSent >= 0;
+		}
+
+		static bool RecvFrom(SOCKET Socket, uint8* Data, int32 BufferSize, int32& BytesRead, sockaddr_in& Addr)
+		{
+			int SocketType = 0;
+			SOCKLEN SocketTypeLength = sizeof(int32);
+			if (getsockopt(Socket, SOL_SOCKET, SO_TYPE, (char*)&SocketType, &SocketTypeLength) == SOCKET_ERROR)
+			{
+				return false;
+			}
+
+			bool bSuccess = false;
+			const bool bStreamSocket = (SocketType == SOCK_STREAM);
+			SOCKLEN Size = sizeof(sockaddr_in);
+
+			// Read into the buffer and set the source address
+			BytesRead = recvfrom(Socket, (char*)Data, BufferSize, 0, (sockaddr*)&Addr, &Size);
+
+			if (BytesRead >= 0)
+			{
+				// For Streaming sockets, 0 indicates a graceful failure
+				bSuccess = !bStreamSocket || (BytesRead > 0);
+			}
+			else
+			{
+				ISocketSubsystem* SocketSubsystem = ISocketSubsystem::Get(PLATFORM_SOCKETSUBSYSTEM);
+				check(SocketSubsystem);
+				// For Streaming sockets, don't treat SE_EWOULDBLOCK as an error
+				bSuccess = bStreamSocket && (SocketSubsystem->TranslateErrorCode(BytesRead) == SE_EWOULDBLOCK);
+				BytesRead = 0;
+			}
+
+			return bSuccess;
+		}
+	};
+
+	// Todo - remove eventloop parameter from IOManager.
+	class FMockEventLoop : public IEventLoop
+	{
+	public:
+		FMockEventLoop() = default;
+		virtual ~FMockEventLoop() = default;
+		virtual bool Init() override { return false; }
+		virtual void RequestShutdown(FOnShutdownComplete&& OnShutdownComplete) override {}
+		virtual FTimerHandle SetTimer(FTimerCallback&& Callback, FTimespan InRate, bool InbRepeat, TOptional<FTimespan> InFirstDelay) override { return FTimerHandle(); }
+		virtual void ClearTimer(FTimerHandle& InHandle, FOnTimerCleared&& OnTimerCleared) override {}
+		virtual void PostAsyncTask(FAsyncTask&& Task) override {}
+		virtual void Run() override {}
+		virtual bool RunOnce(FTimespan WaitTime) override { return false; }
+		virtual FTimespan GetLoopTime() const override { return FTimespan(); }
+	};
+
+	TEST_CASE("IOManagerBSDSocketSelect", "[Online][EventLoop][IOManager][Smoke]")
+	{
+		FMockEventLoop MockEventLoop;
+		FIOManagerBSDSocketSelect IOManager(MockEventLoop, FIOManagerBSDSocketSelect::FParams());
+
+		REQUIRE(IOManager.Init());
+
+		ISocketSubsystem* SocketSubsystem = ISocketSubsystem::Get(PLATFORM_SOCKETSUBSYSTEM);
+		check(SocketSubsystem);
+
+		SECTION("Notify tests")
+		{
+			double StartTime = 0;
+			FTimespan Duration;
+
+			// 1. Waiting while no events are received will wait for at least the requested duration.
+			StartTime = FPlatformTime::Seconds();
+			IOManager.Poll(FTimespan::FromMilliseconds(50));
+			Duration = FTimespan::FromSeconds(FPlatformTime::Seconds() - StartTime);
+			CHECK(Duration > FTimespan::FromMilliseconds(50));
+
+			// 2. Notifying the IO manager will wake the loop up early.
+			StartTime = FPlatformTime::Seconds();
+			IOManager.Notify();
+			IOManager.Poll(FTimespan::FromMilliseconds(50));
+			Duration = FTimespan::FromSeconds(FPlatformTime::Seconds() - StartTime);
+			CHECK(Duration < FTimespan::FromMilliseconds(1));
+		}
+
+		SECTION("Single socket tests")
+		{
+			// Size of buffer used for sending test data.
+			constexpr uint64 WriteBufferSize = 128;
+			TArray<uint8> WriteBuffer;
+			WriteBuffer.AddDefaulted(WriteBufferSize);
+
+			// Size of buffer used for receiving test data.
+			constexpr uint64 ReadBufferSize = 128;
+			TArray<uint8> ReadBuffer;
+			ReadBuffer.AddDefaulted(ReadBufferSize);
+
+			// Socket buffer sizes for send / recv
+			constexpr int32 SocketSendBufferSize = 64;
+			constexpr int32 SocketReceiveBufferSize = 64;
+
+			// Create listen socket.
+			SOCKET ListenSocket = FBSDSocketFuncs::Open(AF_INET, SOCK_STREAM, IPPROTO_TCP);
+			REQUIRE(ListenSocket != INVALID_SOCKET);
+			REQUIRE(FBSDSocketFuncs::SetSendBufferSize(ListenSocket, SocketSendBufferSize));
+			REQUIRE(FBSDSocketFuncs::SetRecvBufferSize(ListenSocket, SocketReceiveBufferSize));
+			REQUIRE(FBSDSocketFuncs::SetNonBlocking(ListenSocket, true));
+			REQUIRE(FBSDSocketFuncs::BindLoopback(ListenSocket));
+			REQUIRE(FBSDSocketFuncs::Listen(ListenSocket, 1));
+			sockaddr_in HostAddr;
+			FMemory::Memzero(&HostAddr, sizeof(HostAddr));
+			REQUIRE(FBSDSocketFuncs::GetAddress(ListenSocket, HostAddr));
+
+			// Create remote socket.
+			SOCKET RemoteSocket1 = FBSDSocketFuncs::Open(AF_INET, SOCK_STREAM, IPPROTO_TCP);
+			REQUIRE(RemoteSocket1 != INVALID_SOCKET);
+			REQUIRE(FBSDSocketFuncs::SetSendBufferSize(RemoteSocket1, SocketSendBufferSize));
+			REQUIRE(FBSDSocketFuncs::SetRecvBufferSize(RemoteSocket1, SocketReceiveBufferSize));
+			REQUIRE(FBSDSocketFuncs::SetNonBlocking(RemoteSocket1, true));
+			REQUIRE(FBSDSocketFuncs::Connect(RemoteSocket1, HostAddr));
+
+			// Create local socket.
+			ESocketErrors ErrorCode1 = SE_NO_ERROR;
+			SOCKET LocalSocket1 = INVALID_SOCKET;
+			do
+			{
+				LocalSocket1 = FBSDSocketFuncs::Accept(ListenSocket);
+				ErrorCode1 = SocketSubsystem->GetLastErrorCode();
+			}
+			while (ErrorCode1 == SE_EWOULDBLOCK);
+			REQUIRE(LocalSocket1 != INVALID_SOCKET);
+
+			SECTION("Single waiting socket tests")
+			{
+				SECTION("Wait on socket readable")
+				{
+					SOCKET TriggeredSocket = INVALID_SOCKET;
+					ESocketIoRequestStatus TriggeredStatus = ESocketIoRequestStatus::Ok;
+					EIOFlags TriggeredFlags = EIOFlags::None;
+					uint32 CallbackTriggerCount = 0;
+
+					FIORequestBSDSocket Request;
+					Request.Socket = LocalSocket1;
+					Request.Flags = EIOFlags::Read;
+					Request.Callback = [&](SOCKET Socket, ESocketIoRequestStatus Status, EIOFlags SignaledFlags){
+						TriggeredSocket = Socket;
+						TriggeredStatus = Status;
+						TriggeredFlags = SignaledFlags;
+						++CallbackTriggerCount;
+					};
+
+					FIORequestHandle RequestHandle = IOManager.GetIOAccess().CreateSocketIORequest(MoveTemp(Request));
+					CHECK(RequestHandle.IsValid());
+
+					// 1. No events triggered when no data has arrived at the socket.
+					IOManager.Poll(FTimespan::Zero());
+					CHECK(CallbackTriggerCount == 0);
+
+					// 2. Send data to the socket.
+					int32 BytesSent = 0;
+					CHECK(FBSDSocketFuncs::Send(RemoteSocket1, WriteBuffer.GetData(), 1, BytesSent));
+					CHECK(BytesSent == 1);
+
+					// 3. See that socket has a read event.
+					TriggeredFlags = EIOFlags::None;
+					IOManager.Poll(FTimespan::Zero());
+					CHECK(CallbackTriggerCount == 1);
+					CHECK(TriggeredSocket == LocalSocket1);
+					CHECK(TriggeredStatus == ESocketIoRequestStatus::Ok);
+					CHECK(TriggeredFlags == EIOFlags::Read);
+
+					// 4. See that socket still has a read event.
+					IOManager.Poll(FTimespan::Zero());
+					CHECK(CallbackTriggerCount == 2);
+					CHECK(TriggeredSocket == LocalSocket1);
+					CHECK(TriggeredStatus == ESocketIoRequestStatus::Ok);
+					CHECK(TriggeredFlags == EIOFlags::Read);
+
+					// 5. Read the waiting data from the socket.
+					int32 BytesRead = 0;
+					CHECK(FBSDSocketFuncs::Recv(LocalSocket1, ReadBuffer.GetData(), ReadBufferSize, BytesRead));
+					CHECK(BytesRead == 1);
+
+					// 6. See that socket no longer has a read event.
+					TriggeredFlags = EIOFlags::None;
+					IOManager.Poll(FTimespan::Zero());
+					CHECK(CallbackTriggerCount == 2);
+				}
+
+				SECTION("Wait on socket writeable")
+				{
+					SOCKET TriggeredSocket = INVALID_SOCKET;
+					ESocketIoRequestStatus TriggeredStatus = ESocketIoRequestStatus::Ok;
+					EIOFlags TriggeredFlags = EIOFlags::None;
+					uint32 CallbackTriggerCount = 0;
+
+					FIORequestBSDSocket Request;
+					Request.Socket = LocalSocket1;
+					Request.Flags = EIOFlags::Write;
+					Request.Callback = [&](SOCKET Socket, ESocketIoRequestStatus Status, EIOFlags SignaledFlags){
+						TriggeredSocket = Socket;
+						TriggeredStatus = Status;
+						TriggeredFlags = SignaledFlags;
+						++CallbackTriggerCount;
+					};
+
+					FIORequestHandle RequestHandle = IOManager.GetIOAccess().CreateSocketIORequest(MoveTemp(Request));
+					CHECK(RequestHandle.IsValid());
+
+					// 1. See that socket has a write event.
+					IOManager.Poll(FTimespan::Zero());
+					CHECK(CallbackTriggerCount == 1);
+					CHECK(TriggeredSocket == LocalSocket1);
+					CHECK(TriggeredStatus == ESocketIoRequestStatus::Ok);
+					CHECK(TriggeredFlags == EIOFlags::Write);
+
+					// 2. Write data to the socket until it returns an error indicating that it would block.
+					int32 BytesSent = 0;
+					int32 TotalBytesSent = 0;
+					while (FBSDSocketFuncs::Send(LocalSocket1, WriteBuffer.GetData(), WriteBufferSize, BytesSent));
+
+					// 3. See that socket no longer has a write event.
+					IOManager.Poll(FTimespan::Zero());
+					CHECK(CallbackTriggerCount == 1);
+
+					// 4. Read the data from the remote end of the connection.
+					int32 BytesRead = 0;
+					while (FBSDSocketFuncs::Recv(RemoteSocket1, ReadBuffer.GetData(), ReadBufferSize, BytesRead))
+					{
+						if (BytesRead == 0)
+						{
+							break;
+						}
+					}
+
+					// 5. See that socket now has a write event.
+					IOManager.Poll(FTimespan::FromMilliseconds(10));
+					CHECK(CallbackTriggerCount == 2);
+					CHECK(TriggeredSocket == LocalSocket1);
+					CHECK(TriggeredStatus == ESocketIoRequestStatus::Ok);
+					CHECK(TriggeredFlags == EIOFlags::Write);
+				}
+
+				SECTION("Receive both read and write events with one request")
+				{
+					SOCKET TriggeredSocket = INVALID_SOCKET;
+					ESocketIoRequestStatus TriggeredStatus = ESocketIoRequestStatus::Ok;
+					EIOFlags TriggeredFlags = EIOFlags::None;
+					uint32 CallbackTriggerCount = 0;
+
+					FIORequestBSDSocket Request;
+					Request.Socket = LocalSocket1;
+					Request.Flags = EIOFlags::Read | EIOFlags::Write;
+					Request.Callback = [&](SOCKET Socket, ESocketIoRequestStatus Status, EIOFlags SignaledFlags){
+						TriggeredSocket = Socket;
+						TriggeredStatus = Status;
+						TriggeredFlags = SignaledFlags;
+						++CallbackTriggerCount;
+					};
+
+					FIORequestHandle RequestHandle = IOManager.GetIOAccess().CreateSocketIORequest(MoveTemp(Request));
+					CHECK(RequestHandle.IsValid());
+
+					// 1. No events triggered when no data has arrived at the socket.
+					IOManager.Poll(FTimespan::Zero());
+					CHECK(CallbackTriggerCount == 1);
+					CHECK(TriggeredSocket == LocalSocket1);
+					CHECK(TriggeredStatus == ESocketIoRequestStatus::Ok);
+					CHECK(EnumHasAnyFlags(TriggeredFlags, EIOFlags::Write));
+					CHECK(!EnumHasAnyFlags(TriggeredFlags, EIOFlags::Read));
+
+					// 2. Send data to the socket.
+					int32 BytesSent = 0;
+					CHECK(FBSDSocketFuncs::Send(RemoteSocket1, WriteBuffer.GetData(), 1, BytesSent));
+					CHECK(BytesSent == 1);
+
+					// 3. See that socket has a read event.
+					IOManager.Poll(FTimespan::Zero());
+					CHECK(CallbackTriggerCount == 2);
+					CHECK(TriggeredSocket == LocalSocket1);
+					CHECK(TriggeredStatus == ESocketIoRequestStatus::Ok);
+					CHECK(EnumHasAnyFlags(TriggeredFlags, EIOFlags::Write));
+					CHECK(EnumHasAnyFlags(TriggeredFlags, EIOFlags::Read));
+				}
+
+				SECTION("Canceling a request stops signaling events")
+				{
+					SOCKET TriggeredSocket = INVALID_SOCKET;
+					ESocketIoRequestStatus TriggeredStatus = ESocketIoRequestStatus::Ok;
+					EIOFlags TriggeredFlags = EIOFlags::None;
+					uint32 CallbackTriggerCount = 0;
+
+					FIORequestBSDSocket Request;
+					Request.Socket = LocalSocket1;
+					Request.Flags = EIOFlags::Read;
+					Request.Callback = [&](SOCKET Socket, ESocketIoRequestStatus Status, EIOFlags SignaledFlags){
+						TriggeredSocket = Socket;
+						TriggeredStatus = Status;
+						TriggeredFlags = SignaledFlags;
+						++CallbackTriggerCount;
+					};
+
+					FIORequestHandle RequestHandle = IOManager.GetIOAccess().CreateSocketIORequest(MoveTemp(Request));
+					CHECK(RequestHandle.IsValid());
+
+					// 1. No events triggered when no data has arrived at the socket.
+					IOManager.Poll(FTimespan::Zero());
+					CHECK(CallbackTriggerCount == 0);
+
+					// 2. Send data to the socket.
+					int32 BytesSent = 0;
+					CHECK(FBSDSocketFuncs::Send(RemoteSocket1, WriteBuffer.GetData(), 1, BytesSent));
+					CHECK(BytesSent == 1);
+
+					// 3. See that socket has a read event.
+					IOManager.Poll(FTimespan::Zero());
+					CHECK(CallbackTriggerCount == 1);
+					CHECK(TriggeredSocket == LocalSocket1);
+					CHECK(TriggeredStatus == ESocketIoRequestStatus::Ok);
+					CHECK(TriggeredFlags == EIOFlags::Read);
+
+					// 4. Cancel the request.
+					IOManager.GetIOAccess().DestroyIORequest(RequestHandle);
+
+					// 5. See that event is not triggered.
+					IOManager.Poll(FTimespan::Zero());
+					CHECK(CallbackTriggerCount == 1);
+				}
+
+				SECTION("TCP socket remote end closing triggers writeable wakeup")
+				{
+					SOCKET TriggeredSocket = INVALID_SOCKET;
+					ESocketIoRequestStatus TriggeredStatus = ESocketIoRequestStatus::Ok;
+					EIOFlags TriggeredFlags = EIOFlags::None;
+					uint32 CallbackTriggerCount = 0;
+
+					FIORequestBSDSocket Request;
+					Request.Socket = LocalSocket1;
+					Request.Flags = EIOFlags::Write;
+					Request.Callback = [&](SOCKET Socket, ESocketIoRequestStatus Status, EIOFlags SignaledFlags){
+						TriggeredSocket = Socket;
+						TriggeredStatus = Status;
+						TriggeredFlags = SignaledFlags;
+						++CallbackTriggerCount;
+					};
+
+					FIORequestHandle RequestHandle = IOManager.GetIOAccess().CreateSocketIORequest(MoveTemp(Request));
+					CHECK(RequestHandle.IsValid());
+
+					// 1. See that socket has a write event.
+					IOManager.Poll(FTimespan::Zero());
+					CHECK(CallbackTriggerCount == 1);
+					CHECK(TriggeredSocket == LocalSocket1);
+					CHECK(TriggeredStatus == ESocketIoRequestStatus::Ok);
+					CHECK(TriggeredFlags == EIOFlags::Write);
+
+					// 2. Write data to the socket until it returns an error indicating that it would block.
+					int32 BytesSent = 0;
+					while (FBSDSocketFuncs::Send(LocalSocket1, WriteBuffer.GetData(), WriteBufferSize, BytesSent));
+
+					// 3. See that socket no longer has a write event.
+					IOManager.Poll(FTimespan::Zero());
+					CHECK(CallbackTriggerCount == 1);
+
+					// 4. Close the remote end of the connection.
+					FBSDSocketFuncs::Close(RemoteSocket1);
+
+					// 5. See that socket now has a write event.
+					IOManager.Poll(FTimespan::Zero());
+					CHECK(CallbackTriggerCount == 2);
+					CHECK(TriggeredSocket == LocalSocket1);
+					CHECK(TriggeredStatus == ESocketIoRequestStatus::Ok);
+					CHECK(TriggeredFlags == EIOFlags::Write);
+				}
+
+				SECTION("More than one request for the same socket")
+				{
+					SOCKET TriggeredSocket = INVALID_SOCKET;
+					ESocketIoRequestStatus TriggeredStatus = ESocketIoRequestStatus::Ok;
+					EIOFlags TriggeredFlags = EIOFlags::None;
+					uint32 CallbackTriggerCount1 = 0;
+					uint32 CallbackTriggerCount2 = 0;
+
+					FIORequestBSDSocket Request1;
+					Request1.Socket = LocalSocket1;
+					Request1.Flags = EIOFlags::Read;
+					Request1.Callback = [&](SOCKET Socket, ESocketIoRequestStatus Status, EIOFlags SignaledFlags){
+						TriggeredSocket = Socket;
+						TriggeredStatus = Status;
+						TriggeredFlags = SignaledFlags;
+						++CallbackTriggerCount1;
+					};
+
+					FIORequestHandle RequestHandle1 = IOManager.GetIOAccess().CreateSocketIORequest(MoveTemp(Request1));
+					CHECK(RequestHandle1.IsValid());
+
+					FIORequestBSDSocket Request2;
+					Request2.Socket = LocalSocket1;
+					Request2.Flags = EIOFlags::Write;
+					Request2.Callback = [&](SOCKET Socket, ESocketIoRequestStatus Status, EIOFlags SignaledFlags){
+						TriggeredSocket = Socket;
+						TriggeredStatus = Status;
+						TriggeredFlags = SignaledFlags;
+						++CallbackTriggerCount2;
+					};
+
+					FIORequestHandle RequestHandle2 = IOManager.GetIOAccess().CreateSocketIORequest(MoveTemp(Request2));
+					CHECK(RequestHandle2.IsValid());
+
+					IOManager.Poll(FTimespan::Zero());
+					CHECK(CallbackTriggerCount1 == 0);
+					CHECK(CallbackTriggerCount2 == 1);
+					CHECK(TriggeredStatus == ESocketIoRequestStatus::Invalid);
+				}
+			}
+
+			SECTION("Multiple waiting TCP socket tests")
+			{
+				// Create remote socket.
+				SOCKET RemoteSocket2 = FBSDSocketFuncs::Open(AF_INET, SOCK_STREAM, IPPROTO_TCP);
+				REQUIRE(RemoteSocket2 != INVALID_SOCKET);
+				REQUIRE(FBSDSocketFuncs::SetSendBufferSize(RemoteSocket2, SocketSendBufferSize));
+				REQUIRE(FBSDSocketFuncs::SetRecvBufferSize(RemoteSocket2, SocketReceiveBufferSize));
+				REQUIRE(FBSDSocketFuncs::SetNonBlocking(RemoteSocket2, true));
+				REQUIRE(FBSDSocketFuncs::Connect(RemoteSocket2, HostAddr));
+
+				// Create local socket.
+				ESocketErrors ErrorCode2 = SE_NO_ERROR;
+				SOCKET LocalSocket2 = INVALID_SOCKET;
+				do
+				{
+					LocalSocket2 = FBSDSocketFuncs::Accept(ListenSocket);
+					ErrorCode2 = SocketSubsystem->GetLastErrorCode();
+				}
+				while (ErrorCode2 == SE_EWOULDBLOCK);
+				REQUIRE(LocalSocket2 != INVALID_SOCKET);
+
+				SECTION("Triggering LocalSocket1 does not trigger LocalSocket2")
+				{
+					SOCKET TriggeredSocket1 = INVALID_SOCKET;
+					ESocketIoRequestStatus TriggeredStatus1 = ESocketIoRequestStatus::Ok;
+					EIOFlags TriggeredFlags1 = EIOFlags::None;
+					uint32 CallbackTriggerCount1 = 0;
+					FIORequestBSDSocket Request1;
+					Request1.Socket = LocalSocket1;
+					Request1.Flags = EIOFlags::Read;
+					Request1.Callback = [&](SOCKET Socket, ESocketIoRequestStatus Status, EIOFlags SignaledFlags){
+						TriggeredSocket1 = Socket;
+						TriggeredStatus1 = Status;
+						TriggeredFlags1 = SignaledFlags;
+						++CallbackTriggerCount1;
+					};
+					FIORequestHandle RequestHandle1 = IOManager.GetIOAccess().CreateSocketIORequest(MoveTemp(Request1));
+					CHECK(RequestHandle1.IsValid());
+
+					SOCKET TriggeredSocket2 = INVALID_SOCKET;
+					ESocketIoRequestStatus TriggeredStatus2 = ESocketIoRequestStatus::Ok;
+					EIOFlags TriggeredFlags2 = EIOFlags::None;
+					uint32 CallbackTriggerCount2 = 0;
+					FIORequestBSDSocket Request2;
+					Request2.Socket = LocalSocket2;
+					Request2.Flags = EIOFlags::Read;
+					Request2.Callback = [&](SOCKET Socket, ESocketIoRequestStatus Status, EIOFlags SignaledFlags){
+						TriggeredSocket2 = Socket;
+						TriggeredStatus2 = Status;
+						TriggeredFlags2 = SignaledFlags;
+						++CallbackTriggerCount2;
+					};
+					FIORequestHandle RequestHandle2 = IOManager.GetIOAccess().CreateSocketIORequest(MoveTemp(Request2));
+					CHECK(RequestHandle2.IsValid());
+
+					// 1. No events triggered when no data has arrived at the socket.
+					IOManager.Poll(FTimespan::Zero());
+					CHECK(CallbackTriggerCount1 == 0);
+					CHECK(CallbackTriggerCount2 == 0);
+
+					// 2. Send data to the socket.
+					int32 BytesSent = 0;
+					CHECK(FBSDSocketFuncs::Send(RemoteSocket1, WriteBuffer.GetData(), 1, BytesSent));
+					CHECK(BytesSent == 1);
+
+					// 3. See that socket has a read event.
+					TriggeredFlags1 = EIOFlags::None;
+					IOManager.Poll(FTimespan::Zero());
+					CHECK(CallbackTriggerCount1 == 1);
+					CHECK(CallbackTriggerCount2 == 0);
+					CHECK(TriggeredSocket1 == LocalSocket1);
+					CHECK(TriggeredStatus1 == ESocketIoRequestStatus::Ok);
+					CHECK(TriggeredFlags1 == EIOFlags::Read);
+
+					// 4. See that socket still has a read event.
+					IOManager.Poll(FTimespan::Zero());
+					CHECK(CallbackTriggerCount1 == 2);
+					CHECK(CallbackTriggerCount2 == 0);
+					CHECK(TriggeredSocket1 == LocalSocket1);
+					CHECK(TriggeredStatus1 == ESocketIoRequestStatus::Ok);
+					CHECK(TriggeredFlags1 == EIOFlags::Read);
+
+					// 5. Read the waiting data from the socket.
+					int32 BytesRead = 0;
+					CHECK(FBSDSocketFuncs::Recv(LocalSocket1, ReadBuffer.GetData(), ReadBufferSize, BytesRead));
+					CHECK(BytesRead == 1);
+
+					// 6. See that socket no longer has a read event.
+					TriggeredFlags1 = EIOFlags::None;
+					IOManager.Poll(FTimespan::Zero());
+					CHECK(CallbackTriggerCount1 == 2);
+					CHECK(CallbackTriggerCount2 == 0);
+				}
+
+				FBSDSocketFuncs::Close(RemoteSocket2);
+				FBSDSocketFuncs::Close(LocalSocket2);
+			}
+
+			SECTION("Waiting on a TCP and UDP socket")
+			{
+				// Create local socket.
+				SOCKET LocalSocket2 = FBSDSocketFuncs::Open(AF_INET, SOCK_DGRAM, IPPROTO_UDP);
+				REQUIRE(LocalSocket2 != INVALID_SOCKET);
+
+				REQUIRE(FBSDSocketFuncs::SetNonBlocking(LocalSocket2, true));
+				REQUIRE(FBSDSocketFuncs::BindLoopback(LocalSocket2));
+				sockaddr_in HostAddr2;
+				FMemory::Memzero(&HostAddr2, sizeof(HostAddr2));
+				REQUIRE(FBSDSocketFuncs::GetAddress(LocalSocket2, HostAddr2));
+
+				// Create remote socket.
+				SOCKET RemoteSocket2 = FBSDSocketFuncs::Open(AF_INET, SOCK_DGRAM, IPPROTO_UDP);
+				REQUIRE(RemoteSocket2 != INVALID_SOCKET);
+				REQUIRE(FBSDSocketFuncs::SetSendBufferSize(RemoteSocket2, SocketSendBufferSize));
+				REQUIRE(FBSDSocketFuncs::SetRecvBufferSize(RemoteSocket2, SocketReceiveBufferSize));
+				REQUIRE(FBSDSocketFuncs::SetNonBlocking(RemoteSocket2, true));
+
+				SECTION("Recevie events for TCP and UDP socket types.")
+				{
+					SOCKET TriggeredSocket1 = INVALID_SOCKET;
+					ESocketIoRequestStatus TriggeredStatus1 = ESocketIoRequestStatus::Ok;
+					EIOFlags TriggeredFlags1 = EIOFlags::None;
+					uint32 CallbackTriggerCount1 = 0;
+					FIORequestBSDSocket Request1;
+					Request1.Socket = LocalSocket1;
+					Request1.Flags = EIOFlags::Read;
+					Request1.Callback = [&](SOCKET Socket, ESocketIoRequestStatus Status, EIOFlags SignaledFlags){
+						TriggeredSocket1 = Socket;
+						TriggeredStatus1 = Status;
+						TriggeredFlags1 = SignaledFlags;
+						++CallbackTriggerCount1;
+					};
+					FIORequestHandle RequestHandle1 = IOManager.GetIOAccess().CreateSocketIORequest(MoveTemp(Request1));
+					CHECK(RequestHandle1.IsValid());
+
+					SOCKET TriggeredSocket2 = INVALID_SOCKET;
+					ESocketIoRequestStatus TriggeredStatus2 = ESocketIoRequestStatus::Ok;
+					EIOFlags TriggeredFlags2 = EIOFlags::None;
+					uint32 CallbackTriggerCount2 = 0;
+					FIORequestBSDSocket Request2;
+					Request2.Socket = LocalSocket2;
+					Request2.Flags = EIOFlags::Read;
+					Request2.Callback = [&](SOCKET Socket, ESocketIoRequestStatus Status, EIOFlags SignaledFlags){
+						TriggeredSocket2 = Socket;
+						TriggeredStatus2 = Status;
+						TriggeredFlags2 = SignaledFlags;
+						++CallbackTriggerCount2;
+					};
+					FIORequestHandle RequestHandle2 = IOManager.GetIOAccess().CreateSocketIORequest(MoveTemp(Request2));
+					CHECK(RequestHandle2.IsValid());
+
+					// 1. No events triggered when no data has arrived at the socket.
+					IOManager.Poll(FTimespan::Zero());
+					CHECK(CallbackTriggerCount1 == 0);
+					CHECK(CallbackTriggerCount2 == 0);
+
+					// 2. Send data to the socket.
+					int32 BytesSent = 0;
+					CHECK(FBSDSocketFuncs::Send(RemoteSocket1, WriteBuffer.GetData(), 1, BytesSent));
+					CHECK(BytesSent == 1);
+
+					// 3. See that socket has a read event.
+					IOManager.Poll(FTimespan::Zero());
+					CHECK(CallbackTriggerCount1 == 1);
+					CHECK(CallbackTriggerCount2 == 0);
+					CHECK(TriggeredSocket1 == LocalSocket1);
+					CHECK(TriggeredStatus1 == ESocketIoRequestStatus::Ok);
+					CHECK(TriggeredFlags1 == EIOFlags::Read);
+
+					// 4. Send data to the UDP socket.
+					BytesSent = 0;
+					CHECK(FBSDSocketFuncs::SendTo(RemoteSocket2, WriteBuffer.GetData(), 1, BytesSent, HostAddr2));
+					CHECK(BytesSent == 1);
+
+					// 5. See that both sockets have a read event.
+					IOManager.Poll(FTimespan::Zero());
+					CHECK(CallbackTriggerCount1 == 2);
+					CHECK(CallbackTriggerCount2 == 1);
+					CHECK(TriggeredSocket1 == LocalSocket1);
+					CHECK(TriggeredStatus1 == ESocketIoRequestStatus::Ok);
+					CHECK(TriggeredFlags1 == EIOFlags::Read);
+					CHECK(TriggeredSocket2 == LocalSocket2);
+					CHECK(TriggeredStatus2 == ESocketIoRequestStatus::Ok);
+					CHECK(TriggeredFlags2 == EIOFlags::Read);
+
+					// 5. Read the waiting data from the TCP socket.
+					int32 BytesRead = 0;
+					CHECK(FBSDSocketFuncs::Recv(LocalSocket1, ReadBuffer.GetData(), ReadBufferSize, BytesRead));
+					CHECK(BytesRead == 1);
+
+					// 6. See that the UDP socket still has a read event.
+					IOManager.Poll(FTimespan::Zero());
+					CHECK(CallbackTriggerCount1 == 2);
+					CHECK(CallbackTriggerCount2 == 2);
+					CHECK(TriggeredSocket2 == LocalSocket2);
+					CHECK(TriggeredStatus2 == ESocketIoRequestStatus::Ok);
+					CHECK(TriggeredFlags2 == EIOFlags::Read);
+				}
+
+				FBSDSocketFuncs::Close(RemoteSocket2);
+				FBSDSocketFuncs::Close(LocalSocket2);
+			}
+
+			FBSDSocketFuncs::Close(ListenSocket);
+			FBSDSocketFuncs::Close(RemoteSocket1);
+			FBSDSocketFuncs::Close(LocalSocket1);
+		}
+
+		SECTION("Wait on too many sockets")
+		{
+			struct FCallbackData
+			{
+				SOCKET Socket = INVALID_SOCKET;
+				ESocketIoRequestStatus Status = ESocketIoRequestStatus::Ok;
+				EIOFlags SignaledFlags = EIOFlags::None;
+			};
+
+			TArray<SOCKET> Sockets;
+			TArray<FCallbackData> CallbackData;
+
+			SOCKET BaseSocket = FBSDSocketFuncs::Open(AF_INET, SOCK_DGRAM, IPPROTO_UDP);
+			REQUIRE(BaseSocket != INVALID_SOCKET);
+			REQUIRE(FBSDSocketFuncs::SetNonBlocking(BaseSocket, true));
+			REQUIRE(FBSDSocketFuncs::SetReusePort(BaseSocket, true));
+			REQUIRE(FBSDSocketFuncs::BindLoopback(BaseSocket));
+
+			sockaddr_in BaseAddr;
+			FMemory::Memzero(&BaseAddr, sizeof(BaseAddr));
+			REQUIRE(FBSDSocketFuncs::GetAddress(BaseSocket, BaseAddr));
+
+			const int32 MaxTestSockets = 10000;
+			for (int32 index = 0; index < MaxTestSockets; ++index)
+			{
+				SOCKET Socket = FBSDSocketFuncs::Open(AF_INET, SOCK_DGRAM, IPPROTO_UDP);
+				REQUIRE(Socket != INVALID_SOCKET);
+				REQUIRE(FBSDSocketFuncs::SetNonBlocking(Socket, true));
+				REQUIRE(FBSDSocketFuncs::SetReusePort(Socket, true));
+				REQUIRE(FBSDSocketFuncs::Bind(Socket, BaseAddr));
+				Sockets.Add(Socket);
+
+				FIORequestBSDSocket Request;
+				Request.Socket = Socket;
+				Request.Flags = EIOFlags::Read;
+				Request.Callback = [&CallbackData](SOCKET Socket, ESocketIoRequestStatus Status, EIOFlags SignaledFlags){
+					CallbackData.Add({Socket, Status, SignaledFlags});
+				};
+				FIORequestHandle RequestHandle = IOManager.GetIOAccess().CreateSocketIORequest(MoveTemp(Request));
+				CHECK(RequestHandle.IsValid());
+
+				IOManager.Poll(FTimespan::Zero());
+
+				if (!CallbackData.IsEmpty())
+				{
+					CHECK(CallbackData.Num() == 1);
+					CHECK(CallbackData[0].Socket == Socket);
+					CHECK(CallbackData[0].Status == ESocketIoRequestStatus::NoResources);
+					CHECK(CallbackData[0].SignaledFlags == EIOFlags::None);
+					break;
+				}
+			}
+
+			// Check that a failure occurred.
+			CHECK(!CallbackData.IsEmpty());
+
+			for (SOCKET Socket : Sockets)
+			{
+				FBSDSocketFuncs::Close(Socket);
+			}
+
+			FBSDSocketFuncs::Close(BaseSocket);
+		}
+
+		IOManager.Shutdown();
+	}
+
+} // UE::EventLoop
+
+#endif // PLATFORM_HAS_BSD_SOCKETS
diff --git a/Engine/Source/Runtime/Online/EventLoopTests/Tests/EventLoop/EventLoopLLTModule.cpp b/Engine/Source/Runtime/Online/EventLoopTests/Tests/EventLoop/EventLoopLLTModule.cpp
new file mode 100644
--- /dev/null
+++ b/Engine/Source/Runtime/Online/EventLoopTests/Tests/EventLoop/EventLoopLLTModule.cpp
@@ -0,0 +1,52 @@
+// Copyright Epic Games, Inc. All Rights Reserved.
+
+#include "CoreTypes.h"
+
+#include "Containers/Array.h"
+#include "Containers/UnrealString.h"
+#include "Modules/ModuleManager.h"
+#include "TestRunner.h"
+
+namespace UE::EventLoop
+{
+
+TArray<FString> GetRequiredModules()
+{
+	static TArray<FString> Modules({
+		FString("EventLoop"),
+		FString("Sockets")
+	});
+
+	return Modules;
+}
+
+void InitializeEventLoopTests()
+{
+	for (const FString& ModuleName : GetRequiredModules())
+	{
+		FModuleManager::LoadModulePtr<IModuleInterface>(*ModuleName);
+	}
+}
+
+void CleanupEventLoopTests()
+{
+	for (const FString& ModuleName : GetRequiredModules())
+	{
+		if (IModuleInterface* Module = FModuleManager::Get().GetModule(*ModuleName))
+		{
+			Module->ShutdownModule();
+		}
+	}
+}
+
+class EventLoopTestGlobalSetup
+{
+public:
+	EventLoopTestGlobalSetup()
+	{
+		FTestDelegates::GetGlobalSetup().BindStatic(&InitializeEventLoopTests);
+		FTestDelegates::GetGlobalTeardown().BindStatic(&CleanupEventLoopTests);
+	}
+} GEventLoopTestGlobalSetup;
+
+} // UE::EventLoop
diff --git a/Engine/Source/Runtime/Online/EventLoopTests/Tests/EventLoop/EventLoopManagedStorageTests.cpp b/Engine/Source/Runtime/Online/EventLoopTests/Tests/EventLoop/EventLoopManagedStorageTests.cpp
new file mode 100644
--- /dev/null
+++ b/Engine/Source/Runtime/Online/EventLoopTests/Tests/EventLoop/EventLoopManagedStorageTests.cpp
@@ -0,0 +1,646 @@
+// Copyright Epic Games, Inc. All Rights Reserved.
+
+#include "EventLoop/EventLoopManagedStorage.h"
+
+#include "TestHarness.h"
+
+namespace UE::EventLoop
+{
+	struct FMockManagedStorageHandleTraits
+	{
+		static constexpr const TCHAR* Name = TEXT("MockManagedStorageHandle");
+	};
+
+	struct FMockEventLoopManagedStorageTraits
+	{
+		static uint32 GetCurrentThreadId()
+		{
+			return CurrentThreadId;
+		}
+
+		static bool IsManagerThread(uint32 ManagerThreadId)
+		{
+			return ManagerThreadId == GetCurrentThreadId();
+		}
+
+		static void CheckNotInitialized(uint32 ManagerThreadId)
+		{
+			bCheckNotInitializedTriggered = (ManagerThreadId != 0);
+		}
+
+		static void CheckIsManagerThread(uint32 ManagerThreadId)
+		{
+			bCheckIsManagerThreadTriggered = !IsManagerThread(ManagerThreadId);
+		}
+
+		static constexpr bool bStorageAccessThreadChecksEnabled = true;
+		using InternalHandleArryAllocatorType = TInlineAllocator<32>;
+		using FExternalHandle = TResourceHandle<FMockManagedStorageHandleTraits>;
+
+		static void ResetTestConditions()
+		{
+			CurrentThreadId = 1;
+			bCheckNotInitializedTriggered = false;
+			bCheckIsManagerThreadTriggered = false;
+		}
+
+		static uint32 CurrentThreadId;
+		static bool bCheckNotInitializedTriggered;
+		static bool bCheckIsManagerThreadTriggered;
+	};
+
+	uint32 FMockEventLoopManagedStorageTraits::CurrentThreadId = 1;
+	bool FMockEventLoopManagedStorageTraits::bCheckNotInitializedTriggered = false;
+	bool FMockEventLoopManagedStorageTraits::bCheckIsManagerThreadTriggered = false;
+
+	TEST_CASE("EventLoop::EventLoopManagedStorage", "[Online][EventLoop][Smoke]")
+	{
+		const uint32 ManagerThreadId = 1;
+		const uint32 WorkerThreadId = 2;
+		FMockEventLoopManagedStorageTraits::ResetTestConditions();
+		FMockEventLoopManagedStorageTraits::CurrentThreadId = ManagerThreadId;
+
+		using MockStorageType = TManagedStorage<int32, FMockEventLoopManagedStorageTraits>;
+		using FMockStorageExternalHandle = MockStorageType::FExternalHandle;
+		using FMockStorageInternalHandle = MockStorageType::FInternalHandle;
+		using FInternalHandleArryType = MockStorageType::FInternalHandleArryType;
+
+		MockStorageType ManagedStorage;
+		const TManagedStorage<int32, FMockEventLoopManagedStorageTraits>& ConstManagedStorage = ManagedStorage;
+
+		SECTION("Single init call")
+		{
+			ManagedStorage.Init();
+			CHECK(!FMockEventLoopManagedStorageTraits::bCheckNotInitializedTriggered);
+		}
+
+		SECTION("Double init call")
+		{
+			ManagedStorage.Init();
+			ManagedStorage.Init();
+			CHECK(FMockEventLoopManagedStorageTraits::bCheckNotInitializedTriggered);
+		}
+
+		SECTION("Call update before init.")
+		{
+			ManagedStorage.Update();
+			CHECK(FMockEventLoopManagedStorageTraits::bCheckIsManagerThreadTriggered);
+		}
+
+		SECTION("Storage access pre-init.")
+		{
+			SECTION("Num called before init.")
+			{
+				ManagedStorage.Num();
+				CHECK(FMockEventLoopManagedStorageTraits::bCheckIsManagerThreadTriggered);
+			}
+
+			SECTION("IsEmpty called before init.")
+			{
+				CHECK(!FMockEventLoopManagedStorageTraits::bCheckIsManagerThreadTriggered);
+				ManagedStorage.IsEmpty();
+				CHECK(FMockEventLoopManagedStorageTraits::bCheckIsManagerThreadTriggered);
+			}
+
+			SECTION("begin called before init.")
+			{
+				ManagedStorage.begin();
+				CHECK(FMockEventLoopManagedStorageTraits::bCheckIsManagerThreadTriggered);
+			}
+
+			SECTION("const begin called before init.")
+			{
+				ConstManagedStorage.begin();
+				CHECK(FMockEventLoopManagedStorageTraits::bCheckIsManagerThreadTriggered);
+			}
+
+			SECTION("end called before init.")
+			{
+				ManagedStorage.end();
+				CHECK(FMockEventLoopManagedStorageTraits::bCheckIsManagerThreadTriggered);
+			}
+
+			SECTION("const end called before init.")
+			{
+				ConstManagedStorage.end();
+				CHECK(FMockEventLoopManagedStorageTraits::bCheckIsManagerThreadTriggered);
+			}
+
+			SECTION("Find external handle called before init.")
+			{
+				ManagedStorage.Find(FMockStorageExternalHandle());
+				CHECK(FMockEventLoopManagedStorageTraits::bCheckIsManagerThreadTriggered);
+			}
+
+			SECTION("const Find external handle called before init.")
+			{
+				ConstManagedStorage.Find(FMockStorageExternalHandle());
+				CHECK(FMockEventLoopManagedStorageTraits::bCheckIsManagerThreadTriggered);
+			}
+
+			SECTION("Find internal handle called before init.")
+			{
+				ManagedStorage.Find(FMockStorageInternalHandle());
+				CHECK(FMockEventLoopManagedStorageTraits::bCheckIsManagerThreadTriggered);
+			}
+
+			SECTION("const Find internal handle called before init.")
+			{
+				ConstManagedStorage.Find(FMockStorageInternalHandle());
+				CHECK(FMockEventLoopManagedStorageTraits::bCheckIsManagerThreadTriggered);
+			}
+		}
+
+		SECTION("Initialized")
+		{
+			int32 Element1Value = 77;
+			int32 Element2Value = 88;
+			ManagedStorage.Init();
+
+			SECTION("Call update from correct thread.")
+			{
+				ManagedStorage.Update();
+				CHECK(!FMockEventLoopManagedStorageTraits::bCheckIsManagerThreadTriggered);
+			}
+
+			SECTION("Call update from wrong thread.")
+			{
+				FMockEventLoopManagedStorageTraits::CurrentThreadId = WorkerThreadId;
+				ManagedStorage.Update();
+				CHECK(FMockEventLoopManagedStorageTraits::bCheckIsManagerThreadTriggered);
+			}
+
+			SECTION("Storage container access.")
+			{
+				SECTION("Add and remove behavior.")
+				{
+					FMockStorageInternalHandle InternalHandle1;
+					FMockStorageExternalHandle Handle1 = ManagedStorage.Add(MoveTemp(Element1Value));
+					FMockStorageExternalHandle Handle2 = ManagedStorage.Add(MoveTemp(Element2Value));
+					FMockStorageExternalHandle StaleHandle1 = Handle1;
+					FMockStorageExternalHandle InvalidHandle;
+					int32 Count = 0;
+					FInternalHandleArryType AddedHandles;
+					FInternalHandleArryType RemovedHandles;
+
+					// Element will not be added until Update is called.
+					CHECK(ManagedStorage.Num() == 0);
+					CHECK(ManagedStorage.IsEmpty());
+
+					AddedHandles.Empty();
+					ManagedStorage.Update(&AddedHandles);
+
+					REQUIRE(AddedHandles.Num() == 2);
+					InternalHandle1 = AddedHandles[0];
+					CHECK(InternalHandle1.GetExternalHandle() == Handle1);
+
+					// Element has been added.
+					CHECK(ManagedStorage.Num() == 2);
+					CHECK(!ManagedStorage.IsEmpty());
+
+					// Request to remove element.
+					ManagedStorage.Remove(Handle1);
+
+					// Element has not yet been removed.
+					CHECK(ManagedStorage.Num() == 2);
+					CHECK(!ManagedStorage.IsEmpty());
+
+					RemovedHandles.Empty();
+					ManagedStorage.Update(nullptr, &RemovedHandles);
+
+					REQUIRE(RemovedHandles.Num() == 1);
+					InternalHandle1 = RemovedHandles[0];
+					CHECK(InternalHandle1.GetExternalHandle() == Handle1);
+
+					// Element has been removed.
+					CHECK(ManagedStorage.Num() == 1);
+					CHECK(!ManagedStorage.IsEmpty());
+
+					// Element added and removed before Update called.
+					Handle1 = ManagedStorage.Add(MoveTemp(Element1Value));
+					ManagedStorage.Remove(Handle1);
+
+					AddedHandles.Empty();
+					RemovedHandles.Empty();
+					ManagedStorage.Update(&AddedHandles, &RemovedHandles);
+
+					// Managed storage still contains the second element.
+					CHECK(ManagedStorage.Num() == 1);
+
+					REQUIRE(AddedHandles.Num() == 1);
+					InternalHandle1 = AddedHandles[0];
+					CHECK(InternalHandle1.GetExternalHandle() == Handle1);
+
+					REQUIRE(RemovedHandles.Num() == 1);
+					InternalHandle1 = RemovedHandles[0];
+					CHECK(InternalHandle1.GetExternalHandle() == Handle1);
+
+					// Try to remove an element using a stale handle.
+					ManagedStorage.Remove(StaleHandle1);
+					RemovedHandles.Empty();
+					ManagedStorage.Update(nullptr, &RemovedHandles);
+					CHECK(!RemovedHandles.IsEmpty());
+					CHECK(ManagedStorage.Num() == 1);
+
+					// Try to remove an element using an invalid handle.
+					ManagedStorage.Remove(InvalidHandle);
+					RemovedHandles.Empty();
+					ManagedStorage.Update(nullptr, &RemovedHandles);
+					CHECK(!RemovedHandles.IsEmpty());
+					CHECK(ManagedStorage.Num() == 1);
+
+					// Remove second element.
+					ManagedStorage.Remove(Handle2);
+					ManagedStorage.Update();
+					CHECK(ManagedStorage.IsEmpty());
+					CHECK(ManagedStorage.Num() == 0);
+				}
+
+				SECTION("Ranged for behavior.")
+				{
+					FMockStorageExternalHandle Handle1 = ManagedStorage.Add(MoveTemp(Element1Value));
+					FMockStorageExternalHandle Handle2 = ManagedStorage.Add(MoveTemp(Element2Value));
+					int32 Count = 0;
+
+					FInternalHandleArryType AddedHandles;
+					ManagedStorage.Update(&AddedHandles);
+					REQUIRE(AddedHandles.Num() == 2);
+
+					FMockStorageInternalHandle InternalHandle1 = AddedHandles[0];
+					FMockStorageInternalHandle InternalHandle2 = AddedHandles[1];
+
+					// Ranged for access.
+					Count = 0;
+					bool bHandle1Tested = false;
+					bool bHandle2Tested = false;
+					for (TPair<FMockStorageInternalHandle, int32&>& Entry : ManagedStorage)
+					{
+						if (Count == 0)
+						{
+							CHECK(Entry.Key == InternalHandle1);
+							CHECK(Entry.Key.GetExternalHandle() == Handle1);
+							CHECK(Entry.Value == Element1Value);
+							bHandle1Tested = true;
+						}
+						if (Count == 1)
+						{
+							CHECK(Entry.Key == InternalHandle2);
+							CHECK(Entry.Key.GetExternalHandle() == Handle2);
+							CHECK(Entry.Value == Element2Value);
+							bHandle2Tested = true;
+						}
+
+						++Count;
+					}
+					CHECK(Count == ManagedStorage.Num());
+					CHECK(bHandle1Tested);
+					CHECK(bHandle2Tested);
+
+					// Modify element during ranged iteration.
+					const int32 ModifiedValue = 56;
+					for (TPair<FMockStorageInternalHandle, int32&>& Entry : ManagedStorage)
+					{
+						if (Entry.Key == InternalHandle1)
+						{
+							Entry.Value = ModifiedValue;
+						}
+					}
+
+					// Verify element 1 changed.
+					int32* Element1 = ManagedStorage.Find(InternalHandle1);
+					REQUIRE(Element1);
+					CHECK(*Element1 == ModifiedValue);
+
+					// Verify element 2 did not change.
+					int32* Element2 = ManagedStorage.Find(InternalHandle2);
+					REQUIRE(Element2);
+					CHECK(*Element2 == Element2Value);
+
+					// Remove elements.
+					ManagedStorage.Remove(Handle1);
+					ManagedStorage.Remove(Handle2);
+					ManagedStorage.Update();
+					REQUIRE(ManagedStorage.IsEmpty());
+
+					// Ranged for access.
+					Count = 0;
+					for (TPair<FMockStorageInternalHandle, int32&>& Entry : ManagedStorage)
+					{
+						++Count;
+					}
+					CHECK(Count == 0);
+				}
+
+				SECTION("Ranged for const behavior.")
+				{
+					FMockStorageExternalHandle Handle1 = ManagedStorage.Add(MoveTemp(Element1Value));
+					FMockStorageExternalHandle Handle2 = ManagedStorage.Add(MoveTemp(Element2Value));
+					int32 Count = 0;
+
+					FInternalHandleArryType AddedHandles;
+					ManagedStorage.Update(&AddedHandles);
+					REQUIRE(AddedHandles.Num() == 2);
+
+					FMockStorageInternalHandle InternalHandle1 = AddedHandles[0];
+					FMockStorageInternalHandle InternalHandle2 = AddedHandles[1];
+
+					// Ranged for const access.
+					Count = 0;
+					bool bHandle1Tested = false;
+					bool bHandle2Tested = false;
+					for (const TPair<FMockStorageInternalHandle, const int32&>& Entry : ConstManagedStorage)
+					{
+						if (Count == 0)
+						{
+							CHECK(Entry.Key == InternalHandle1);
+							CHECK(Entry.Key.GetExternalHandle() == Handle1);
+							CHECK(Entry.Value == Element1Value);
+							bHandle1Tested = true;
+						}
+						if (Count == 1)
+						{
+							CHECK(Entry.Key == InternalHandle2);
+							CHECK(Entry.Key.GetExternalHandle() == Handle2);
+							CHECK(Entry.Value == Element2Value);
+							bHandle2Tested = true;
+						}
+
+						++Count;
+					}
+					CHECK(Count == ManagedStorage.Num());
+					CHECK(bHandle1Tested);
+					CHECK(bHandle2Tested);
+
+					// Remove elements.
+					ManagedStorage.Remove(Handle1);
+					ManagedStorage.Remove(Handle2);
+					ManagedStorage.Update();
+					REQUIRE(ManagedStorage.IsEmpty());
+
+					// Ranged for const access.
+					Count = 0;
+					for (const TPair<FMockStorageInternalHandle, const int32&>& Entry : ConstManagedStorage)
+					{
+						++Count;
+					}
+					CHECK(Count == 0);
+				}
+			}
+
+			SECTION("Element access.")
+			{
+				FInternalHandleArryType AddedHandles;
+
+				SECTION("New element added.")
+				{
+					FMockStorageExternalHandle ExternalHandle1 = ManagedStorage.Add(MoveTemp(Element1Value));
+					FMockStorageExternalHandle ExternalHandle2 = ManagedStorage.Add(MoveTemp(Element2Value));
+					FMockStorageExternalHandle StaleExternalHandle2 = ExternalHandle2;
+
+					AddedHandles.Empty();
+					ManagedStorage.Update(&AddedHandles);
+
+					CHECK(AddedHandles.Num() == 2);
+					FMockStorageInternalHandle InternalHandle1 = AddedHandles[0];
+					FMockStorageInternalHandle InternalHandle2 = AddedHandles[1];
+					FMockStorageInternalHandle StaleInternalHandle2 = InternalHandle2;
+
+					ManagedStorage.Remove(ExternalHandle2);
+					ManagedStorage.Update();
+
+					SECTION("Mutable access.")
+					{
+						int32* Element = nullptr;
+
+						// Retrieve value using internal handle.
+						Element = ManagedStorage.Find(InternalHandle1);
+						REQUIRE(Element);
+						CHECK(*Element == Element1Value);
+
+						// Retrieve value using external handle.
+						Element = ManagedStorage.Find(ExternalHandle1);
+						REQUIRE(Element);
+						CHECK(*Element == Element1Value);
+
+						// Retrieve value using stale internal handle.
+						Element = ManagedStorage.Find(StaleInternalHandle2);
+						CHECK(Element == nullptr);
+
+						// Retrieve value using stale external handle.
+						Element = ManagedStorage.Find(StaleExternalHandle2);
+						CHECK(Element == nullptr);
+					}
+
+					SECTION("Const access.")
+					{
+						const int32* Element = nullptr;
+
+						// Retrieve value using internal handle.
+						Element = ConstManagedStorage.Find(InternalHandle1);
+						REQUIRE(Element);
+						CHECK(*Element == Element1Value);
+
+						// Retrieve value using external handle.
+						Element = ConstManagedStorage.Find(ExternalHandle1);
+						REQUIRE(Element);
+						CHECK(*Element == Element1Value);
+
+						// Retrieve value using stale internal handle.
+						Element = ManagedStorage.Find(StaleInternalHandle2);
+						CHECK(Element == nullptr);
+
+						// Retrieve value using stale external handle.
+						Element = ManagedStorage.Find(StaleExternalHandle2);
+						CHECK(Element == nullptr);
+					}
+				}
+
+				SECTION("New element removed before update called.")
+				{
+					FMockStorageExternalHandle ExternalHandle = ManagedStorage.Add(MoveTemp(Element1Value));
+					ManagedStorage.Remove(ExternalHandle);
+					ManagedStorage.Update(&AddedHandles);
+					REQUIRE(AddedHandles.Num() == 1);
+					FMockStorageInternalHandle InternalHandle = AddedHandles[0];
+
+					SECTION("Mutable access.")
+					{
+						int32* Element = nullptr;
+
+						// Retrieve value using internal handle.
+						Element = ManagedStorage.Find(InternalHandle);
+						CHECK(Element == nullptr);
+
+						// Retrieve value using external handle.
+						Element = ManagedStorage.Find(ExternalHandle);
+						CHECK(Element == nullptr);
+					}
+
+					SECTION("Const access.")
+					{
+						const int32* Element = nullptr;
+
+						// Retrieve value using internal handle.
+						Element = ConstManagedStorage.Find(InternalHandle);
+						CHECK(Element == nullptr);
+
+						// Retrieve value using external handle.
+						Element = ConstManagedStorage.Find(InternalHandle);
+						CHECK(Element == nullptr);
+					}
+				}
+			}
+
+			SECTION("Storage access thread safety.")
+			{
+				FMockStorageExternalHandle ExternalHandle = ManagedStorage.Add(MoveTemp(Element1Value));
+				FInternalHandleArryType AddedHandles;
+				ManagedStorage.Update(&AddedHandles);
+				REQUIRE(AddedHandles.Num() == 1);
+				FMockStorageInternalHandle InternalHandle = AddedHandles[0];
+
+				FMockEventLoopManagedStorageTraits::ResetTestConditions();
+				FMockEventLoopManagedStorageTraits::CurrentThreadId = WorkerThreadId;
+
+				SECTION("Num called from wrong thread.")
+				{
+					ManagedStorage.Num();
+					CHECK(FMockEventLoopManagedStorageTraits::bCheckIsManagerThreadTriggered);
+				}
+
+				SECTION("IsEmpty called from wrong thread.")
+				{
+					ManagedStorage.IsEmpty();
+					CHECK(FMockEventLoopManagedStorageTraits::bCheckIsManagerThreadTriggered);
+				}
+
+				SECTION("begin called from wrong thread.")
+				{
+					ManagedStorage.begin();
+					CHECK(FMockEventLoopManagedStorageTraits::bCheckIsManagerThreadTriggered);
+				}
+
+				SECTION("const begin called from wrong thread.")
+				{
+					ConstManagedStorage.begin();
+					CHECK(FMockEventLoopManagedStorageTraits::bCheckIsManagerThreadTriggered);
+				}
+
+				SECTION("end called from wrong thread.")
+				{
+					ManagedStorage.end();
+					CHECK(FMockEventLoopManagedStorageTraits::bCheckIsManagerThreadTriggered);
+				}
+
+				SECTION("const end called from wrong thread.")
+				{
+					ConstManagedStorage.end();
+					CHECK(FMockEventLoopManagedStorageTraits::bCheckIsManagerThreadTriggered);
+				}
+
+				SECTION("Find external handle called from wrong thread.")
+				{
+					ManagedStorage.Find(ExternalHandle);
+					CHECK(FMockEventLoopManagedStorageTraits::bCheckIsManagerThreadTriggered);
+				}
+
+				SECTION("const Find external handle called from wrong thread.")
+				{
+					ConstManagedStorage.Find(ExternalHandle);
+					CHECK(FMockEventLoopManagedStorageTraits::bCheckIsManagerThreadTriggered);
+				}
+
+				SECTION("Find internal handle called from wrong thread.")
+				{
+					ManagedStorage.Find(InternalHandle);
+					CHECK(FMockEventLoopManagedStorageTraits::bCheckIsManagerThreadTriggered);
+				}
+
+				SECTION("const Find internal handle called from wrong thread.")
+				{
+					ConstManagedStorage.Find(InternalHandle);
+					CHECK(FMockEventLoopManagedStorageTraits::bCheckIsManagerThreadTriggered);
+				}
+			}
+
+			SECTION("Handle string conversion.")
+			{
+				FInternalHandleArryType AddedHandles;
+				FMockStorageExternalHandle ExternalHandle;
+				FString ExternalStringValue;
+				FString InternalStringValue;
+				FString ExpectedStringValue;
+
+				// Test that ToString resolves name from traits.
+				ExternalHandle = ManagedStorage.Add(MoveTemp(Element1Value));
+				ExternalStringValue = ExternalHandle.ToString();
+				CHECK(ExternalStringValue.Contains(FMockEventLoopManagedStorageTraits::FExternalHandle::GetTypeName()));
+
+				// Test that internal handle as a string contains the external handle string.
+				ManagedStorage.Update(&AddedHandles);
+				REQUIRE(AddedHandles.Num() == 1);
+				InternalStringValue = AddedHandles[0].ToString();
+				CHECK(InternalStringValue.Contains(ExternalStringValue));
+			}
+		}
+	}
+
+	TEST_CASE("EventLoop::EventLoopManagedStorage::MoveOnlyType", "[Online][EventLoop][Smoke]")
+	{
+		FMockEventLoopManagedStorageTraits::ResetTestConditions();
+
+		struct FManagedData
+		{
+			FManagedData(const FManagedData&) = delete;
+			FManagedData& operator=(const FManagedData&) = delete;
+
+			FManagedData()
+				: IntValue(0)
+			{
+			}
+
+			FManagedData(FManagedData&& Other)
+				: IntValue(MoveTemp(Other.IntValue))
+			{
+			}
+
+			FManagedData& operator=(FManagedData&& Other)
+			{
+				IntValue = MoveTemp(Other.IntValue);
+				return *this;
+			}
+
+			int32 IntValue;
+		};
+
+		using MockStorageType = TManagedStorage<FManagedData, FMockEventLoopManagedStorageTraits>;
+		using FMockStorageExternalHandle = MockStorageType::FExternalHandle;
+		using FMockStorageInternalHandle = MockStorageType::FInternalHandle;
+		using FInternalHandleArryType = MockStorageType::FInternalHandleArryType;
+
+		MockStorageType ManagedStorage;
+
+		SECTION("Element access.")
+		{
+			const int32 ExpectedValue = 5;
+			FManagedData ManagedData;
+			ManagedData.IntValue = ExpectedValue;
+			FMockStorageExternalHandle Handle = ManagedStorage.Add(MoveTemp(ManagedData));
+
+			FInternalHandleArryType AddedHandles;
+			ManagedStorage.Update(&AddedHandles);
+			REQUIRE(!AddedHandles.IsEmpty());
+
+			const int32 ModifiedValue = 33;
+			for (TPair<FMockStorageInternalHandle, FManagedData&>& Entry : ManagedStorage)
+			{
+				CHECK(Entry.Value.IntValue == ExpectedValue);
+				Entry.Value.IntValue = ModifiedValue;
+			}
+
+			FManagedData* FoundManagedData = ManagedStorage.Find(AddedHandles[0]);
+			REQUIRE(FoundManagedData);
+			CHECK(FoundManagedData->IntValue == ModifiedValue);
+		}
+	}
+} // UE::EventLoop
diff --git a/Engine/Source/Runtime/Online/EventLoopTests/Tests/EventLoop/EventLoopTests.cpp b/Engine/Source/Runtime/Online/EventLoopTests/Tests/EventLoop/EventLoopTests.cpp
new file mode 100644
--- /dev/null
+++ b/Engine/Source/Runtime/Online/EventLoopTests/Tests/EventLoop/EventLoopTests.cpp
@@ -0,0 +1,491 @@
+// Copyright Epic Games, Inc. All Rights Reserved.
+
+#include "EventLoop/EventLoop.h"
+#include "Templates/SharedPointer.h"
+
+#include "TestHarness.h"
+
+namespace UE::EventLoop
+{
+	struct FMockTimerHandleTraits
+	{
+		static constexpr const TCHAR* Name = TEXT("MockTimerHandle");
+	};
+
+	struct FEventLoopMockTraits
+	{
+		static FORCEINLINE uint32 GetCurrentThreadId()
+		{
+			return CurrentThreadId;
+		}
+
+		static FORCEINLINE FTimespan GetCurrentTime()
+		{
+			return CurrentTime;
+		}
+
+		static FORCEINLINE bool IsEventLoopThread(uint32 EventLoopThreadId)
+		{
+			return EventLoopThreadId == GetCurrentThreadId();
+		}
+
+		static FORCEINLINE bool IsInitialized(uint32 EventLoopThreadId)
+		{
+			return (EventLoopThreadId != 0);
+		}
+
+		static FORCEINLINE void CheckInitialized(uint32 EventLoopThreadId)
+		{
+			bCheckInitializedTriggered = (EventLoopThreadId == 0);
+		}
+
+		static FORCEINLINE void CheckNotInitialized(uint32 EventLoopThreadId)
+		{
+			bCheckNotInitializedTriggered = (EventLoopThreadId != 0);
+		}
+
+		static FORCEINLINE void CheckIsEventLoopThread(uint32 EventLoopThreadId)
+		{
+			bCheckIsEventLoopThreadTriggered = !IsEventLoopThread(EventLoopThreadId);
+		}
+
+		static FORCEINLINE bool IsShutdownRequested(bool bInShutdownRequested)
+		{
+			return bShutdownRequested || bInShutdownRequested;
+		}
+
+		struct FTimerManagerTraits : public FTimerManagerTraitsBase
+		{
+			static void CheckIsManagerThread(bool bIsManagerThread)
+			{
+				bCheckIsManagerThreadTriggered = !bIsManagerThread;
+			}
+
+			using FTimerHeapAllocatorType = TInlineAllocator<32>;
+			using FTimerRepeatAllocatorType = TInlineAllocator<32>;
+
+			struct FStorageTraits
+			{
+				static uint32 GetCurrentThreadId()
+				{
+					return CurrentThreadId;
+				}
+
+				static bool IsManagerThread(uint32 ManagerThreadId)
+				{
+					return ManagerThreadId == GetCurrentThreadId();
+				}
+
+				static void CheckNotInitialized(uint32 ManagerThreadId)
+				{
+					bCheckNotInitializedTriggered = (ManagerThreadId != 0);
+				}
+
+				static void CheckIsManagerThread(uint32 ManagerThreadId)
+				{
+					bCheckIsManagerThreadTriggered = !IsManagerThread(ManagerThreadId);
+				}
+
+				static constexpr bool bStorageAccessThreadChecksEnabled = true;
+				using InternalHandleArryAllocatorType = TInlineAllocator<32>;
+				using FExternalHandle = FTimerHandle;
+
+				static void ResetTestConditions()
+				{
+					CurrentThreadId = 1;
+					bCheckNotInitializedTriggered = false;
+					bCheckIsManagerThreadTriggered = false;
+				}
+
+				static uint32 CurrentThreadId;
+				static bool bCheckNotInitializedTriggered;
+				static bool bCheckIsManagerThreadTriggered;
+			};
+
+			static void ResetTestConditions()
+			{
+				FStorageTraits::ResetTestConditions();
+				bCheckIsManagerThreadTriggered = false;
+			}
+
+			static bool bCheckIsManagerThreadTriggered;
+		};
+
+		static void ResetTestConditions()
+		{
+			FTimerManagerTraits::ResetTestConditions();
+			CurrentThreadId = 1;
+			CurrentTime = FTimespan::Zero();
+			bCheckInitializedTriggered = false;
+			bCheckNotInitializedTriggered = false;
+			bCheckIsEventLoopThreadTriggered = false;
+			bShutdownRequested = false;
+		}
+
+		static uint32 CurrentThreadId;
+		static FTimespan CurrentTime;
+		static bool bCheckInitializedTriggered;
+		static bool bCheckNotInitializedTriggered;
+		static bool bCheckIsEventLoopThreadTriggered;
+		static bool bShutdownRequested;
+	};
+
+	uint32 FEventLoopMockTraits::CurrentThreadId = 1;
+	FTimespan FEventLoopMockTraits::CurrentTime = FTimespan::Zero();
+	bool FEventLoopMockTraits::bCheckInitializedTriggered = false;
+	bool FEventLoopMockTraits::bCheckNotInitializedTriggered = false;
+	bool FEventLoopMockTraits::bCheckIsEventLoopThreadTriggered = false;
+	bool FEventLoopMockTraits::bShutdownRequested = false;
+	bool FEventLoopMockTraits::FTimerManagerTraits::bCheckIsManagerThreadTriggered = false;
+	uint32 FEventLoopMockTraits::FTimerManagerTraits::FStorageTraits::CurrentThreadId = 1;
+	bool FEventLoopMockTraits::FTimerManagerTraits::FStorageTraits::bCheckNotInitializedTriggered = false;
+	bool FEventLoopMockTraits::FTimerManagerTraits::FStorageTraits::bCheckIsManagerThreadTriggered = false;
+
+	class FIOAccessMock final : public FNoncopyable
+	{
+	public:
+		// Nothing needed for now.
+	};
+
+	class FIOManagerMock final : public IIOManager
+	{
+	public:
+		using FIOAccess = FIOAccessMock;
+
+		struct FParams
+		{
+		};
+
+		FIOManagerMock(IEventLoop&, FParams&&)
+		{
+		}
+
+		virtual ~FIOManagerMock() = default;
+
+		virtual bool Init() override
+		{
+			bHasInitialized = true;
+			return bInitReturnValue;
+		}
+
+		virtual void Shutdown() override
+		{
+			bHasShutdown = true;
+		}
+
+		virtual void Notify() override
+		{
+			bHasNotified = true;
+		}
+
+		virtual void Poll(FTimespan WaitTime) override
+		{
+			CHECK(bHasNotified == bExpectedNotify);
+			CHECK(WaitTime == ExpectedWaitTime);
+		}
+
+		FIOAccess& GetIOAccess()
+		{
+			return IOAccess;
+		}
+
+		static void ResetTestConditions()
+		{
+			bHasInitialized = false;
+			bInitReturnValue = true;
+			bHasNotified = false;
+			bExpectedNotify = false;
+			bHasShutdown = false;
+			ExpectedWaitTime = FTimespan::Zero();
+		}
+
+		static bool bHasInitialized;
+		static bool bInitReturnValue;
+		static bool bHasNotified;
+		static bool bExpectedNotify;
+		static bool bHasShutdown;
+		static FTimespan ExpectedWaitTime;
+
+		FIOAccess IOAccess;
+	};
+
+	bool FIOManagerMock::bHasInitialized = false;
+	bool FIOManagerMock::bInitReturnValue = true;
+	bool FIOManagerMock::bHasNotified = false;
+	bool FIOManagerMock::bExpectedNotify = false;
+	bool FIOManagerMock::bHasShutdown = false;
+	FTimespan FIOManagerMock::ExpectedWaitTime = FTimespan::Zero();
+
+	TEST_CASE("EventLoop", "[Online][EventLoop][Smoke]")
+	{
+		TEventLoop<FIOManagerMock, FEventLoopMockTraits> EventLoop;
+		FIOManagerMock::ResetTestConditions();
+		FEventLoopMockTraits::ResetTestConditions();
+
+		const FTimespan InfiniteWait(std::numeric_limits<int64>::max());
+		const uint32 EventLoopThreadId = 1;
+		FEventLoopMockTraits::CurrentThreadId = EventLoopThreadId;
+		FEventLoopMockTraits::FTimerManagerTraits::FStorageTraits::CurrentThreadId = EventLoopThreadId;
+
+		SECTION("Single init call")
+		{
+			EventLoop.Init();
+			CHECK(!FEventLoopMockTraits::bCheckNotInitializedTriggered);
+		}
+
+		SECTION("Double init call")
+		{
+			EventLoop.Init();
+			EventLoop.Init();
+			CHECK(FEventLoopMockTraits::bCheckNotInitializedTriggered);
+		}
+
+		SECTION("Call Run before Init")
+		{
+			FEventLoopMockTraits::bShutdownRequested = true;
+			FIOManagerMock::ExpectedWaitTime = InfiniteWait;
+
+			EventLoop.Run();
+			CHECK(FEventLoopMockTraits::bCheckInitializedTriggered);
+		}
+
+		SECTION("Call RunOnce before Init")
+		{
+			FEventLoopMockTraits::bShutdownRequested = true;
+
+			EventLoop.RunOnce(FTimespan::Zero());
+			CHECK(FEventLoopMockTraits::bCheckInitializedTriggered);
+		}
+
+		SECTION("Call RequestShutdown before Init")
+		{
+			EventLoop.RequestShutdown();
+
+			// Request shutdown is allowed to be called before init.
+			CHECK(!FEventLoopMockTraits::bCheckInitializedTriggered);
+
+			// Check that loop is shutdown when run after calling init.
+			CHECK(EventLoop.Init());
+
+			FIOManagerMock::bExpectedNotify = true;
+			FIOManagerMock::ExpectedWaitTime = InfiniteWait;
+			CHECK(!EventLoop.RunOnce(InfiniteWait));
+		}
+
+		SECTION("Call SetTimer before Init")
+		{
+			bool bTimerRan = false;
+			FTimerHandle Handle = EventLoop.SetTimer([&bTimerRan](){ bTimerRan = true; }, FTimespan::FromSeconds(1));
+			CHECK(Handle.IsValid());
+
+			// Request shutdown is allowed to be called before init.
+			CHECK(!FEventLoopMockTraits::bCheckInitializedTriggered);
+
+			// Check that timer is run when run after calling init.
+			CHECK(EventLoop.Init());
+
+			// The first run schedules the timer. It will not wait because it has been interrupted.
+			FIOManagerMock::bExpectedNotify = true;
+			FIOManagerMock::ExpectedWaitTime = InfiniteWait;
+			CHECK(EventLoop.RunOnce(InfiniteWait));
+
+			// The second run will run until the timer duration and fire the timer.
+			FIOManagerMock::bHasNotified = false;
+			FIOManagerMock::bExpectedNotify = false;
+			FIOManagerMock::ExpectedWaitTime = FTimespan::FromSeconds(1);
+			FEventLoopMockTraits::CurrentTime = FTimespan::FromSeconds(1);
+			CHECK(EventLoop.RunOnce(InfiniteWait));
+			CHECK(bTimerRan);
+		}
+
+		SECTION("Call ClearTimer before Init")
+		{
+			bool bTimerRan = false;
+			FTimerHandle Handle = EventLoop.SetTimer([&bTimerRan](){ bTimerRan = true; }, FTimespan::FromSeconds(1));
+			CHECK(Handle.IsValid());
+
+			EventLoop.ClearTimer(Handle);
+
+			// Request shutdown is allowed to be called before init.
+			CHECK(!FEventLoopMockTraits::bCheckInitializedTriggered);
+
+			// Check that timer is not run when run after calling init.
+			CHECK(EventLoop.Init());
+
+			// The first run schedules the timer. It will not wait because it has been interrupted.
+			FIOManagerMock::bExpectedNotify = true;
+			FIOManagerMock::ExpectedWaitTime = InfiniteWait;
+			CHECK(EventLoop.RunOnce(InfiniteWait));
+
+			// The second run would have fired the timer, but it is no longer expected to be present.
+			FIOManagerMock::bHasNotified = false;
+			FIOManagerMock::bExpectedNotify = false;
+			FIOManagerMock::ExpectedWaitTime = InfiniteWait;
+			FEventLoopMockTraits::CurrentTime = FTimespan::FromSeconds(1);
+			CHECK(EventLoop.RunOnce(InfiniteWait));
+			CHECK(!bTimerRan);
+		}
+
+		SECTION("Call PostAsyncTask before Init")
+		{
+			bool bTaskTriggered = false;
+			EventLoop.PostAsyncTask([&bTaskTriggered](){ bTaskTriggered = true; });
+
+			// PostAsyncTask is allowed to be called before init.
+			CHECK(!FEventLoopMockTraits::bCheckInitializedTriggered);
+
+			// Check that task is run when run after calling init.
+			CHECK(EventLoop.Init());
+
+			// Loop is expected to have been interrupted due to scheduling a task.
+			FIOManagerMock::bExpectedNotify = true;
+			FIOManagerMock::ExpectedWaitTime = InfiniteWait;
+			CHECK(EventLoop.RunOnce(InfiniteWait));
+			CHECK(bTaskTriggered);
+		}
+
+		SECTION("Initialized")
+		{
+			EventLoop.Init();
+
+			SECTION("Call Run")
+			{
+				FEventLoopMockTraits::bShutdownRequested = true;
+				FIOManagerMock::ExpectedWaitTime = InfiniteWait;
+
+				EventLoop.Run();
+
+				CHECK(!FEventLoopMockTraits::bCheckInitializedTriggered);
+				CHECK(!FEventLoopMockTraits::bCheckIsEventLoopThreadTriggered);
+			}
+
+			SECTION("Call RunOnce - no shutdown")
+			{
+				FIOManagerMock::ExpectedWaitTime = FTimespan::FromSeconds(1);
+
+				bool bContinueRunning = EventLoop.RunOnce(FIOManagerMock::ExpectedWaitTime);
+				CHECK(bContinueRunning == true);
+
+				CHECK(!FEventLoopMockTraits::bCheckInitializedTriggered);
+				CHECK(!FEventLoopMockTraits::bCheckIsEventLoopThreadTriggered);
+			}
+
+			SECTION("Call RunOnce - shutdown")
+			{
+				FEventLoopMockTraits::bShutdownRequested = true;
+				FIOManagerMock::ExpectedWaitTime = FTimespan::FromSeconds(1);
+
+				bool bContinueRunning = EventLoop.RunOnce(FIOManagerMock::ExpectedWaitTime);
+				CHECK(bContinueRunning == false);
+
+				CHECK(!FEventLoopMockTraits::bCheckInitializedTriggered);
+				CHECK(!FEventLoopMockTraits::bCheckIsEventLoopThreadTriggered);
+			}
+
+			SECTION("Call SetTimer")
+			{
+				bool bTimerRan = false;
+				FIOManagerMock::ExpectedWaitTime = InfiniteWait;
+				FIOManagerMock::bExpectedNotify = true;
+
+				EventLoop.SetTimer([&](){ bTimerRan = true; }, FTimespan::FromSeconds(1));
+				EventLoop.RunOnce(InfiniteWait);
+				CHECK(!FEventLoopMockTraits::bCheckInitializedTriggered);
+				CHECK(!FEventLoopMockTraits::bCheckIsEventLoopThreadTriggered);
+				CHECK(!bTimerRan);
+
+				FIOManagerMock::ExpectedWaitTime = FTimespan::FromSeconds(1);
+				FEventLoopMockTraits::CurrentTime = FTimespan::FromSeconds(1);
+				FIOManagerMock::bHasNotified = false;
+				FIOManagerMock::bExpectedNotify = false;
+				EventLoop.RunOnce(InfiniteWait);
+				CHECK(bTimerRan);
+
+				bTimerRan = false;
+				FIOManagerMock::ExpectedWaitTime = InfiniteWait;
+				FEventLoopMockTraits::CurrentTime = FTimespan::FromSeconds(2);
+				FIOManagerMock::bExpectedNotify = false;
+				EventLoop.RunOnce(InfiniteWait);
+				CHECK(!bTimerRan);
+			}
+
+			SECTION("Call SetTimer - repeating")
+			{
+				bool bTimerRan = false;
+				FIOManagerMock::ExpectedWaitTime = InfiniteWait;
+				FIOManagerMock::bExpectedNotify = true;
+				EventLoop.SetTimer([&](){ bTimerRan = true; }, FTimespan::FromSeconds(1), true);
+				EventLoop.RunOnce(InfiniteWait);
+				CHECK(!FEventLoopMockTraits::bCheckInitializedTriggered);
+				CHECK(!FEventLoopMockTraits::bCheckIsEventLoopThreadTriggered);
+				CHECK(!bTimerRan);
+
+				FIOManagerMock::ExpectedWaitTime = FTimespan::FromSeconds(1);
+				FEventLoopMockTraits::CurrentTime = FTimespan::FromSeconds(1);
+				FIOManagerMock::bHasNotified = false;
+				FIOManagerMock::bExpectedNotify = false;
+				EventLoop.RunOnce(InfiniteWait);
+				CHECK(bTimerRan);
+
+				bTimerRan = false;
+				FIOManagerMock::ExpectedWaitTime = FTimespan::FromSeconds(1);
+				FEventLoopMockTraits::CurrentTime = FTimespan::FromSeconds(2);
+				FIOManagerMock::bExpectedNotify = false;
+				EventLoop.RunOnce(InfiniteWait);
+				CHECK(bTimerRan);
+			}
+
+			SECTION("Call ClearTimer")
+			{
+				FIOManagerMock::ExpectedWaitTime = InfiniteWait;
+				FIOManagerMock::bExpectedNotify = true;
+
+				FTimerHandle Handle;
+				EventLoop.ClearTimer(Handle);
+				EventLoop.RunOnce(FIOManagerMock::ExpectedWaitTime);
+
+				CHECK(!FEventLoopMockTraits::bCheckInitializedTriggered);
+				CHECK(!FEventLoopMockTraits::bCheckIsEventLoopThreadTriggered);
+			}
+
+			SECTION("Call PostAsyncTask")
+			{
+				FIOManagerMock::ExpectedWaitTime = InfiniteWait;
+				FIOManagerMock::bExpectedNotify = true;
+
+				bool bTaskRan = false;
+				EventLoop.PostAsyncTask([&bTaskRan](){ bTaskRan = true; });
+				EventLoop.RunOnce(FIOManagerMock::ExpectedWaitTime);
+
+				CHECK(!FEventLoopMockTraits::bCheckInitializedTriggered);
+				CHECK(!FEventLoopMockTraits::bCheckIsEventLoopThreadTriggered);
+				CHECK(bTaskRan);
+			}
+
+			SECTION("Call RequestShutdown")
+			{
+				bool bContinueRunning = false;
+
+				FIOManagerMock::ExpectedWaitTime = InfiniteWait;
+				FIOManagerMock::bExpectedNotify = true;
+
+				bool bShutdownCompleted = false;
+				EventLoop.RequestShutdown([&bShutdownCompleted](){ bShutdownCompleted = true; });
+				bContinueRunning = EventLoop.RunOnce(InfiniteWait);
+
+				CHECK(!FEventLoopMockTraits::bCheckInitializedTriggered);
+				CHECK(!FEventLoopMockTraits::bCheckIsEventLoopThreadTriggered);
+				CHECK(bShutdownCompleted);
+				CHECK(!bContinueRunning);
+
+				// Request again after requested.
+				bShutdownCompleted = false;
+				EventLoop.RequestShutdown([&bShutdownCompleted](){ bShutdownCompleted = true; });
+				CHECK(bShutdownCompleted);
+				CHECK(!bContinueRunning);
+
+				EventLoop.RunOnce(InfiniteWait);
+			}
+		}
+	}
+} // UE::EventLoop
diff --git a/Engine/Source/Runtime/Online/EventLoopTests/Tests/EventLoop/EventLoopTimerTests.cpp b/Engine/Source/Runtime/Online/EventLoopTests/Tests/EventLoop/EventLoopTimerTests.cpp
new file mode 100644
--- /dev/null
+++ b/Engine/Source/Runtime/Online/EventLoopTests/Tests/EventLoop/EventLoopTimerTests.cpp
@@ -0,0 +1,606 @@
+// Copyright Epic Games, Inc. All Rights Reserved.
+
+#include "EventLoop/EventLoopTimer.h"
+
+#include "TestHarness.h"
+
+namespace UE::EventLoop
+{
+	struct FMockTimerHandleTraits
+	{
+		static constexpr const TCHAR* Name = TEXT("MockTimerHandle");
+	};
+
+	struct FEventLoopTimerManagerMockTraits : public FTimerManagerTraitsBase
+	{
+		static void CheckIsManagerThread(bool bIsManagerThread)
+		{
+			bCheckIsManagerThreadTriggered = !bIsManagerThread;
+		}
+
+		using FTimerHeapAllocatorType = TInlineAllocator<32>;
+		using FTimerRepeatAllocatorType = TInlineAllocator<32>;
+
+		struct FStorageTraits
+		{
+			static uint32 GetCurrentThreadId()
+			{
+				return CurrentThreadId;
+			}
+
+			static bool IsManagerThread(uint32 ManagerThreadId)
+			{
+				return ManagerThreadId == GetCurrentThreadId();
+			}
+
+			static void CheckNotInitialized(uint32 ManagerThreadId)
+			{
+				bCheckNotInitializedTriggered = (ManagerThreadId != 0);
+			}
+
+			static void CheckIsManagerThread(uint32 ManagerThreadId)
+			{
+				bCheckIsManagerThreadTriggered = !IsManagerThread(ManagerThreadId);
+			}
+
+			static constexpr bool bStorageAccessThreadChecksEnabled = true;
+			using InternalHandleArryAllocatorType = TInlineAllocator<32>;
+			using FExternalHandle = TResourceHandle<FMockTimerHandleTraits>;
+
+			static void ResetTestConditions()
+			{
+				CurrentThreadId = 1;
+				bCheckNotInitializedTriggered = false;
+				bCheckIsManagerThreadTriggered = false;
+			}
+
+			static uint32 CurrentThreadId;
+			static bool bCheckNotInitializedTriggered;
+			static bool bCheckIsManagerThreadTriggered;
+		};
+
+		static void ResetTestConditions()
+		{
+			FStorageTraits::ResetTestConditions();
+			bCheckIsManagerThreadTriggered = false;
+		}
+
+		static bool bCheckIsManagerThreadTriggered;
+	};
+
+	bool FEventLoopTimerManagerMockTraits::bCheckIsManagerThreadTriggered = false;
+	uint32 FEventLoopTimerManagerMockTraits::FStorageTraits::CurrentThreadId = 1;
+	bool FEventLoopTimerManagerMockTraits::FStorageTraits::bCheckNotInitializedTriggered = false;
+	bool FEventLoopTimerManagerMockTraits::FStorageTraits::bCheckIsManagerThreadTriggered = false;
+
+	TEST_CASE("EventLoop::Timer", "[Online][EventLoop][Smoke]")
+	{
+		using FMockEventLoopTimerManager = TTimerManager<FEventLoopTimerManagerMockTraits>;
+		using FMockEventLoopTimerHandle = FMockEventLoopTimerManager::FTimerHandle;
+		FMockEventLoopTimerManager TimerManager;
+		const uint32 ManagerThreadId = 1;
+		FEventLoopTimerManagerMockTraits::ResetTestConditions();
+		FEventLoopTimerManagerMockTraits::FStorageTraits::CurrentThreadId = ManagerThreadId;
+
+		SECTION("Single init call")
+		{
+			TimerManager.Init();
+			CHECK(!FEventLoopTimerManagerMockTraits::FStorageTraits::bCheckNotInitializedTriggered);
+		}
+
+		SECTION("Double init call")
+		{
+			TimerManager.Init();
+			TimerManager.Init();
+			CHECK(FEventLoopTimerManagerMockTraits::FStorageTraits::bCheckNotInitializedTriggered);
+		}
+
+		SECTION("Call Tick before Init")
+		{
+			TimerManager.Tick(FTimespan());
+			CHECK(FEventLoopTimerManagerMockTraits::bCheckIsManagerThreadTriggered);
+		}
+
+		SECTION("Initialized")
+		{
+			uint32 TimerRunCount1 = 0;
+			auto TimerFunc1 = [&]() { ++TimerRunCount1; };
+
+			uint32 TimerRunCount2 = 0;
+			auto TimerFunc2 = [&]() { ++TimerRunCount2; };
+
+			TimerManager.Init();
+
+			SECTION("Call Tick after Init")
+			{
+				TimerManager.Tick(FTimespan());
+				CHECK(!FEventLoopTimerManagerMockTraits::bCheckIsManagerThreadTriggered);
+			}
+
+			SECTION("Timer paramater validation.")
+			{
+				SECTION("Timer with zero rate.")
+				{
+					FMockEventLoopTimerHandle Handle = TimerManager.SetTimer(MoveTemp(TimerFunc1), FTimespan::Zero());
+					CHECK(Handle.IsValid());
+				}
+
+				SECTION("Timer with positive rate.")
+				{
+					FMockEventLoopTimerHandle Handle = TimerManager.SetTimer(MoveTemp(TimerFunc1), FTimespan::FromSeconds(1));
+					CHECK(Handle.IsValid());
+				}
+
+				SECTION("Timer with negative rate.")
+				{
+					FMockEventLoopTimerHandle Handle = TimerManager.SetTimer(MoveTemp(TimerFunc1), FTimespan::FromSeconds(-1));
+					CHECK(!Handle.IsValid());
+				}
+
+				SECTION("Timer with zero first delay.")
+				{
+					FMockEventLoopTimerHandle Handle = TimerManager.SetTimer(MoveTemp(TimerFunc1), FTimespan::FromSeconds(1), true, FTimespan::Zero());
+					CHECK(Handle.IsValid());
+				}
+
+				SECTION("Timer with positive first delay.")
+				{
+					FMockEventLoopTimerHandle Handle = TimerManager.SetTimer(MoveTemp(TimerFunc1), FTimespan::FromSeconds(1), true, FTimespan::FromSeconds(1));
+					CHECK(Handle.IsValid());
+				}
+
+				SECTION("Timer with negative first delay.")
+				{
+					FMockEventLoopTimerHandle Handle = TimerManager.SetTimer(MoveTemp(TimerFunc1), FTimespan::FromSeconds(1), true, FTimespan::FromSeconds(-1));
+					CHECK(!Handle.IsValid());
+				}
+
+				SECTION("Timer with invalid first delay.")
+				{
+					FMockEventLoopTimerHandle Handle = TimerManager.SetTimer(MoveTemp(TimerFunc1), FTimespan::FromSeconds(1), false, FTimespan::FromSeconds(1));
+					CHECK(!Handle.IsValid());
+				}
+
+				SECTION("Timer with invalid callback.")
+				{
+					FMockEventLoopTimerHandle Handle = TimerManager.SetTimer(nullptr, FTimespan::FromSeconds(1));
+					CHECK(!Handle.IsValid());
+				}
+			}
+
+			SECTION("Non-repeating")
+			{
+				SECTION("Simple timer")
+				{
+					FMockEventLoopTimerHandle Handle = TimerManager.SetTimer(MoveTemp(TimerFunc1), FTimespan::FromSeconds(1));
+					REQUIRE(Handle.IsValid());
+
+					// Tick to add the timer.
+					TimerManager.Tick(FTimespan::Zero());
+					CHECK(TimerRunCount1 == 0);
+
+					// Tick just before expiration time.
+					TimerManager.Tick(FTimespan::FromMilliseconds(999));
+					CHECK(TimerRunCount1 == 0);
+
+					// Tick to make the timer run.
+					TimerManager.Tick(FTimespan::FromMilliseconds(1));
+					CHECK(TimerRunCount1 == 1);
+
+					// Make sure timer does not run again.
+					TimerRunCount1 = 0;
+					TimerManager.Tick(FTimespan::FromSeconds(10));
+					CHECK(TimerRunCount1 == 0);
+				}
+
+				SECTION("Multiple timers fired during tick.")
+				{
+					FMockEventLoopTimerHandle Handle1 = TimerManager.SetTimer(MoveTemp(TimerFunc1), FTimespan::FromSeconds(1));
+					REQUIRE(Handle1.IsValid());
+
+					FMockEventLoopTimerHandle Handle2 = TimerManager.SetTimer(MoveTemp(TimerFunc2), FTimespan::FromSeconds(2));
+					REQUIRE(Handle2.IsValid());
+
+					// Tick to add the timers.
+					TimerManager.Tick(FTimespan::Zero());
+					CHECK(TimerRunCount1 == 0);
+					CHECK(TimerRunCount2 == 0);
+
+					// Tick to fire the timers.
+					TimerManager.Tick(FTimespan::FromSeconds(5));
+					CHECK(TimerRunCount1 == 1);
+					CHECK(TimerRunCount2 == 1);
+				}
+
+				SECTION("Tiemr with zero rate called when scheduled.")
+				{
+					FMockEventLoopTimerHandle Handle = TimerManager.SetTimer(MoveTemp(TimerFunc1), FTimespan::Zero());
+					REQUIRE(Handle.IsValid());
+
+					// Tick to add and fire the timer.
+					TimerManager.Tick(FTimespan::Zero());
+					CHECK(TimerRunCount1 == 1);
+				}
+
+				SECTION("Timer removed before firing.")
+				{
+					FMockEventLoopTimerHandle Handle = TimerManager.SetTimer(MoveTemp(TimerFunc1), FTimespan::FromSeconds(1));
+					REQUIRE(Handle.IsValid());
+
+					// Tick to add the timer.
+					TimerManager.Tick(FTimespan::Zero());
+					CHECK(TimerRunCount1 == 0);
+
+					TimerManager.ClearTimer(Handle);
+
+					// Make sure timer does not run.
+					TimerManager.Tick(FTimespan::FromSeconds(10));
+					CHECK(TimerRunCount1 == 0);
+				}
+
+				SECTION("Timer not fired more than once.")
+				{
+					FMockEventLoopTimerHandle Handle = TimerManager.SetTimer(MoveTemp(TimerFunc1), FTimespan::FromSeconds(1));
+					REQUIRE(Handle.IsValid());
+
+					// Tick to add the timer.
+					TimerManager.Tick(FTimespan::Zero());
+
+					// Check that timer ran.
+					TimerManager.Tick(FTimespan::FromSeconds(10));
+					CHECK(TimerRunCount1 == 1);
+				}
+
+				SECTION("Timer pending reschedule is not triggered.")
+				{
+					FMockEventLoopTimerHandle Handle = TimerManager.SetTimer(MoveTemp(TimerFunc1), FTimespan::FromSeconds(1));
+					REQUIRE(Handle.IsValid());
+
+					// Tick to add the timer.
+					TimerManager.Tick(FTimespan::Zero());
+
+					// Check that no timer is pending reschedule.
+					CHECK(!TimerManager.HasPendingRepeatTimer());
+
+					// Check that timer ran.
+					TimerManager.Tick(FTimespan::FromSeconds(10));
+					CHECK(TimerRunCount1 == 1);
+
+					// Check that no timer is pending reschedule.
+					CHECK(!TimerManager.HasPendingRepeatTimer());
+				}
+
+				SECTION("Timer storage")
+				{
+					FMockEventLoopTimerHandle Handle = TimerManager.SetTimer(MoveTemp(TimerFunc1), FTimespan::FromSeconds(1));
+					REQUIRE(Handle.IsValid());
+
+					// Tick to add the timer.
+					TimerManager.Tick(FTimespan::Zero());
+					CHECK(TimerRunCount1 == 0);
+					CHECK(TimerManager.GetNumTimers() == 1);
+
+					// Tick to fire the timer.
+					TimerRunCount1 = 0;
+					TimerManager.Tick(FTimespan::FromSeconds(10));
+					CHECK(TimerRunCount1 == 1);
+					CHECK(TimerManager.GetNumTimers() == 0);
+				}
+			}
+
+			SECTION("Repeating")
+			{
+				SECTION("Timer with zero rate does not block tick")
+				{
+					FMockEventLoopTimerHandle Handle1 = TimerManager.SetTimer(MoveTemp(TimerFunc1), FTimespan::FromSeconds(1));
+					REQUIRE(Handle1.IsValid());
+
+					// Tick to add first timer.
+					TimerManager.Tick(FTimespan::Zero());
+					CHECK(TimerRunCount1 == 0);
+
+					FMockEventLoopTimerHandle Handle2 = TimerManager.SetTimer(MoveTemp(TimerFunc2), FTimespan::Zero());
+					REQUIRE(Handle2.IsValid());
+
+					// Check that both timers ran.
+					TimerManager.Tick(FTimespan::FromSeconds(10));
+					CHECK(TimerRunCount1 == 1);
+					CHECK(TimerRunCount2 == 1);
+				}
+
+				SECTION("Timer with zero rate called each tick")
+				{
+					FMockEventLoopTimerHandle Handle = TimerManager.SetTimer(MoveTemp(TimerFunc1), FTimespan::Zero(), true);
+					REQUIRE(Handle.IsValid());
+
+					TimerManager.Tick(FTimespan::Zero());
+					CHECK(TimerRunCount1 == 1);
+
+					TimerRunCount1 = 0;
+					TimerManager.Tick(FTimespan::Zero());
+					CHECK(TimerRunCount1 == 1);
+
+					TimerManager.ClearTimer(Handle);
+
+					TimerRunCount1 = 0;
+					TimerManager.Tick(FTimespan::Zero());
+					CHECK(TimerRunCount1 == 0);
+				}
+
+				SECTION("Timer with positive rate called when elapsed")
+				{
+					FMockEventLoopTimerHandle Handle = TimerManager.SetTimer(MoveTemp(TimerFunc1), FTimespan::FromSeconds(1), true);
+					REQUIRE(Handle.IsValid());
+
+					// Tick to add the timer.
+					TimerManager.Tick(FTimespan::Zero());
+					CHECK(TimerRunCount1 == 0);
+
+					// Tick not enough time to fire.
+					TimerManager.Tick(FTimespan::FromMilliseconds(500));
+					CHECK(TimerRunCount1 == 0);
+
+					// Tick enough time to fire.
+					TimerManager.Tick(FTimespan::FromMilliseconds(500));
+					CHECK(TimerRunCount1 == 1);
+					TimerRunCount1 = 0;
+
+					// Tick to reschedule timer.
+					// This looks funky, the intent is that a timer will never fire early.
+					// The way to prevent the extra call to tick to reschedule while providing that
+					// guarantee is to get the current time again when rescheduling which this
+					// implementation wants to avoid.
+					TimerManager.Tick(FTimespan::FromMilliseconds(500));
+					CHECK(TimerRunCount1 == 0);
+
+					// Tick not enough time to fire.
+					TimerManager.Tick(FTimespan::FromMilliseconds(500));
+					CHECK(TimerRunCount1 == 0);
+
+					// Tick enough time to fire.
+					TimerManager.Tick(FTimespan::FromMilliseconds(500));
+					CHECK(TimerRunCount1 == 1);
+
+					// Timer not fired after cleared.
+					TimerManager.ClearTimer(Handle);
+					TimerRunCount1 = 0;
+					TimerManager.Tick(FTimespan::FromSeconds(10));
+					CHECK(TimerRunCount1 == 0);
+				}
+
+				SECTION("Timer with zero first delay fires as expected.")
+				{
+					FMockEventLoopTimerHandle Handle = TimerManager.SetTimer(MoveTemp(TimerFunc1), FTimespan::FromSeconds(1), true, FTimespan::Zero());
+					REQUIRE(Handle.IsValid());
+
+					// Tick to add and fire the timer.
+					TimerManager.Tick(FTimespan::Zero());
+					CHECK(TimerRunCount1 == 1);
+
+					// Tick to reschedule the timer.
+					TimerRunCount1 = 0;
+					TimerManager.Tick(FTimespan::Zero());
+					CHECK(TimerRunCount1 == 0);
+
+					// Tick to see timer does not fire again at zero interval.
+					TimerManager.Tick(FTimespan::Zero());
+					CHECK(TimerRunCount1 == 0);
+
+					// Tick to fire the timer.
+					TimerManager.Tick(FTimespan::FromSeconds(1));
+					CHECK(TimerRunCount1 == 1);
+				}
+
+				SECTION("Timer with positive first delay fires as expected.")
+				{
+					FMockEventLoopTimerHandle Handle = TimerManager.SetTimer(MoveTemp(TimerFunc1), FTimespan::FromSeconds(1), true, FTimespan::FromSeconds(5));
+					REQUIRE(Handle.IsValid());
+
+					// Tick to add the timer.
+					TimerManager.Tick(FTimespan::Zero());
+					CHECK(TimerRunCount1 == 0);
+
+					// Tick to see timer not fire once rate has elapsed.
+					TimerManager.Tick(FTimespan::FromSeconds(1));
+					CHECK(TimerRunCount1 == 0);
+
+					// Tick near remaining time to see timer has still not fired.
+					TimerManager.Tick(FTimespan::FromMilliseconds(3999));
+					CHECK(TimerRunCount1 == 0);
+
+					// Tick remaining time to see timer has fired.
+					TimerManager.Tick(FTimespan::FromMilliseconds(1));
+					CHECK(TimerRunCount1 == 1);
+				}
+
+				SECTION("Timer pending reschedule is triggered.")
+				{
+					FMockEventLoopTimerHandle Handle = TimerManager.SetTimer(MoveTemp(TimerFunc1), FTimespan::FromSeconds(1), true);
+					REQUIRE(Handle.IsValid());
+
+					// Tick to add the timer.
+					TimerManager.Tick(FTimespan::Zero());
+
+					// Check that no timer is pending reschedule.
+					CHECK(!TimerManager.HasPendingRepeatTimer());
+
+					// Check that timer ran.
+					TimerManager.Tick(FTimespan::FromSeconds(10));
+					CHECK(TimerRunCount1 == 1);
+
+					// Check that timer is pending reschedule.
+					CHECK(TimerManager.HasPendingRepeatTimer());
+				}
+
+				SECTION("Timer storage")
+				{
+					FMockEventLoopTimerHandle Handle = TimerManager.SetTimer(MoveTemp(TimerFunc1), FTimespan::FromSeconds(1), true);
+					REQUIRE(Handle.IsValid());
+
+					// Tick to add the timer.
+					TimerRunCount1 = 0;
+					TimerManager.Tick(FTimespan::Zero());
+					CHECK(TimerRunCount1 == 0);
+					CHECK(TimerManager.GetNumTimers() == 1);
+
+					// Tick to fire the timer.
+					TimerRunCount1 = 0;
+					TimerManager.Tick(FTimespan::FromSeconds(2));
+					CHECK(TimerRunCount1 == 1);
+					CHECK(TimerManager.GetNumTimers() == 1);
+
+					SECTION("Remove before reschedule")
+					{
+						// Tick to clear the timer.
+						TimerRunCount1 = 0;
+						TimerManager.ClearTimer(Handle);
+						TimerManager.Tick(FTimespan::FromSeconds(2));
+						CHECK(TimerRunCount1 == 0);
+						CHECK(TimerManager.GetNumTimers() == 0);
+					}
+
+					SECTION("Remove after reschedule")
+					{
+						// Tick to reschedule the timer.
+						TimerRunCount1 = 0;
+						TimerManager.Tick(FTimespan::Zero());
+						CHECK(TimerRunCount1 == 0);
+						CHECK(TimerManager.GetNumTimers() == 1);
+
+						// Tick to clear the timer.
+						TimerRunCount1 = 0;
+						TimerManager.ClearTimer(Handle);
+						TimerManager.Tick(FTimespan::FromSeconds(2));
+						CHECK(TimerRunCount1 == 0);
+						CHECK(TimerManager.GetNumTimers() == 0);
+					}
+				}
+			}
+
+			SECTION("Primed timer removed before it is processed during tick.")
+			{
+				FMockEventLoopTimerHandle Handle1 = TimerManager.SetTimer(MoveTemp(TimerFunc1), FTimespan::FromSeconds(2));
+				REQUIRE(Handle1.IsValid());
+
+				bool Timer1Removed = false;
+				auto Timer1RemovedCallback = [&]()
+				{
+					Timer1Removed = true;
+				};
+
+				auto Timer2Callback = [&]()
+				{
+					TimerManager.ClearTimer(Handle1, MoveTemp(Timer1RemovedCallback));
+				};
+
+				FMockEventLoopTimerHandle Handle2 = TimerManager.SetTimer(MoveTemp(Timer2Callback), FTimespan::FromSeconds(1));
+				REQUIRE(Handle1.IsValid());
+
+				// Tick to add the timers.
+				TimerManager.Tick(FTimespan::Zero());
+				CHECK(TimerRunCount1 == 0);
+
+				// Tick so that both timers are expired.
+				// See that timer 1 did not fire, but will not be removed until the following tick.
+				TimerManager.Tick(FTimespan::FromSeconds(2));
+				CHECK(TimerRunCount1 == 0);
+				CHECK(Timer1Removed == false);
+
+				// Tick again to see that timer has been removed.
+				TimerManager.Tick(FTimespan::Zero());
+				CHECK(Timer1Removed);
+			}
+
+			SECTION("Timer with zero rate added during tick is not fired until the next tick.")
+			{
+				FMockEventLoopTimerHandle Handle1;
+				FMockEventLoopTimerHandle Handle2;
+
+				auto Timer1Callback = [&]()
+				{
+					Handle2 = TimerManager.SetTimer(MoveTemp(TimerFunc2), FTimespan::Zero());
+					REQUIRE(Handle2.IsValid());
+				};
+
+				FString value = Handle1.ToString();
+
+				Handle1 = TimerManager.SetTimer(MoveTemp(Timer1Callback), FTimespan::Zero());
+				REQUIRE(Handle1.IsValid());
+
+				value = Handle1.ToString();
+
+				// Tick to add and fire the timers.
+				TimerManager.Tick(FTimespan::Zero());
+
+				// Timer 2 has been set, but has not fired.
+				CHECK(Handle2.IsValid());
+				CHECK(TimerRunCount2 == 0);
+
+				// Tick to fire timer 2.
+				TimerManager.Tick(FTimespan::Zero());
+				CHECK(TimerRunCount2 == 1);
+			}
+
+			SECTION("GetNextTimeout.")
+			{
+				FTimespan NextTimeout = FTimespan::Zero();
+
+				FMockEventLoopTimerHandle Handle1 = TimerManager.SetTimer(MoveTemp(TimerFunc1), FTimespan::FromSeconds(1));
+				REQUIRE(Handle1.IsValid());
+
+				FMockEventLoopTimerHandle Handle2 = TimerManager.SetTimer(MoveTemp(TimerFunc2), FTimespan::FromSeconds(2));
+				REQUIRE(Handle2.IsValid());
+
+				// No next timeout set yet.
+				CHECK(!TimerManager.GetNextTimeout(NextTimeout));
+
+				// Tick to add the timers.
+				TimerManager.Tick(FTimespan::Zero());
+				CHECK(TimerRunCount1 == 0);
+				CHECK(TimerRunCount2 == 0);
+
+				CHECK(TimerManager.GetNextTimeout(NextTimeout));
+				CHECK(NextTimeout == FTimespan::FromSeconds(1));
+
+				TimerManager.Tick(FTimespan::FromMilliseconds(500));
+				CHECK(TimerRunCount1 == 0);
+				CHECK(TimerRunCount2 == 0);
+
+				CHECK(TimerManager.GetNextTimeout(NextTimeout));
+				CHECK(NextTimeout == FTimespan::FromMilliseconds(500));
+
+				TimerManager.Tick(NextTimeout);
+				CHECK(TimerRunCount1 == 1);
+				CHECK(TimerRunCount2 == 0);
+
+				CHECK(TimerManager.GetNextTimeout(NextTimeout));
+				CHECK(NextTimeout == FTimespan::FromSeconds(1));
+
+				TimerManager.Tick(NextTimeout);
+				CHECK(TimerRunCount2 == 1);
+
+				// There is no longer any valid timer.
+				CHECK(!TimerManager.GetNextTimeout(NextTimeout));
+			}
+
+			SECTION("Handle string conversion.")
+			{
+				FMockEventLoopTimerHandle Handle1 = TimerManager.SetTimer(MoveTemp(TimerFunc1), FTimespan::Zero());
+				REQUIRE(Handle1.IsValid());
+
+				// Test that ToString resolves name from traits.
+				FString StringValue = Handle1.ToString();
+				CHECK(StringValue.Contains(FEventLoopTimerManagerMockTraits::FStorageTraits::FExternalHandle::GetTypeName()));
+			}
+
+			SECTION("Timer handle is cleared when timer is removed.")
+			{
+				FMockEventLoopTimerHandle Handle1 = TimerManager.SetTimer(MoveTemp(TimerFunc1), FTimespan::Zero());
+				REQUIRE(Handle1.IsValid());
+
+				TimerManager.ClearTimer(Handle1);
+				CHECK(!Handle1.IsValid());
+			}
+		}
+	}
+}
\ No newline at end of file
diff --git a/Engine/Source/Runtime/Online/EventLoopTests/Tests/TestGroupEvents.cpp b/Engine/Source/Runtime/Online/EventLoopTests/Tests/TestGroupEvents.cpp
new file mode 100644
--- /dev/null
+++ b/Engine/Source/Runtime/Online/EventLoopTests/Tests/TestGroupEvents.cpp
@@ -0,0 +1,9 @@
+// Copyright Epic Games, Inc. All Rights Reserved.
+#include "TestCommon/Initialization.h"
+
+#include <catch2/catch_test_macros.hpp>
+
+GROUP_AFTER_GLOBAL(Catch::DefaultGroup)
+{
+	CleanupLocalization();
+}
diff --git a/Engine/Source/Runtime/Online/Experimental/EventLoop/EventLoop.Build.cs b/Engine/Source/Runtime/Online/Experimental/EventLoop/EventLoop.Build.cs
deleted file mode 100644
diff --git a/Engine/Source/Runtime/Online/Experimental/EventLoop/Private/EventLoop/BSDSocket/BSDSocketTypesPrivate.h b/Engine/Source/Runtime/Online/Experimental/EventLoop/Private/EventLoop/BSDSocket/BSDSocketTypesPrivate.h
deleted file mode 100644
diff --git a/Engine/Source/Runtime/Online/Experimental/EventLoop/Private/EventLoop/BSDSocket/EventLoopIOManagerBSDSocketSelect.cpp b/Engine/Source/Runtime/Online/Experimental/EventLoop/Private/EventLoop/BSDSocket/EventLoopIOManagerBSDSocketSelect.cpp
deleted file mode 100644
diff --git a/Engine/Source/Runtime/Online/Experimental/EventLoop/Private/EventLoop/BSDSocket/EventLoopNotifierEventFd.cpp b/Engine/Source/Runtime/Online/Experimental/EventLoop/Private/EventLoop/BSDSocket/EventLoopNotifierEventFd.cpp
deleted file mode 100644
diff --git a/Engine/Source/Runtime/Online/Experimental/EventLoop/Private/EventLoop/BSDSocket/EventLoopNotifierSocket.cpp b/Engine/Source/Runtime/Online/Experimental/EventLoop/Private/EventLoop/BSDSocket/EventLoopNotifierSocket.cpp
deleted file mode 100644
diff --git a/Engine/Source/Runtime/Online/Experimental/EventLoop/Private/EventLoop/EventLoopLog.cpp b/Engine/Source/Runtime/Online/Experimental/EventLoop/Private/EventLoop/EventLoopLog.cpp
deleted file mode 100644
diff --git a/Engine/Source/Runtime/Online/Experimental/EventLoop/Private/EventLoop/EventLoopLog.h b/Engine/Source/Runtime/Online/Experimental/EventLoop/Private/EventLoop/EventLoopLog.h
deleted file mode 100644
diff --git a/Engine/Source/Runtime/Online/Experimental/EventLoop/Private/EventLoop/EventLoopModule.cpp b/Engine/Source/Runtime/Online/Experimental/EventLoop/Private/EventLoop/EventLoopModule.cpp
deleted file mode 100644
diff --git a/Engine/Source/Runtime/Online/Experimental/EventLoop/Public/EventLoop/BSDSocket/BSDSocketTypes.h b/Engine/Source/Runtime/Online/Experimental/EventLoop/Public/EventLoop/BSDSocket/BSDSocketTypes.h
deleted file mode 100644
diff --git a/Engine/Source/Runtime/Online/Experimental/EventLoop/Public/EventLoop/BSDSocket/EventLoopIOAccessBSDSocket.h b/Engine/Source/Runtime/Online/Experimental/EventLoop/Public/EventLoop/BSDSocket/EventLoopIOAccessBSDSocket.h
deleted file mode 100644
diff --git a/Engine/Source/Runtime/Online/Experimental/EventLoop/Public/EventLoop/BSDSocket/EventLoopIOManagerBSDSocket.h b/Engine/Source/Runtime/Online/Experimental/EventLoop/Public/EventLoop/BSDSocket/EventLoopIOManagerBSDSocket.h
deleted file mode 100644
diff --git a/Engine/Source/Runtime/Online/Experimental/EventLoop/Public/EventLoop/BSDSocket/EventLoopIOManagerBSDSocketSelect.h b/Engine/Source/Runtime/Online/Experimental/EventLoop/Public/EventLoop/BSDSocket/EventLoopIOManagerBSDSocketSelect.h
deleted file mode 100644
diff --git a/Engine/Source/Runtime/Online/Experimental/EventLoop/Public/EventLoop/BSDSocket/EventLoopNotifier.h b/Engine/Source/Runtime/Online/Experimental/EventLoop/Public/EventLoop/BSDSocket/EventLoopNotifier.h
deleted file mode 100644
diff --git a/Engine/Source/Runtime/Online/Experimental/EventLoop/Public/EventLoop/BSDSocket/EventLoopNotifierEventFd.h b/Engine/Source/Runtime/Online/Experimental/EventLoop/Public/EventLoop/BSDSocket/EventLoopNotifierEventFd.h
deleted file mode 100644
diff --git a/Engine/Source/Runtime/Online/Experimental/EventLoop/Public/EventLoop/BSDSocket/EventLoopNotifierSocket.h b/Engine/Source/Runtime/Online/Experimental/EventLoop/Public/EventLoop/BSDSocket/EventLoopNotifierSocket.h
deleted file mode 100644
diff --git a/Engine/Source/Runtime/Online/Experimental/EventLoop/Public/EventLoop/EventLoop.h b/Engine/Source/Runtime/Online/Experimental/EventLoop/Public/EventLoop/EventLoop.h
deleted file mode 100644
diff --git a/Engine/Source/Runtime/Online/Experimental/EventLoop/Public/EventLoop/EventLoopHandle.h b/Engine/Source/Runtime/Online/Experimental/EventLoop/Public/EventLoop/EventLoopHandle.h
deleted file mode 100644
diff --git a/Engine/Source/Runtime/Online/Experimental/EventLoop/Public/EventLoop/EventLoopIOManagerNull.h b/Engine/Source/Runtime/Online/Experimental/EventLoop/Public/EventLoop/EventLoopIOManagerNull.h
deleted file mode 100644
diff --git a/Engine/Source/Runtime/Online/Experimental/EventLoop/Public/EventLoop/EventLoopManagedStorage.h b/Engine/Source/Runtime/Online/Experimental/EventLoop/Public/EventLoop/EventLoopManagedStorage.h
deleted file mode 100644
diff --git a/Engine/Source/Runtime/Online/Experimental/EventLoop/Public/EventLoop/EventLoopModule.h b/Engine/Source/Runtime/Online/Experimental/EventLoop/Public/EventLoop/EventLoopModule.h
deleted file mode 100644
diff --git a/Engine/Source/Runtime/Online/Experimental/EventLoop/Public/EventLoop/EventLoopTimer.h b/Engine/Source/Runtime/Online/Experimental/EventLoop/Public/EventLoop/EventLoopTimer.h
deleted file mode 100644
diff --git a/Engine/Source/Runtime/Online/Experimental/EventLoop/Public/EventLoop/IEventLoop.h b/Engine/Source/Runtime/Online/Experimental/EventLoop/Public/EventLoop/IEventLoop.h
deleted file mode 100644
diff --git a/Engine/Source/Runtime/Online/Experimental/EventLoop/Public/EventLoop/IEventLoopIOManager.h b/Engine/Source/Runtime/Online/Experimental/EventLoop/Public/EventLoop/IEventLoopIOManager.h
deleted file mode 100644
diff --git a/Engine/Source/Runtime/Online/Experimental/EventLoopTests/Tests/EventLoop/BSDSocketTypesPrivate.h b/Engine/Source/Runtime/Online/Experimental/EventLoopTests/Tests/EventLoop/BSDSocketTypesPrivate.h
deleted file mode 100644
diff --git a/Engine/Source/Runtime/Online/Experimental/EventLoopTests/Tests/EventLoop/EventLoopHandleTests.cpp b/Engine/Source/Runtime/Online/Experimental/EventLoopTests/Tests/EventLoop/EventLoopHandleTests.cpp
deleted file mode 100644
diff --git a/Engine/Source/Runtime/Online/Experimental/EventLoopTests/Tests/EventLoop/EventLoopIOManagerBSDSocketSelectTests.cpp b/Engine/Source/Runtime/Online/Experimental/EventLoopTests/Tests/EventLoop/EventLoopIOManagerBSDSocketSelectTests.cpp
deleted file mode 100644
diff --git a/Engine/Source/Runtime/Online/Experimental/EventLoopTests/Tests/EventLoop/EventLoopLLTModule.cpp b/Engine/Source/Runtime/Online/Experimental/EventLoopTests/Tests/EventLoop/EventLoopLLTModule.cpp
deleted file mode 100644
diff --git a/Engine/Source/Runtime/Online/Experimental/EventLoopTests/Tests/EventLoop/EventLoopManagedStorageTests.cpp b/Engine/Source/Runtime/Online/Experimental/EventLoopTests/Tests/EventLoop/EventLoopManagedStorageTests.cpp
deleted file mode 100644
diff --git a/Engine/Source/Runtime/Online/Experimental/EventLoopTests/Tests/EventLoop/EventLoopTests.cpp b/Engine/Source/Runtime/Online/Experimental/EventLoopTests/Tests/EventLoop/EventLoopTests.cpp
deleted file mode 100644
diff --git a/Engine/Source/Runtime/Online/Experimental/EventLoopTests/Tests/EventLoop/EventLoopTimerTests.cpp b/Engine/Source/Runtime/Online/Experimental/EventLoopTests/Tests/EventLoop/EventLoopTimerTests.cpp
deleted file mode 100644
diff --git a/Engine/Source/Runtime/Online/Experimental/EventLoopTests/Tests/TestGroupEvents.cpp b/Engine/Source/Runtime/Online/Experimental/EventLoopTests/Tests/TestGroupEvents.cpp
deleted file mode 100644
