BPS: Fix for deadlock on cancelation during certain write scenarios.

#rnx
#Synced-CL 0

diff --git a/Engine/Source/Runtime/Online/BuildPatchServices/Private/BuildPatchFileConstructor.cpp b/Engine/Source/Runtime/Online/BuildPatchServices/Private/BuildPatchFileConstructor.cpp
--- a/Engine/Source/Runtime/Online/BuildPatchServices/Private/BuildPatchFileConstructor.cpp
+++ b/Engine/Source/Runtime/Online/BuildPatchServices/Private/BuildPatchFileConstructor.cpp
@@ -2373,6 +2373,9 @@ IConstructorChunkSource::FRequestProcessFn FBuildPatchFileConstructor::CreateWri
 	{
 		if (bIsAbort)
 		{
+			// Report the batch finished otherwise we hard lock. This happens if we abort after the write thread grabs the jobs
+			// but before we pass the abort check in the generic thread dispatch.
+			Batch->bIsFinished.store(true, std::memory_order_release);
 			return;
 		}
 
