Add CSV profiler timing regions for http requests (enabled via the cvar http.csvRegions). Disabled by default
#Synced-CL 48595085

diff --git a/Engine/Source/Runtime/Online/HTTP/Private/GenericPlatform/HttpRequestCommon.cpp b/Engine/Source/Runtime/Online/HTTP/Private/GenericPlatform/HttpRequestCommon.cpp
--- a/Engine/Source/Runtime/Online/HTTP/Private/GenericPlatform/HttpRequestCommon.cpp
+++ b/Engine/Source/Runtime/Online/HTTP/Private/GenericPlatform/HttpRequestCommon.cpp
@@ -9,8 +9,11 @@
 #include "HttpManager.h"
 #include "Logging/StructuredLog.h"
 #include "Misc/CommandLine.h"
+#include "ProfilingDebugging/CsvProfiler.h"
 #include "Stats/Stats.h"
 
+CSV_DEFINE_CATEGORY(HttpRequest, true);
+
 TAutoConsoleVariable<bool> CVarHttpLogJsonContentOnly(
 	TEXT("http.LogJsonContentOnly"),
 	true,
@@ -18,6 +21,44 @@ TAutoConsoleVariable<bool> CVarHttpLogJsonContentOnly(
 	ECVF_SaveForNextBoot
 );
 
+TAutoConsoleVariable<bool> CVarHttpCsvRegions(
+	TEXT("http.csvRegions"),
+	false,
+	TEXT("Add CSV regions for requests"),
+	ECVF_Default
+);
+
+static void CsvRegion(const FHttpRequestCommon* Request, bool bBegin)
+{
+#if CSV_PROFILER_STATS && !UE_BUILD_SHIPPING
+	if (CVarHttpCsvRegions.GetValueOnAnyThread(true) && FCsvProfiler::Get()->IsCategoryEnabled(CSV_CATEGORY_INDEX(HttpRequest)))
+	{
+		// Trim the https:// and anything passed via URL params
+		FString Name=Request->GetURL();
+		if (Name.StartsWith(TEXT("https://")))
+		{
+			Name.RightChopInline(8);
+		}
+		int32 Index;
+		if (Name.FindChar(TCHAR('?'), Index))
+		{
+			Name.MidInline(0,Index);
+		}
+		Name = FString::Printf(TEXT("%s (%s)"), *Name, *Request->GetVerb());
+		if (bBegin)
+		{
+			CSV_REGION_BEGIN(HttpRequest, *Name);
+		}
+		else
+		{
+			CSV_REGION_END(HttpRequest, *Name);
+		}		
+	}
+#endif
+}
+
+
+
 FHttpRequestCommon::FHttpRequestCommon()
 	: RequestStartTimeAbsoluteSeconds(FPlatformTime::Seconds())
 	, ActivityTimeoutAt(0.0)
@@ -280,6 +321,8 @@ void FHttpRequestCommon::SetOption(const FName Option, const FString& OptionValu
 
 void FHttpRequestCommon::HandleRequestSucceed()
 {
+	CsvRegion(this, false);
+
 	SetStatus(EHttpRequestStatus::Succeeded);
 
 	LogResponse(ResponseCommon);
@@ -289,6 +332,7 @@ void FHttpRequestCommon::HandleRequestSucceed()
 
 void FHttpRequestCommon::HandleRequestFailed()
 {
+	CsvRegion(this, false);
 	if (FailureReason == EHttpFailureReason::None) // Failure reason was not set by platform, will set it here
 	{
 		if (bCanceled)
@@ -664,6 +708,7 @@ float FHttpRequestCommon::GetElapsedTime() const
 
 bool FHttpRequestCommon::StartThreadedRequest()
 {
+	CsvRegion(this, true);
 	LogRequest();
 	return true;
 }
