BPS: Fix for install source chunk crash during completed file chunk harvesting.

#rnx
#Synced-CL 48917355

diff --git a/Engine/Source/Runtime/Online/BuildPatchServices/Private/BuildPatchFileConstructor.cpp b/Engine/Source/Runtime/Online/BuildPatchServices/Private/BuildPatchFileConstructor.cpp
--- a/Engine/Source/Runtime/Online/BuildPatchServices/Private/BuildPatchFileConstructor.cpp
+++ b/Engine/Source/Runtime/Online/BuildPatchServices/Private/BuildPatchFileConstructor.cpp
@@ -1753,6 +1753,31 @@ void FBuildPatchFileConstructor::GenericThreadFn(int32 ThreadIndex)
 	ThreadCompleteEvents[ThreadIndex]->Trigger();
 }
 
+struct FChunkHarvestContext
+{
+	// This gets unlocked when all harvests are done to signal the main thread
+	UE::FMutex HarvestsCompleteLock;
+
+	std::atomic_int32_t PendingHarvestRequests = 0;
+};
+
+// Called from basically any thread.
+void FBuildPatchFileConstructor::ChunkHarvestCompletedFn(const FGuid& Guid, bool bAborted, bool bFailedToRead, void* UserPtr)
+{
+	if (bFailedToRead)
+	{
+		// We tell the main thread this failed by setting the location since that's thread
+		// safe.
+		SetChunkLocation(Guid, EConstructorChunkLocation::Cloud);
+	}
+
+	FChunkHarvestContext* HarvestContext = (FChunkHarvestContext*)UserPtr;
+	if (HarvestContext->PendingHarvestRequests.fetch_add(-1, std::memory_order_acq_rel) == 1)
+	{
+		HarvestContext->HarvestsCompleteLock.Unlock();
+	}
+}
+
 bool FBuildPatchFileConstructor::HarvestChunksForCompletedFile(const FString& CompletedFullPathFileName)
 {
 	UE_LOG(LogBuildPatchServices, Verbose, TEXT("Harvesting source: %s"), *CompletedFullPathFileName);
@@ -1801,6 +1826,7 @@ bool FBuildPatchFileConstructor::HarvestChunksForCompletedFile(const FString& Co
 		}
 	}
 
+	UE_LOG(LogBuildPatchServices, Verbose, TEXT("    Chunks Needed: %d"), ChunksFromFileWeNeed.Num());
 	if (ChunksFromFileWeNeed.Num() == 0)
 	{
 		return true;
@@ -1836,8 +1862,10 @@ bool FBuildPatchFileConstructor::HarvestChunksForCompletedFile(const FString& Co
 			BatchSize += ChunksFromFileWeNeed[ChunkEndIndex].ChunkSize;
 		}
 
-		PendingHarvestRequests.store(ChunkEndIndex - ChunkIndex, std::memory_order_release);
-
+		FChunkHarvestContext ChunkHarvestContext;
+		ChunkHarvestContext.HarvestsCompleteLock.Lock();
+		ChunkHarvestContext.PendingHarvestRequests.store(ChunkEndIndex - ChunkIndex, std::memory_order_release);
+		UE_LOG(LogBuildPatchServices, Verbose, TEXT("    Harvesting %d chunks"), ChunkEndIndex - ChunkIndex);
 		for (int32 DispatchIndex = ChunkIndex; DispatchIndex < ChunkEndIndex; DispatchIndex++)
 		{
 			FNeededChunk& Chunk = ChunksFromFileWeNeed[DispatchIndex];
@@ -1849,11 +1877,12 @@ bool FBuildPatchFileConstructor::HarvestChunksForCompletedFile(const FString& Co
 			{
 				// Call the completion function so we decrement the request count,
 				// but this is a consistency failure so we can't use the results.
-				ChunkHarvestCompletedFn(Chunk.Id, false, false, 0);
+				ChunkHarvestCompletedFn(Chunk.Id, false, false, &ChunkHarvestContext);
 				bHarvestSuccess = false;
 			}
 
-			IConstructorChunkSource::FRequestProcessFn HarvestFn = InstallSource->CreateRequest(Chunk.Id, Destination, 0, IConstructorChunkSource::FChunkRequestCompleteDelegate::CreateRaw(this, &FBuildPatchFileConstructor::ChunkHarvestCompletedFn));
+			IConstructorChunkSource::FRequestProcessFn HarvestFn = InstallSource->CreateRequest(Chunk.Id, Destination, &ChunkHarvestContext, 
+				IConstructorChunkSource::FChunkRequestCompleteDelegate::CreateRaw(this, &FBuildPatchFileConstructor::ChunkHarvestCompletedFn));
 
 			HarvestFn(false);
 		}
@@ -1861,7 +1890,9 @@ bool FBuildPatchFileConstructor::HarvestChunksForCompletedFile(const FString& Co
 		// The read is synchronous but the verification is not, so we still need to do the wait.
 		{
 			TRACE_CPUPROFILER_EVENT_SCOPE(HarvestWait);
-			WakeUpDispatchThreadEvent->Wait();
+			UE_LOG(LogBuildPatchServices, Verbose, TEXT("    Waiting..."));
+			ChunkHarvestContext.HarvestsCompleteLock.Lock();
+			UE_LOG(LogBuildPatchServices, Verbose, TEXT("    Done."));
 		}
 
 		if (bHarvestSuccess)
@@ -2255,22 +2286,6 @@ struct FBatchState
 	}
 };
 
-// Called from basically any thread.
-void FBuildPatchFileConstructor::ChunkHarvestCompletedFn(const FGuid& Guid, bool bAborted, bool bFailedToRead, void* UserPtr)
-{
-	if (bFailedToRead)
-	{
-		// We tell the main thread this failed by setting the location since that's thread
-		// safe.
-		SetChunkLocation(Guid, EConstructorChunkLocation::Cloud);
-	}
-
-	if (PendingHarvestRequests.fetch_add(-1, std::memory_order_relaxed) == 1)
-	{
-		WakeUpDispatchThreadEvent->Trigger();
-	}
-}
-
 // Called from basically any thread.
 void FBuildPatchFileConstructor::RequestCompletedFn(const FGuid& Guid, bool bAborted, bool bFailedToRead, void* UserPtr)
 {
diff --git a/Engine/Source/Runtime/Online/BuildPatchServices/Private/BuildPatchFileConstructor.h b/Engine/Source/Runtime/Online/BuildPatchServices/Private/BuildPatchFileConstructor.h
--- a/Engine/Source/Runtime/Online/BuildPatchServices/Private/BuildPatchFileConstructor.h
+++ b/Engine/Source/Runtime/Online/BuildPatchServices/Private/BuildPatchFileConstructor.h
@@ -324,7 +324,6 @@ namespace BuildPatchServices
 		void CompleteReadBatch(const FFileManifest& FileManifest, FBatchState& Buffer);
 		void RequestCompletedFn(const FGuid& Guid, bool bAborted, bool bFailedToRead, void* UserPtr);
 
-		std::atomic_int32_t PendingHarvestRequests = 0;
 		void ChunkHarvestCompletedFn(const FGuid& Guid, bool bAborted, bool bFailedToRead, void* UserPtr);
 
 		void StartFile(FFileConstructionState& CurrentFile, const FResumeData& ResumeData);
