Fix BuildPatchInstaller: preventing empty folders deletion in installDir when retrying installation

#tests Tested locally and on luna
#rnx

diff --git a/Engine/Source/Runtime/Online/BuildPatchServices/Private/BuildPatchInstaller.cpp b/Engine/Source/Runtime/Online/BuildPatchServices/Private/BuildPatchInstaller.cpp
--- a/Engine/Source/Runtime/Online/BuildPatchServices/Private/BuildPatchInstaller.cpp
+++ b/Engine/Source/Runtime/Online/BuildPatchServices/Private/BuildPatchInstaller.cpp
@@ -880,6 +880,11 @@ namespace BuildPatchServices
 					CorruptFiles.Empty();
 					bProcessSuccess = bInstallSuccess && RunVerification(CorruptFiles);
 
+					// Set if we can retry
+					--InstallRetries;
+					bCanRetry = InstallRetries > 0 && !InstallerError->IsCancelled() && InstallerError->CanRetry();
+					const bool bWillRetry = !bProcessSuccess && bCanRetry;
+
 					// Clean staging if INSTALL success, we still do cleanup if we failed at the verify stage.
 					if (bInstallSuccess && 
 						!Configuration.bInstallToMemory)
@@ -894,17 +899,15 @@ namespace BuildPatchServices
 						{
 							UE_LOG(LogBuildPatchServices, Log, TEXT("Deleting staging area."));
 							IFileManager::Get().DeleteDirectory(*Configuration.StagingDirectory, false, true);
-							CleanupEmptyDirectories(Configuration.InstallDirectory);
+							if (!bWillRetry)
+							{
+								CleanupEmptyDirectories(Configuration.InstallDirectory);
+							}
 						}
 						CleanUpTimer.Stop();
 					}
 					BuildProgress.SetStateProgress(EBuildPatchState::CleanUp, 1.0f);
 
-					// Set if we can retry
-					--InstallRetries;
-					bCanRetry = InstallRetries > 0 && !InstallerError->IsCancelled() && InstallerError->CanRetry();
-					const bool bWillRetry = !bProcessSuccess && bCanRetry;
-
 					// If successful or we will retry, remove the moved files marker
 					if (!Configuration.bInstallToMemory)
 					{
