Fix for patch descriptors not emitting patch data for new bundles.

#rnx
#Synced-CL 49586837

diff --git a/Engine/Source/Runtime/Online/BuildPatchServices/Private/Diffing/DiffManifests.cpp b/Engine/Source/Runtime/Online/BuildPatchServices/Private/Diffing/DiffManifests.cpp
--- a/Engine/Source/Runtime/Online/BuildPatchServices/Private/Diffing/DiffManifests.cpp
+++ b/Engine/Source/Runtime/Online/BuildPatchServices/Private/Diffing/DiffManifests.cpp
@@ -611,8 +611,8 @@ namespace BuildPatchServices
 
 			if (Configuration.OutputPatchDescriptorPath.Len())
 			{
-				TSet<FString> SharedTags = TagsA.Intersect(TagsB);
-				ExportPatchDescriptorFiles(SharedTags, ManifestA, ManifestB);
+				// We want to use all NEW tags so that we catch new bundles, not just changed bundles.
+				ExportPatchDescriptorFiles(TagsB, ManifestA, ManifestB);
 
 				if (Configuration.bOnlyPatchDescriptors)
 				{
@@ -991,4 +991,4 @@ namespace BuildPatchServices
 	{
 		return new FDiffManifests(Configuration);
 	}
-}
\ No newline at end of file
+}
