BPS: Bringing over delta filename aliases, tweak some logging to be more useful.

#Synced-CL 49586837

diff --git a/Engine/Source/Runtime/Online/BuildPatchServices/Private/BuildPatchDelta.cpp b/Engine/Source/Runtime/Online/BuildPatchServices/Private/BuildPatchDelta.cpp
--- a/Engine/Source/Runtime/Online/BuildPatchServices/Private/BuildPatchDelta.cpp
+++ b/Engine/Source/Runtime/Online/BuildPatchServices/Private/BuildPatchDelta.cpp
@@ -4,9 +4,9 @@
 #include "BuildPatchUtil.h"
 #include "BuildPatchMergeManifests.h"
 
-BUILDPATCHSERVICES_API FString BuildPatchServices::GetChunkDeltaFilename(const IBuildManifestRef& SourceManifest, const IBuildManifestRef& DestinationManifest)
+BUILDPATCHSERVICES_API FString BuildPatchServices::GetChunkDeltaFilename(const IBuildManifestRef& SourceManifest, const IBuildManifestRef& DestinationManifest, const FString& FilenameTrailer)
 {
-	return FBuildPatchUtils::GetChunkDeltaFilename(StaticCastSharedRef<FBuildPatchAppManifest>(SourceManifest).Get(), StaticCastSharedRef<FBuildPatchAppManifest>(DestinationManifest).Get());
+	return FBuildPatchUtils::GetChunkDeltaFilename(StaticCastSharedRef<FBuildPatchAppManifest>(SourceManifest).Get(), StaticCastSharedRef<FBuildPatchAppManifest>(DestinationManifest).Get(), FilenameTrailer);
 }
 
 BUILDPATCHSERVICES_API IBuildManifestPtr BuildPatchServices::MergeDeltaManifest(const IBuildManifestRef& Manifest, const IBuildManifestRef& Delta)
diff --git a/Engine/Source/Runtime/Online/BuildPatchServices/Private/BuildPatchInstaller.cpp b/Engine/Source/Runtime/Online/BuildPatchServices/Private/BuildPatchInstaller.cpp
--- a/Engine/Source/Runtime/Online/BuildPatchServices/Private/BuildPatchInstaller.cpp
+++ b/Engine/Source/Runtime/Online/BuildPatchServices/Private/BuildPatchInstaller.cpp
@@ -112,6 +112,7 @@ namespace ConfigHelpers
 		}
 		OptimisedDeltaConfiguration.DeltaPolicy = Config.DeltaPolicy;
 		OptimisedDeltaConfiguration.InstallMode = Config.InstallMode;
+		OptimisedDeltaConfiguration.DeltaFilenameTrailer = Config.DeltaFilenameTrailer;
 		return OptimisedDeltaConfiguration;
 	}
 }
@@ -634,9 +635,10 @@ namespace BuildPatchServices
 			ValidTags.Add(TEXT(""));
 			InstallerAction.GetInstallOrCurrentManifest().GetFileTagList(ValidTags);
 			ValidTags.Add(TEXT(""));
-			if (InstallerAction.GetInstallTags().Difference(ValidTags).Num() > 0)
+			TSet<FString> InvalidTags = InstallerAction.GetInstallTags().Difference(ValidTags);
+			if (InvalidTags.Num() > 0)
 			{
-				UE_LOG(LogBuildPatchServices, Error, TEXT("Installer configuration: Invalid InstallTags provided."));
+				UE_LOG(LogBuildPatchServices, Error, TEXT("Installer configuration: Invalid InstallTags provided: %s"), *FString::Join(InvalidTags, TEXT(",")));
 				InstallerError->SetError(EBuildPatchInstallError::InitializationError, InitializationErrorCodes::InvalidInstallTags, 0,
 					NSLOCTEXT("BuildPatchInstallError", "InvalidInstallTags", "This installation could not continue due to a configuration issue. Please contact support."));
 				bInstallerInitSuccess = false;
@@ -1428,18 +1430,24 @@ namespace BuildPatchServices
 			// Set up our chunk location tracking.
 			int32 DownloadChunkCount = CloudChunks.Num();
 			TMap<FGuid, EConstructorChunkLocation> ChunkLocations;	
+			uint64 InstallBytes = 0;
+			uint64 ChunkDbBytes = 0;
+			uint64 CloudBytes = 0;
 			{
 				for (const FGuid& InstallChunk : InstallChunks)
 				{
 					ChunkLocations.Add(InstallChunk, EConstructorChunkLocation::Install);
+					InstallBytes += ManifestSet->GetDownloadSize(InstallChunk);
 				}
 				for (const FGuid& ChunkDbChunk : ChunkDbChunks)
 				{
 					ChunkLocations.Add(ChunkDbChunk, EConstructorChunkLocation::ChunkDb);
+					ChunkDbBytes += ManifestSet->GetDownloadSize(ChunkDbChunk);
 				}
 				for (const FGuid& CloudChunk : CloudChunks)
 				{
 					ChunkLocations.Add(CloudChunk, EConstructorChunkLocation::Cloud);
+					CloudBytes += ManifestSet->GetDownloadSize(CloudChunk);
 				}
 			}
 
@@ -1523,8 +1531,10 @@ namespace BuildPatchServices
 				BuildStats.ChunksQueuedForDownload = CloudChunks.Num();
 				BuildStats.ChunksLocallyAvailable = InstallChunks.Num();
 				BuildStats.ChunksInChunkDbs = ChunkDbChunks.Num();
-				UE_LOG(LogBPSInstallerConfig, Display, TEXT("Chunk Locations: Cloud %d Install %d ChunkDb %d"), 
-					BuildStats.ChunksQueuedForDownload, BuildStats.ChunksLocallyAvailable, BuildStats.ChunksInChunkDbs);
+				UE_LOG(LogBPSInstallerConfig, Display, TEXT("Chunk Locations: Cloud %d (%s bytes) Install %d (%s bytes) ChunkDb %d (%s bytes)"), 
+					BuildStats.ChunksQueuedForDownload, *FText::AsNumber(CloudBytes).ToString(), 
+					BuildStats.ChunksLocallyAvailable, *FText::AsNumber(InstallBytes).ToString(),
+					BuildStats.ChunksInChunkDbs, *FText::AsNumber(ChunkDbBytes).ToString());
 			}
 
 			// Setup some weightings for the progress tracking
diff --git a/Engine/Source/Runtime/Online/BuildPatchServices/Private/BuildPatchUtil.cpp b/Engine/Source/Runtime/Online/BuildPatchServices/Private/BuildPatchUtil.cpp
--- a/Engine/Source/Runtime/Online/BuildPatchServices/Private/BuildPatchUtil.cpp
+++ b/Engine/Source/Runtime/Online/BuildPatchServices/Private/BuildPatchUtil.cpp
@@ -196,8 +196,12 @@ FString FBuildPatchUtils::GetChunkDeltaDirectory(const FBuildPatchAppManifest& D
 	return TEXT("Deltas") / DestinationManifest.GetBuildId();
 }
 
-FString FBuildPatchUtils::GetChunkDeltaFilename(const FBuildPatchAppManifest& SourceManifest, const FBuildPatchAppManifest& DestinationManifest)
+FString FBuildPatchUtils::GetChunkDeltaFilename(const FBuildPatchAppManifest& SourceManifest, const FBuildPatchAppManifest& DestinationManifest, const FString& FilenameTrailer)
 {
+    if (FilenameTrailer.Len())
+	{
+		return GetChunkDeltaDirectory(DestinationManifest) / SourceManifest.GetBuildId() + TEXT("-") + FilenameTrailer + TEXT(".delta");
+	}
 	return GetChunkDeltaDirectory(DestinationManifest) / SourceManifest.GetBuildId() + TEXT(".delta");
 }
 
diff --git a/Engine/Source/Runtime/Online/BuildPatchServices/Private/BuildPatchUtil.h b/Engine/Source/Runtime/Online/BuildPatchServices/Private/BuildPatchUtil.h
--- a/Engine/Source/Runtime/Online/BuildPatchServices/Private/BuildPatchUtil.h
+++ b/Engine/Source/Runtime/Online/BuildPatchServices/Private/BuildPatchUtil.h
@@ -136,9 +136,10 @@ struct FBuildPatchUtils
 	 * Based on the source and destination manifests, get the filename for the delta that optimises patching from source to destination.
 	 * @param SourceManifest        The source manifest.
 	 * @param DestinationManifest   The destination manifest.
+	 * @param FilenameTrailer		If a nonstandard optimized delta is used, this is the trailer to add to the filename.
 	 * @return the CloudDir relative delta filename.
 	 */
-	static FString GetChunkDeltaFilename(const FBuildPatchAppManifest& SourceManifest, const FBuildPatchAppManifest& DestinationManifest);
+	static FString GetChunkDeltaFilename(const FBuildPatchAppManifest& SourceManifest, const FBuildPatchAppManifest& DestinationManifest, const FString& FilenameTrailer);
 
 	/**
 	 * Checks a file against SHA1 hashes. The function takes two so that it can return no match, match with Hash1, or match with Hash2, that way we can check the file for being the same as an old manifest or new manifest
diff --git a/Engine/Source/Runtime/Online/BuildPatchServices/Private/Generation/ChunkDeltaOptimiser.cpp b/Engine/Source/Runtime/Online/BuildPatchServices/Private/Generation/ChunkDeltaOptimiser.cpp
--- a/Engine/Source/Runtime/Online/BuildPatchServices/Private/Generation/ChunkDeltaOptimiser.cpp
+++ b/Engine/Source/Runtime/Online/BuildPatchServices/Private/Generation/ChunkDeltaOptimiser.cpp
@@ -753,7 +753,8 @@ namespace BuildPatchServices
 
 			// Check that an optimisation does not already exist, and skip long process if so.
 			FBuildPatchAppManifest DeltaManifest;
-			const FString OutputDeltaFilename = Configuration.CloudDirectory / FBuildPatchUtils::GetChunkDeltaFilename(*ManifestA.Get(), *ManifestB.Get());
+			// The Portal stream version of this has support for filename trailers, but we don't use it in this stream so we pass "".
+			const FString OutputDeltaFilename = Configuration.CloudDirectory / FBuildPatchUtils::GetChunkDeltaFilename(*ManifestA.Get(), *ManifestB.Get(), TEXT(""));
 			const bool bDeltaPreviouslyCompleted = FileSystem->FileExists(*OutputDeltaFilename);
 			if (bDeltaPreviouslyCompleted)
 			{
diff --git a/Engine/Source/Runtime/Online/BuildPatchServices/Private/Installer/OptimisedDelta.cpp b/Engine/Source/Runtime/Online/BuildPatchServices/Private/Installer/OptimisedDelta.cpp
--- a/Engine/Source/Runtime/Online/BuildPatchServices/Private/Installer/OptimisedDelta.cpp
+++ b/Engine/Source/Runtime/Online/BuildPatchServices/Private/Installer/OptimisedDelta.cpp
@@ -65,7 +65,7 @@ namespace BuildPatchServices
 	FOptimisedDelta::FOptimisedDelta(const FOptimisedDeltaConfiguration& InConfiguration, FOptimisedDeltaDependencies&& InDependencies)
 		: Configuration(InConfiguration)
 		, DownloadService(InDependencies.DownloadService)
-		, RelativeDeltaFilePath(Configuration.SourceManifest.IsValid() ? FBuildPatchUtils::GetChunkDeltaFilename(*Configuration.SourceManifest.Get(), Configuration.DestinationManifest.Get()) : TEXT(""))
+		, RelativeDeltaFilePath(Configuration.SourceManifest.IsValid() ? FBuildPatchUtils::GetChunkDeltaFilename(*Configuration.SourceManifest.Get(), Configuration.DestinationManifest.Get(), Configuration.DeltaFilenameTrailer) : TEXT(""))
 		, DeltaRetries(ConfigHelpers::LoadDeltaRetries(Configuration.CloudDirectories.Num()))
 		, OnComplete(MoveTemp(InDependencies.OnComplete))
 		, DeltaPolicy(Configuration.DeltaPolicy)
@@ -89,7 +89,7 @@ namespace BuildPatchServices
 		// Kick off the request if we should be.
 		if (DeltaPolicy != EDeltaPolicy::Skip)
 		{
-			UE_LOG(LogOptimisedDelta, Log, TEXT("Requesting optimised delta file %s"), *RelativeDeltaFilePath);
+			UE_LOG(LogOptimisedDelta, Display, TEXT("Requesting optimised delta file %s"), *(Configuration.CloudDirectories[CloudDirIdx] / RelativeDeltaFilePath));
 			DownloadService->RequestFile(Configuration.CloudDirectories[CloudDirIdx] / RelativeDeltaFilePath, ChunkDeltaComplete, ChunkDeltaProgress);
 		}
 		// Otherwise we provide the standard destination manifest.
diff --git a/Engine/Source/Runtime/Online/BuildPatchServices/Private/Installer/OptimisedDelta.h b/Engine/Source/Runtime/Online/BuildPatchServices/Private/Installer/OptimisedDelta.h
--- a/Engine/Source/Runtime/Online/BuildPatchServices/Private/Installer/OptimisedDelta.h
+++ b/Engine/Source/Runtime/Online/BuildPatchServices/Private/Installer/OptimisedDelta.h
@@ -55,6 +55,8 @@ namespace BuildPatchServices
 		EDeltaPolicy DeltaPolicy;
 		// The mode for installation.
 		EInstallMode InstallMode;
+		// The trailer to use for selecting a specific optimized delta or empty if using the default.
+		FString DeltaFilenameTrailer;
 	};
 
 	/**
diff --git a/Engine/Source/Runtime/Online/BuildPatchServices/Public/BuildPatchDelta.h b/Engine/Source/Runtime/Online/BuildPatchServices/Public/BuildPatchDelta.h
--- a/Engine/Source/Runtime/Online/BuildPatchServices/Public/BuildPatchDelta.h
+++ b/Engine/Source/Runtime/Online/BuildPatchServices/Public/BuildPatchDelta.h
@@ -29,9 +29,10 @@ namespace BuildPatchServices
 	 * Based on the source and destination manifests, get the filename for the delta that optimises patching from source to destination.
 	 * @param SourceManifest        The source manifest.
 	 * @param DestinationManifest   The destination manifest.
+	 * @param FilenameTrailer		The trailer to use for selecting a non standard optimized delta, or empty if unneeded.
 	 * @return the CloudDir relative delta filename.
 	 */
-	BUILDPATCHSERVICES_API FString GetChunkDeltaFilename(const IBuildManifestRef& SourceManifest, const IBuildManifestRef& DestinationManifest);
+	BUILDPATCHSERVICES_API FString GetChunkDeltaFilename(const IBuildManifestRef& SourceManifest, const IBuildManifestRef& DestinationManifest, const FString& FilenameTrailer);
 
 	BUILDPATCHSERVICES_API IBuildManifestPtr MergeDeltaManifest(const IBuildManifestRef& Manifest, const IBuildManifestRef& Delta);
 }
diff --git a/Engine/Source/Runtime/Online/BuildPatchServices/Public/BuildPatchSettings.h b/Engine/Source/Runtime/Online/BuildPatchServices/Public/BuildPatchSettings.h
--- a/Engine/Source/Runtime/Online/BuildPatchServices/Public/BuildPatchSettings.h
+++ b/Engine/Source/Runtime/Online/BuildPatchServices/Public/BuildPatchSettings.h
@@ -335,6 +335,8 @@ namespace BuildPatchServices
 
 		// If set, no initial size check will be performed prior to starting the installation. Running out of disk space during installation is unchanged.
 		bool bSkipInitialDiskSizeCheck = false;
+		// Optional specification for non-standard optimized delta
+		FString DeltaFilenameTrailer;
 	};
 
 	/**
