Get rid of the outdated http config redirections;
For Curl based impl, set max concurrent requests based on max total connections and max idle connections, which means MaxConcurrentRequests = MaxTotalConnections - MaxIdleConnections;
Use CVarHttpCurlMaxTotalConnections for test and hotfix support, user don't have to config it, but still load Http.Curl.MaxTotalConnections config, to support game customization;
Use CVarHttpCurlMaxIdleConnections for test and hotfix support, user don't have to config it, but still load Http.Curl.MaxConnections config, to support game customization;
If http connection reuse is disabled, CURLMOPT_MAXCONNECTS will be set as 0 instead of MaxIdleConnections;
Extract curl cached error message output into a separated function LogCurlErrors.

#Synced-CL 49852690

#Synced-CL 49852690

diff --git a/Engine/Config/ConfigRedirects.ini b/Engine/Config/ConfigRedirects.ini
--- a/Engine/Config/ConfigRedirects.ini
+++ b/Engine/Config/ConfigRedirects.ini
@@ -55,24 +55,6 @@ EncryptionKey=ClientEncryptionKey
 [/Script/OnlineSubsystemEOS.EOSSettings]
 EncryptionKey=ClientEncryptionKey
 
-[HTTP]
-; The legacy behavior uses HttpTimeout as HttpActivityTimeout, it actually should use HttpActivityTimeout instead. 
-; Now the value of HttpTimeout will be read into HttpActivityTimeout just to keep existing behavior.
-; If it's really an requirement to actually have total timeout for the http request, use HttpTotalTimeout instead.
-; By default HttpTotalTimeout should be 0, means no time out for the entire http request
-HttpTimeout=HttpActivityTimeout
-; Use HttpActivityTimeout instead for checking any activity on the establish connection
-HttpReceiveTimeout=HttpActivityTimeout
-; The legacy behavior doesn't make use of HttpSendTimeout at all, and there is no such support on some platforms.
-; HttpActivityTimeout should be enough to decide if there is still any ongoing activity
-HttpSendTimeout=HttpActivityTimeout
-; FlushSoftTimeLimitBackground has been removed and will no longer be used
-FlushSoftTimeLimitBackground=Removed
-; FlushHardTimeLimitBackground has been removed and will no longer be used
-FlushHardTimeLimitBackground=Removed
-
-[HTTP.HttpThread]
-;RunningThreadedRequestLimitEditor has been deprecated, change it through DefaultMaxConcurrentRequests in build script instead
-RunningThreadedRequestLimitEditor=Removed
-;RunningThreadedRequestLimit has been deprecated, change it through DefaultMaxConcurrentRequests in build script instead
-RunningThreadedRequestLimit=Removed
\ No newline at end of file
+[HTTP.Curl]
+;bForbidReuse has been removed, config it through http.CurlReuseConnectionEnable or change it through CVarHttpCurlReuseConnectionEnabled instead
+bForbidReuse=Removed
\ No newline at end of file
diff --git a/Engine/Source/Programs/WebTests/Private/TestHttp.cpp b/Engine/Source/Programs/WebTests/Private/TestHttp.cpp
--- a/Engine/Source/Programs/WebTests/Private/TestHttp.cpp
+++ b/Engine/Source/Programs/WebTests/Private/TestHttp.cpp
@@ -533,6 +533,9 @@ TEST_CASE_METHOD(FWaitUntilCompleteHttpFixture, "Can process https request", HTT
 
 TEST_CASE_METHOD(FWaitUntilCompleteHttpFixture, "Can process multiple https requests", HTTP_TAG)
 {
+	// Disable possible warnings "GET https ://www.epicgames.com/ waited in queue for 10.10s before starting" when max concurrent connections setting is low
+	DisableWarningsInThisTest();
+
 	TArray<FString> Urls{
 		TEXT("https://www.unrealengine.com/"),
 		TEXT("https://www.epicgames.com/"),
diff --git a/Engine/Source/Runtime/Online/HTTP/HTTP.Build.cs b/Engine/Source/Runtime/Online/HTTP/HTTP.Build.cs
--- a/Engine/Source/Runtime/Online/HTTP/HTTP.Build.cs
+++ b/Engine/Source/Runtime/Online/HTTP/HTTP.Build.cs
@@ -31,7 +31,7 @@ public class HTTP : ModuleRules
 	{
 		get
 		{
-			return bPlatformSupportsLibCurl || bPlatformSupportsXCurl;
+			return true;
 		}
 	}
 
@@ -53,7 +53,26 @@ public class HTTP : ModuleRules
 		} 
 	}
 
-	protected virtual int DefaultMaxConcurrentRequests { get { return 256; } }
+	protected virtual int CurlDefaultMaxTotalConnections { get { return 256; } }
+	protected virtual int CurlDefaultMaxIdleConnections { get { return 5; } } // Default value of CURLMOPT_MAXCONNECTS is 5
+
+	protected virtual int DefaultMaxConcurrentRequests
+	{
+		get
+		{
+			if (bPlatformSupportsLibCurl)
+			{
+				if (CurlDefaultMaxTotalConnections <= CurlDefaultMaxIdleConnections)
+				{
+					throw new BuildException($"CurlDefaultMaxTotalConnections {CurlDefaultMaxTotalConnections} must be greater than CurlDefaultMaxIdleConnections {CurlDefaultMaxIdleConnections}");
+				}
+
+				return CurlDefaultMaxTotalConnections - CurlDefaultMaxIdleConnections; // Means max active connections
+			}
+
+			return 256;
+		}
+	}
 
 	private bool bPlatformSupportsCurl { get { return bPlatformSupportsLibCurl || bPlatformSupportsXCurl; } }
 
@@ -165,6 +184,8 @@ public class HTTP : ModuleRules
 		}
 
 		PrivateDefinitions.Add("UE_HTTP_DEFAULT_MAX_CONCURRENT_REQUESTS=" + DefaultMaxConcurrentRequests);
+		PrivateDefinitions.Add("UE_HTTP_CURL_DEFAULT_MAX_TOTAL_CONNECTIONS=" + CurlDefaultMaxTotalConnections);
+		PrivateDefinitions.Add("UE_HTTP_CURL_DEFAULT_MAX_IDLE_CONNECTIONS=" + CurlDefaultMaxIdleConnections);
 
 		float PlatformConnectionTimeoutMaxDeviation = 0.5f;
 		if (bPlatformSupportsXCurl)
diff --git a/Engine/Source/Runtime/Online/HTTP/Private/Curl/CurlHttp.cpp b/Engine/Source/Runtime/Online/HTTP/Private/Curl/CurlHttp.cpp
--- a/Engine/Source/Runtime/Online/HTTP/Private/Curl/CurlHttp.cpp
+++ b/Engine/Source/Runtime/Online/HTTP/Private/Curl/CurlHttp.cpp
@@ -1241,27 +1241,7 @@ void FCurlHttpRequest::FinishRequest()
 	bool bSucceeded = Response.IsValid() && Response->bSucceeded;
 	if (!bSucceeded)
 	{
-		if (CurlAddToMultiResult != CURLM_OK)
-		{
-			UE_LOG(LogHttp, Warning, TEXT("%p: request failed, libcurl multi error: %d (%s)"), this, (int32)CurlAddToMultiResult, ANSI_TO_TCHAR(curl_multi_strerror(CurlAddToMultiResult)));
-		}
-		else if (CurlCompletionResult != CURLE_OK)
-		{
-			UE_LOG(LogHttp, Warning, TEXT("%p: request failed, libcurl error: %d (%s)"), this, (int32)CurlCompletionResult, ANSI_TO_TCHAR(curl_easy_strerror(CurlCompletionResult)));
-		}
-
-		const bool bAborted = (bCanceled || bTimedOut || bActivityTimedOut);
-		if (!bAborted)
-		{
-			const FScopeLock CacheLock(&InfoMessageCacheCriticalSection);
-			for (int32 i = 0; i < InfoMessageCache.Num(); ++i)
-			{
-				if (InfoMessageCache[(LeastRecentlyCachedInfoMessageIndex + i) % InfoMessageCache.Num()].Len() > 0)
-				{
-					UE_LOG(LogHttp, Warning, TEXT("%p: libcurl info message cache %d (%s)"), this, (LeastRecentlyCachedInfoMessageIndex + i) % InfoMessageCache.Num(), *(InfoMessageCache[(LeastRecentlyCachedInfoMessageIndex + i) % NumberOfInfoMessagesToCache]));
-				}
-			}
-		}
+		LogCurlErrors();
 
 		if (bCurlRequestCompleted)
 		{
@@ -1287,6 +1267,32 @@ void FCurlHttpRequest::FinishRequest()
 	OnFinishRequest(bSucceeded);
 }
 
+void FCurlHttpRequest::LogCurlErrors()
+{
+	if (CurlAddToMultiResult != CURLM_OK)
+	{
+		UE_LOG(LogHttp, Warning, TEXT("%p: request failed, libcurl multi error: %d (%s)"), this, (int32)CurlAddToMultiResult, ANSI_TO_TCHAR(curl_multi_strerror(CurlAddToMultiResult)));
+	}
+	else if (CurlCompletionResult != CURLE_OK)
+	{
+		UE_LOG(LogHttp, Warning, TEXT("%p: request failed, libcurl error: %d (%s)"), this, (int32)CurlCompletionResult, ANSI_TO_TCHAR(curl_easy_strerror(CurlCompletionResult)));
+	}
+
+	const bool bAborted = (bCanceled || bTimedOut || bActivityTimedOut);
+	if (!bAborted)
+	{
+		const FScopeLock CacheLock(&InfoMessageCacheCriticalSection);
+		for (int32 i = 0; i < InfoMessageCache.Num(); ++i)
+		{
+			int RelativeIndex = (LeastRecentlyCachedInfoMessageIndex + i) % InfoMessageCache.Num();
+			if (!InfoMessageCache[RelativeIndex].IsEmpty())
+			{
+				UE_LOG(LogHttp, Warning, TEXT("%p: libcurl info message cache %d (%s)"), this, RelativeIndex, *(InfoMessageCache[RelativeIndex]));
+			}
+		}
+	}
+}
+
 void FCurlHttpRequest::CleanupRequest()
 {
 	CloseRequestPayloadDefaultImpl();
diff --git a/Engine/Source/Runtime/Online/HTTP/Private/Curl/CurlHttp.h b/Engine/Source/Runtime/Online/HTTP/Private/Curl/CurlHttp.h
--- a/Engine/Source/Runtime/Online/HTTP/Private/Curl/CurlHttp.h
+++ b/Engine/Source/Runtime/Online/HTTP/Private/Curl/CurlHttp.h
@@ -258,6 +258,8 @@ private:
 	void SetupOptionUnixSocketPath();
 	void SetupOptionHttpVersion();
 
+	void LogCurlErrors();
+
 private:
 
 	/** Pointer to an easy handle specific to this request */
diff --git a/Engine/Source/Runtime/Online/HTTP/Private/Curl/CurlHttpManager.cpp b/Engine/Source/Runtime/Online/HTTP/Private/Curl/CurlHttpManager.cpp
--- a/Engine/Source/Runtime/Online/HTTP/Private/Curl/CurlHttpManager.cpp
+++ b/Engine/Source/Runtime/Online/HTTP/Private/Curl/CurlHttpManager.cpp
@@ -67,6 +67,22 @@ CURLSH* FCurlHttpManager::GShareHandle = nullptr;
 
 FCurlHttpManager::FCurlRequestOptions FCurlHttpManager::CurlRequestOptions;
 
+extern TAutoConsoleVariable<int32> CVarHttpMaxConcurrentRequests;
+
+TAutoConsoleVariable<int32> CVarHttpCurlMaxTotalConnections(
+	TEXT("HTTP.Curl.MaxTotalConnections"),
+	UE_HTTP_CURL_DEFAULT_MAX_TOTAL_CONNECTIONS,
+	TEXT("The max number of total http connections"),
+	ECVF_SaveForNextBoot
+);
+
+TAutoConsoleVariable<int32> CVarHttpCurlMaxIdleConnections(
+	TEXT("HTTP.Curl.MaxIdleConnections"),
+	UE_HTTP_CURL_DEFAULT_MAX_IDLE_CONNECTIONS,
+	TEXT("The max number of idle http connections"),
+	ECVF_SaveForNextBoot
+);
+
 TAutoConsoleVariable<bool> CVarHttpCurlReuseConnectionEnabled(
 	TEXT("http.CurlReuseConnectionEnabled"),
 	UE_HTTP_CURL_REUSE_CONNECTION_ENABLED_BY_DEFAULT,
@@ -155,29 +171,45 @@ void FCurlHttpManager::InitCurl()
 			UE_LOG(LogInit, Fatal, TEXT("Could not initialize create libcurl multi handle! HTTP transfers will not function properly."));
 		}
 
+#if !WITH_CURL_XCURL
 		int32 MaxTotalConnections = 0;
+		// Support of customization by game settings
 		if (GConfig->GetInt(TEXT("HTTP.Curl"), TEXT("MaxTotalConnections"), MaxTotalConnections, GEngineIni) && MaxTotalConnections > 0)
 		{
-			const CURLMcode SetOptResult = curl_multi_setopt(GMultiHandle, CURLMOPT_MAX_TOTAL_CONNECTIONS, static_cast<long>(MaxTotalConnections));
-			if (SetOptResult != CURLM_OK)
-			{
-				UE_LOG(LogInit, Warning, TEXT("Failed to set libcurl max total connections options (%d), error %d ('%s')"),
-					MaxTotalConnections, static_cast<int32>(SetOptResult), StringCast<TCHAR>(curl_multi_strerror(SetOptResult)).Get());
-			}
+			// If CVarHttpCurlMaxTotalConnections was set by ECVF_SetByHotfix through ECVF_SaveForNextBoot, this call will be an none operation
+			CVarHttpCurlMaxTotalConnections.AsVariable()->Set(MaxTotalConnections, ECVF_SetByGameSetting);
 		}
 
-		int32 MaxConnects = 0;
-		if (GConfig->GetInt(TEXT("HTTP.Curl"), TEXT("MaxConnects"), MaxConnects, GEngineIni) && MaxConnects >= 0)
+		CURLMcode SetOptResult = curl_multi_setopt(GMultiHandle, CURLMOPT_MAX_TOTAL_CONNECTIONS, static_cast<long>(CVarHttpCurlMaxTotalConnections.GetValueOnAnyThread()));
+		if (SetOptResult != CURLM_OK)
 		{
-			const CURLMcode SetOptResult = curl_multi_setopt(GMultiHandle, CURLMOPT_MAXCONNECTS, static_cast<long>(MaxConnects));
-			if (SetOptResult != CURLM_OK)
-			{
-				UE_LOG(LogInit, Warning, TEXT("Failed to set libcurl max connects options (%d), error %d ('%s')"),
-					MaxConnects, static_cast<int32>(SetOptResult), StringCast<TCHAR>(curl_multi_strerror(SetOptResult)).Get());
-			}
+			UE_LOG(LogInit, Warning, TEXT("Failed to set libcurl max total connections options (%d), error %d ('%s')"),
+				CVarHttpCurlMaxTotalConnections.GetValueOnAnyThread(), static_cast<int32>(SetOptResult), StringCast<TCHAR>(curl_multi_strerror(SetOptResult)).Get());
+		}
+
+		// Support of customization by game settings
+		bool bCurlReuseConnectionEnabled = true;
+		if (GConfig->GetBool(TEXT("http"), TEXT("CurlReuseConnectionEnabled"), bCurlReuseConnectionEnabled, GEngineIni))
+		{
+			// If CVarHttpCurlReuseConnectionEnabled was set by ECVF_SetByHotfix through ECVF_SaveForNextBoot, this call will be an none operation
+			CVarHttpCurlReuseConnectionEnabled.AsVariable()->Set(bCurlReuseConnectionEnabled, ECVF_SetByGameSetting);
+		}
+
+		// Support of customization by game settings
+		int32 MaxIdleConnections = 0;
+		if (GConfig->GetInt(TEXT("HTTP.Curl"), TEXT("MaxConnects"), MaxIdleConnections, GEngineIni) && MaxIdleConnections >= 0)
+		{
+			// If CVarHttpCurlMaxIdleConnections was set by ECVF_SetByHotfix through ECVF_SaveForNextBoot, this call will be an none operation
+			CVarHttpCurlMaxIdleConnections.AsVariable()->Set(MaxIdleConnections, ECVF_SetByGameSetting);
+		}
+
+		SetOptResult = curl_multi_setopt(GMultiHandle, CURLMOPT_MAXCONNECTS, static_cast<long>(CVarHttpCurlReuseConnectionEnabled.GetValueOnAnyThread() ? CVarHttpCurlMaxIdleConnections.GetValueOnAnyThread() : 0));
+		if (SetOptResult != CURLM_OK)
+		{
+			UE_LOG(LogInit, Warning, TEXT("Failed to set libcurl max connects options (%d), error %d ('%s')"),
+				CVarHttpCurlMaxIdleConnections.GetValueOnAnyThread(), static_cast<int32>(SetOptResult), StringCast<TCHAR>(curl_multi_strerror(SetOptResult)).Get());
 		}
 
-#if !WITH_CURL_XCURL
 		GShareHandle = curl_share_init();
 		if (NULL != GShareHandle)
 		{
@@ -189,7 +221,7 @@ void FCurlHttpManager::InitCurl()
 		{
 			UE_LOG(LogInit, Fatal, TEXT("Could not initialize libcurl share handle!"));
 		}
-#endif
+#endif // !WITH_CURL_XCURL
 	}
 	else
 	{
