[BuildPatchServices]
* Added missing include when building mac non unity

#Synced-CL 50214969

diff --git a/Engine/Source/Runtime/Online/BuildPatchServices/Private/BuildPatchFileConstructor.cpp b/Engine/Source/Runtime/Online/BuildPatchServices/Private/BuildPatchFileConstructor.cpp
--- a/Engine/Source/Runtime/Online/BuildPatchServices/Private/BuildPatchFileConstructor.cpp
+++ b/Engine/Source/Runtime/Online/BuildPatchServices/Private/BuildPatchFileConstructor.cpp
@@ -33,6 +33,7 @@
 
 #if PLATFORM_MAC
 #include <sys/stat.h>
+#include <unistd.h>
 #endif
 
 using namespace BuildPatchServices;
