BPS: Fix for hash error when reading an uncompressed chunk through a smaller destination buffer.

#rnx

#Synced-CL 49917480

diff --git a/Engine/Source/Runtime/Online/BuildPatchServices/Private/Data/ChunkData.cpp b/Engine/Source/Runtime/Online/BuildPatchServices/Private/Data/ChunkData.cpp
--- a/Engine/Source/Runtime/Online/BuildPatchServices/Private/Data/ChunkData.cpp
+++ b/Engine/Source/Runtime/Online/BuildPatchServices/Private/Data/ChunkData.cpp
@@ -17,6 +17,8 @@
 // The chunkdb header magic codeword, for quick checking that the opened file is a chunkdb file.
 #define CHUNKDB_HEADER_MAGIC    0xB1FE3AA3
 
+DEFINE_LOG_CATEGORY_STATIC(LogChunkData, Log, All);
+
 namespace BuildPatchServices
 {
 	const TCHAR* ToString(const EChunkLoadResult& ChunkLoadResult)
@@ -567,37 +569,68 @@ namespace BuildPatchServices
 
 
 		virtual bool DecompressValidatedRead(const FChunkHeader& InHeader, FMutableMemoryView InDestinationBuffer, const FUniqueBuffer& InCompressedBuffer) const override
-		{			
-			FMutableMemoryView UseDestinationBuffer = InDestinationBuffer;
-
-			FUniqueBuffer TempDestinationBuffer;
-			if (InHeader.DataSizeUncompressed != InDestinationBuffer.GetSize())
-			{
-				TempDestinationBuffer = FUniqueBuffer::Alloc(InHeader.DataSizeUncompressed);
-				UseDestinationBuffer = TempDestinationBuffer.GetView();
-			}
+		{
+			// ValidateAndRead sets up the data we get on the dispatch thread so we can decompress and hash here
+			// in another thread. Where the data lives depends on compression and whether the read requested the 
+			// full buffer.
+			
 
+			FMemoryView FullChunkBuffer;
+			FUniqueBuffer DecompressionBuffer;
 			if ((InHeader.StoredAs & EChunkStorageFlags::Compressed) != EChunkStorageFlags::None)
 			{
+				// Compressed - we usually can decompress directly to the output, but sometimes
+				// the request isn't for the full chunk so we use a temp.
+				FMutableMemoryView DecompressTo = InDestinationBuffer;
+				if (DecompressTo.GetSize() != InHeader.DataSizeUncompressed)
+				{
+					DecompressionBuffer = FUniqueBuffer::Alloc(InHeader.DataSizeUncompressed);
+					DecompressTo = DecompressionBuffer.GetView();
+				}
+
 				if (!FCompression::UncompressMemory(
 					NAME_Zlib,
-					UseDestinationBuffer.GetData(),
-					UseDestinationBuffer.GetSize(),
+					DecompressTo.GetData(),
+					DecompressTo.GetSize(),
 					InCompressedBuffer.GetData(),
 					InCompressedBuffer.GetSize()))
 				{
+					// already logged.
 					return false;
 				}
+
+				FullChunkBuffer = DecompressTo;
+			}
+			else
+			{
+				// Uncompressed - where ValidateAndRead put the data depends on whether the destination buffer
+				// was big enough for the full chunk.
+				FullChunkBuffer = InDestinationBuffer;
+				if (InDestinationBuffer.GetSize() != InHeader.DataSizeUncompressed)
+				{
+					// The data was read into InCompressedBuffer instead.
+					FullChunkBuffer = InCompressedBuffer.GetView();
+				}
 			}
 
-			bool bHashMatches = CheckChunkHash(InHeader, UseDestinationBuffer.GetData(), UseDestinationBuffer.GetSize());
-			if (bHashMatches && InHeader.DataSizeUncompressed != InDestinationBuffer.GetSize())
+			if (!CheckChunkHash(InHeader, FullChunkBuffer.GetData(), FullChunkBuffer.GetSize()))
 			{
-				// We routed through a temp buffer, copy out as necessary.
-				FMemory::Memcpy(InDestinationBuffer.GetData(), TempDestinationBuffer.GetData(), InDestinationBuffer.GetSize());
+				UE_LOG(LogChunkData, Error, TEXT("Failed hash check for chunk %s"), *WriteToString<64>(InHeader.Guid));
+				return false;
 			}
 
-			return bHashMatches;
+			// If we had to use a temp buffer, copy to the dest.
+			if (FullChunkBuffer.GetData() != InDestinationBuffer.GetData())
+			{
+				// Destination buffer shuold always be <= the chunk size.
+				if (InDestinationBuffer.GetSize() > FullChunkBuffer.GetSize())
+				{
+					UE_LOG(LogChunkData, Error, TEXT("Invalid destination buffer size for chunk %s"), *WriteToString<64>(InHeader.Guid));
+					return false;
+				}
+				FMemory::Memcpy(InDestinationBuffer.GetData(), FullChunkBuffer.GetData(), InDestinationBuffer.GetSize());
+			}
+			return true;
 		}
 
 
