BPS: "Experimental" support for per file subdirectories, with different directories for the read location (install source / patching) and the write location (everything else).

#rnx

#Synced-CL 49917480

diff --git a/Engine/Source/Runtime/Online/BuildPatchServices/Private/BuildPatchFileConstructor.cpp b/Engine/Source/Runtime/Online/BuildPatchServices/Private/BuildPatchFileConstructor.cpp
--- a/Engine/Source/Runtime/Online/BuildPatchServices/Private/BuildPatchFileConstructor.cpp
+++ b/Engine/Source/Runtime/Online/BuildPatchServices/Private/BuildPatchFileConstructor.cpp
@@ -3021,7 +3021,9 @@ void FBuildPatchFileConstructor::CompleteConstructedFile(FFileConstructionState&
 		Configuration.InstallMode == EInstallMode::DestructiveInstall &&
 		!Configuration.bInstallToMemory) // Install to memory never patches/has an existing installation dir.
 	{
-		FString FileToDelete = Configuration.InstallDirectory / CurrentFile.BuildFilename;
+		TMap<FString, FString> EmptyPerFileSubdirectories;
+		FString FileToDelete = FBuildPatchUtils::ResolveInstallationFileName(Configuration.InstallDirectory, Configuration.PerFileSubdirectories ? *Configuration.PerFileSubdirectories : EmptyPerFileSubdirectories, CurrentFile.BuildFilename);
+
 		FPaths::NormalizeFilename(FileToDelete);
 		FPaths::CollapseRelativeDirectories(FileToDelete);
 		if (FileSystem->FileExists(*FileToDelete))
diff --git a/Engine/Source/Runtime/Online/BuildPatchServices/Private/BuildPatchFileConstructor.h b/Engine/Source/Runtime/Online/BuildPatchServices/Private/BuildPatchFileConstructor.h
--- a/Engine/Source/Runtime/Online/BuildPatchServices/Private/BuildPatchFileConstructor.h
+++ b/Engine/Source/Runtime/Online/BuildPatchServices/Private/BuildPatchFileConstructor.h
@@ -78,7 +78,10 @@ namespace BuildPatchServices
 		// The location where memory overflow will get written to.
 		FString BackingStoreDirectory;
 
-		IBuildInstallerSharedContext* SharedContext;
+		// Overrides for file locations.
+		const TMap<FString, FString>* PerFileSubdirectories = nullptr;
+
+		IBuildInstallerSharedContext* SharedContext = nullptr;
 
 		// See comments in installer config
 		bool bInstallToMemory = false;
diff --git a/Engine/Source/Runtime/Online/BuildPatchServices/Private/BuildPatchInstaller.cpp b/Engine/Source/Runtime/Online/BuildPatchServices/Private/BuildPatchInstaller.cpp
--- a/Engine/Source/Runtime/Online/BuildPatchServices/Private/BuildPatchInstaller.cpp
+++ b/Engine/Source/Runtime/Online/BuildPatchServices/Private/BuildPatchInstaller.cpp
@@ -758,7 +758,9 @@ namespace BuildPatchServices
 
 		// We can now build out any systems that need late construction but can survive between retries.
 		ManifestSet.Reset(FBuildManifestSetFactory::Create(InstallerActions));
-		Verifier.Reset(FVerifierFactory::Create(FileSystem.Get(), VerifierStatistics.Get(), Configuration.VerifyMode, Configuration.SharedContext, ManifestSet.Get(), Configuration.InstallDirectory, Configuration.InstallMode == EInstallMode::StageFiles ? InstallStagingDir : FString()));
+		Verifier.Reset(FVerifierFactory::Create(FileSystem.Get(), VerifierStatistics.Get(), Configuration.VerifyMode, Configuration.SharedContext, 
+			ManifestSet.Get(), Configuration.InstallDirectory, Configuration.InstallMode == EInstallMode::StageFiles ? InstallStagingDir : FString(),
+			NewPerFileSubdirectories));
 
 		// Add systems to controllables.
 		{
@@ -796,6 +798,9 @@ namespace BuildPatchServices
 		ManifestSet->GetExpectedFiles(ExpectedFiles);
 		for (const FString& ExpectedFile : ExpectedFiles)
 		{
+			// We're not checking per file subdirs here, which could allow things to escape, but we also are specifically
+			// adding per file subdirs so that it can since we're wanting to move files to a sibling directory of the install
+			// location. 
 			const FString InstallConstructionFile = FPaths::ConvertRelativePathToFull(Configuration.InstallDirectory, ExpectedFile);
 			if (!InstallConstructionFile.StartsWith(InstallDirectoryWithSlash))
 			{
@@ -1024,7 +1029,8 @@ namespace BuildPatchServices
 		// Check if any of the provided files exist on disk. If they do, we'll be starting the installation with a verify to find out what work needs to be done.
 		for (const FString& FileToCheck : FilesToCheck)
 		{
-			if (FileSystem->FileExists(*(Configuration.InstallDirectory / FileToCheck)))
+			const FString DiskFileName = FBuildPatchUtils::ResolveInstallationFileName(Configuration.InstallDirectory, NewPerFileSubdirectories, FileToCheck);
+			if (FileSystem->FileExists(*DiskFileName))
 			{
 				return true;
 			}
@@ -1298,7 +1304,11 @@ namespace BuildPatchServices
 		}
 		else
 		{
-			ManifestSet->GetOutdatedFiles(Configuration.InstallDirectory, OutdatedFiles);
+			// The only reason the install directory is passed to this is to do a "just in case" check where even if a file isn't
+			// supposed to have changed, if the file size on disk DID change, we add it to the file list as outdated. IMO this should be done as a separate function,
+			// but for the time being we don't really care about this verification so we remove the install directory to disable this when we have
+			// per file subdirectories rather than try and pipe the per file subdirectory stuff all through the manifest set.
+			ManifestSet->GetOutdatedFiles(NewPerFileSubdirectories.Num() ? FString() : Configuration.InstallDirectory, OutdatedFiles);
 		}
 
 		const bool bIsPrereqOnly = Configuration.InstallMode == EInstallMode::PrereqOnly;
@@ -1343,7 +1353,7 @@ namespace BuildPatchServices
 		{
 			for (const FString& FileToConstruct : FilesToConstruct)
 			{
-				const FString InstallConstructionFile = Configuration.InstallDirectory / FileToConstruct;
+				const FString InstallConstructionFile = FBuildPatchUtils::ResolveInstallationFileName(Configuration.InstallDirectory, NewPerFileSubdirectories, FileToConstruct);
 				const FString StagedConstructionFile = InstallStagingDir / FileToConstruct;
 				if (InstallConstructionFile.Len() >= FPlatformMisc::GetMaxPathLength())
 				{
@@ -1409,7 +1419,9 @@ namespace BuildPatchServices
 					FileSystem.Get(),
 					InstallChunkSourceStatistics.Get(),
 					InstallationInfo,
-					ReferencedChunks));
+					ReferencedChunks,
+					OldPerFileSubdirectories
+				));
 			}
 
 			// Sub out the chunks we have available on disk - anything we don't we have to download.
@@ -1478,6 +1490,7 @@ namespace BuildPatchServices
 			FCC.InstallDirectory = Configuration.InstallDirectory;
 			FPaths::NormalizeDirectoryName(FCC.InstallDirectory);
 			FCC.InstallMode = Configuration.InstallMode;
+			FCC.PerFileSubdirectories = &NewPerFileSubdirectories;
 			FCC.ManifestSet = ManifestSet.Get();
 			FCC.MetaDirectory = MetaStagingDir;
 			FCC.SharedContext = Configuration.SharedContext.Get();
@@ -1675,11 +1688,14 @@ namespace BuildPatchServices
 		}
 	}
 
-	void FBuildPatchInstaller::FilterToExistingFiles(const FString& RootDirectory, TSet<FString>& Files)
+	void FBuildPatchInstaller::FilterToExistingFiles(const FString& RootDirectory, TSet<FString>& Files, bool bResolveInstallationFileName)
 	{
 		for (auto FilesIt = Files.CreateIterator(); FilesIt; ++FilesIt)
 		{
-			if (!FileSystem->FileExists(*(RootDirectory / *FilesIt)))
+			const FString DiskFileName = bResolveInstallationFileName ? 
+				FBuildPatchUtils::ResolveInstallationFileName(Configuration.InstallDirectory, NewPerFileSubdirectories, *FilesIt) 
+			: (RootDirectory / *FilesIt);
+			if (!FileSystem->FileExists(*DiskFileName))
 			{
 				FilesIt.RemoveCurrent();
 			}
@@ -1764,13 +1780,13 @@ namespace BuildPatchServices
 				// Get list of all expected files
 				TSet<FString> FilesToRelocate;
 				ManifestSet->GetExpectedFiles(FilesToRelocate);
-				FilterToExistingFiles(InstallStagingDir, FilesToRelocate);
+				FilterToExistingFiles(InstallStagingDir, FilesToRelocate, false);
 				FilesToRelocate.Sort(TLess<FString>());
 
 				// First handle files that should be removed for patching
 				TSet<FString> FilesToRemove;
 				ManifestSet->GetRemovableFiles(FilesToRemove);
-				FilterToExistingFiles(Configuration.InstallDirectory, FilesToRemove);
+				FilterToExistingFiles(Configuration.InstallDirectory, FilesToRemove, true);
 				FilesToRemove.Sort(TLess<FString>());
 
 				// Counters for progress tracking.
@@ -1792,7 +1808,7 @@ namespace BuildPatchServices
 				for (auto FileToRemove = FilesToRemove.CreateConstIterator(); FileToRemove && !InstallerError->HasError(); ++FileToRemove)
 				{
 					BackupFileIfNecessary(*FileToRemove);
-					const FString FullFilename = Configuration.InstallDirectory / *FileToRemove;
+					const FString FullFilename = FBuildPatchUtils::ResolveInstallationFileName(Configuration.InstallDirectory, NewPerFileSubdirectories, *FileToRemove);
 					uint32 ErrorCode = 0;
 					bool bDeleteSuccess = RemoveFileWithRetries(FullFilename, ErrorCode);
 					if (bDeleteSuccess)
@@ -1813,7 +1829,7 @@ namespace BuildPatchServices
 				for (auto FileToRelocate = FilesToRelocate.CreateConstIterator(); FileToRelocate && !InstallerError->HasError(); ++FileToRelocate)
 				{
 					const FString SrcFilename = InstallStagingDir / *FileToRelocate;
-					const FString DestFilename = Configuration.InstallDirectory / *FileToRelocate;
+					const FString DestFilename = FBuildPatchUtils::ResolveInstallationFileName(Configuration.InstallDirectory, NewPerFileSubdirectories, *FileToRelocate);
 					const bool bOldFileExists = FileSystem->FileExists(*DestFilename);
 					if (bOldFileExists)
 					{
@@ -1873,7 +1889,7 @@ namespace BuildPatchServices
 
 		// Construct the attributes class
 		FileAttributesTimer.Start();
-		TUniquePtr<IFileAttribution> Attributes(FFileAttributionFactory::Create(FileSystem.Get(), ManifestSet.Get(), FilesToConstruct, Configuration.InstallDirectory, OptionalStageDirectory, &BuildProgress));
+		TUniquePtr<IFileAttribution> Attributes(FFileAttributionFactory::Create(FileSystem.Get(), ManifestSet.Get(), FilesToConstruct, Configuration.InstallDirectory, OptionalStageDirectory, &BuildProgress, NewPerFileSubdirectories));
 		FScopedControllables ScopedControllables(&ThreadLock, Controllables, bIsPaused, bShouldAbort);
 		ScopedControllables.Register(Attributes.Get());
 		Attributes->ApplyAttributes();
@@ -1941,7 +1957,7 @@ namespace BuildPatchServices
 				}
 				if (Configuration.InstallMode != EInstallMode::StageFiles)
 				{
-					FString InstalledFile = Configuration.InstallDirectory / CorruptFile;
+					FString InstalledFile = FBuildPatchUtils::ResolveInstallationFileName(Configuration.InstallDirectory, NewPerFileSubdirectories, CorruptFile);
 					if (FileSystem->FileExists(*InstalledFile))
 					{
 						FileSystem->DeleteFile(*InstalledFile);
@@ -1963,7 +1979,7 @@ namespace BuildPatchServices
 
 	bool FBuildPatchInstaller::BackupFileIfNecessary(const FString& Filename, bool bDiscoveredByVerification /*= false */)
 	{
-		const FString InstalledFilename = Configuration.InstallDirectory / Filename;
+		const FString InstalledFilename = FBuildPatchUtils::ResolveInstallationFileName(Configuration.InstallDirectory, NewPerFileSubdirectories, Filename);
 		const FString BackupFilename = Configuration.BackupDirectory / Filename;
 		const bool bBackupOriginals = !Configuration.BackupDirectory.IsEmpty();
 		// Skip if not doing backups
@@ -2113,6 +2129,12 @@ namespace BuildPatchServices
 		return FilesInstalledToMemory;
 	}
 
+	void FBuildPatchInstaller::SetPerFileSubdirectories(TMap<FString, FString>&& InOldPerFileSubdirectories, TMap<FString, FString>&& InNewPerFileSubdirectories)
+	{
+		NewPerFileSubdirectories = MoveTemp(InNewPerFileSubdirectories);
+		OldPerFileSubdirectories = MoveTemp(InOldPerFileSubdirectories);
+	}
+
 	//@todo this is deprecated and shouldn't be used anymore [6/4/2014 justin.sargent]
 	FText FBuildPatchInstaller::GetPercentageText() const
 	{
diff --git a/Engine/Source/Runtime/Online/BuildPatchServices/Private/BuildPatchInstaller.h b/Engine/Source/Runtime/Online/BuildPatchServices/Private/BuildPatchInstaller.h
--- a/Engine/Source/Runtime/Online/BuildPatchServices/Private/BuildPatchInstaller.h
+++ b/Engine/Source/Runtime/Online/BuildPatchServices/Private/BuildPatchInstaller.h
@@ -139,6 +139,12 @@ namespace BuildPatchServices
 		// With bInstallToMemory, files end up here after construction.
 		TMap<FString, TArray64<uint8>> FilesInstalledToMemory;
 
+		// Keyed off the file manifest filename, value is appended to the installation directory.
+		// Where to find the files in the installation when patching.
+		TMap<FString, FString> OldPerFileSubdirectories;
+		// Where to put the files in the new installation.
+		TMap<FString, FString> NewPerFileSubdirectories;
+
 		// Map of registered installations.
 		TMultiMap<FString, FBuildPatchAppManifestRef> InstallationInfo;
 
@@ -260,6 +266,7 @@ namespace BuildPatchServices
 		virtual void GetDebugText(TArray<FString>& Output) override;
 #endif
 		virtual TMap<FString, TArray64<uint8>>& GetFilesInstalledToMemory() override;
+		virtual void SetPerFileSubdirectories(TMap<FString, FString>&& InOldPerFileSubdirectories, TMap<FString, FString>&& InNewPerFileSubdirectories) override;
 		// IBuildInstaller interface end.
 
 		/**
@@ -397,8 +404,9 @@ namespace BuildPatchServices
 		 * Remove from the set any files that do not exist.
 		 * @param RootDirectory     Root Directory for search.
 		 * @param Files             Set of fles to filter.
+		 * @param bResolveInstallationFileName If true, instead of appending the RootDirectory / FilesName, call ResolveInstallationFileName on the file.
 		 */
-		void FilterToExistingFiles(const FString& RootDirectory, TSet<FString>& Files);
+		void FilterToExistingFiles(const FString& RootDirectory, TSet<FString>& Files, bool bResolveInstallationFileName);
 
 		/**
 		 * Attempt to remove a file, using configured retries.
diff --git a/Engine/Source/Runtime/Online/BuildPatchServices/Private/BuildPatchUtil.cpp b/Engine/Source/Runtime/Online/BuildPatchServices/Private/BuildPatchUtil.cpp
--- a/Engine/Source/Runtime/Online/BuildPatchServices/Private/BuildPatchUtil.cpp
+++ b/Engine/Source/Runtime/Online/BuildPatchServices/Private/BuildPatchUtil.cpp
@@ -259,6 +259,15 @@ uint8 FBuildPatchUtils::VerifyFile(IFileSystem* FileSystem, const FString& FileT
 	return ReturnValue;
 }
 
+FString FBuildPatchUtils::ResolveInstallationFileName(const FString& BaseInstallationDirectory, const TMap<FString, FString>& PerFileSubdirectories, const FString& BuildFilename)
+{
+	const FString* PerFileSubdir = PerFileSubdirectories.Find(BuildFilename);
+	if (PerFileSubdir)
+	{
+		return BaseInstallationDirectory / *PerFileSubdir / BuildFilename;
+	}
+	return BaseInstallationDirectory / BuildFilename;
+}
 
 uint64 FBuildPatchUtils::CalculateDiskSpaceRequirementsWithDeleteDuringInstall(
 	const TArray<FString>& InFilesToConstruct, 
diff --git a/Engine/Source/Runtime/Online/BuildPatchServices/Private/BuildPatchUtil.h b/Engine/Source/Runtime/Online/BuildPatchServices/Private/BuildPatchUtil.h
--- a/Engine/Source/Runtime/Online/BuildPatchServices/Private/BuildPatchUtil.h
+++ b/Engine/Source/Runtime/Online/BuildPatchServices/Private/BuildPatchUtil.h
@@ -172,4 +172,10 @@ struct FBuildPatchUtils
 		BuildPatchServices::IBuildManifestSet* InManifestSet, 
 		const TArray<uint64>& InChunkDbSizesAtPosition, 
 		uint64 InTotalChunkDbSize);
+
+    /**
+	* Generates the path to the file we're installing by adjusting the path for per file subdirectories. Currently this does not handle
+	* installation subdirectories which are appended as part of the manifest set.
+	*/
+	static FString ResolveInstallationFileName(const FString& BaseInstallationDirectory, const TMap<FString, FString>& PerFileSubdirectories, const FString& BuildFilename);
 };
diff --git a/Engine/Source/Runtime/Online/BuildPatchServices/Private/Installer/FileAttribution.cpp b/Engine/Source/Runtime/Online/BuildPatchServices/Private/Installer/FileAttribution.cpp
--- a/Engine/Source/Runtime/Online/BuildPatchServices/Private/Installer/FileAttribution.cpp
+++ b/Engine/Source/Runtime/Online/BuildPatchServices/Private/Installer/FileAttribution.cpp
@@ -4,13 +4,15 @@
 #include "Common/FileSystem.h"
 #include "BuildPatchProgress.h"
 #include "IBuildManifestSet.h"
+#include "BuildPatchUtil.h"
 
 namespace BuildPatchServices
 {
 	class FFileAttribution : public IFileAttribution
 	{
 	public:
-		FFileAttribution(IFileSystem* FileSystem, IBuildManifestSet* ManifestSet, TSet<FString> TouchedFiles, const FString& InstallDirectory, const FString& StagedFileDirectory, bool bUseStageDirectory, FBuildPatchProgress* BuildProgress);
+		FFileAttribution(IFileSystem* FileSystem, IBuildManifestSet* ManifestSet, TSet<FString> TouchedFiles, const FString& InstallDirectory, 
+			const FString& StagedFileDirectory, bool bUseStageDirectory, FBuildPatchProgress* BuildProgress, const TMap<FString, FString>& PerFileSubdirectories);
 
 		virtual ~FFileAttribution() {}
 
@@ -38,9 +40,11 @@ namespace BuildPatchServices
 		FBuildPatchProgress* BuildProgress;
 		FThreadSafeBool bIsPaused;
 		FThreadSafeBool bShouldAbort;
+		const TMap<FString, FString>& PerFileSubdirectories;
 	};
 
-	FFileAttribution::FFileAttribution(IFileSystem* InFileSystem, IBuildManifestSet* InManifestSet, TSet<FString> InTouchedFiles, const FString& InInstallDirectory, const FString& InStagedFileDirectory, bool bInUseStageDirectory, FBuildPatchProgress* InBuildProgress)
+	FFileAttribution::FFileAttribution(IFileSystem* InFileSystem, IBuildManifestSet* InManifestSet, TSet<FString> InTouchedFiles, const FString& InInstallDirectory, 
+		const FString& InStagedFileDirectory, bool bInUseStageDirectory, FBuildPatchProgress* InBuildProgress, const TMap<FString, FString>& InPerFileSubdirectories)
 		: FileSystem(InFileSystem)
 		, ManifestSet(InManifestSet)
 		, TouchedFiles(MoveTemp(InTouchedFiles))
@@ -50,6 +54,7 @@ namespace BuildPatchServices
 		, BuildProgress(InBuildProgress)
 		, bIsPaused(false)
 		, bShouldAbort(false)
+		, PerFileSubdirectories(InPerFileSubdirectories)
 	{
 		BuildProgress->SetStateProgress(EBuildPatchState::SettingAttributes, 0.0f);
 	}
@@ -112,8 +117,8 @@ namespace BuildPatchServices
 				return InstallFilename;
 			}
 		}
-		InstallFilename = InstallDirectory / BuildFile;
-		return InstallFilename;
+
+		return FBuildPatchUtils::ResolveInstallationFileName(InstallDirectory, PerFileSubdirectories, BuildFile);
 	}
 
 	bool FFileAttribution::HasSameAttributes(const FFileManifest* NewFileManifest, const FFileManifest* OldFileManifest)
@@ -186,9 +191,10 @@ namespace BuildPatchServices
 		}
 	}
 
-	IFileAttribution* FFileAttributionFactory::Create(IFileSystem* FileSystem, IBuildManifestSet* ManifestSet, TSet<FString> TouchedFiles, const FString& InstallDirectory, const FString& StagedFileDirectory, FBuildPatchProgress* BuildProgress)
+	IFileAttribution* FFileAttributionFactory::Create(IFileSystem* FileSystem, IBuildManifestSet* ManifestSet, TSet<FString> TouchedFiles, 
+		const FString& InstallDirectory, const FString& StagedFileDirectory, FBuildPatchProgress* BuildProgress, const TMap<FString, FString>& PerFileSubdirectories)
 	{
 		check(BuildProgress != nullptr);
-		return new FFileAttribution(FileSystem, ManifestSet, MoveTemp(TouchedFiles), InstallDirectory, StagedFileDirectory, !StagedFileDirectory.IsEmpty(), BuildProgress);
+		return new FFileAttribution(FileSystem, ManifestSet, MoveTemp(TouchedFiles), InstallDirectory, StagedFileDirectory, !StagedFileDirectory.IsEmpty(), BuildProgress, PerFileSubdirectories);
 	}
 }
\ No newline at end of file
diff --git a/Engine/Source/Runtime/Online/BuildPatchServices/Private/Installer/FileAttribution.h b/Engine/Source/Runtime/Online/BuildPatchServices/Private/Installer/FileAttribution.h
--- a/Engine/Source/Runtime/Online/BuildPatchServices/Private/Installer/FileAttribution.h
+++ b/Engine/Source/Runtime/Online/BuildPatchServices/Private/Installer/FileAttribution.h
@@ -39,6 +39,7 @@ namespace BuildPatchServices
 		 * @param BuildProgress         The build progress class to update with progress information.
 		 * @return     Ref of an object that can be used to perform the operation.
 		 */
-		static IFileAttribution* Create(IFileSystem* FileSystem, IBuildManifestSet* ManifestSet, TSet<FString> TouchedFiles, const FString& InstallDirectory, const FString& StagedFileDirectory, FBuildPatchProgress* BuildProgress);
+		static IFileAttribution* Create(IFileSystem* FileSystem, IBuildManifestSet* ManifestSet, TSet<FString> TouchedFiles, const FString& InstallDirectory, 
+			const FString& StagedFileDirectory, FBuildPatchProgress* BuildProgress, const TMap<FString, FString>& PerFileSubdirectories);
 	};
 }
diff --git a/Engine/Source/Runtime/Online/BuildPatchServices/Private/Installer/InstallChunkSource.cpp b/Engine/Source/Runtime/Online/BuildPatchServices/Private/Installer/InstallChunkSource.cpp
--- a/Engine/Source/Runtime/Online/BuildPatchServices/Private/Installer/InstallChunkSource.cpp
+++ b/Engine/Source/Runtime/Online/BuildPatchServices/Private/Installer/InstallChunkSource.cpp
@@ -12,6 +12,7 @@
 #include "Misc/Paths.h"
 #include "Serialization/MemoryReader.h"
 #include "Tasks/Task.h"
+#include "BuildPatchUtil.h"
 
 DEFINE_LOG_CATEGORY_STATIC(LogInstallChunkSource, Log, All);
 
@@ -28,7 +29,7 @@ namespace BuildPatchServices
 	{
 	public:
 		FInstallChunkSource(IFileSystem* FileSystem, IInstallChunkSourceStat* InInstallChunkSourceStat, const TMultiMap<FString, FBuildPatchAppManifestRef>& InInstallationSources, 
-			const TSet<FGuid>& ChunksThatWillBeNeeded);
+			const TSet<FGuid>& ChunksThatWillBeNeeded, const TMap<FString, FString>& PerFileSubdirectories);
 		~FInstallChunkSource();
 		
 		virtual FRequestProcessFn CreateRequest(const FGuid& DataId, FMutableMemoryView DestinationBuffer, void* UserPtr, FChunkRequestCompleteDelegate CompleteFn) override;
@@ -119,6 +120,7 @@ namespace BuildPatchServices
 		};
 
 		UE::FMutex OpenedFileHandlesLock; 
+		// This is keyed off the installation path and the filename, but NOT the per file subdirectory.
 		TMap<FString, TUniquePtr<FOpenedFile>> OpenedFileHandles;
 
 		// The index (ChunkReferenceTracker->GetCurrentUsageIndex) at which our files will get deleted due to destructive install to make room for
@@ -136,12 +138,18 @@ namespace BuildPatchServices
 
 		// Contains the list of pieces to assemble chunks from, for each chunk.
 		TMap<FGuid, TArray<FChunkSourceDetails>> ChunkSources;
+
+		// Contains files that have a custom subdirectory to install into. This is accessed from multiple threads
+		// but should be static across execution and so requires no lock.
+		const TMap<FString, FString>& PerFileSubdirectories;
 	};
 
 	FInstallChunkSource::FInstallChunkSource(IFileSystem* InFileSystem, IInstallChunkSourceStat* InInstallChunkSourceStat, 
-		const TMultiMap<FString, FBuildPatchAppManifestRef>& InInstallationSources, const TSet<FGuid>& ChunksThatWillBeNeeded)
+		const TMultiMap<FString, FBuildPatchAppManifestRef>& InInstallationSources, const TSet<FGuid>& ChunksThatWillBeNeeded,
+		const TMap<FString, FString>& InPerFileSubdirectories)
 		: FileSystem(InFileSystem)
 		, InstallChunkSourceStat(InInstallChunkSourceStat)
+		, PerFileSubdirectories(InPerFileSubdirectories)
 	{
 		TRACE_CPUPROFILER_EVENT_SCOPE(InstallChunkSource_ctor);
 
@@ -184,6 +192,8 @@ namespace BuildPatchServices
 
 				bool bExistenceCheckPassed = false;
 
+				int32 ChunksForFile = 0;
+
 				uint32 ChunkPartSize = 0;
 				uint64 FileOffset = 0;
 				for (int32 ChunkPartIndex = 0; ChunkPartIndex < InstalledFileManifest->ChunkParts.Num(); ChunkPartIndex++, FileOffset += ChunkPartSize)
@@ -201,8 +211,10 @@ namespace BuildPatchServices
 					// 5 of them.
 					if (!bExistenceCheckPassed)
 					{
+						const FString DiskFileName = FBuildPatchUtils::ResolveInstallationFileName(InstallationDirectory, PerFileSubdirectories, InstalledFileName);
+
 						int64 SizeOnDisk = 0;
-						bool bGotFileSizeSuccessfully = InFileSystem->GetFileSize(*(InstallationDirectory / InstalledFileName), SizeOnDisk);
+						bool bGotFileSizeSuccessfully = InFileSystem->GetFileSize(*DiskFileName, SizeOnDisk);
 						uint64 ExpectedSize = InstalledManifest->GetFileSize(InstalledFileName);
 						if (bGotFileSizeSuccessfully &&
 							(uint64)SizeOnDisk == ExpectedSize)
@@ -240,6 +252,7 @@ namespace BuildPatchServices
 							NewInfo.SizeOfChunkPart = ChunkPart.Size;
 							ChunkAssemblyInfo.SourceDetails.Add(NewInfo);
 							ChunkAssemblyInfo.bComplete = true;
+							ChunksForFile++;
 							continue;
 						}
 
@@ -275,8 +288,11 @@ namespace BuildPatchServices
 						// Done with this chunk, free the structures space.
 						ChunkAssemblyInfo.bComplete = true;
 						ChunkStructures.Remove(ChunkPart.Guid);
-					}					
+						ChunksForFile++;
+					}
 				}
+
+				UE_CLOG(ChunksForFile, LogInstallChunkSource, Verbose, TEXT("%s produced %d chunks"), *InstalledFileName, ChunksForFile);
 			}
 
 			// Move the full chunks we have available over to our internal structures.
@@ -430,12 +446,13 @@ namespace BuildPatchServices
 				FOpenedFile* OpenedFile = nullptr;
 				{
 					OpenedFileHandlesLock.Lock();
-					TUniquePtr<FOpenedFile>* OpenedFilePtr = OpenedFileHandles.Find(FullFilename);
+					TUniquePtr<FOpenedFile>* OpenedFilePtr = OpenedFileHandles.Find(FullFilename); // this is purposely keyed WITHOUT the per file subdir.
 					if (OpenedFilePtr == nullptr)
 					{
 						TRACE_CPUPROFILER_EVENT_SCOPE(Install_OpenSource);
 
-						TUniquePtr<FArchive> NewReader = FileSystem->CreateFileReader(*FullFilename);
+						const FString DiskFileName = FBuildPatchUtils::ResolveInstallationFileName(FoundInstallDirectory, PerFileSubdirectories, FileChunkPart.InstalledFileManifest->Filename);
+						TUniquePtr<FArchive> NewReader = FileSystem->CreateFileReader(*DiskFileName);
 						if (NewReader.IsValid())
 						{
 							TUniquePtr<FOpenedFile> NewOpenedFile = MakeUnique<FOpenedFile>(MoveTemp(NewReader));
@@ -601,9 +618,9 @@ namespace BuildPatchServices
 	}
 
 	IConstructorInstallChunkSource* IConstructorInstallChunkSource::CreateInstallSource(IFileSystem* FileSystem, IInstallChunkSourceStat* InstallChunkSourceStat, 
-		const TMultiMap<FString, FBuildPatchAppManifestRef>& InstallationSources, const TSet<FGuid>& ChunksThatWillBeNeeded)
+		const TMultiMap<FString, FBuildPatchAppManifestRef>& InstallationSources, const TSet<FGuid>& ChunksThatWillBeNeeded, const TMap<FString, FString>& PerFileSubdirectories)
 	{
-		return new FInstallChunkSource(FileSystem, InstallChunkSourceStat, InstallationSources, ChunksThatWillBeNeeded);
+		return new FInstallChunkSource(FileSystem, InstallChunkSourceStat, InstallationSources, ChunksThatWillBeNeeded, PerFileSubdirectories);
 	}
 
 	const TCHAR* ToString(const IInstallChunkSourceStat::ELoadResult& LoadResult)
diff --git a/Engine/Source/Runtime/Online/BuildPatchServices/Private/Installer/InstallChunkSource.h b/Engine/Source/Runtime/Online/BuildPatchServices/Private/Installer/InstallChunkSource.h
--- a/Engine/Source/Runtime/Online/BuildPatchServices/Private/Installer/InstallChunkSource.h
+++ b/Engine/Source/Runtime/Online/BuildPatchServices/Private/Installer/InstallChunkSource.h
@@ -51,7 +51,8 @@ namespace BuildPatchServices
 			IFileSystem* FileSystem, 
 			IInstallChunkSourceStat* InstallChunkSourceStat, 
 			const TMultiMap<FString, FBuildPatchAppManifestRef>& InstallationSources,
-			const TSet<FGuid>& ChunksThatWillBeNeeded);
+			const TSet<FGuid>& ChunksThatWillBeNeeded,
+			const TMap<FString, FString>& InPerFileSubdirectories);
 	};
 
 	/**
diff --git a/Engine/Source/Runtime/Online/BuildPatchServices/Private/Installer/Verifier.cpp b/Engine/Source/Runtime/Online/BuildPatchServices/Private/Installer/Verifier.cpp
--- a/Engine/Source/Runtime/Online/BuildPatchServices/Private/Installer/Verifier.cpp
+++ b/Engine/Source/Runtime/Online/BuildPatchServices/Private/Installer/Verifier.cpp
@@ -41,7 +41,10 @@ namespace BuildPatchServices
 	class FVerifier : public IVerifier
 	{
 	public:
-		FVerifier(IFileSystem* FileSystem, IVerifierStat* InVerificationStat, EVerifyMode InVerifyMode, IBuildInstallerSharedContextPtr InSharedContext, IBuildManifestSet* InManifestSet, FString InVerifyDirectory, FString InStagedFileDirectory);
+		FVerifier(IFileSystem* FileSystem, IVerifierStat* InVerificationStat, EVerifyMode InVerifyMode, 
+			IBuildInstallerSharedContextPtr InSharedContext, IBuildManifestSet* InManifestSet, 
+			FString InVerifyDirectory, FString InStagedFileDirectory, const TMap<FString, FString>& InPerFileSubdirectories);
+
 		~FVerifier() {}
 
 		// IControllable interface begin.
@@ -66,6 +69,7 @@ namespace BuildPatchServices
 		IVerifierStat* const VerifierStat;
 		IBuildManifestSet* const ManifestSet;
 		IBuildInstallerSharedContextPtr SharedContext;
+		const TMap<FString, FString>& PerFileSubdirectories;
 
 		EVerifyMode VerifyMode;
 		TSet<FString> FilesToVerify;
@@ -92,13 +96,17 @@ namespace BuildPatchServices
 		void ProcessVerifyJobs(FSemaphore* ThreadDoneSem);
 	};
 
-	FVerifier::FVerifier(IFileSystem* InFileSystem, IVerifierStat* InVerificationStat, EVerifyMode InVerifyMode, IBuildInstallerSharedContextPtr InSharedContext, IBuildManifestSet* InManifestSet, FString InVerifyDirectory, FString InStagedFileDirectory)
+	FVerifier::FVerifier(IFileSystem* InFileSystem, IVerifierStat* InVerificationStat, EVerifyMode InVerifyMode, 
+		IBuildInstallerSharedContextPtr InSharedContext, IBuildManifestSet* InManifestSet, FString InVerifyDirectory, FString InStagedFileDirectory,
+		const TMap<FString, FString>& InPerFileSubdirectories
+		)
 		: VerifyDirectory(MoveTemp(InVerifyDirectory))
 		, StagedFileDirectory(MoveTemp(InStagedFileDirectory))
 		, FileSystem(InFileSystem)
 		, VerifierStat(InVerificationStat)
 		, ManifestSet(InManifestSet)
 		, SharedContext(InSharedContext)
+		, PerFileSubdirectories(InPerFileSubdirectories)
 		, VerifyMode(InVerifyMode)
 		, bIsPaused(false)
 		, bShouldAbort(false)
@@ -294,8 +302,8 @@ namespace BuildPatchServices
 				return FullFilePath;
 			}
 		}
-		FullFilePath = VerifyDirectory / BuildFile;
-		return FullFilePath;
+
+		return FBuildPatchUtils::ResolveInstallationFileName(VerifyDirectory, PerFileSubdirectories, BuildFile);
 	}
 
 	EVerifyResult FVerifier::VerfiyFileSha(TArray<uint8>& ReadBuffer, const FString& BuildFile, const FFileManifest& BuildFileManifest)
@@ -411,10 +419,13 @@ namespace BuildPatchServices
 		return VerifyResult;
 	}
 
-	IVerifier* FVerifierFactory::Create(IFileSystem* FileSystem, IVerifierStat* VerifierStat, EVerifyMode VerifyMode, IBuildInstallerSharedContextPtr SharedContext, IBuildManifestSet* ManifestSet, FString VerifyDirectory, FString StagedFileDirectory)
+	IVerifier* FVerifierFactory::Create(IFileSystem* FileSystem, IVerifierStat* VerifierStat, EVerifyMode VerifyMode, 
+		IBuildInstallerSharedContextPtr SharedContext, IBuildManifestSet* ManifestSet, FString VerifyDirectory, FString StagedFileDirectory,
+		const TMap<FString, FString>& PerFileSubdirectories
+		)
 	{
 		check(FileSystem != nullptr);
 		check(VerifierStat != nullptr);
-		return new FVerifier(FileSystem, VerifierStat, VerifyMode, SharedContext, ManifestSet, MoveTemp(VerifyDirectory), MoveTemp(StagedFileDirectory));
+		return new FVerifier(FileSystem, VerifierStat, VerifyMode, SharedContext, ManifestSet, MoveTemp(VerifyDirectory), MoveTemp(StagedFileDirectory), PerFileSubdirectories);
 	}
 }
diff --git a/Engine/Source/Runtime/Online/BuildPatchServices/Private/Installer/Verifier.h b/Engine/Source/Runtime/Online/BuildPatchServices/Private/Installer/Verifier.h
--- a/Engine/Source/Runtime/Online/BuildPatchServices/Private/Installer/Verifier.h
+++ b/Engine/Source/Runtime/Online/BuildPatchServices/Private/Installer/Verifier.h
@@ -79,19 +79,16 @@ namespace BuildPatchServices
 	{
 	public:
 		/**
-		 * Creates a verifier class that will verify a local directory structure against a given manifest, optionally
-		 * taking account of a staging directory where alternative files are used.
-		 * NOTE: This function is blocking and will not return until finished. Don't run on a UI thread.
-		 * @param FileSystem            The file system interface.
-		 * @param VerifierStat          Pointer to the class which will receive status updates.
-		 * @param VerifyMode            The verify mode to run.
-		 * @param InstallTags           The install tags, will be used when verifying all files.
-		 * @param Manifest              The manifest describing the build data.
-		 * @param VerifyDirectory       The directory to analyze.
-		 * @param StagedFileDirectory   A stage directory for updated files, ignored if empty string. If a file exists here, it will be checked instead of the one in VerifyDirectory.
-		 * @return     Ref of an object that can be used to perform the operation.
-		 */
-		static IVerifier* Create(IFileSystem* FileSystem, IVerifierStat* VerifierStat, EVerifyMode VerifyMode, IBuildInstallerSharedContextPtr SharedContext, IBuildManifestSet* ManifestSet, FString VerifyDirectory, FString StagedFileDirectory);
+		* Creates a verifier class that will verify a local directory structure against a given manifest, optionally
+		* taking account of a staging directory where alternative files are used.
+		* @param StagedFileDirectory   A stage directory for updated files, ignored if empty string. If a file exists here, it will be checked instead of the one in VerifyDirectory.
+		* @param PerFileSubdirectories An map of path adjustments to use keyed off the manifest filename, leave empty if unneeded.
+		* @return The verifier object. Delete when complete. See IVerifier class for usage.
+		*/
+		static IVerifier* Create(IFileSystem* FileSystem, IVerifierStat* VerifierStat, EVerifyMode VerifyMode, 
+			IBuildInstallerSharedContextPtr SharedContext, IBuildManifestSet* ManifestSet, FString VerifyDirectory, FString StagedFileDirectory,
+			const TMap<FString, FString>& PerFileSubdirectories
+			);
 	};
 
 	/**
diff --git a/Engine/Source/Runtime/Online/BuildPatchServices/Private/Tests/Unit/FileAttribution.spec.cpp b/Engine/Source/Runtime/Online/BuildPatchServices/Private/Tests/Unit/FileAttribution.spec.cpp
--- a/Engine/Source/Runtime/Online/BuildPatchServices/Private/Tests/Unit/FileAttribution.spec.cpp
+++ b/Engine/Source/Runtime/Online/BuildPatchServices/Private/Tests/Unit/FileAttribution.spec.cpp
@@ -358,13 +358,15 @@ void FFileAttributionSpec::MakeUnit()
 		MockFileSystem->FileAttributes.Add(InstallDirectory / File, EAttributeFlags::Exists);
 	}
 	ManifestSet.Reset(FBuildManifestSetFactory::Create({ BuildPatchServices::FInstallerAction::MakeInstallOrUpdate(MockOldManifest, MockNewManifest.ToSharedRef()) }));
+	TMap<FString, FString> EmptyPerFileSubdirectories;
 	FileAttribution.Reset(BuildPatchServices::FFileAttributionFactory::Create(
 		MockFileSystem.Get(),
 		ManifestSet.Get(),
 		TouchedFiles,
 		InstallDirectory,
 		StagedFileDirectory,
-		MockBuildProgress.Get()));
+		MockBuildProgress.Get(),
+		EmptyPerFileSubdirectories));
 }
 
 TFuture<void> FFileAttributionSpec::PauseFor(float Seconds)
diff --git a/Engine/Source/Runtime/Online/BuildPatchServices/Private/Tests/Unit/InstallChunkSource.spec.cpp b/Engine/Source/Runtime/Online/BuildPatchServices/Private/Tests/Unit/InstallChunkSource.spec.cpp
--- a/Engine/Source/Runtime/Online/BuildPatchServices/Private/Tests/Unit/InstallChunkSource.spec.cpp
+++ b/Engine/Source/Runtime/Online/BuildPatchServices/Private/Tests/Unit/InstallChunkSource.spec.cpp
@@ -329,11 +329,13 @@ void FInstallChunkSourceSpec::MakeUnit()
 	using namespace BuildPatchServices;
 	TSet<FGuid> ChunksThatWillBeNeeded;
 	ManifestSet->GetReferencedChunks(ChunksThatWillBeNeeded);
+	static TMap<FString, FString> EmptyPerFileSubdirectories;
 	InstallChunkSource.Reset(IConstructorInstallChunkSource::CreateInstallSource(
 		FakeFileSystem.Get(),
 		MockInstallChunkSourceStat.Get(),
 		InstallationSources,
-		ChunksThatWillBeNeeded));
+		ChunksThatWillBeNeeded,
+		EmptyPerFileSubdirectories));
 }
 
 void FInstallChunkSourceSpec::InventUsableChunkData()
diff --git a/Engine/Source/Runtime/Online/BuildPatchServices/Private/Tests/Unit/Verifier.spec.cpp b/Engine/Source/Runtime/Online/BuildPatchServices/Private/Tests/Unit/Verifier.spec.cpp
--- a/Engine/Source/Runtime/Online/BuildPatchServices/Private/Tests/Unit/Verifier.spec.cpp
+++ b/Engine/Source/Runtime/Online/BuildPatchServices/Private/Tests/Unit/Verifier.spec.cpp
@@ -552,6 +552,7 @@ void FVerifierSpec::MakeUnit(BuildPatchServices::EVerifyMode Mode)
 {
 	using namespace BuildPatchServices;
 	// TODO: Verifier behavioural change needs unit test update.
+	TMap<FString, FString> PerFileSubdirectories;
 	Verifier.Reset(FVerifierFactory::Create(
 		FakeFileSystem.Get(),
 		MockVerificationStat.Get(),
@@ -560,7 +561,8 @@ void FVerifierSpec::MakeUnit(BuildPatchServices::EVerifyMode Mode)
 		/*TouchedFiles,*/
 		ManifestSet.Get(),
 		VerifyDirectory,
-		StagedFileDirectory));
+		StagedFileDirectory,
+		PerFileSubdirectories));
 
 	if (EVerifyMode::FileSizeCheckTouchedFiles == Mode || EVerifyMode::ShaVerifyTouchedFiles == Mode)
 	{
diff --git a/Engine/Source/Runtime/Online/BuildPatchServices/Public/Interfaces/IBuildInstaller.h b/Engine/Source/Runtime/Online/BuildPatchServices/Public/Interfaces/IBuildInstaller.h
--- a/Engine/Source/Runtime/Online/BuildPatchServices/Public/Interfaces/IBuildInstaller.h
+++ b/Engine/Source/Runtime/Online/BuildPatchServices/Public/Interfaces/IBuildInstaller.h
@@ -473,6 +473,14 @@ public:
 	 * the installer does not touch this other than to delete it on destruction.
 	 */
 	virtual TMap<FString, TArray64<uint8>>& GetFilesInstalledToMemory() = 0; 
+
+	/**
+	* [Experimental]
+	* Pass the installation folder on a per file basis. This path will be appended to the installation directory passed
+	* in the FBuildInstallerConfiguration::InstallDirectory member. The old version is used for patch source data, the new
+	* for everything else.
+	*/
+	virtual void SetPerFileSubdirectories(TMap<FString, FString>&& InOldPerFileSubdirectories, TMap<FString, FString>&& InNewPerFileSubdirectories) = 0;
 };
 
 static_assert((uint32)EBuildPatchInstallError::NumInstallErrors == 12, "Please add support for the extra values to the Lex functions below.");
