PR #14312: Linux: Fix SDL voice capture

#rnx

#Synced-CL 50595024

diff --git a/Engine/Source/Runtime/Online/Voice/Private/Linux/VoiceModuleLinux.cpp b/Engine/Source/Runtime/Online/Voice/Private/Linux/VoiceModuleLinux.cpp
--- a/Engine/Source/Runtime/Online/Voice/Private/Linux/VoiceModuleLinux.cpp
+++ b/Engine/Source/Runtime/Online/Voice/Private/Linux/VoiceModuleLinux.cpp
@@ -115,8 +115,11 @@ static void SDLCALL VoiceCallbackWrapper(void *UserData, SDL_AudioStream *AudioS
 		Uint8 *Data = SDL_stack_alloc(Uint8, AdditionalAmount);
 		if(Data)
 		{
-			FVoiceCaptureSDL::AudioCallback(UserData, Data, AdditionalAmount);
-			SDL_PutAudioStreamData(AudioStream, Data, AdditionalAmount);
+			int BytesRead = SDL_GetAudioStreamData(AudioStream, Data, AdditionalAmount);
+			if (BytesRead > 0)
+			{
+				FVoiceCaptureSDL::AudioCallback(UserData, Data, BytesRead);
+			}
 			SDL_stack_free(Data);
 		}
 	}
@@ -206,8 +209,8 @@ bool FVoiceCaptureSDL::Start()
 	AudioBufferWritePos = 0;
 	AudioBufferAvailableData = 0;
 
-	SDL_ResumeAudioDevice(SDL_GetAudioStreamDevice(AudioStream));
-	if (SDL_AudioDevicePaused(SDL_GetAudioStreamDevice(AudioStream)) == false)
+	SDL_ResumeAudioStreamDevice(AudioStream);
+	if (SDL_AudioStreamDevicePaused(AudioStream) == false)
 	{
 		VoiceCaptureState = EVoiceCaptureState::Ok;
 		return true;
@@ -222,7 +225,7 @@ bool FVoiceCaptureSDL::Start()
 void FVoiceCaptureSDL::Stop()
 {
 	check(IsCapturing());
-	SDL_PauseAudioDevice(SDL_GetAudioStreamDevice(AudioStream));
+	SDL_PauseAudioStreamDevice(AudioStream);
 
 	VoiceCaptureState = EVoiceCaptureState::NotCapturing;
 	AudioBufferWritePos = 0;
@@ -318,7 +321,7 @@ bool InitVoiceCapture()
 	// AudioMixerSDL may interfere with this
 	if (!SDL_WasInit(SDL_INIT_AUDIO))
 	{
-		GVoiceCaptureNeedsToShutdownSDLAudio = SDL_InitSubSystem(SDL_INIT_AUDIO) == 0;
+		GVoiceCaptureNeedsToShutdownSDLAudio = SDL_InitSubSystem(SDL_INIT_AUDIO);
 		return GVoiceCaptureNeedsToShutdownSDLAudio;
 	}
 
