BuildPatchFileConstructor harvest chunks fix for subdirectory mappings

#Synced-CL 50783704

diff --git a/Engine/Source/Runtime/Online/BuildPatchServices/Private/BuildPatchFileConstructor.cpp b/Engine/Source/Runtime/Online/BuildPatchServices/Private/BuildPatchFileConstructor.cpp
--- a/Engine/Source/Runtime/Online/BuildPatchServices/Private/BuildPatchFileConstructor.cpp
+++ b/Engine/Source/Runtime/Online/BuildPatchServices/Private/BuildPatchFileConstructor.cpp
@@ -1783,16 +1783,16 @@ void FBuildPatchFileConstructor::ChunkHarvestCompletedFn(const FGuid& Guid, bool
 	}
 }
 
-bool FBuildPatchFileConstructor::HarvestChunksForCompletedFile(const FString& CompletedFullPathFileName)
+bool FBuildPatchFileConstructor::HarvestChunksForCompletedFile(const FString& InstallationDirectory, const FString& FilenameInManifest)
 {
-	UE_LOG(LogBuildPatchServices, Verbose, TEXT("Harvesting source: %s"), *CompletedFullPathFileName);
+	UE_LOG(LogBuildPatchServices, Verbose, TEXT("Harvesting source: %s / %s"), *InstallationDirectory, *FilenameInManifest);
 
 	TRACE_CPUPROFILER_EVENT_SCOPE(Harvest);
 	// We need to grab any chunks from install sources that are no longer available.
 	// Anything that's already been loaded is already placed into the memory store appropriately,
 	// but anything that _hasn't_ needs to be pulled out.	
 	TSet<FGuid> FileChunks;
-	InstallSource->GetChunksForFile(CompletedFullPathFileName, FileChunks);
+	InstallSource->GetChunksForFile(InstallationDirectory, FilenameInManifest, FileChunks);
 
 	struct FNeededChunk
 	{
@@ -3028,7 +3028,7 @@ void FBuildPatchFileConstructor::CompleteConstructedFile(FFileConstructionState&
 		FPaths::CollapseRelativeDirectories(FileToDelete);
 		if (FileSystem->FileExists(*FileToDelete))
 		{
-			if (InstallSource && !HarvestChunksForCompletedFile(FileToDelete))
+			if (InstallSource && !HarvestChunksForCompletedFile(Configuration.InstallDirectory, CurrentFile.BuildFilename))
 			{
 				CurrentFile.bSuccess = false;
 				CurrentFile.ConstructionError = EConstructionError::InternalConsistencyError;
diff --git a/Engine/Source/Runtime/Online/BuildPatchServices/Private/BuildPatchFileConstructor.h b/Engine/Source/Runtime/Online/BuildPatchServices/Private/BuildPatchFileConstructor.h
--- a/Engine/Source/Runtime/Online/BuildPatchServices/Private/BuildPatchFileConstructor.h
+++ b/Engine/Source/Runtime/Online/BuildPatchServices/Private/BuildPatchFileConstructor.h
@@ -333,7 +333,7 @@ namespace BuildPatchServices
 		void ResumeFile(FFileConstructionState& FileToResume);
 		void OpenFileToConstruct(FFileConstructionState& CurrentFile);
 		bool HandleInitialDiskSizeCheck(const FFileManifest& CurrentFileManifest, int64 BytesInToCurrentFile);
-		bool HarvestChunksForCompletedFile(const FString& CompletedBuildFileName);
+		bool HarvestChunksForCompletedFile(const FString& InstallationDirectory, const FString& FilenameInManifest);
 		void CompleteConstructedFile(FFileConstructionState& CurrentFile);
 		void InitFile(FFileConstructionState& CurrentFile, const FResumeData& ResumeData);
 
diff --git a/Engine/Source/Runtime/Online/BuildPatchServices/Private/Installer/InstallChunkSource.cpp b/Engine/Source/Runtime/Online/BuildPatchServices/Private/Installer/InstallChunkSource.cpp
--- a/Engine/Source/Runtime/Online/BuildPatchServices/Private/Installer/InstallChunkSource.cpp
+++ b/Engine/Source/Runtime/Online/BuildPatchServices/Private/Installer/InstallChunkSource.cpp
@@ -69,7 +69,7 @@ namespace BuildPatchServices
 			FileRetirementPositions = MoveTemp(InFileRetirementPositions);
 		}
 
-		virtual void GetChunksForFile(const FString& FilePath, TSet<FGuid>& OutChunks) const override;
+		virtual void GetChunksForFile(const FString& InstallationDirectory, const FString& FilenameInManifest, TSet<FGuid>& OutChunks) const override;
 
 		virtual void EnumerateFilesForChunk(const FGuid& DataId, TUniqueFunction<void(const FString& NormalizedInstallDirectory, const FString& NormalizedFileName)>&& Callback) const override
 		{
@@ -335,16 +335,18 @@ namespace BuildPatchServices
 		
 	}
 
-	void FInstallChunkSource::GetChunksForFile(const FString& FilePath, TSet<FGuid>& OutChunks) const
+	void FInstallChunkSource::GetChunksForFile(const FString& InstallationDirectory, const FString& FileInManifest, TSet<FGuid>& OutChunks) const
 	{
 		const FFileManifest* FileManifest = nullptr;
+		FString NormalisedInstallDirectory = InstallationDirectory;
+		FPaths::NormalizeDirectoryName(NormalisedInstallDirectory);
 		for (const TPair<FString, FBuildPatchAppManifestRef>& Pair : InstallationSources)
 		{
-			if (FilePath.StartsWith(Pair.Key))
+			FString NormaliasedInstallSourcePath = Pair.Key;
+			FPaths::NormalizeDirectoryName(NormaliasedInstallSourcePath);
+			if (FPaths::IsSamePath(NormalisedInstallDirectory, NormaliasedInstallSourcePath))
 			{
-				FString BuildRelativeFilePath = FilePath;
-				FPaths::MakePathRelativeTo(BuildRelativeFilePath, *(Pair.Key / TEXT("")));
-				FileManifest = Pair.Value->GetFileManifest(BuildRelativeFilePath);
+				FileManifest = Pair.Value->GetFileManifest(FileInManifest);
 				break;
 			}
 		}
diff --git a/Engine/Source/Runtime/Online/BuildPatchServices/Private/Installer/InstallChunkSource.h b/Engine/Source/Runtime/Online/BuildPatchServices/Private/Installer/InstallChunkSource.h
--- a/Engine/Source/Runtime/Online/BuildPatchServices/Private/Installer/InstallChunkSource.h
+++ b/Engine/Source/Runtime/Online/BuildPatchServices/Private/Installer/InstallChunkSource.h
@@ -34,7 +34,7 @@ namespace BuildPatchServices
 
 		// Return the list of chunks for a file so that the constructor can harvest them before the file
 		// is deleted
-		virtual void GetChunksForFile(const FString& FilePath, TSet<FGuid>& OutChunks) const = 0;
+		virtual void GetChunksForFile(const FString& InstallationDirectory, const FString& FilenameInManifest, TSet<FGuid>& OutChunks) const = 0;
 
 		// Call the lambda on each file the chunk needs to load bits from.
 		virtual void EnumerateFilesForChunk(const FGuid& DataId, TUniqueFunction<void(const FString& NormalizedInstallDirectory, const FString& NormalizedFileName)>&& Callback) const = 0;
