Fix for open file handle tracking when patching with BPS using multiple subdirectories
- Store the resolved file names in the OpenedFileHandles container

Fix for files not being deleted when using subdirectory mappings during patching
- Files that exist in the old manifest but not in the new one should have their paths resolved with the old subdirectory mapping

#Synced-CL 50859770

diff --git a/Engine/Source/Runtime/Online/BuildPatchServices/Private/BuildPatchInstaller.cpp b/Engine/Source/Runtime/Online/BuildPatchServices/Private/BuildPatchInstaller.cpp
--- a/Engine/Source/Runtime/Online/BuildPatchServices/Private/BuildPatchInstaller.cpp
+++ b/Engine/Source/Runtime/Online/BuildPatchServices/Private/BuildPatchInstaller.cpp
@@ -1688,12 +1688,12 @@ namespace BuildPatchServices
 		}
 	}
 
-	void FBuildPatchInstaller::FilterToExistingFiles(const FString& RootDirectory, TSet<FString>& Files, bool bResolveInstallationFileName)
+	void FBuildPatchInstaller::FilterToExistingFiles(const FString& RootDirectory, const TMap<FString, FString>& PerFileSubdirectories, TSet<FString>& Files, bool bResolveInstallationFileName)
 	{
 		for (auto FilesIt = Files.CreateIterator(); FilesIt; ++FilesIt)
 		{
 			const FString DiskFileName = bResolveInstallationFileName ? 
-				FBuildPatchUtils::ResolveInstallationFileName(Configuration.InstallDirectory, NewPerFileSubdirectories, *FilesIt) 
+				FBuildPatchUtils::ResolveInstallationFileName(Configuration.InstallDirectory, PerFileSubdirectories, *FilesIt)
 			: (RootDirectory / *FilesIt);
 			if (!FileSystem->FileExists(*DiskFileName))
 			{
@@ -1780,13 +1780,14 @@ namespace BuildPatchServices
 				// Get list of all expected files
 				TSet<FString> FilesToRelocate;
 				ManifestSet->GetExpectedFiles(FilesToRelocate);
-				FilterToExistingFiles(InstallStagingDir, FilesToRelocate, false);
+				FilterToExistingFiles(InstallStagingDir, NewPerFileSubdirectories, FilesToRelocate, false);
 				FilesToRelocate.Sort(TLess<FString>());
 
 				// First handle files that should be removed for patching
 				TSet<FString> FilesToRemove;
 				ManifestSet->GetRemovableFiles(FilesToRemove);
-				FilterToExistingFiles(Configuration.InstallDirectory, FilesToRemove, true);
+				// Note, this is using the OLD subdirectory mapping since we are concerned with files that may not exist in the new manifest
+				FilterToExistingFiles(Configuration.InstallDirectory, OldPerFileSubdirectories, FilesToRemove, true);
 				FilesToRemove.Sort(TLess<FString>());
 
 				// Counters for progress tracking.
@@ -1807,8 +1808,8 @@ namespace BuildPatchServices
 				// Perform all of the removals
 				for (auto FileToRemove = FilesToRemove.CreateConstIterator(); FileToRemove && !InstallerError->HasError(); ++FileToRemove)
 				{
-					BackupFileIfNecessary(*FileToRemove);
-					const FString FullFilename = FBuildPatchUtils::ResolveInstallationFileName(Configuration.InstallDirectory, NewPerFileSubdirectories, *FileToRemove);
+					BackupFileIfNecessary(*FileToRemove, OldPerFileSubdirectories);
+					const FString FullFilename = FBuildPatchUtils::ResolveInstallationFileName(Configuration.InstallDirectory, OldPerFileSubdirectories, *FileToRemove);
 					uint32 ErrorCode = 0;
 					bool bDeleteSuccess = RemoveFileWithRetries(FullFilename, ErrorCode);
 					if (bDeleteSuccess)
@@ -1833,7 +1834,7 @@ namespace BuildPatchServices
 					const bool bOldFileExists = FileSystem->FileExists(*DestFilename);
 					if (bOldFileExists)
 					{
-						BackupFileIfNecessary(*FileToRelocate);
+						BackupFileIfNecessary(*FileToRelocate, NewPerFileSubdirectories);
 						uint32 ErrorCode = 0;
 						bool bDeleteSuccess = RemoveFileWithRetries(DestFilename, ErrorCode);
 						if (!bDeleteSuccess)
@@ -1949,7 +1950,7 @@ namespace BuildPatchServices
 		{
 			for (const FString& CorruptFile : CorruptFiles)
 			{
-				BackupFileIfNecessary(CorruptFile, true);
+				BackupFileIfNecessary(CorruptFile, NewPerFileSubdirectories, true);
 				FString StagedFile = InstallStagingDir / CorruptFile;
 				if (FileSystem->FileExists(*StagedFile))
 				{
@@ -1977,9 +1978,9 @@ namespace BuildPatchServices
 		return bVerifySuccessful;
 	}
 
-	bool FBuildPatchInstaller::BackupFileIfNecessary(const FString& Filename, bool bDiscoveredByVerification /*= false */)
+	bool FBuildPatchInstaller::BackupFileIfNecessary(const FString& Filename, const TMap<FString, FString>& PerFileSubdirectories, bool bDiscoveredByVerification /*= false */)
 	{
-		const FString InstalledFilename = FBuildPatchUtils::ResolveInstallationFileName(Configuration.InstallDirectory, NewPerFileSubdirectories, Filename);
+		const FString InstalledFilename = FBuildPatchUtils::ResolveInstallationFileName(Configuration.InstallDirectory, PerFileSubdirectories, Filename);
 		const FString BackupFilename = Configuration.BackupDirectory / Filename;
 		const bool bBackupOriginals = !Configuration.BackupDirectory.IsEmpty();
 		// Skip if not doing backups
diff --git a/Engine/Source/Runtime/Online/BuildPatchServices/Private/BuildPatchInstaller.h b/Engine/Source/Runtime/Online/BuildPatchServices/Private/BuildPatchInstaller.h
--- a/Engine/Source/Runtime/Online/BuildPatchServices/Private/BuildPatchInstaller.h
+++ b/Engine/Source/Runtime/Online/BuildPatchServices/Private/BuildPatchInstaller.h
@@ -389,10 +389,11 @@ namespace BuildPatchServices
 		/**
 		 * Checks a particular file in the install directory to see if it needs backing up, and does so if necessary.
 		 * @param Filename                      The filename to check, which should match a filename in a manifest.
+		 * @param PerFileSubdirectories         Per-file subdirectory mapping to use when resolving paths
 		 * @param bDiscoveredByVerification     Optional, whether the file was detected changed already by verification stage. Default: false.
 		 * @return    Returns true if there were no errors.
 		 */
-		bool BackupFileIfNecessary(const FString& Filename, bool bDiscoveredByVerification = false);
+		bool BackupFileIfNecessary(const FString& Filename, const TMap<FString, FString>& PerFileSubdirectories, bool bDiscoveredByVerification = false);
 
 		/**
 		 * Delete empty directories from an installation.
@@ -402,11 +403,12 @@ namespace BuildPatchServices
 
 		/**
 		 * Remove from the set any files that do not exist.
-		 * @param RootDirectory     Root Directory for search.
-		 * @param Files             Set of fles to filter.
+		 * @param RootDirectory             Root Directory for search.
+		 * @param PerFileSubdirectories     Per-file subdirectory mapping to use when resolving paths
+		 * @param Files                     Set of files to filter.
 		 * @param bResolveInstallationFileName If true, instead of appending the RootDirectory / FilesName, call ResolveInstallationFileName on the file.
 		 */
-		void FilterToExistingFiles(const FString& RootDirectory, TSet<FString>& Files, bool bResolveInstallationFileName);
+		void FilterToExistingFiles(const FString& RootDirectory, const TMap<FString, FString>& PerFileSubdirectories, TSet<FString>& Files, bool bResolveInstallationFileName);
 
 		/**
 		 * Attempt to remove a file, using configured retries.
diff --git a/Engine/Source/Runtime/Online/BuildPatchServices/Private/Installer/InstallChunkSource.cpp b/Engine/Source/Runtime/Online/BuildPatchServices/Private/Installer/InstallChunkSource.cpp
--- a/Engine/Source/Runtime/Online/BuildPatchServices/Private/Installer/InstallChunkSource.cpp
+++ b/Engine/Source/Runtime/Online/BuildPatchServices/Private/Installer/InstallChunkSource.cpp
@@ -120,7 +120,7 @@ namespace BuildPatchServices
 		};
 
 		UE::FMutex OpenedFileHandlesLock; 
-		// This is keyed off the installation path and the filename, but NOT the per file subdirectory.
+		// This is keyed off the resolved file name (including subdirectory mapping)
 		TMap<FString, TUniquePtr<FOpenedFile>> OpenedFileHandles;
 
 		// The index (ChunkReferenceTracker->GetCurrentUsageIndex) at which our files will get deleted due to destructive install to make room for
@@ -438,29 +438,25 @@ namespace BuildPatchServices
 					break;
 				}
 
-				FString FullFilename = FoundInstallDirectory / FileChunkPart.InstalledFileManifest->Filename;
-
-
 				// We use the internal FArchive pointer so that we don't have to hold the file handle
 				// lock over the read - the uniqueptr might move around, but the managed pointer will not.
 				// We do know that our specific file won't get deleted until we are done, but it's possible that
 				// multiple threads are here looking for the same file.
 				FOpenedFile* OpenedFile = nullptr;
 				{
+					const FString DiskFileName = FBuildPatchUtils::ResolveInstallationFileName(FoundInstallDirectory, PerFileSubdirectories, FileChunkPart.InstalledFileManifest->Filename);
 					OpenedFileHandlesLock.Lock();
-					TUniquePtr<FOpenedFile>* OpenedFilePtr = OpenedFileHandles.Find(FullFilename); // this is purposely keyed WITHOUT the per file subdir.
+					TUniquePtr<FOpenedFile>* OpenedFilePtr = OpenedFileHandles.Find(DiskFileName);
 					if (OpenedFilePtr == nullptr)
 					{
-						TRACE_CPUPROFILER_EVENT_SCOPE(Install_OpenSource);
-
-						const FString DiskFileName = FBuildPatchUtils::ResolveInstallationFileName(FoundInstallDirectory, PerFileSubdirectories, FileChunkPart.InstalledFileManifest->Filename);
+						TRACE_CPUPROFILER_EVENT_SCOPE(Install_OpenSource);					
 						TUniquePtr<FArchive> NewReader = FileSystem->CreateFileReader(*DiskFileName);
 						if (NewReader.IsValid())
 						{
 							TUniquePtr<FOpenedFile> NewOpenedFile = MakeUnique<FOpenedFile>(MoveTemp(NewReader));
 
-							OpenedFileHandles.Add(FullFilename, MoveTemp(NewOpenedFile));
-							OpenedFilePtr = OpenedFileHandles.Find(FullFilename);
+							OpenedFileHandles.Add(DiskFileName, MoveTemp(NewOpenedFile));
+							OpenedFilePtr = OpenedFileHandles.Find(DiskFileName);
 						}
 					}
 					if (OpenedFilePtr)
