Initial Implementation for Thin Client bandwidth distributor

#rnx
#redirectsaresafe

#Synced-CL 50897393

#Synced-CL 50897393

diff --git a/Engine/Source/Runtime/Experimental/IoStore/OnDemand/IoStoreOnDemand.Build.cs b/Engine/Source/Runtime/Experimental/IoStore/OnDemand/IoStoreOnDemand.Build.cs
--- a/Engine/Source/Runtime/Experimental/IoStore/OnDemand/IoStoreOnDemand.Build.cs
+++ b/Engine/Source/Runtime/Experimental/IoStore/OnDemand/IoStoreOnDemand.Build.cs
@@ -9,6 +9,7 @@ public class IoStoreOnDemand : ModuleRules
 		PublicDependencyModuleNames.Add("Core");
 		PublicDependencyModuleNames.Add("IoStoreOnDemandCore");
 		PublicDependencyModuleNames.Add("TraceLog");
+		PublicDependencyModuleNames.Add("DownlinkBandwidthManager");
 		PrivateDependencyModuleNames.AddRange(
 			new string[] {
 				"Analytics",
diff --git a/Engine/Source/Runtime/Experimental/IoStore/OnDemand/Private/OnDemandHttpIoDispacher.cpp b/Engine/Source/Runtime/Experimental/IoStore/OnDemand/Private/OnDemandHttpIoDispacher.cpp
--- a/Engine/Source/Runtime/Experimental/IoStore/OnDemand/Private/OnDemandHttpIoDispacher.cpp
+++ b/Engine/Source/Runtime/Experimental/IoStore/OnDemand/Private/OnDemandHttpIoDispacher.cpp
@@ -7,6 +7,7 @@
 #include "AtomicFlags.h"
 #include "Containers/AnsiString.h"
 #include "CVarUtilities.h"
+#include "DownlinkBandwidthManager.h"
 #include "Experimental/ConcurrentLinearAllocator.h"
 #include "HAL/Platform.h"
 #include "HAL/Runnable.h"
@@ -1186,6 +1187,7 @@ private:
 	void							RecreateHttpClientIfNeeded();
 
 	bool							TryReadFromCache(FHttpRequest* Request);
+	int32							GetHttpRateLimitKiBPerSecond() const;
 	bool							ProcessHttpRequests();
 	void							CancelRemainingHttpRequests();
 	FHttpRequest*					DequeueHttpRequests();
@@ -1218,6 +1220,8 @@ private:
 	std::atomic_bool						bRecreateHttpClient = false;
 	std::atomic_bool						bStopRequested = false;
 	std::atomic_int32_t						InflightRequestCount{0};
+	FThreadSafeCounter						DLBWHttpRateLimitKiBPerSecond = 0;
+	UE::BandwidthManager::FServiceHandleID	BandwidthManagerServiceHandle = UE_INVALID_BANDWIDTH_SERVICE_HANDLE;
 
 #if UE_ALLOW_HTTP_PAUSE
 	FConsoleCommandPtr						PauseCommand;
@@ -1299,6 +1303,22 @@ bool FHttpDispatcher::Init()
 			WakeUp->Trigger();
 		});
 
+	if (FDownlinkBandwidthManager::IsEnabled())
+	{
+		FDownlinkBandwidthManager& DBManager = FDownlinkBandwidthManager::Get();
+		UE::BandwidthManager::FManagedServiceConfigurableData BandwidthConfig;
+		BandwidthConfig.ComparisonValue = MakeShared<int32>(GIasHttpRateLimitKiBPerSecond);
+		BandwidthConfig.IdleGameplayPriority = UE::BandwidthManager::EBandwidthPriority::HighPriority;
+		BandwidthConfig.InGameplayPriority = UE::BandwidthManager::EBandwidthPriority::LowPriority;
+
+		BandwidthManagerServiceHandle = DBManager.RegisterMonitoredService(TEXT("HttpDispatcher"), BandwidthConfig, true,
+			[this](int32 AllottedBandwidthKbps) 
+			{
+				// 1 Kilobyte is equal to (1000 / 1024) kibibytes
+				DLBWHttpRateLimitKiBPerSecond.Set(static_cast<int32>((float)AllottedBandwidthKbps * (1000.0f / 1024.0f)));
+			});
+	}
+
 	return TryCreateHttpClient();
 }
 
@@ -1306,6 +1326,12 @@ uint32 FHttpDispatcher::Run()
 {
 	check(HttpClient.IsValid());
 
+	if (FDownlinkBandwidthManager::IsEnabled())
+	{
+		FDownlinkBandwidthManager& DBManager = FDownlinkBandwidthManager::Get();
+		DBManager.SetActivationForService(BandwidthManagerServiceHandle, true);
+	}
+
 	for (;;)
 	{
 		UpdateThreadPriorityIfNeeded();
@@ -1353,11 +1379,21 @@ void FHttpDispatcher::Stop()
 	{
 		WakeUp->Trigger();
 	}
+
+	if (FDownlinkBandwidthManager::IsEnabled())
+	{
+		FDownlinkBandwidthManager& DBManager = FDownlinkBandwidthManager::Get();
+		DBManager.SetActivationForService(BandwidthManagerServiceHandle, false);
+	}
 }
 
 void FHttpDispatcher::Exit()
 {
-	
+	if (FDownlinkBandwidthManager::IsEnabled())
+	{
+		FDownlinkBandwidthManager& DBManager = FDownlinkBandwidthManager::Get();
+		DBManager.UnregisterMonitoredService(BandwidthManagerServiceHandle);
+	}
 }
 
 void FHttpDispatcher::Tick()
@@ -1736,6 +1772,11 @@ bool FHttpDispatcher::TryReadFromCache(FHttpRequest* Request)
 	return false;
 }
 
+int32 FHttpDispatcher::GetHttpRateLimitKiBPerSecond() const
+{
+	return (FDownlinkBandwidthManager::IsEnabled()) ? DLBWHttpRateLimitKiBPerSecond.GetValue() : GIasHttpRateLimitKiBPerSecond;
+}
+
 bool FHttpDispatcher::ProcessHttpRequests()
 {
 	TRACE_CPUPROFILER_EVENT_SCOPE(FHttpDispatcher::ProcessHttpRequests);
@@ -1844,7 +1885,7 @@ bool FHttpDispatcher::ProcessHttpRequests()
 
 				while (NumConcurrentRequests >= MaxConcurrentRequests) //-V654
 				{
-					HttpClient->Tick(MAX_uint32, GIasHttpRateLimitKiBPerSecond); 
+					HttpClient->Tick(MAX_uint32, GetHttpRateLimitKiBPerSecond());
 				}
 			}
 
@@ -1857,7 +1898,7 @@ bool FHttpDispatcher::ProcessHttpRequests()
 
 			// Keep processing pending connections until all requests are completed or a new one is issued
 			TRACE_CPUPROFILER_EVENT_SCOPE(FHttpDispatcher::TickHttp);
-			while (NextHttpRequest == nullptr && HttpClient->Tick(GIasHttpPollTimeoutMs, GIasHttpRateLimitKiBPerSecond))
+			while (NextHttpRequest == nullptr && HttpClient->Tick(GIasHttpPollTimeoutMs, GetHttpRateLimitKiBPerSecond()))
 			{
 				NextHttpRequest = DequeueHttpRequests();
 			}
diff --git a/Engine/Source/Runtime/Online/Experimental/DownlinkBandwidthManager/DownlinkBandwidthManager.Build.cs b/Engine/Source/Runtime/Online/Experimental/DownlinkBandwidthManager/DownlinkBandwidthManager.Build.cs
new file mode 100644
--- /dev/null
+++ b/Engine/Source/Runtime/Online/Experimental/DownlinkBandwidthManager/DownlinkBandwidthManager.Build.cs
@@ -0,0 +1,20 @@
+// Copyright Epic Games, Inc. All Rights Reserved.
+
+using UnrealBuildTool;
+
+public class DownlinkBandwidthManager : ModuleRules
+{
+	public DownlinkBandwidthManager(ReadOnlyTargetRules Target) : base(Target)
+	{
+		PCHUsage = PCHUsageMode.UseExplicitOrSharedPCHs;
+		CppCompileWarningSettings.UnsafeTypeCastWarningLevel = WarningLevel.Error;
+
+		PublicDependencyModuleNames.AddRange(new string[] {
+			"Core",
+		});
+
+		PrivateDependencyModuleNames.AddRange(new string[] {
+			"CoreUObject",
+		});
+	}
+}
diff --git a/Engine/Source/Runtime/Online/Experimental/DownlinkBandwidthManager/Private/DownlinkBandwidthManager.cpp b/Engine/Source/Runtime/Online/Experimental/DownlinkBandwidthManager/Private/DownlinkBandwidthManager.cpp
new file mode 100644
--- /dev/null
+++ b/Engine/Source/Runtime/Online/Experimental/DownlinkBandwidthManager/Private/DownlinkBandwidthManager.cpp
@@ -0,0 +1,421 @@
+// Copyright Epic Games, Inc. All Rights Reserved.
+
+#include "DownlinkBandwidthManager.h"
+#include "HAL/IConsoleManager.h"
+#include "Misc/ConfigCacheIni.h"
+
+DEFINE_LOG_CATEGORY(LogDownlinkBandwidthManager);
+
+namespace DLBW
+{
+	/// Console Variables to determine Manager behavior
+	int32 GDefaultBandwidthToAllocate = 300000;
+	static FAutoConsoleVariableRef CVarDLBWDefaultBandwidthToAllocate(
+		TEXT("net.DLBW.DefaultBandwidth"),
+		GDefaultBandwidthToAllocate,
+		TEXT("Default value to use for bandwidth cap distribution in Kbps. Example: Average american home internet is 200-300Mbps, so we would use 200,000 - 300,000 for this value."),
+		ECVF_Default);
+
+	static TAutoConsoleVariable<int32> CVarDLBWPriorityWeightExponent(
+		TEXT("net.DLBW.PriorityWeightExponent"),
+		3,
+		TEXT("Bias weight used as an exponent for bucket distribution."),
+		ECVF_Default);
+
+	static TAutoConsoleVariable<float> CVarDLBWDistributionPolling(
+		TEXT("net.DLBW.DistributionPolling"),
+		1.0f,
+		TEXT("Dealy Time, per frame cycle to poll for changes needed in distribution."),
+		ECVF_Default);
+
+	static TAutoConsoleVariable<bool> CVarDLBWEnabled(
+		TEXT("net.DLBW.Enable"),
+		false,
+		TEXT("Toggles the activation state for the Downlink Bandwidth Manager."));
+
+	bool GEnableBandwidthLogging = false;
+	static FAutoConsoleVariableRef CVarDLBWEnableBandwidthLogging(
+		TEXT("net.DLBW.EnableLogging"),
+		GEnableBandwidthLogging,
+		TEXT("Turns on Logging for the manager, and will write out based on CVarDLBWDistributionPolling rate."),
+		ECVF_Default);
+}
+
+/// Non-Member Functions
+FString LexToString(UE::BandwidthManager::EBandwidthPriority Value)
+{
+	switch (Value)
+	{
+	case UE::BandwidthManager::EBandwidthPriority::HighPriority:
+		return TEXT("High Priority");
+	case UE::BandwidthManager::EBandwidthPriority::MediumPriority:
+		return TEXT("Medium Prority");
+	case UE::BandwidthManager::EBandwidthPriority::LowPriority:
+		return TEXT("Low Priority");
+	default:
+		return TEXT("When Available");
+	}
+}
+
+// ---------------------------
+// Singleton backing pointer
+// ---------------------------
+TUniquePtr<FDownlinkBandwidthManager> FDownlinkBandwidthManager::Instance = nullptr;
+
+/// Static Functions
+FDownlinkBandwidthManager& FDownlinkBandwidthManager::Get()
+{
+	if (!Instance)
+	{
+		Instance = MakeUnique<FDownlinkBandwidthManager>();
+	}
+	check(Instance);
+	return *Instance.Get();
+}
+
+bool FDownlinkBandwidthManager::IsEnabled()
+{
+	return DLBW::CVarDLBWEnabled.GetValueOnAnyThread();
+}
+
+/// Member Functions
+FDownlinkBandwidthManager::FDownlinkBandwidthManager()
+{
+
+}
+
+FDownlinkBandwidthManager::~FDownlinkBandwidthManager()
+{
+	TrackedServices.Empty();
+}
+
+void FDownlinkBandwidthManager::Initialize()
+{
+	SetIncomingBandwidthCap(DLBW::GDefaultBandwidthToAllocate);
+	
+	// return true for reschedule
+	TickHandle = FTSTicker::GetCoreTicker().AddTicker(FTickerDelegate::CreateLambda([this](float InDeltaSeconds)->bool
+		{
+			// let's us evaulate distribution on time cycle needs should anything change
+			Tick();
+			return true;
+		}), DLBW::CVarDLBWDistributionPolling.GetValueOnAnyThread());
+}
+
+void FDownlinkBandwidthManager::Shutdown()
+{
+	FTSTicker::GetCoreTicker().RemoveTicker(TickHandle);
+	PriorityBucketMapping.Empty();
+}
+
+void FDownlinkBandwidthManager::SetIncomingBandwidthCap(int32 InKbps)
+{
+	UE::TScopeLock Lock(Mutex);
+
+	if (CapKbps != InKbps)
+	{
+		CapKbps = InKbps;
+		// a change to the cap value means we need to perform a full re-distribution
+		GeneratePriorityBucketMap();
+	}
+}
+
+UE::BandwidthManager::FServiceHandleID FDownlinkBandwidthManager::RegisterMonitoredService(
+	const FString& ServiceName,
+	UE::BandwidthManager::FManagedServiceConfigurableData& ServiceConfigData,
+	bool StartActive /* = false*/,
+	FOnBandwidthUpdate&& Callback /* = nullptr*/)
+{
+	UE::TScopeLock Lock(Mutex);
+
+	FManagedService NewService(MoveTemp(Callback));
+	NewService.DebugName = ServiceName;
+	NewService.ConfigData = ServiceConfigData;
+	NewService.bActiveConnection = StartActive;
+
+	NextServiceId++;
+	TrackedServices.Emplace(NextServiceId, NewService);
+
+	TrackedServices.ValueSort([&](const FManagedService& ServiceA, const FManagedService& ServiceB) 
+		{
+			uint8 ValueA = (CurrentGameplayPriorityState == UE::BandwidthManager::EBandwidthDistributionPriorityState::IdleGameplay) ? ServiceA.ConfigData.IdleGameplayPriority : ServiceA.ConfigData.InGameplayPriority;
+			uint8 ValueB = (CurrentGameplayPriorityState == UE::BandwidthManager::EBandwidthDistributionPriorityState::IdleGameplay) ? ServiceB.ConfigData.IdleGameplayPriority : ServiceB.ConfigData.InGameplayPriority;
+			// Sorts in ascending priority, so lower numerical (higher listed) services can be alloted to first
+			return ValueA > ValueB;
+		});
+
+	if (StartActive)
+	{
+		GeneratePriorityBucketMap();
+	}
+
+	return NextServiceId;
+}
+
+
+void FDownlinkBandwidthManager::UnregisterMonitoredService(const UE::BandwidthManager::FServiceHandleID& ServiceID)
+{
+	UE::TScopeLock Lock(Mutex);
+
+	TrackedServices.Remove(ServiceID);
+	GeneratePriorityBucketMap();
+}
+
+void FDownlinkBandwidthManager::SetActivationForService(const UE::BandwidthManager::FServiceHandleID& ServiceID, bool ActiveState)
+{
+	UE::TScopeLock Lock(Mutex);
+
+	if (FManagedService* Service = TrackedServices.Find(ServiceID))
+	{
+		Service->bActiveConnection = ActiveState;
+		GeneratePriorityBucketMap();
+	}
+}
+
+void FDownlinkBandwidthManager::ChangeServicePriority(const UE::BandwidthManager::FServiceHandleID& ServiceID, UE::BandwidthManager::EBandwidthPriority NewPriority, UE::BandwidthManager::EBandwidthDistributionPriorityState PriorityState)
+{
+	UE::TScopeLock Lock(Mutex);
+
+	if (FManagedService* Service = TrackedServices.Find(ServiceID))
+	{
+		switch (PriorityState) 
+		{
+		case UE::BandwidthManager::EBandwidthDistributionPriorityState::IdleGameplay:
+			Service->ConfigData.IdleGameplayPriority = NewPriority;
+			break;
+		case UE::BandwidthManager::EBandwidthDistributionPriorityState::InGameplay:
+			Service->ConfigData.InGameplayPriority = NewPriority;
+			break;
+		default:
+			Service->ConfigData.IdleGameplayPriority = NewPriority;
+			Service->ConfigData.InGameplayPriority = NewPriority;
+			break;
+		}
+
+		GeneratePriorityBucketMap();
+	}
+}
+
+void FDownlinkBandwidthManager::SetDistributionGameplayState(UE::BandwidthManager::EBandwidthDistributionPriorityState State)
+{
+	UE::TScopeLock Lock(Mutex);
+
+	if (CurrentGameplayPriorityState != State)
+	{
+		CurrentGameplayPriorityState = State;
+		GeneratePriorityBucketMap();
+	}
+}
+
+void FDownlinkBandwidthManager::Tick()
+{
+	if (CapKbps != DLBW::GDefaultBandwidthToAllocate)
+	{
+		// Updates if Cvar is changed, and will handle DistributeBandwidth
+		SetIncomingBandwidthCap(DLBW::GDefaultBandwidthToAllocate);
+	}
+	else
+	{
+		// Called as a background poll to see if redistribution is needed
+		DistributeBandwidth();
+	}
+
+	if (DLBW::GEnableBandwidthLogging)
+	{
+		Log();
+	}
+}
+
+void FDownlinkBandwidthManager::GeneratePriorityBucketMap()
+{
+	UE::TScopeLock Lock(Mutex);
+
+	int32 BandwidthBudgetBytes = CapKbps;
+	if (BandwidthBudgetBytes <= 0)
+	{
+		UE_LOG(LogDownlinkBandwidthManager, Warning, TEXT("%s - BandwidthBudgetPerFrame was set to less than 0. Value: %i, CapKbps: %i"),
+			ANSI_TO_TCHAR(__FUNCTION__), BandwidthBudgetBytes, CapKbps);
+		return;
+	}
+
+	// If we are generating a new set of buckets, then we need to remove stale data
+	if (PriorityBucketMapping.Num() != 0)
+	{
+		PriorityBucketMapping.Reset();
+	}
+
+	// We need to generate the current priority buckets since services could have changed
+	for (TPair<UE::BandwidthManager::FServiceHandleID, FManagedService>& ServicePair : TrackedServices)
+	{
+		const FManagedService& ServiceRef = ServicePair.Value;
+		if (!ServiceRef.bActiveConnection)
+		{
+			//if not active, then don't distribute bandwidth for it
+			continue;
+		}
+
+		UE::BandwidthManager::FPriorityBucket& PriorityBucket = PriorityBucketMapping.FindOrAdd(GetPriorityValueForService(ServiceRef));
+		PriorityBucket.NumberOfServicesAtPriority++;
+	}
+	
+	// we do not need an exponent larger than 255, so this lets us make sure that value is not too large and within the limits of the formula
+	int32 SanitizedExponent = FMath::Min<int32>(DLBW::CVarDLBWPriorityWeightExponent.GetValueOnAnyThread(), UINT8_MAX);
+	uint8 WeightBiasExponent = FMath::Max<uint8>(1, static_cast<uint8>(DLBW::CVarDLBWPriorityWeightExponent.GetValueOnAnyThread()));
+	const int32 BucketCount = PriorityBucketMapping.Num();
+	int32 BucketNumber = 0;
+	double OverallReduction = 0.0;
+	for (TPair<UE::BandwidthManager::EBandwidthPriority, UE::BandwidthManager::FPriorityBucket>& PriorityPair : PriorityBucketMapping)
+	{
+		if (BandwidthBudgetBytes <= 0)
+		{
+			// we seem to be out of bandwidth
+			break;
+		}
+
+		UE::BandwidthManager::FPriorityBucket& BucketRef = PriorityPair.Value;
+		// This method of distribution allows the highest priority bucket to take the 
+		// highest percentage based on a weighted exponent value, then trickle down the
+		// remainder while still considering the weight.
+		// Example:
+		// PriorityList = [High, Medium, Low], WeightBiasExponent = 3
+		// This would yield a starting WeightedPercentage = 0.702
+		// The next Cycle would be interpolate from 0 - 0.702 and then WeightedPercentage = 0.209
+		// The Overal Reduction becomes 0.911, meaning the final bucket gets 1-0.911 or 0.089
+		// In Percentages of the total bandwidth that is 70.2, 20.9, and 8.9%
+		// This is the total bandwidth each bucket is allowed and that bucket then needs to
+		// subdivide its allowance across it's own services
+
+		if (++BucketNumber == BucketCount)
+		{
+			// Last Item, we want to give out what's left to maximize bandwidth useage
+			BucketRef.BandwidthAllocationForPriority = BandwidthBudgetBytes;
+		}
+		else
+		{
+			double WeightedPercentage = FMath::InterpEaseOut<double>(0.0, 1.0 - OverallReduction, 1.0f / (float)BucketCount, WeightBiasExponent);
+			OverallReduction += WeightedPercentage;
+
+			// This is why we are tracking the remainder by bandwidth and not used percentage
+			int32 WeightedBandwidth = static_cast<int32>(static_cast<double>(BandwidthBudgetBytes) * WeightedPercentage);
+			BandwidthBudgetBytes -= FMath::Min<int32>(WeightedBandwidth, BandwidthBudgetBytes);
+			// ^^^ Prevents any potential negative roll over
+			BucketRef.BandwidthAllocationForPriority = WeightedBandwidth;
+
+		}
+	}
+
+	// Now that the buckets have been created/changed, we need to distribute bandwidth to services
+	DistributeBandwidth();
+}
+
+void FDownlinkBandwidthManager::DistributeBandwidth()
+{
+	UE::TScopeLock Lock(Mutex);
+
+	int32 BandwidthBudgetBytes = CapKbps;
+	if (BandwidthBudgetBytes <= 0)
+	{
+		UE_LOG(LogDownlinkBandwidthManager, Warning, TEXT("%s - BandwidthBudgetPerFrame was set to less than 0. Value: %i, CapKbps: %i"), 
+			ANSI_TO_TCHAR(__FUNCTION__), BandwidthBudgetBytes, CapKbps);
+		return;
+	}
+
+	// Distributes Allocations based on bucket allowance
+	for (TPair<UE::BandwidthManager::FServiceHandleID, FManagedService>& ServicePair : TrackedServices)
+	{
+		if (BandwidthBudgetBytes <= 0)
+		{
+			// we seem to be out of bandwidth
+			break;
+		}
+
+		FManagedService& ServiceRef = ServicePair.Value;
+		if (!ServiceRef.bActiveConnection)
+		{
+			//if not active, then don't distribute bandwidth for it
+			ServiceRef.AllotedBandwidthKbps = 0;
+			if (ServiceRef.OnBandwidthUpdate.IsSet() && ServiceRef.AllotedBandwidthKbps > 0)
+			{
+				ServiceRef.OnBandwidthUpdate(MoveTemp(ServiceRef.AllotedBandwidthKbps));
+			}
+			continue;
+		}
+
+		if (UE::BandwidthManager::FPriorityBucket* PriorityBucket = PriorityBucketMapping.Find(GetPriorityValueForService(ServiceRef)))
+		{
+			// Gives us the highest available bandwidth this priority can have based on total service needs
+			int32 CurrentServiceAbsoluteBandwidth = PriorityBucket->BandwidthAllocationForPriority / PriorityBucket->NumberOfServicesAtPriority;
+
+			// If reducing BandwidthBudgetBytes by CurrentServiceAbsoluteBandwidth would yield less than 0 (rollover) then use the smallest
+			int32 BandwidthToUse = FMath::Min(BandwidthBudgetBytes,CurrentServiceAbsoluteBandwidth);
+			if (ServiceRef.ConfigData.ComparisonValue.IsValid())
+			{
+				// Then take the minimum bandwidth between the "need" and available
+				int32 ComparisonBandwidth = FMath::Max(0, *ServiceRef.ConfigData.ComparisonValue.Get());
+				ServiceRef.AllotedBandwidthKbps = FMath::Min(BandwidthToUse, ComparisonBandwidth);
+			}
+			else
+			{
+				ServiceRef.AllotedBandwidthKbps = BandwidthToUse;
+			}
+
+			// reduce total amount by amount distributed
+			BandwidthBudgetBytes -= ServiceRef.AllotedBandwidthKbps;
+
+			if (ServiceRef.OnBandwidthUpdate.IsSet() && ServiceRef.AllotedBandwidthKbps > 0)
+			{
+				ServiceRef.OnBandwidthUpdate(MoveTemp(ServiceRef.AllotedBandwidthKbps));
+			}
+		}
+	}
+}
+
+UE::BandwidthManager::EBandwidthPriority FDownlinkBandwidthManager::GetPriorityValueForService(const FManagedService& MonitoredService) const
+{
+	return (CurrentGameplayPriorityState == UE::BandwidthManager::EBandwidthDistributionPriorityState::IdleGameplay) ? MonitoredService.ConfigData.IdleGameplayPriority : MonitoredService.ConfigData.InGameplayPriority;
+}
+
+void FDownlinkBandwidthManager::Log()
+{
+	UE_LOG(LogDownlinkBandwidthManager, Log, TEXT("FDownlinkBandwidthManager Logging Breakdown:\n"));
+
+	UE_LOG(LogDownlinkBandwidthManager, Log, TEXT("FDownlinkBandwidthManager Internal Data:"));
+	UE_LOG(LogDownlinkBandwidthManager, Log, TEXT("Has the manager been enabled: %s"), (DLBW::CVarDLBWEnabled.GetValueOnAnyThread()) ? TEXT("True") : TEXT("False"));
+	UE_LOG(LogDownlinkBandwidthManager, Log, TEXT("Default Bandwidth Cap in Kbps: %i"), DLBW::GDefaultBandwidthToAllocate);
+	UE_LOG(LogDownlinkBandwidthManager, Log, TEXT("Current Bandwidth Cap in Kbps: %i"), CapKbps);
+	UE_LOG(LogDownlinkBandwidthManager, Log, TEXT("Current Gameplay Priority State: %s"), 
+		(CurrentGameplayPriorityState == UE::BandwidthManager::EBandwidthDistributionPriorityState::IdleGameplay) ? TEXT("Outside Normal Gameplay") : TEXT("In Active Gameplay"));
+	UE_LOG(LogDownlinkBandwidthManager, Log, TEXT("Tick Polling in seconds: %0.3f"), DLBW::CVarDLBWDistributionPolling.GetValueOnAnyThread());
+
+	UE_LOG(LogDownlinkBandwidthManager, Log, TEXT("FDownlinkBandwidthManager Prioritization:"));
+	UE_LOG(LogDownlinkBandwidthManager, Log, TEXT("Priority Weight Exponent: %i"), DLBW::CVarDLBWPriorityWeightExponent.GetValueOnAnyThread());
+
+	for (TPair<UE::BandwidthManager::EBandwidthPriority, UE::BandwidthManager::FPriorityBucket>& BucketData : PriorityBucketMapping)
+	{
+		UE_LOG(LogDownlinkBandwidthManager, Log, TEXT("Bucket Priority Value: %s"), *LexToString(BucketData.Key));
+		UE_LOG(LogDownlinkBandwidthManager, Log, TEXT("Bucket Service Count: %i"), BucketData.Value.NumberOfServicesAtPriority);
+		UE_LOG(LogDownlinkBandwidthManager, Log, TEXT("Bandwidth for this Bucket: %i"), BucketData.Value.BandwidthAllocationForPriority);
+	}
+
+	UE_LOG(LogDownlinkBandwidthManager, Log, TEXT("FDownlinkBandwidthManager Tracked Services:"));
+	for (TPair<UE::BandwidthManager::FServiceHandleID, FManagedService>& ServiceData : TrackedServices)
+	{
+		UE_LOG(LogDownlinkBandwidthManager, Log, TEXT("Tracked Service ID: %i, Name: %s"), ServiceData.Key, *ServiceData.Value.DebugName);
+		UE_LOG(LogDownlinkBandwidthManager, Log, TEXT("Service Gameplay Priority Value: %s"), *LexToString(ServiceData.Value.ConfigData.InGameplayPriority));
+		UE_LOG(LogDownlinkBandwidthManager, Log, TEXT("Service Gameplay Priority Value: %s"), *LexToString(ServiceData.Value.ConfigData.IdleGameplayPriority));
+		UE_LOG(LogDownlinkBandwidthManager, Log, TEXT("Is Service Currently Active: %s"), (ServiceData.Value.bActiveConnection) ? TEXT("True") : TEXT("False"));
+		if (!ServiceData.Value.ConfigData.ComparisonValue.IsValid())
+		{
+			UE_LOG(LogDownlinkBandwidthManager, Log, TEXT("Service Comparision Value: Null"));
+		}
+		else
+		{
+			UE_LOG(LogDownlinkBandwidthManager, Log, TEXT("Service Comparision Value: %i"), ServiceData.Value.ConfigData.ComparisonValue.Get());
+		}
+
+		UE_LOG(LogDownlinkBandwidthManager, Log, TEXT("Service Alloted Bandwidth in Kbps: %i"), ServiceData.Value.AllotedBandwidthKbps);
+		UE_LOG(LogDownlinkBandwidthManager, Log, TEXT("Service Bound Callback: %s\n"), (ServiceData.Value.OnBandwidthUpdate.IsSet()) ? TEXT("Set") : TEXT("Unset"));
+	}
+
+	UE_LOG(LogDownlinkBandwidthManager, Log, TEXT("FDownlinkBandwidthManager Logging Complete."));
+}
\ No newline at end of file
diff --git a/Engine/Source/Runtime/Online/Experimental/DownlinkBandwidthManager/Private/DownlinkBandwidthManagerModule.cpp b/Engine/Source/Runtime/Online/Experimental/DownlinkBandwidthManager/Private/DownlinkBandwidthManagerModule.cpp
new file mode 100644
--- /dev/null
+++ b/Engine/Source/Runtime/Online/Experimental/DownlinkBandwidthManager/Private/DownlinkBandwidthManagerModule.cpp
@@ -0,0 +1,25 @@
+// Copyright Epic Games, Inc. All Rights Reserved.
+
+#include "Modules/ModuleInterface.h"
+#include "Modules/ModuleManager.h"
+#include "DownlinkBandwidthManager.h"
+
+// DownlinkBandwidthManager module
+// Central service for dividing available downlink bandwidth across systems.
+// Provides per-profile and per-socket throttling, driven by INI config,
+// so HTTP/streaming subsystems can share bandwidth fairly by priority.
+
+class FDownlinkBandwidthManagerModule : public IModuleInterface
+{
+	void StartupModule()
+	{
+		FDownlinkBandwidthManager::Get().Initialize();
+	}
+
+	void ShutdownModule()
+	{
+		FDownlinkBandwidthManager::Get().Shutdown();
+	}
+};
+
+IMPLEMENT_MODULE(FDownlinkBandwidthManagerModule, DownlinkBandwidthManager)
\ No newline at end of file
diff --git a/Engine/Source/Runtime/Online/Experimental/DownlinkBandwidthManager/Public/DownlinkBandwidthManager.h b/Engine/Source/Runtime/Online/Experimental/DownlinkBandwidthManager/Public/DownlinkBandwidthManager.h
new file mode 100644
--- /dev/null
+++ b/Engine/Source/Runtime/Online/Experimental/DownlinkBandwidthManager/Public/DownlinkBandwidthManager.h
@@ -0,0 +1,167 @@
+// Copyright Epic Games, Inc. All Rights Reserved.
+#pragma once
+#include "Containers/Ticker.h"
+
+#define UE_API DOWNLINKBANDWIDTHMANAGER_API
+#define UE_INVALID_BANDWIDTH_SERVICE_HANDLE -1
+
+UE_API DECLARE_LOG_CATEGORY_EXTERN(LogDownlinkBandwidthManager, Log, All);
+
+// Useage for this callback will recieve a parameter in Kbps
+using FOnBandwidthUpdate = TUniqueFunction<void(int32&&)>;
+
+namespace UE::BandwidthManager
+{
+	enum EBandwidthPriority : uint8
+	{
+		HighPriority = 0,
+		MediumPriority = 1,
+		LowPriority = 2,
+		WhenAvailable = 3
+	};
+
+	enum EBandwidthDistributionPriorityState : uint8
+	{
+		InGameplay,
+		IdleGameplay
+	};
+
+	typedef int32 FServiceHandleID;
+
+	struct FManagedServiceConfigurableData
+	{
+		FManagedServiceConfigurableData()
+			: ComparisonValue(nullptr),
+			InGameplayPriority(EBandwidthPriority::WhenAvailable),
+			IdleGameplayPriority(EBandwidthPriority::WhenAvailable) {
+		};
+
+		FManagedServiceConfigurableData(const FManagedServiceConfigurableData& Copyable)
+			: ComparisonValue(Copyable.ComparisonValue),
+			InGameplayPriority(Copyable.InGameplayPriority),
+			IdleGameplayPriority(Copyable.IdleGameplayPriority) {
+		};
+
+		// Shared Pointer for a preferred bandwidth cap, distribution will 
+		// prefer the minimum between this and the calculated bandwidth.
+		TSharedPtr<int32> ComparisonValue;
+		// The priority to use when the client is considered in active gameplay
+		EBandwidthPriority InGameplayPriority;
+		// The priority to use when the client is considered in less active gameplay
+		EBandwidthPriority IdleGameplayPriority;
+	};
+
+	struct FPriorityBucket
+	{
+		FPriorityBucket() {};
+		int32 BandwidthAllocationForPriority = 0;
+		uint8 NumberOfServicesAtPriority = 0;
+	};
+}
+
+struct FManagedService
+{
+
+	FManagedService(FOnBandwidthUpdate&& CallBack)
+		: bActiveConnection(false),
+		ConfigData(),
+		AllotedBandwidthKbps(0),
+		OnBandwidthUpdate(MoveTemp(CallBack)) {
+	};
+
+	FManagedService(const FManagedService& Copyable)
+		: bActiveConnection(Copyable.bActiveConnection), 
+		ConfigData(Copyable.ConfigData),
+		AllotedBandwidthKbps(Copyable.AllotedBandwidthKbps),
+		OnBandwidthUpdate(MoveTemp(Copyable.OnBandwidthUpdate)) {};
+
+	FString DebugName;
+	bool bActiveConnection;
+	UE::BandwidthManager::FManagedServiceConfigurableData ConfigData;
+	int32 AllotedBandwidthKbps;
+	// callback for subscriber to handle the new allotment that reports report back in Kbps
+	FOnBandwidthUpdate&& OnBandwidthUpdate; 
+};
+
+class FDownlinkBandwidthManager
+{
+
+public:
+	UE_NONCOPYABLE(FDownlinkBandwidthManager);
+
+	FDownlinkBandwidthManager();
+	~FDownlinkBandwidthManager();
+	// --- Singleton access ---
+	UE_API static FDownlinkBandwidthManager& Get();
+
+	// --- CVar Enabled Check ---
+	UE_API static bool IsEnabled();
+
+	// --- Lifecycle (called by module) ---
+	void Initialize();
+	void Shutdown();
+
+	// Service Initialization and registration with Distribution Manager, FOnBandwidthUpdate will report back in Kbps 
+	UE_API UE::BandwidthManager::FServiceHandleID RegisterMonitoredService(
+		const FString& ServiceName,
+		UE::BandwidthManager::FManagedServiceConfigurableData& ServiceConfigData,
+		bool StartActive = false,
+		FOnBandwidthUpdate&& Callback = nullptr);
+
+	// Removes service from manager considerations
+	UE_API void UnregisterMonitoredService(const UE::BandwidthManager::FServiceHandleID& ServiceID);
+
+	// Determines if service should be considered for distribution
+	UE_API void SetActivationForService(const UE::BandwidthManager::FServiceHandleID& ServiceID, bool ActiveState);
+
+	// Changes the associated Services Priority according to the priority state
+	UE_API void ChangeServicePriority(const UE::BandwidthManager::FServiceHandleID& ServiceID, UE::BandwidthManager::EBandwidthPriority NewPriority, UE::BandwidthManager::EBandwidthDistributionPriorityState PriorityState);
+
+	// Ways for the game to communicate back or check to see what priority distribution pattern the manager is in
+	UE_API void SetDistributionGameplayState(UE::BandwidthManager::EBandwidthDistributionPriorityState State);
+	UE::BandwidthManager::EBandwidthDistributionPriorityState GetDistributionGameplayState() const { return CurrentGameplayPriorityState; };
+
+	// Sets the bandwidth Cap for the Manager to distribute across all services
+	void SetIncomingBandwidthCap(int32 InKbps);
+
+private:
+	// Called once per frame/tick to apply an updated allotment based on active services
+	void Tick();
+
+	// Source of Bandwidth Assignment
+	void DistributeBandwidth();
+
+	// Called when services are altered to maintain PriorityBucketMapping
+	void GeneratePriorityBucketMap();
+
+	// Helper function to get priority based on CurrentGameplayPriority
+	UE::BandwidthManager::EBandwidthPriority GetPriorityValueForService(const FManagedService& MonitoredService) const;
+
+	// Allows for logging output of FDownlinkBandwidthManager, enabled via CVarDLBWEnableBandwidthLogging
+	void Log();
+
+	// Singleton instance pointer
+	static TUniquePtr<FDownlinkBandwidthManager> Instance;
+
+	mutable FTransactionallySafeCriticalSection Mutex;
+
+	// Handle for the polling tick
+	FTSTicker::FDelegateHandle TickHandle;
+
+	// The hard cap used for determining bandwidth distribution
+	int32 CapKbps = 0;
+
+	// Incrementing Handle value for external services to track their FManagedService data within the manager
+	int32 NextServiceId = UE_INVALID_BANDWIDTH_SERVICE_HANDLE;
+
+	// Service Data Mapping
+	TMap<UE::BandwidthManager::FServiceHandleID, FManagedService> TrackedServices;
+
+	// Mapping of expected priorities to their calculated allotment of bandwidth
+	TMap<UE::BandwidthManager::EBandwidthPriority, UE::BandwidthManager::FPriorityBucket> PriorityBucketMapping;
+
+	// The Managers perceived state of gameplay to consider for distribution
+	UE::BandwidthManager::EBandwidthDistributionPriorityState CurrentGameplayPriorityState = UE::BandwidthManager::EBandwidthDistributionPriorityState::IdleGameplay;
+};
+
+#undef UE_API
