FWebImage: add ESlateBrushDrawType and FSlateBrushOutlineSettings config for downloaded brush

#rnx

#Synced-CL 50908424
#Synced-CL 50908424

diff --git a/Engine/Source/Runtime/Online/ImageDownload/Private/WebImage.cpp b/Engine/Source/Runtime/Online/ImageDownload/Private/WebImage.cpp
--- a/Engine/Source/Runtime/Online/ImageDownload/Private/WebImage.cpp
+++ b/Engine/Source/Runtime/Online/ImageDownload/Private/WebImage.cpp
@@ -11,8 +11,9 @@
 
 
 FWebImage::FWebImage()
-: StandInBrush(FCoreStyle::Get().GetDefaultBrush())
-, bDownloadSuccess(false)
+	: StandInBrush(FCoreStyle::Get().GetDefaultBrush())
+	, DownloadedBrushDrawType(ESlateBrushDrawType::Image)
+	, bDownloadSuccess(false)
 {
 }
 
@@ -21,6 +22,24 @@ FWebImage::~FWebImage()
 	CancelDownload();
 }
 
+void FWebImage::SetDownloadedBrushDrawType(ESlateBrushDrawType::Type DrawType)
+{
+	DownloadedBrushDrawType = DrawType;
+	if (DownloadedBrush)
+	{
+		DownloadedBrush->DrawAs = DownloadedBrushDrawType;
+	}
+}
+
+void FWebImage::SetDownloadedBrushOutlineSettings(FSlateBrushOutlineSettings OutlineSettings)
+{
+	DownloadedBrushOutlineSettings = MoveTemp(OutlineSettings);
+	if (DownloadedBrush)
+	{
+		DownloadedBrush->OutlineSettings = DownloadedBrushOutlineSettings;
+	}
+}
+
 TAttribute< const FSlateBrush* > FWebImage::Attr() const
 {
 	return TAttribute< const FSlateBrush* >(AsShared(), &FWebImage::GetBrush);
@@ -127,6 +146,11 @@ bool FWebImage::ProcessHttpResponse(const FString& RequestUrl, FHttpResponsePtr
 	// FImage.RawData is a TArray64<uint8> while FSlateDynamicImageBrush::CreateWithImageData expects a TArray<uint8>. We need to copy the image data to adapt
 	TArray<uint8> RawImageCopy(DownloadedImage.RawData);
 	DownloadedBrush = FSlateDynamicImageBrush::CreateWithImageData(ResourceName, FVector2D((float)DownloadedImage.SizeX, (float)DownloadedImage.SizeY), RawImageCopy);
+	if (DownloadedBrush)
+	{
+		DownloadedBrush->DrawAs = DownloadedBrushDrawType;
+		DownloadedBrush->OutlineSettings = DownloadedBrushOutlineSettings;
+	}
 	return DownloadedBrush.IsValid();
 }
 
diff --git a/Engine/Source/Runtime/Online/ImageDownload/Public/WebImage.h b/Engine/Source/Runtime/Online/ImageDownload/Public/WebImage.h
--- a/Engine/Source/Runtime/Online/ImageDownload/Public/WebImage.h
+++ b/Engine/Source/Runtime/Online/ImageDownload/Public/WebImage.h
@@ -53,7 +53,18 @@ public:
 	DECLARE_DELEGATE_OneParam(FOnImageDownloaded, bool);
 
 	/** Set the brush that is currently being returned (this will be overridden when any async download completes) */
-	inline FWebImage& SetStandInBrush(TAttribute<const FSlateBrush*> StandInBrushIn) { StandInBrush = StandInBrushIn; DownloadedBrush.Reset(); return *this; }
+	inline FWebImage& SetStandInBrush(TAttribute<const FSlateBrush*> StandInBrushIn)
+	{
+		StandInBrush = StandInBrushIn;
+		DownloadedBrush.Reset();
+		return *this;
+	}
+
+	/** Set draw type for the downloaded image brush */
+	IMAGEDOWNLOAD_API void SetDownloadedBrushDrawType(ESlateBrushDrawType::Type DrawType);
+
+	/** Set outline settings for the downloaded image brush */
+	IMAGEDOWNLOAD_API void SetDownloadedBrushOutlineSettings(FSlateBrushOutlineSettings OutlineSettings);
 
 	/** Begin downloading an image. This will automatically set the current brush to the downloaded image when it completes (if successful) */
 	IMAGEDOWNLOAD_API bool BeginDownload(const FString& InUrl, const TOptional<FString>& StandInETag = TOptional<FString>(), const FOnImageDownloaded& DownloadCallback = FOnImageDownloaded());
@@ -108,6 +119,12 @@ private:
 	/** The most recently downloaded and generated brush */
 	TSharedPtr<FSlateDynamicImageBrush> DownloadedBrush;
 
+	/** Outline settings for the downloaded image brush */
+	FSlateBrushOutlineSettings DownloadedBrushOutlineSettings;
+
+	/** Draw type for the downloaded image brush */
+	ESlateBrushDrawType::Type DownloadedBrushDrawType;
+
 	/** Any pending request */
 	TSharedPtr<IHttpRequest, ESPMode::ThreadSafe> PendingRequest;
 
