Fixing Macro errors for engine tests

#redirectsaresafe
#rnx

#Synced-CL 50897393

diff --git a/Engine/Source/Runtime/Online/Experimental/DownlinkBandwidthManager/Public/DownlinkBandwidthManager.h b/Engine/Source/Runtime/Online/Experimental/DownlinkBandwidthManager/Public/DownlinkBandwidthManager.h
--- a/Engine/Source/Runtime/Online/Experimental/DownlinkBandwidthManager/Public/DownlinkBandwidthManager.h
+++ b/Engine/Source/Runtime/Online/Experimental/DownlinkBandwidthManager/Public/DownlinkBandwidthManager.h
@@ -7,7 +7,7 @@
 
 UE_API DECLARE_LOG_CATEGORY_EXTERN(LogDownlinkBandwidthManager, Log, All);
 
-// Useage for this callback will recieve a parameter in Kbps
+// Useage for this callback will receive a parameter in Kbps
 using FOnBandwidthUpdate = TUniqueFunction<void(int32&&)>;
 
 namespace UE::BandwidthManager
@@ -122,7 +122,7 @@ public:
 	UE::BandwidthManager::EBandwidthDistributionPriorityState GetDistributionGameplayState() const { return CurrentGameplayPriorityState; };
 
 	// Sets the bandwidth Cap for the Manager to distribute across all services
-	void SetIncomingBandwidthCap(int32 InKbps);
+	UE_API void SetIncomingBandwidthCap(int32 InKbps);
 
 private:
 	// Called once per frame/tick to apply an updated allotment based on active services
