p4 integrate //UE5/Dev-Distro-5.5/Engine/Source/Programs/BuildPatchTool/... //UE5/Main/Engine/Source/Programs/BuildPatchTool/...

diff --git a/Engine/Source/Runtime/Online/BuildPatchServices/Private/BuildPatchFileConstructor.cpp b/Engine/Source/Runtime/Online/BuildPatchServices/Private/BuildPatchFileConstructor.cpp
--- a/Engine/Source/Runtime/Online/BuildPatchServices/Private/BuildPatchFileConstructor.cpp
+++ b/Engine/Source/Runtime/Online/BuildPatchServices/Private/BuildPatchFileConstructor.cpp
@@ -2,6 +2,7 @@
 
 #include "BuildPatchFileConstructor.h"
 #include "IBuildManifestSet.h"
+#include "Algo/Transform.h"
 #include "HAL/UESemaphore.h"
 #include "HAL/PlatformFileManager.h"
 #include "HAL/FileManager.h"
@@ -10,6 +11,7 @@
 #include "Misc/ConfigCacheIni.h"
 #include "Misc/Paths.h"
 #include "Misc/ScopeLock.h"
+#include "Misc/ScopeRWLock.h"
 #include "Misc/OutputDeviceRedirector.h"
 #include "BuildPatchServicesPrivate.h"
 #include "Interfaces/IBuildInstaller.h"
@@ -36,6 +38,9 @@
 #include <unistd.h>
 #endif
 
+DECLARE_LOG_CATEGORY_EXTERN(LogBuildPatchFileConstructor, Log, All);
+DEFINE_LOG_CATEGORY(LogBuildPatchFileConstructor);
+
 using namespace BuildPatchServices;
 
 static int32 SleepTimeWhenFileSystemThrottledSeconds = 1;
@@ -79,6 +84,27 @@ namespace FileConstructorHelpers
 		}
 	}
 
+	bool CheckRemainingDiskSpace(const FString& InstallDirectory, uint64 RemainingBytesRequired, uint64& OutAvailableDiskSpace)
+	{
+		bool bContinueConstruction = true;
+		uint64 TotalSize = 0;
+		OutAvailableDiskSpace = 0;
+		if (FPlatformMisc::GetDiskTotalAndFreeSpace(InstallDirectory, TotalSize, OutAvailableDiskSpace))
+		{
+			if (OutAvailableDiskSpace < RemainingBytesRequired)
+			{
+				bContinueConstruction = false;
+			}
+		}
+		else
+		{
+			// If we can't get the disk space free then the most likely reason is the drive is no longer around...
+			bContinueConstruction = false;
+		}
+
+		return bContinueConstruction;
+	}
+
 	uint64 CalculateRequiredDiskSpace(const FBuildPatchAppManifestPtr& CurrentManifest, const FBuildPatchAppManifestRef& BuildManifest, const EInstallMode& InstallMode, const TSet<FString>& InInstallTags)
 	{
 		// Make tags expected
@@ -123,6 +149,24 @@ namespace FileConstructorHelpers
 		}
 		return FMath::Max<int64>(DiskSpaceDeltaPeak, 0);
 	}
+
+	static float ReadRWDurationThresholdSec()
+	{
+		float RWDurationThresholdSec = 4.0;
+		GConfig->GetFloat(TEXT("BuildPatchServices"), TEXT("ConstructionRWDurationThresholdSec"), RWDurationThresholdSec, GEngineIni);
+		return RWDurationThresholdSec;
+	}
+
+	static void LogLowRWDurationTime(const ISpeedRecorder::FRecord& ActivityRecord, const TCHAR* ProcedureType)
+	{
+		// We've seen reports of abnormally poor HDD response times so we log when things are truly bad (e.g. multiple seconds for an access).
+		static const double RWDurationThresholdSec = ReadRWDurationThresholdSec();
+		const double Duration = FStatsCollector::CyclesToSeconds(ActivityRecord.CyclesEnd - ActivityRecord.CyclesStart);
+		if (Duration > RWDurationThresholdSec)
+		{
+			UE_LOG(LogBuildPatchServices, Warning, TEXT("FBuildPatchFileConstructor: %s time of %.3f seconds (%d bytes)"), ProcedureType, Duration, ActivityRecord.Size);
+		}
+	}
 }
 
 
@@ -148,9 +192,11 @@ struct FReadScope
 {
 	ISpeedRecorder::FRecord ActivityRecord;
 	IFileConstructorStat* FileConstructorStat;
-	FReadScope(IFileConstructorStat* InFileConstructorStat, int64 Size)
+	const TCHAR* LogDetails = nullptr;
+	FReadScope(IFileConstructorStat* InFileConstructorStat, int64 Size, const TCHAR* InLogDetails)
 	{
-		FileConstructorStat = InFileConstructorStat;		
+		LogDetails = InLogDetails;
+		FileConstructorStat = InFileConstructorStat;
 		ActivityRecord.CyclesStart = FStatsCollector::GetCycles();
 		ActivityRecord.Size = Size;
 		FileConstructorStat->OnBeforeRead();
@@ -158,6 +204,7 @@ struct FReadScope
 	~FReadScope()
 	{
 		ActivityRecord.CyclesEnd = FStatsCollector::GetCycles();
+		FileConstructorHelpers::LogLowRWDurationTime(ActivityRecord, LogDetails);
 		FileConstructorStat->OnAfterRead(ActivityRecord);
 	}
 };
@@ -175,9 +222,18 @@ enum class EConstructionError : uint8
 	InternalConsistencyError,
 	Aborted,
 	MissingFileInfo,
+	InvalidSymlink,
 	CloseError
 };
 
+enum class EFileStage : uint8
+{
+	Unknown,	 // The file stage is unknown due to a failure (FFileConstructionState::bSuccess is false).
+	Complete,    // The file is complete; it exists and contains the expected data size.
+	Incomplete,  // The file is incomplete; it exists but does not contain the expected data size. Resume the operation.
+	New          // The file is missing; a new file should be created and populated from the start.
+};
+
 
 // Since we can have more than one file in flight, store state here.
 struct FFileConstructionState
@@ -216,7 +272,9 @@ struct FFileConstructionState
 	int32 ConstructionIndex = -1;
 	const FFileManifest* FileManifest;
 	const FString& BuildFilename; // name as it is in the manifest, references the ConstructList in the configuration.
-	FString NewFilename; // with output path and such
+	const FString NewFilename; // with output path and such
+
+	EFileStage Stage = EFileStage::Unknown;
 
 	FFileConstructionState(const FFileManifest* InFileManifest, const FString& InBuildFilename, FString&& InNewFilename) 
 	: FileManifest(InFileManifest)
@@ -232,7 +290,7 @@ struct FFileConstructionState
 
 	void SetAttributes()
 	{
-#if PLATFORM_MAC
+#if PLATFORM_MAC || PLATFORM_LINUX
 		if (bSuccess && 
 			EnumHasAllFlags(FileManifest->FileMetaFlags, EFileMetaFlags::UnixExecutable))
 		{
@@ -281,14 +339,8 @@ public:
 	// The resume ids that we loaded from disk
 	TSet<FString> LoadedResumeIds;
 
-	// The set of files that were started
-	TSet<FString> FilesStarted;
-
-	// The set of files that were completed, determined by expected file size
-	TSet<FString> FilesCompleted;
-
-	// The set of files that exist but are not able to assume resumable
-	TSet<FString> FilesIncompatible;
+	// Maps each resumed file to its current stage: Complete or Incomplete. Files not found on disk or treated as New are not present in this map.
+	TMap<FString, EFileStage> FilesStage;
 
 	// Whether we have any resume data for this install
 	bool bHasResumeData = false;
@@ -353,11 +405,24 @@ public:
 		}
 	}
 
+	/**
+	 * Deletes the resume data
+	 */
+	void RemoveResumeData()
+	{
+		if (bResumeEnabled)
+		{
+			bHasResumeData = false;
+			FileSystem->DeleteFile(*ResumeDataFilename);
+		}
+	}
+
 	/**
 	 * Checks whether the file was completed during last install attempt and adds it to FilesCompleted if so
 	 * @param Filename    The filename to check
+	 * @return A flag indicating that the file is incompatible and should be removed and constracted from scratch.
 	 */
-	void CheckFile(const FString& Filename)
+	bool CheckFile(const FString& Filename)
 	{
 		// If we had resume data, check if this file might have been resumable
 		if (bHasResumeData)
@@ -374,25 +439,27 @@ public:
 				if (NewFileManifest && bFileExists)
 				{
 					const uint64 UnsignedDiskFileSize = DiskFileSize;
-					if (UnsignedDiskFileSize > 0 && UnsignedDiskFileSize <= NewFileManifest->FileSize)
+					if (UnsignedDiskFileSize == NewFileManifest->FileSize)
 					{
-						FilesStarted.Add(Filename);
+						FilesStage.Add(Filename, EFileStage::Complete);
 					}
-					if (UnsignedDiskFileSize == NewFileManifest->FileSize)
+					else if (UnsignedDiskFileSize > 0 && UnsignedDiskFileSize < NewFileManifest->FileSize)
 					{
-						FilesCompleted.Add(Filename);
+						FilesStage.Add(Filename, EFileStage::Incomplete);
 					}
-					if (UnsignedDiskFileSize > NewFileManifest->FileSize)
+					else if (UnsignedDiskFileSize > NewFileManifest->FileSize)
 					{
-						FilesIncompatible.Add(Filename);
+						return true;
 					}
+					// If the file size is zero, we just treat it as EFileStage::New
 				}
 			}
 			else if (bFileExists)
 			{
-				FilesIncompatible.Add(Filename);
+				return true;
 			}
 		}
+		return false;
 	}
 };
 
@@ -1440,6 +1507,14 @@ FBuildPatchFileConstructor::FBuildPatchFileConstructor(
 
 		FFileToConstruct FileToConstruct;
 		FileToConstruct.FileManifest = FileManifest;
+		if (Configuration.bUseSHA1StageFilenames)
+		{
+			FBuildPatchUtils::SHAToBase32(FileManifest->SHA1Hash, FileToConstruct.StagingFilename);
+		}
+		else
+		{
+			FileToConstruct.StagingFilename = FileToConstructName;
+		}
 
 		// If we are missing the file manifest, we will fail to install when we get to the file. However,
 		// we guarantee a 1:1 mapping with the arrays we are filling here so we use invalid data for those
@@ -1605,7 +1680,6 @@ FBuildPatchFileConstructor::FBuildPatchFileConstructor(
 
 		Threads[ThreadIndex]->RunTask([this, ThreadIndex]() { GenericThreadFn(ThreadIndex); });
 	}
-
 }
 
 FBuildPatchFileConstructor::~FBuildPatchFileConstructor()
@@ -1725,35 +1799,25 @@ void FBuildPatchFileConstructor::QueueGenericThreadTask(int32 ThreadIndex, ICons
 
 void FBuildPatchFileConstructor::GenericThreadFn(int32 ThreadIndex)
 {
-	for (;;)
+	bool bLastShouldAbort = false;
+	do
 	{
 		ThreadWakeups[ThreadIndex]->Wait();
 
 		TArray<IConstructorChunkSource::FRequestProcessFn> GrabbedJobs;
-		
 		{
-			ThreadJobPostingLocks[ThreadIndex].Lock();
+			const UE::TScopeLock Lock(ThreadJobPostingLocks[ThreadIndex]);
 			GrabbedJobs = MoveTemp(ThreadJobPostings[ThreadIndex]);
-			ThreadJobPostingLocks[ThreadIndex].Unlock();
-		}
-
-		if (bShouldAbort)
-		{
-			for (IConstructorChunkSource::FRequestProcessFn& AbortJob : GrabbedJobs)
-			{
-				AbortJob(true);
-			}
-			WakeUpDispatchThreadEvent->Trigger();
-			break;
 		}
 
-		for (IConstructorChunkSource::FRequestProcessFn& Job : GrabbedJobs)
+		bLastShouldAbort = bShouldAbort;
+		for (const IConstructorChunkSource::FRequestProcessFn& Job : GrabbedJobs)
 		{
-			Job(false);
+			Job(bLastShouldAbort);
 		}
 
 		WakeUpDispatchThreadEvent->Trigger();
-	}
+	} while (!bLastShouldAbort);
 
 	ThreadCompleteEvents[ThreadIndex]->Trigger();
 }
@@ -1824,7 +1888,7 @@ bool FBuildPatchFileConstructor::HarvestChunksForCompletedFile(const FString& In
 					continue;
 				}
 
-				int32 ChunkSize = Configuration.ManifestSet->GetChunkInfo(FileChunk)->WindowSize;
+				int32 ChunkSize = Configuration.ManifestSet->GetChunkInfo(FileChunk)->DataSizeUncompressed;
 
 				ChunksFromFileWeNeed.Add({FileChunk, LastUsageIndex, NextUsageIndex, ChunkSize});
 			}
@@ -1957,7 +2021,7 @@ void FBuildPatchFileConstructor::Run()
 	for (const TPair<FGuid, EConstructorChunkLocation>& ChunkLocation : ChunkLocations)
 	{
 		FChunkInfo const* ChunkInfo = Configuration.ManifestSet->GetChunkInfo(ChunkLocation.Key);
-		WindowSizes.FindOrAdd(ChunkInfo->WindowSize)++;
+		WindowSizes.FindOrAdd(ChunkInfo->DataSizeUncompressed)++;
 
 		if (ChunkLocation.Value == EConstructorChunkLocation::Cloud)
 		{
@@ -2026,26 +2090,32 @@ void FBuildPatchFileConstructor::Run()
 
 		ResumeData.InitResume();
 
-		// Remove incompatible files
-		if (ResumeData.bHasResumeData)
+		// Check the stage of resumed files and remove any incompatible files.
+		if (ResumeData.bHasResumeData && !Configuration.bUseSHA1StageFilenames)
 		{
 			for (const FString& FileToConstruct : Configuration.ConstructList)
 			{
-				ResumeData.CheckFile(FileToConstruct);
-				const bool bFileIncompatible = ResumeData.FilesIncompatible.Contains(FileToConstruct);
+				const bool bFileIncompatible = ResumeData.CheckFile(FileToConstruct);
 				if (bFileIncompatible)
 				{
-					GLog->Logf(TEXT("FBuildPatchFileConstructor: Deleting incompatible stage file %s"), *FileToConstruct);
+					UE_LOG(LogBuildPatchFileConstructor, Log, TEXT("Deleting incompatible stage file %s"), *FileToConstruct);
 					FileSystem->DeleteFile(*(Configuration.StagingDirectory / FileToConstruct));
 				}
 			}
 
 			if (Configuration.bConstructInMemory)
 			{
-				UE_LOG(LogBuildPatchServices, Log, TEXT("Emptying Resume Data FilesStarted: construct in memory"));
-				ResumeData.FilesStarted.Empty();
+				UE_LOG(LogBuildPatchServices, Log, TEXT("Removing incomplete entries from resume data: construct in memory"));
+				ResumeData.FilesStage = ResumeData.FilesStage.FilterByPredicate([](const TPair<FString, EFileStage>& Entry) { return Entry.Value != EFileStage::Incomplete; });
 			}
+		}
 
+		// Uninstall should not have any files to construct
+		const bool bIsUninstall = Configuration.ConstructList.IsEmpty();
+		if (ResumeData.bHasResumeData && !bIsUninstall)
+		{
+			// For uninstall, we will remove staged files at an upper level
+			RemoveUnnecessaryStagedFiles(Configuration.StagingDirectory);
 		}
 
 		// Save for started versions
@@ -2067,6 +2137,12 @@ void FBuildPatchFileConstructor::Run()
 	{
 		FileConstructorStat->OnResumeCompleted();
 	}
+
+	if (InstallerError->HasError() == false)
+	{
+		ResumeData.RemoveResumeData();
+	}
+
 	FileConstructorStat->OnConstructionCompleted();
 }
 
@@ -2080,11 +2156,21 @@ uint64 FBuildPatchFileConstructor::GetAvailableDiskSpace()
 	return AvailableDiskSpace.load(std::memory_order_relaxed);
 }
 
+bool FBuildPatchFileConstructor::HasCorruptFiles() const
+{
+	return !CorruptFiles.IsEmpty();
+}
+
 FBuildPatchFileConstructor::FOnBeforeDeleteFile& FBuildPatchFileConstructor::OnBeforeDeleteFile()
 {
 	return BeforeDeleteFileEvent;
 }
 
+FBuildPatchFileConstructor::FOnFileIsStaged& FBuildPatchFileConstructor::OnFileIsStaged()
+{
+	return FileIsStagedEvent;
+}
+
 void FBuildPatchFileConstructor::CountBytesProcessed( const int64& ByteCount )
 {
 	ByteProcessed += ByteCount;
@@ -2116,6 +2202,9 @@ uint64 FBuildPatchFileConstructor::CalculateInProgressDiskSpaceRequired(const FF
 {
 	if (Configuration.InstallMode == EInstallMode::DestructiveInstall)
 	{
+		TSet<FString> MissedCurrentFiles; // Files that were not installed, possibly due to filtering by tags.
+		Configuration.ManifestSet->GetMissedFiles(Configuration.InstallDirectory, MissedCurrentFiles);
+
 		// The simplest method will be to run through each high level file operation, tracking peak disk usage delta.
 
 		// We know we need enough space to finish writing this file
@@ -2154,8 +2243,8 @@ uint64 FBuildPatchFileConstructor::CalculateInProgressDiskSpaceRequired(const FF
 			{
 				DiskSpaceDeltaPeak = DiskSpaceDelta;
 			}
-			// Then we can remove the current existing file.
-			if (OldFileManifest)
+			// Then we can remove the currently existing file if it was previously created.
+			if (OldFileManifest && !MissedCurrentFiles.Contains(Configuration.ConstructList[ConstructionIndex]))
 			{
 				DiskSpaceDelta -= OldFileManifest->FileSize;
 			}
@@ -2171,19 +2260,15 @@ uint64 FBuildPatchFileConstructor::CalculateInProgressDiskSpaceRequired(const FF
 	}
 }
 
-uint64 FBuildPatchFileConstructor::CalculateDiskSpaceRequirementsWithDeleteDuringInstall()
+uint64 FBuildPatchFileConstructor::CalculateDiskSpaceRequirementsWithDeleteDuringInstall(int64 InProgressFileSize)
 {
-	if (ChunkDbSource == nullptr)
-	{
-		// invalid use.
-		return 0;
-	}
+	checkf(ChunkDbSource, TEXT("CalculateDiskSpaceRequirementsWithDeleteDuringInstall can't be called without valid ChunkDbSource"));
 
 	// These are the sizes at after each file that we _started_ with. This is the size after retirement for the
 	// file at those positions.
 	TArray<uint64> ChunkDbSizesAtPosition;
-	uint64 TotalChunkDbSize = ChunkDbSource->GetChunkDbSizesAtIndexes(FileCompletionPositions, ChunkDbSizesAtPosition);
-
+	ChunkDbSizesAtPosition.Reserve(FileCompletionPositions.Num());
+	const uint64 TotalChunkDbSize = ChunkDbSource->GetChunkDbSizesAtIndexes(FileCompletionPositions, ChunkDbSizesAtPosition);
 	// Strip off the files we've completed.
 	int32 CompletedFileCount = NextIndexToConstruct.load(std::memory_order_acquire);
 
@@ -2193,12 +2278,12 @@ uint64 FBuildPatchFileConstructor::CalculateDiskSpaceRequirementsWithDeleteDurin
 	{
 		CompletedFileCount--;
 	}
-
-	uint64 MaxDiskSize = FBuildPatchUtils::CalculateDiskSpaceRequirementsWithDeleteDuringInstall(
-		Configuration.ConstructList, CompletedFileCount, Configuration.ManifestSet, ChunkDbSizesAtPosition, TotalChunkDbSize);
+	const uint64 MaxDiskSize = FBuildPatchUtils::CalculateDiskSpaceRequirementsWithDeleteDuringInstall(
+		Configuration.ConstructList, CompletedFileCount, InProgressFileSize, Configuration.ManifestSet, ChunkDbSizesAtPosition,
+		Configuration.InstallDirectory, Configuration.InstallMode);
 
 	// Strip off the data we already have on disk.
-	uint64 PostDlSize = 0;
+	uint64 PostDlSize = 0; // Zero means no additional free space is needed for action, as it will be freed by removing chunks.
 	if (MaxDiskSize > TotalChunkDbSize)
 	{
 		PostDlSize = MaxDiskSize - TotalChunkDbSize;
@@ -2431,6 +2516,8 @@ IConstructorChunkSource::FRequestProcessFn FBuildPatchFileConstructor::CreateWri
 
 		ActivityRecord.Size = WriteBuffer.GetSize();
 		ActivityRecord.CyclesEnd = FStatsCollector::GetCycles();
+		FileConstructorHelpers::LogLowRWDurationTime(ActivityRecord, TEXT("Write"));
+
 		FileConstructorStat->OnAfterWrite(ActivityRecord);
 		
 		Batch->bIsFinished.store(true, std::memory_order_release);
@@ -2520,6 +2607,7 @@ void FBuildPatchFileConstructor::CompleteReadBatch(const FFileManifest& FileMani
 				Batch.ErrorContextGuid = FileManifest.ChunkParts[Batch.StartChunkPartIndex + i].Guid;
 			}
 		}
+		RemainingChunkCount = ChunkReferenceTracker->GetRemainingChunkCount();
 	}
 
 	if (Batch.Error == EConstructionError::None)
@@ -2559,7 +2647,7 @@ void FBuildPatchFileConstructor::StartReadBatch(FFileConstructionState& CurrentF
 			FRequestInfo& Request = Batch.Requests.FindOrAdd(ChunkPart.Guid);
 			Request.Guid = ChunkPart.Guid;
 			Request.Batch = &Batch;
-			Request.ChunkSize = Configuration.ManifestSet->GetChunkInfo(ChunkPart.Guid)->WindowSize;
+			Request.ChunkSize = Configuration.ManifestSet->GetChunkInfo(ChunkPart.Guid)->DataSizeUncompressed;
 			Request.File = &CurrentFile;
 			Request.ChunkIndexInFile = EndChunkIdx;
 			
@@ -2705,6 +2793,19 @@ void FBuildPatchFileConstructor::StartReadBatch(FFileConstructionState& CurrentF
 				IConstructorChunkSource::FChunkRequestCompleteDelegate::CreateRaw(this, &FBuildPatchFileConstructor::RequestCompletedFn));
 						
 			QueueGenericThreadTask(ThreadAssignments[ChunkLocation], MoveTemp(RequestProcess));
+
+			if (ThisChunkSource == ChunkDbSource)
+			{
+				IConstructorChunkSource::FRequestProcessFn TruncateTask = [this](bool bIsAborted) {
+					if (!bIsAborted)
+					{
+						ChunkDbSource->TryTruncateFiles(RemainingChunkCount);
+					}
+				};
+				// Execute the truncating on the same thread that reads chunks from chunkdb to make archive access lock-free.
+				QueueGenericThreadTask(ThreadAssignments[ChunkLocation], MoveTemp(TruncateTask));
+
+			}
 		}
 	}
 
@@ -2716,21 +2817,61 @@ void FBuildPatchFileConstructor::WakeUpDispatch()
 	WakeUpDispatchThreadEvent->Trigger();
 }
 
-void FBuildPatchFileConstructor::InitFile(FFileConstructionState& CurrentFile, const FResumeData& ResumeData)
+EFileStage FBuildPatchFileConstructor::GetFileStage(const FFileConstructionState& CurrentFile, const FResumeData& ResumeData) const
 {
 	if (!CurrentFile.bSuccess)
 	{
-		return;
+		return EFileStage::Unknown;
 	}
 
-	const int64 FileSize = CurrentFile.FileManifest->FileSize;
+	if (Configuration.bUseSHA1StageFilenames)
+	{
+		// We have to check file existence here
+		int64 DiskFileSize;
+		if (const bool bFileExists = FileSystem->GetFileSize(*CurrentFile.NewFilename, DiskFileSize))
+		{
+			const uint64 UnsignedDiskFileSize = DiskFileSize;
+			if (UnsignedDiskFileSize == CurrentFile.FileManifest->FileSize)
+			{
+				return EFileStage::Complete;
+			}
+			
+			if (UnsignedDiskFileSize > 0 && UnsignedDiskFileSize < CurrentFile.FileManifest->FileSize) // we only want to resume if the file hasn't been completed.
+			{
+				return EFileStage::Incomplete;
+			}
+			
+			if (UnsignedDiskFileSize > CurrentFile.FileManifest->FileSize)
+			{
+				UE_LOG(LogBuildPatchFileConstructor, Log, TEXT("Deleting incompatible stage file %s (%s)"), *CurrentFile.BuildFilename, *CurrentFile.NewFilename);
+				FileSystem->DeleteFile(*CurrentFile.NewFilename);
+			}
+			// A zero-size file is treated as EFileStage::New.
+		}
+	}
+	else if(const EFileStage * StagePtr = ResumeData.FilesStage.Find(CurrentFile.BuildFilename))
+	{
+		// Can only be EFileStage::Complete or EFileStage::Incomplete.
+		const EFileStage Stage = *StagePtr; // Avoid SA C6269
+		ensure((Stage == EFileStage::Complete) || (Stage == EFileStage::Incomplete));
+		return Stage;
+	}
+
+	// Treat the file as new.
+	return EFileStage::New;
+}
 
-	// Check resume status for this file.
-	const bool bFilePreviouslyComplete = ResumeData.FilesCompleted.Contains(CurrentFile.BuildFilename);	
+void FBuildPatchFileConstructor::InitFile(FFileConstructionState& CurrentFile, const FResumeData& ResumeData)
+{
+	if (!CurrentFile.bSuccess)
+	{
+		return;
+	}
 
 	// Construct or skip the file.
-	if (bFilePreviouslyComplete)
+	if (CurrentFile.Stage == EFileStage::Complete)
 	{
+		const int64 FileSize = CurrentFile.FileManifest->FileSize;
 		CountBytesProcessed(FileSize);
 
 		UE_LOG(LogBuildPatchServices, Display, TEXT("Skipping completed file %s"), *CurrentFile.BuildFilename);
@@ -2744,26 +2885,26 @@ void FBuildPatchFileConstructor::InitFile(FFileConstructionState& CurrentFile, c
 				break;
 			}
 		}
+		RemainingChunkCount = ChunkReferenceTracker->GetRemainingChunkCount();
 
-		CurrentFile.bSkippedConstruction = true;
+		// Can't skip if we are failing.
+		CurrentFile.bSkippedConstruction = CurrentFile.bSuccess;
 		return;
 	}
 
-	if (!CurrentFile.bSuccess &&
+	if (CurrentFile.bSuccess &&
 		!CurrentFile.FileManifest->SymlinkTarget.IsEmpty())
 	{
-#if PLATFORM_MAC
-		CurrentFile.bSuccess = true;
-		if (!Configuration.bInstallToMemory && !Configuration.bConstructInMemory)
+		// Symlinks are handled separately from file construction
+		if (Configuration.bRejectSymlinks)
 		{
-			CurrentFile.bSuccess = symlink(TCHAR_TO_UTF8(*CurrentFile.FileManifest->SymlinkTarget), TCHAR_TO_UTF8(*CurrentFile.NewFilename)) == 0;
+			CurrentFile.bSuccess = false;
+			CurrentFile.ConstructionError = EConstructionError::InvalidSymlink;
 		}
-		CurrentFile.bSkippedConstruction = true;
-#else
-		const bool bSymlinkNotImplemented = false;
-		check(bSymlinkNotImplemented);
-		CurrentFile.bSuccess = false;
-#endif
+
+		// Can't skip if we are failing.
+		CurrentFile.bSkippedConstruction = CurrentFile.bSuccess;
+		return;
 	}
 
 	if (Configuration.bInstallToMemory || Configuration.bConstructInMemory)
@@ -2791,13 +2932,18 @@ bool FBuildPatchFileConstructor::HandleInitialDiskSizeCheck(const FFileManifest&
 		bInitialDiskSizeCheck = true;
 
 		// Normal operation can just use the classic calculation
-		uint64 LocalDiskSpaceRequired = CalculateInProgressDiskSpaceRequired(FileManifest, StartPosition);
+		uint64 LocalDiskSpaceRequired = 0;
 
 		// If we are delete-during-install this gets more complicated because we'll be freeing up
 		// space as we add.
-		if (Configuration.bDeleteChunkDBFilesAfterUse)
+		if (Configuration.bDeleteChunkDBFilesAfterUse && ChunkDbSource)
+		{
+			LocalDiskSpaceRequired = CalculateDiskSpaceRequirementsWithDeleteDuringInstall(StartPosition);
+		}
+		else
 		{
-			LocalDiskSpaceRequired = CalculateDiskSpaceRequirementsWithDeleteDuringInstall();
+			// Normal operation can just use the classic calculation
+			LocalDiskSpaceRequired = CalculateInProgressDiskSpaceRequired(FileManifest, StartPosition);
 		}
 
 		bool bHaveSufficientFreeSpace = false;
@@ -2825,7 +2971,7 @@ bool FBuildPatchFileConstructor::HandleInitialDiskSizeCheck(const FFileManifest&
 			UE_LOG(LogBuildPatchServices, Error, TEXT("Insufficient disk space on initial check."));
 			InstallerError->SetError(
 				EBuildPatchInstallError::OutOfDiskSpace,
-				DiskSpaceErrorCodes::InitialSpaceCheck,
+				StartPosition != 0 ? DiskSpaceErrorCodes::DuringInstallation : DiskSpaceErrorCodes::InitialSpaceCheck,
 				0,
 				BuildPatchServices::GetDiskSpaceMessage(Configuration.InstallDirectory, LocalDiskSpaceRequired, LocalDiskSpaceAvailable));
 			return false;
@@ -2875,7 +3021,7 @@ void FBuildPatchFileConstructor::ResumeFile(FFileConstructionState& FileToResume
 			ReadBuffer.SetNumUninitialized(ChunkPart.Size, EAllowShrinking::No);
 
 			{
-				FReadScope _(FileConstructorStat, ChunkPart.Size);
+				FReadScope _(FileConstructorStat, ChunkPart.Size, TEXT("ReadResume"));
 				NewFileReader->Serialize(ReadBuffer.GetData(), ChunkPart.Size);
 			}
 
@@ -2905,6 +3051,7 @@ void FBuildPatchFileConstructor::ResumeFile(FFileConstructionState& FileToResume
 			break;
 		}
 	}
+	RemainingChunkCount = ChunkReferenceTracker->GetRemainingChunkCount();
 
 	NewFileReader->Close();
 	FileToResume.bIsResumedFile = true;
@@ -2947,7 +3094,7 @@ void FBuildPatchFileConstructor::OpenFileToConstruct(FFileConstructionState& Cur
 }
 
 
-void FBuildPatchFileConstructor::CompleteConstructedFile(FFileConstructionState& CurrentFile)
+void FBuildPatchFileConstructor::CompleteConstructedFile(FFileConstructionState& CurrentFile, FResumeData& ResumeData)
 {
 	UE_LOG(LogBuildPatchServices, Display, TEXT("Completing: %s"), *CurrentFile.BuildFilename);
 	bShouldLogNextDependencyWait = true;
@@ -3002,7 +3149,7 @@ void FBuildPatchFileConstructor::CompleteConstructedFile(FFileConstructionState&
 
 			FSHAHash HashValue;
 			CurrentFile.HashState.GetHash(HashValue.Hash);
-			CurrentFile.bSuccess = HashValue == CurrentFile.FileManifest->FileHash;
+			CurrentFile.bSuccess = HashValue == CurrentFile.FileManifest->SHA1Hash;
 			if (!CurrentFile.bSuccess)
 			{
 				CurrentFile.ConstructionError = EConstructionError::OutboundDataError;
@@ -3145,13 +3292,8 @@ void FBuildPatchFileConstructor::CompleteConstructedFile(FFileConstructionState&
 		{
 		case EConstructionError::OutboundDataError:
 			{
-				// Only report if the first error
-				if (bReportAnalytic)
-				{
-					InstallerAnalytics->RecordConstructionError(CurrentFile.BuildFilename, INDEX_NONE, TEXT("Serialised Verify Fail"));
-				}
 				UE_LOG(LogBuildPatchServices, Error, TEXT("FBuildPatchFileConstructor: Verify failed after constructing %s"), *CurrentFile.BuildFilename);
-				InstallerError->SetError(EBuildPatchInstallError::FileConstructionFail, ConstructionErrorCodes::OutboundCorrupt);
+				CorruptFiles.Add(CurrentFile.BuildFilename);
 				break;
 			}
 		case EConstructionError::FailedWrite:
@@ -3168,7 +3310,7 @@ void FBuildPatchFileConstructor::CompleteConstructedFile(FFileConstructionState&
 				uint64 LocalDiskSpaceRequired;
 				if (Configuration.bDeleteChunkDBFilesAfterUse)
 				{
-					LocalDiskSpaceRequired = CalculateDiskSpaceRequirementsWithDeleteDuringInstall();
+					LocalDiskSpaceRequired = CalculateDiskSpaceRequirementsWithDeleteDuringInstall(CurrentFile.Progress);
 				}
 				else
 				{
@@ -3244,6 +3386,16 @@ void FBuildPatchFileConstructor::CompleteConstructedFile(FFileConstructionState&
 				InstallerError->SetError(EBuildPatchInstallError::FileConstructionFail, ConstructionErrorCodes::TrackingError);
 				break;
 			}
+		case EConstructionError::InvalidSymlink:
+			{
+				if (bReportAnalytic)
+				{
+					InstallerAnalytics->RecordConstructionError(CurrentFile.BuildFilename, INDEX_NONE, TEXT("Symlink present when not allowed"));
+				}
+				UE_LOG(LogBuildPatchFileConstructor, Error, TEXT("File manifest includes symlinks when they are forbidden %s"), *CurrentFile.BuildFilename);
+				InstallerError->SetError(EBuildPatchInstallError::FileConstructionFail, ConstructionErrorCodes::SymlinkCreateForbidden);
+				break;
+			}
 		case EConstructionError::InternalConsistencyError:
 			{
 				if (bReportAnalytic)
@@ -3267,6 +3419,15 @@ void FBuildPatchFileConstructor::CompleteConstructedFile(FFileConstructionState&
 		case EConstructionError::FailedInitialSizeCheck:
 			{
 				// Error already set back when we had the info.
+
+				// If we failed the initial disk space check before we did anything,
+				// delete the resume data so that we don't leave any trash. We know there wasn't any existing resume
+				// data because we failed on the first file at the beginning of the file.
+				if (CurrentFile.ConstructionIndex == 0 &&
+					CurrentFile.StartPosition == 0)
+				{
+					ResumeData.RemoveResumeData();
+				}
 				break;
 			}
 		case EConstructionError::Aborted:
@@ -3303,8 +3464,16 @@ void FBuildPatchFileConstructor::CompleteConstructedFile(FFileConstructionState&
 			}
 		}
 
-		// Stop trying to install further files.
-		Abort();
+		// Before aborting the installation, we need to identify and process all corrupted files.
+		if (CurrentFile.ConstructionError != EConstructionError::OutboundDataError)
+		{
+			// Stop trying to install further files.
+			Abort();
+		}
+	}
+	else
+	{
+		OnFileIsStaged().Broadcast(CurrentFile.BuildFilename);
 	}
 }
 
@@ -3319,7 +3488,7 @@ void FBuildPatchFileConstructor::StartFile(FFileConstructionState& CurrentFile,
 		return;
 	}
 
-	if (ResumeData.FilesStarted.Contains(CurrentFile.BuildFilename))
+	if (CurrentFile.Stage == EFileStage::Incomplete && CurrentFile.bSuccess)
 	{
 		ResumeFile(CurrentFile);
 
@@ -3364,7 +3533,7 @@ void FBuildPatchFileConstructor::StartFile(FFileConstructionState& CurrentFile,
 
 }
 
-void FBuildPatchFileConstructor::ConstructFiles(const FResumeData& ResumeData)
+void FBuildPatchFileConstructor::ConstructFiles(FResumeData& ResumeData)
 {
 	MaxWriteBatchSize = FFileConstructorConfig::DefaultIOBatchSizeMB;
 	if (GConfig->GetInt(TEXT("Portal.BuildPatch"), TEXT("ConstructorIOBatchSizeMB"), (int32&)MaxWriteBatchSize, GEngineIni))
@@ -3527,7 +3696,7 @@ void FBuildPatchFileConstructor::ConstructFiles(const FResumeData& ResumeData)
 				if (File->NextChunkPartToRead == File->FileManifest->ChunkParts.Num() &&
 					File->OutstandingBatches == 0)
 				{
-					CompleteConstructedFile(*File);
+					CompleteConstructedFile(*File, ResumeData);
 
 					UE_LOG(LogBuildPatchServices, Verbose, TEXT("Completed file: %s, New ActiveFileCount = %d"), *File->BuildFilename, ActiveFiles.Num() - 1);
 
@@ -3729,49 +3898,66 @@ void FBuildPatchFileConstructor::ConstructFiles(const FResumeData& ResumeData)
 
 								if (!bDelayForDependencies)
 								{
-									// Now commit this file since we are starting it.
-									NextIndexToConstruct.fetch_add(1, std::memory_order_acq_rel);
+									const FString &StagingFileName = ConstructionList[IndexToConstruct].StagingFilename;
+									TUniquePtr<FFileConstructionState> AnotherFileCandidate = MakeUnique<FFileConstructionState>(
+										ConstructionList[IndexToConstruct].FileManifest,
+										Configuration.ConstructList[IndexToConstruct],
+										Configuration.StagingDirectory / StagingFileName);
+
+									AnotherFileCandidate->Stage = GetFileStage(*AnotherFileCandidate, ResumeData);
+									
+									// The complete file will modify ChunkReferenceTracker, which should be done sequentially after processing active files.
+									const bool bDelayForCompleteFile = AnotherFileCandidate->Stage == EFileStage::Complete && ActiveFiles.Num() > 0;
+
+									if (!bDelayForCompleteFile)
+									{
+										// Now commit this file since we are starting it.
+										NextIndexToConstruct.fetch_add(1, std::memory_order_acq_rel);
 
-									FString StagingFileName = Configuration.StagingDirectory / Configuration.ConstructList[IndexToConstruct];
-									TUniquePtr<FFileConstructionState> AnotherFile = MakeUnique<FFileConstructionState>(
-										ConstructionList[IndexToConstruct].FileManifest, 
-										Configuration.ConstructList[IndexToConstruct], 
-										MoveTemp(StagingFileName));
+										// This log has to be here for tests to succeed, but we don't want spam so we make it Verbose.
+										UE_LOG(LogBuildPatchFileConstructor, Log, TEXT("Constructing file %s (%s)"), *Configuration.ConstructList[IndexToConstruct], *StagingFileName);
 
-									AnotherFile->BaseReferenceIndex = 0;
-									if (IndexToConstruct)
-									{
-										AnotherFile->BaseReferenceIndex = FileCompletionPositions[IndexToConstruct-1];
-									}
+										TUniquePtr<FFileConstructionState> AnotherFile = MoveTemp(AnotherFileCandidate);
+
+										AnotherFile->BaseReferenceIndex = 0;
+										if (IndexToConstruct)
+										{
+											AnotherFile->BaseReferenceIndex = FileCompletionPositions[IndexToConstruct - 1];
+										}
 
-									AnotherFile->ConstructionIndex = IndexToConstruct;
+										AnotherFile->ConstructionIndex = IndexToConstruct;
 
-									TRACE_BOOKMARK(TEXT("Starting File: %s [%u]"), *AnotherFile->BuildFilename, AnotherFile->FileManifest->ChunkParts.Num());
+										TRACE_BOOKMARK(TEXT("Starting File: %s [%u]"), *AnotherFile->BuildFilename, AnotherFile->FileManifest->ChunkParts.Num());
 
-									UE_LOG(LogBuildPatchServices, Display, TEXT("Starting File: %s [%s bytes, %d chunks], New ActiveFileCount = %d"), 
-										   *AnotherFile->BuildFilename, 
-										   *FormatNumber(AnotherFile->FileManifest->FileSize),
-										   AnotherFile->FileManifest->ChunkParts.Num(),
-										   ActiveFiles.Num() + 1);
+										UE_LOG(LogBuildPatchServices, Display, TEXT("Starting File: %s [%s bytes, %d chunks], New ActiveFileCount = %d"),
+											*AnotherFile->BuildFilename,
+											*FormatNumber(AnotherFile->FileManifest->FileSize),
+											AnotherFile->FileManifest->ChunkParts.Num(),
+											ActiveFiles.Num() + 1);
 
-									StartFile(*AnotherFile.Get(), ResumeData);
+										StartFile(*AnotherFile.Get(), ResumeData);
 
-									if (AnotherFile->bSkippedConstruction)
-									{
-										// Nothing else needs to happen with this file - we'll loop around to try for another.
-										bCheckForAnotherRead = true;
+										if (AnotherFile->bSkippedConstruction)
+										{
+											// Nothing else needs to happen with this file - we'll loop around to try for another.
+											bCheckForAnotherRead = true;
+										}
+										else
+										{
+											// Only start the file if there's not currently an active file - otherwise we
+											// do it when the current one finishes.
+											if (ActiveFileCount == 0)
+											{
+												FileConstructorStat->OnFileStarted(AnotherFile->BuildFilename, AnotherFile->FileManifest->FileSize);
+											}
+
+											FileToStart = AnotherFile.Get();
+											ActiveFiles.Push(MoveTemp(AnotherFile));
+										}
 									}
 									else
 									{
-										// Only start the file if there's not currently an active file - otherwise we
-										// do it when the current one finishes.
-										if (ActiveFileCount == 0)
-										{
-											FileConstructorStat->OnFileStarted(AnotherFile->BuildFilename, AnotherFile->FileManifest->FileSize);
-										}
-
-										FileToStart = AnotherFile.Get();
-										ActiveFiles.Push(MoveTemp(AnotherFile));
+										UE_LOG(LogBuildPatchServices, Verbose, TEXT("Delaying %s due to it complete and should be processed after active files."), *Configuration.ConstructList[IndexToConstruct]);
 									}
 								} // end if not delaying for dependencies
 
@@ -4095,7 +4281,7 @@ void FBuildPatchFileConstructor::ConstructFiles(const FResumeData& ResumeData)
 		{
 			if (!File->bSuccess && File->ConstructionError != EConstructionError::Aborted)
 			{
-				CompleteConstructedFile(*File.Get());
+				CompleteConstructedFile(*File.Get(), ResumeData);
 			}
 		}
 
@@ -4112,10 +4298,17 @@ void FBuildPatchFileConstructor::ConstructFiles(const FResumeData& ResumeData)
 				File->bSuccess = false;
 				File->ConstructionError = EConstructionError::Aborted;
 			}
-			CompleteConstructedFile(*File.Get());
+			CompleteConstructedFile(*File.Get(), ResumeData);
 		}
 	}
 
+	if (CorruptFiles.Num() > 0)
+	{
+		Abort();
+		UE_LOG(LogBuildPatchServices, Error, TEXT("FBuildPatchFileConstructor: The %d files are corrupt."), CorruptFiles.Num());
+		InstallerError->SetError(EBuildPatchInstallError::FileConstructionFail, ConstructionErrorCodes::OutboundCorrupt);
+	}
+
 	uint64 ConstructCycles = FPlatformTime::Cycles64() - ConstructStartCycles;
 	uint64 UnaccountedForCycles = ConstructCycles - HashCycles - WaitCycles - ReadCheckCycles - ReadDoneCycles - WriteStartCycles - WriteDoneCycles - CloudTickCycles;
 
@@ -4132,3 +4325,31 @@ void FBuildPatchFileConstructor::ConstructFiles(const FResumeData& ResumeData)
 		100.0 * UnaccountedForCycles / ConstructCycles
 		);
 }
+
+void FBuildPatchFileConstructor::RemoveUnnecessaryStagedFiles(const FString& StagingDirectory) const
+{
+	// Read staging folder to get the list of existing staged files
+	TSet<FString> AllStagedFiles;
+	IFileManager::Get().IterateDirectory(*StagingDirectory, [&AllStagedFiles, StagingDirectory](const TCHAR* Filename, bool) {
+			AllStagedFiles.Add(FString(Filename).Mid(StagingDirectory.Len() + 1)); // get relative path.
+			return true;
+		});
+
+	// Get a set of potentially staged files based on ConstructionList
+	TSet<FString> NecessaryStagedFiles;
+	Algo::Transform(ConstructionList, NecessaryStagedFiles, [](const FFileToConstruct& FileToConstruct)
+		{
+			return FileToConstruct.StagingFilename;
+		});
+
+	const TSet<FString> UnnecessaryFiles = AllStagedFiles.Difference(NecessaryStagedFiles);
+	for (const FString& Filename : UnnecessaryFiles)
+	{
+		UE_LOG(LogBuildPatchFileConstructor, Log, TEXT("RemoveUnnecessaryStagedFiles: Deleting unnecessary staged file %s"), *Filename);
+		if (!FileSystem->DeleteFile(*(StagingDirectory / Filename)))
+		{
+			UE_LOG(LogBuildPatchServices, Error, TEXT("RemoveUnnecessaryStagedFiles: Error deleting file: %s (Error Code %i)"), *Filename, FPlatformMisc::GetLastError());
+		}
+	}
+}
+
diff --git a/Engine/Source/Runtime/Online/BuildPatchServices/Private/BuildPatchFileConstructor.h b/Engine/Source/Runtime/Online/BuildPatchServices/Private/BuildPatchFileConstructor.h
--- a/Engine/Source/Runtime/Online/BuildPatchServices/Private/BuildPatchFileConstructor.h
+++ b/Engine/Source/Runtime/Online/BuildPatchServices/Private/BuildPatchFileConstructor.h
@@ -12,6 +12,7 @@
 #include "Installer/Controllable.h"
 #include "Common/SpeedRecorder.h"
 #include "BuildPatchInstall.h"
+#include "Async/Future.h"
 
 
 // Forward declarations
@@ -20,6 +21,7 @@ struct FResumeData;
 struct FFileConstructionState;
 class FBuildPatchAppManifest;
 enum class EConstructionError : uint8;
+enum class EFileStage : uint8;
 
 class IBuildInstallerSharedContext;
 class FChunkBackingStore;
@@ -74,6 +76,11 @@ namespace BuildPatchServices
 
 		// The install mode used for this installation.
 		EInstallMode InstallMode;
+		// Whether we are needing to use SHA1 hash string value as filename for temp staging files.
+		bool bUseSHA1StageFilenames;
+
+		// If true, fail the installation if the manifest contains a symlink.
+		bool bRejectSymlinks = false;
 
 		// The location where memory overflow will get written to.
 		FString BackingStoreDirectory;
@@ -136,10 +143,16 @@ namespace BuildPatchServices
 		// IControllable interface begin.
 		virtual void SetPaused(bool bInIsPaused) override;
 		virtual void Abort() override;
+
+		virtual void Reset() override
+		{
+			// This should _never_ be called. I'm not sure how you can Reset() a constructor..
+			check(false);
+		}
 		// IControllable interface end.
 
 		void WakeUpDispatch();
-
+		
 		/**
 		 * Get the disk space that was required to perform the installation. This can change over time and indicates the required
 		 * space to _finish_ the installation from the current state. It is not initialized until after resume is processed and returns
@@ -157,13 +170,26 @@ namespace BuildPatchServices
 		 */
 		uint64 GetAvailableDiskSpace();
 
+		/**
+		 * Gets flag indicating that at least one file is corrupt due to invalid hash sum.
+		 * @return	true if at least one file is corrupt.
+		 */
+		bool HasCorruptFiles() const;
+
 		/**
 		 * Broadcasts with full filepath to file that the constructor is about to delete in order to free up space.
 		 * @return	Reference to the event object.
 		 */
-		DECLARE_EVENT_OneParam(FBuildPatchFileConstructor, FOnBeforeDeleteFile, const FString& /*BuildFile*/);
+		DECLARE_EVENT_OneParam(FBuildPatchFileConstructor, FOnBeforeDeleteFile, const FString& /*FullFilePath*/);
 		FOnBeforeDeleteFile& OnBeforeDeleteFile();
 
+		/**
+		 * Broadcasts with the build relative path for a file that has finished being staged (including by a previous incomplete installation).
+		 * @return	Reference to the event object.
+		 */
+		DECLARE_EVENT_OneParam(FBuildPatchFileConstructor, FOnFileIsStaged, const FString& /*BuildFilename*/);
+		FOnFileIsStaged& OnFileIsStaged();
+
 		void GrabFilesInstalledToMemory(TMap<FString, TArray64<uint8>>& OutFilesInstalledToMemory)
 		{
 			OutFilesInstalledToMemory = MoveTemp(MemoryOutputFiles);
@@ -192,19 +218,32 @@ namespace BuildPatchServices
 		 */
 		uint64 CalculateInProgressDiskSpaceRequired(const FFileManifest& InProgressFileManifest, uint64 InProgressFileSize);
 
-		// Calculates the amount of disk space we need to finish the install, needs to be called on file boundaries on the construct thread.
-		uint64 CalculateDiskSpaceRequirementsWithDeleteDuringInstall();
+		/**
+		 * Calculates the amount of disk space we need to finish the install, needs to be called on file boundaries on the construct thread.
+		 * @param InProgressFileSize		The remaining size required for the file currently being constructed.
+		 * @return Number of bytes needed to complete the installation, or zero if space will be freed by removing chunks.
+		 */
+		uint64 CalculateDiskSpaceRequirementsWithDeleteDuringInstall(int64 InProgressFileSize);
 
 		/**
 		 * Sequentially build the files required, potentially skipping already created files or resuming partial
 		 * files.
 		 */
-		void ConstructFiles(const FResumeData& ResumeData);
+		void ConstructFiles(FResumeData& ResumeData);
 
 	private:
 		// The configuration for the constructor.
 		const FFileConstructorConfig Configuration;
 
+		// A flag marking that we a running.
+		bool bIsRunning;
+
+		// A flag marking that we initialized correctly.
+		bool bIsInited;
+
+		// A flag marking that our init returned a failure (true means failed).
+		bool bInitFailed;
+
 		// A flag marking that we told the chunk cache to queue required downloads.
 		bool bIsDownloadStarted;
 
@@ -236,6 +275,8 @@ namespace BuildPatchServices
 			// constructed in order, we only track the file that will be harvested last. If this
 			// file has no dependencies this is -1
 			int32 LatestDependentInstallSource = -1;
+			// Name of the staging file. It may be obfuscated if bUseSHA1StageFilenames is true.
+			FString StagingFilename;
 		};
 
 		// The in-oder list of files to construct. The array is parallel with Configuration.ConstructList.
@@ -263,6 +304,8 @@ namespace BuildPatchServices
 		// Pointer to the chunk reference tracker.
 		IChunkReferenceTracker* ChunkReferenceTracker;
 
+		std::atomic<int32> RemainingChunkCount = TNumericLimits<int32>::Max();
+
 		// Pointer to the installer error class.
 		IInstallerError* InstallerError;
 
@@ -300,7 +343,9 @@ namespace BuildPatchServices
 
 		// Event executed before deleting an old installation file.
 		FOnBeforeDeleteFile BeforeDeleteFileEvent;
-
+		// Event executed when a file has now been staged or was previously staged.
+		FOnFileIsStaged FileIsStagedEvent;
+		
 
 		TArray<FEvent*> ThreadWakeups;
 		TArray<TArray<IConstructorChunkSource::FRequestProcessFn>> ThreadJobPostings;
@@ -334,10 +379,19 @@ namespace BuildPatchServices
 		void OpenFileToConstruct(FFileConstructionState& CurrentFile);
 		bool HandleInitialDiskSizeCheck(const FFileManifest& CurrentFileManifest, int64 BytesInToCurrentFile);
 		bool HarvestChunksForCompletedFile(const FString& InstallationDirectory, const FString& FilenameInManifest);
-		void CompleteConstructedFile(FFileConstructionState& CurrentFile);
-		void InitFile(FFileConstructionState& CurrentFile, const FResumeData& ResumeData);
+		void CompleteConstructedFile(FFileConstructionState& CurrentFile, FResumeData& ResumeData);
 
+		// Determine and return the file stage based on its existence and file size.
+		EFileStage GetFileStage(const FFileConstructionState& CurrentFile, const FResumeData& ResumeData) const;
+
+		// Initialize and validate.
+		void InitFile(FFileConstructionState& CurrentFile, const FResumeData& ResumeData);
+		
+		// Removes staged files that are no longer required for the current operation.
+		void RemoveUnnecessaryStagedFiles(const FString& StagingDirectory) const;
 
+		// Set of files considered corrupt due to an invalid hash sum.
+		TSet<FString> CorruptFiles;
 
 	public:
 		struct FBackingStoreStats
@@ -464,4 +518,5 @@ namespace BuildPatchServices
 namespace FileConstructorHelpers
 {
 	uint64 CalculateRequiredDiskSpace(const FBuildPatchAppManifestPtr& CurrentManifest, const FBuildPatchAppManifestRef& BuildManifest, const BuildPatchServices::EInstallMode& InstallMode, const TSet<FString>& InstallTags);
+	bool CheckRemainingDiskSpace(const FString& InstallDirectory, uint64 RemainingBytesRequired, uint64& OutAvailableDiskSpace);
 }
